IKRig plugin first version

- you can add one type of solver - transform and then that will give you one goal and you can run it.
- 3 modules - IKRig, IKRigDeveloper, IKRigEditor
- IKRigDefinition is the main data for anim node.

diff --git a/Engine/Plugins/Animation/IKRig/IKRig.uplugin b/Engine/Plugins/Animation/IKRig/IKRig.uplugin
new file mode 100644
--- /dev/null
+++ b/Engine/Plugins/Animation/IKRig/IKRig.uplugin
@@ -0,0 +1,35 @@
+{
+	"FileVersion": 3,
+	"Version": 1,
+	"VersionName": "1.0",
+	"FriendlyName": "IKRig",
+	"Description": "",
+	"Category": "Other",
+	"CreatedBy": "",
+	"CreatedByURL": "",
+	"DocsURL": "",
+	"MarketplaceURL": "",
+	"SupportURL": "",
+	"EnabledByDefault": false,
+	"CanContainContent": true,
+	"IsBetaVersion": true,
+	"IsExperimentalVersion": false,
+	"Installed": false,
+	"Modules": [
+		{
+			"Name": "IKRig",
+			"Type": "Runtime",
+			"LoadingPhase": "Default"
+		},
+		{
+			"Name": "IKRigDeveloper",
+			"Type": "UncookedOnly",
+			"LoadingPhase": "Default"
+		},
+		{
+			"Name": "IKRigEditor",
+			"Type": "Editor",
+			"LoadingPhase": "Default"
+		}
+	]
+}
\ No newline at end of file
diff --git a/Engine/Plugins/Animation/IKRig/Source/IKRig/IKRig.Build.cs b/Engine/Plugins/Animation/IKRig/Source/IKRig/IKRig.Build.cs
new file mode 100644
--- /dev/null
+++ b/Engine/Plugins/Animation/IKRig/Source/IKRig/IKRig.Build.cs
@@ -0,0 +1,37 @@
+// Copyright Epic Games, Inc. All Rights Reserved.
+
+using UnrealBuildTool;
+
+public class IKRig : ModuleRules
+{
+	public IKRig(ReadOnlyTargetRules Target) : base(Target)
+	{
+		PCHUsage = ModuleRules.PCHUsageMode.UseExplicitOrSharedPCHs;
+		
+		PrivateDependencyModuleNames.AddRange(
+			new string[]
+			{
+				"CoreUObject",
+				"Engine",
+				"Slate",
+				"SlateCore",
+				"AnimationCore", 
+				"AnimGraphRuntime"
+			}
+			);
+
+		PublicDependencyModuleNames.AddRange(
+			new string[]
+			{
+				"Core",
+			}
+			);
+
+		DynamicallyLoadedModuleNames.AddRange(
+			new string[]
+			{
+				// ... add any modules that your module loads dynamically here ...
+			}
+			);
+	}
+}
diff --git a/Engine/Plugins/Animation/IKRig/Source/IKRig/Private/AnimNodes/AnimNode_IKRig.cpp b/Engine/Plugins/Animation/IKRig/Source/IKRig/Private/AnimNodes/AnimNode_IKRig.cpp
new file mode 100644
--- /dev/null
+++ b/Engine/Plugins/Animation/IKRig/Source/IKRig/Private/AnimNodes/AnimNode_IKRig.cpp
@@ -0,0 +1,197 @@
+// Copyright Epic Games, Inc. All Rights Reserved.
+
+#include "AnimNodes/AnimNode_IKRig.h"
+#include "IKRigProcessor.h"
+#include "IKRigDefinition.h"
+#include "Animation/AnimInstanceProxy.h"
+/////////////////////////////////////////////////////
+// AnimNode_IKRig
+
+FAnimNode_IKRig::FAnimNode_IKRig()
+	: RigDefinitionAsset (nullptr)
+#if WITH_EDITORONLY_DATA
+	, bEnableDebugDraw(false)
+#endif
+{
+}
+
+void FAnimNode_IKRig::Evaluate_AnyThread(FPoseContext& Output) 
+{
+	DECLARE_SCOPE_HIERARCHICAL_COUNTER_FUNC()
+	FPoseContext SourcePose(Output);
+
+	if (Source.GetLinkNode())
+	{
+		Source.Evaluate(SourcePose);
+	}
+	else
+	{
+		// apply refpose
+		SourcePose.ResetToRefPose();
+	}
+
+	if (RigProcessor)
+	{
+		FIKRigTransformModifier& TransformModifier = RigProcessor->GetIKRigTransformModifier();
+
+		if (bStartFromRefPose)
+		{
+			RigProcessor->ResetToRefPose();
+		}
+		
+		FCompactPose& OutPose = SourcePose.Pose;
+
+		for (FCompactPoseBoneIndex CPIndex : OutPose.ForEachBoneIndex())
+		{
+			int32* Index = CompactPoseToRigIndices.Find(CPIndex);
+			if (Index)
+			{
+				// Todo: this is again slow. 
+				// the tricky part is that if we go start from child to parent
+				// we may still missing joints that don't exists
+				// so this is slow and for safety but needs revisit
+				TransformModifier.SetLocalTransform(*Index, OutPose[CPIndex], true);
+			}
+		}
+
+		for (int32 GoalIndex = 0; GoalIndex < GoalNames.Num(); ++GoalIndex)
+		{
+			if (ensure(GoalTransforms.IsValidIndex(GoalIndex)))
+			{
+				RigProcessor->SetGoalPosition( GoalNames[GoalIndex], GoalTransforms[GoalIndex].GetLocation());
+				RigProcessor->SetGoalRotation(GoalNames[GoalIndex], GoalTransforms[GoalIndex].GetRotation().Rotator());
+			}
+		}
+
+		// fill up current pose
+		RigProcessor->Solve();
+
+		// get component pose from control rig
+//		FCSPose<FCompactPose> MeshPoses;
+//		MeshPoses.InitPose(SourcePose.Pose);
+
+		// now we copy back to output pose
+		for (FCompactPoseBoneIndex CPIndex : OutPose.ForEachBoneIndex())
+		{
+			int32* Index = CompactPoseToRigIndices.Find(CPIndex);
+			if (Index)
+			{
+//				MeshPoses.SetComponentSpaceTransform(CPIndex, TransformModifier.GetGlobalTransform(*Index));
+				Output.Pose[CPIndex] = TransformModifier.GetLocalTransform(*Index);
+			}
+		}
+
+//		FCSPose<FCompactPose>::ConvertComponentPosesToLocalPosesSafe(MeshPoses, Output.Pose);
+		Output.Pose.NormalizeRotations();
+	}
+	else
+	{
+		Output = SourcePose;
+	}
+}
+
+void FAnimNode_IKRig::GatherDebugData(FNodeDebugData& DebugData)
+{
+	DECLARE_SCOPE_HIERARCHICAL_COUNTER_FUNC()
+	// alpha later?
+	Source.GatherDebugData(DebugData.BranchFlow(1.f));
+}
+
+void FAnimNode_IKRig::Initialize_AnyThread(const FAnimationInitializeContext& Context)
+{
+	DECLARE_SCOPE_HIERARCHICAL_COUNTER_FUNC()
+
+	FAnimNode_Base::Initialize_AnyThread(Context);
+	Source.Initialize(Context);
+
+	if (RigProcessor)
+	{
+		RigProcessor->SetIKRigDefinition(RigDefinitionAsset, true);
+		GoalNames.Reset();
+		RigProcessor->GetGoals(GoalNames);
+	}
+}
+
+void FAnimNode_IKRig::OnInitializeAnimInstance(const FAnimInstanceProxy* InProxy, const UAnimInstance* InAnimInstance)
+{
+	DECLARE_SCOPE_HIERARCHICAL_COUNTER_FUNC()
+
+	if (RigDefinitionAsset && !RigProcessor)
+	{
+		RigProcessor = NewObject<UIKRigProcessor>(InAnimInstance->GetOwningComponent());
+	}
+
+	FAnimNode_Base::OnInitializeAnimInstance(InProxy, InAnimInstance);
+}
+
+void FAnimNode_IKRig::RebuildGoalList()
+{
+	if (RigDefinitionAsset)
+	{
+		GoalTransforms.SetNum(RigDefinitionAsset->GetGoals().Num());
+	}
+}
+
+FName FAnimNode_IKRig::GetGoalName(int32 Index) const
+{
+	if (RigDefinitionAsset)
+	{
+		// TODO: fix  this vs RigProcessor goals. I think RigProcessor
+		// only issue with RigProcessor is you don't know until it compiles
+		TArray<FName> Names = RigDefinitionAsset->GetGoalsNames();
+		if (Names.IsValidIndex(Index))
+		{
+			return Names[Index];
+		}
+	}
+
+	return FName(TEXT("Invalid"));
+}
+
+void FAnimNode_IKRig::Update_AnyThread(const FAnimationUpdateContext& Context)
+{
+	GetEvaluateGraphExposedInputs().Execute(Context);
+
+	FAnimNode_Base::Update_AnyThread(Context);
+	Source.Update(Context);
+}
+
+void FAnimNode_IKRig::CacheBones_AnyThread(const FAnimationCacheBonesContext& Context)
+{
+	FAnimNode_Base::CacheBones_AnyThread(Context);
+	Source.CacheBones(Context);
+
+	if (RigProcessor)
+	{
+		// fill up node names
+		const FBoneContainer& RequiredBones = Context.AnimInstanceProxy->GetRequiredBones();
+
+		CompactPoseToRigIndices.Reset();
+
+		if (RequiredBones.IsValid())
+		{
+			const TArray<FBoneIndexType>& RequiredBonesArray = RequiredBones.GetBoneIndicesArray();
+			const int32 NumBones = RequiredBonesArray.Num();
+
+			const FReferenceSkeleton& RefSkeleton = RequiredBones.GetReferenceSkeleton();
+			const FIKRigHierarchy* Hierarchy = RigProcessor->GetHierarchy();
+
+			if (Hierarchy)
+			{
+
+				// even if not mapped, we map only node that exists in the controlrig
+				for (uint16 Index = 0; Index < NumBones; ++Index)
+				{
+					// get the name
+					int32 MeshBone = RequiredBonesArray[Index];
+					if (ensure(MeshBone != INDEX_NONE))
+					{
+						FCompactPoseBoneIndex CPIndex = RequiredBones.MakeCompactPoseIndex(FMeshPoseBoneIndex(MeshBone));
+						FName Name = RequiredBones.GetReferenceSkeleton().GetBoneName(MeshBone);
+						CompactPoseToRigIndices.Add(CPIndex) = Hierarchy->GetIndex(Name);
+					}
+				}
+			}
+		}
+	}
+}
\ No newline at end of file
diff --git a/Engine/Plugins/Animation/IKRig/Source/IKRig/Private/IKRig.cpp b/Engine/Plugins/Animation/IKRig/Source/IKRig/Private/IKRig.cpp
new file mode 100644
--- /dev/null
+++ b/Engine/Plugins/Animation/IKRig/Source/IKRig/Private/IKRig.cpp
@@ -0,0 +1,22 @@
+// Copyright Epic Games, Inc. All Rights Reserved.
+
+#include "IKRig.h"
+
+#define LOCTEXT_NAMESPACE "FIKRigModule"
+
+DEFINE_LOG_CATEGORY(LogIKRig);
+
+void FIKRigModule::StartupModule()
+{
+	// This code will execute after your module is loaded into memory; the exact timing is specified in the .uplugin file per-module
+}
+
+void FIKRigModule::ShutdownModule()
+{
+	// This function may be called during shutdown to clean up your module.  For modules that support dynamic reloading,
+	// we call this function before unloading the module.
+}
+
+#undef LOCTEXT_NAMESPACE
+	
+IMPLEMENT_MODULE(FIKRigModule, IKRig)
\ No newline at end of file
diff --git a/Engine/Plugins/Animation/IKRig/Source/IKRig/Private/IKRigConstraint.cpp b/Engine/Plugins/Animation/IKRig/Source/IKRig/Private/IKRigConstraint.cpp
new file mode 100644
--- /dev/null
+++ b/Engine/Plugins/Animation/IKRig/Source/IKRig/Private/IKRigConstraint.cpp
@@ -0,0 +1,14 @@
+// Copyright Epic Games, Inc. All Rights Reserved.
+
+/**
+ * Contains IKRigConstraint Runtime Implementation
+ *
+ */
+
+#include "IKRigConstraint.h"
+
+// UIKRigConstraint::UIKRigConstraint(const FObjectInitializer& ObjectInitializer)
+// 	: Super(ObjectInitializer)
+// {
+// 
+// }
diff --git a/Engine/Plugins/Animation/IKRig/Source/IKRig/Private/IKRigConstraintDefinition.cpp b/Engine/Plugins/Animation/IKRig/Source/IKRig/Private/IKRigConstraintDefinition.cpp
new file mode 100644
--- /dev/null
+++ b/Engine/Plugins/Animation/IKRig/Source/IKRig/Private/IKRigConstraintDefinition.cpp
@@ -0,0 +1,16 @@
+// Copyright Epic Games, Inc. All Rights Reserved.
+
+/**
+ * Contains IKRigConstraint Definition Implementation
+ *
+ */
+
+#include "IKRigConstraintDefinition.h"
+
+const FName UIKRigConstraintDefinition::DefaultProfileName = FName(TEXT("Default"));
+
+UIKRigConstraintDefinition::UIKRigConstraintDefinition()
+{
+	FIKRigConstraintProfile& DefaultProfile = ConstraintProfiles.Add(DefaultProfileName);
+	DefaultProfile.Name = DefaultProfileName;
+}
diff --git a/Engine/Plugins/Animation/IKRig/Source/IKRig/Private/IKRigConstraintSolver.cpp b/Engine/Plugins/Animation/IKRig/Source/IKRig/Private/IKRigConstraintSolver.cpp
new file mode 100644
--- /dev/null
+++ b/Engine/Plugins/Animation/IKRig/Source/IKRig/Private/IKRigConstraintSolver.cpp
@@ -0,0 +1,18 @@
+// Copyright Epic Games, Inc. All Rights Reserved.
+
+/**
+ * Contains IKRig Constraint Solver Definition 
+ *
+ */
+#include "IKRigConstraintSolver.h"
+
+// register/unregister query function
+void UIKRigConstraintSolver::RegisterQueryConstraintHandler(const FIKRigQueryConstraint& InQueryConstraintHandler)
+{
+
+}
+
+void UIKRigConstraintSolver::UnregisterQueryConstraintHandler()
+{
+
+}
\ No newline at end of file
diff --git a/Engine/Plugins/Animation/IKRig/Source/IKRig/Private/IKRigDataTypes.cpp b/Engine/Plugins/Animation/IKRig/Source/IKRig/Private/IKRigDataTypes.cpp
new file mode 100644
--- /dev/null
+++ b/Engine/Plugins/Animation/IKRig/Source/IKRig/Private/IKRigDataTypes.cpp
@@ -0,0 +1,153 @@
+// Copyright Epic Games, Inc. All Rights Reserved.
+
+/*=============================================================================
+	IKRigDataTypes.cpp: IKRig Data Type impelmentation file
+=============================================================================*/
+
+#include "IKRigDataTypes.h"
+#include "IKRigHierarchy.h"
+
+FIKRigTransformModifier::FIKRigTransformModifier(const FIKRigHierarchy* InHierarchy)
+	: Hierarchy(InHierarchy)
+{
+	ensure(Hierarchy != nullptr);
+}
+
+// we don't want to allocate local transforms until required
+// so we don't do this in the constructor but on demand when to be used
+void FIKRigTransformModifier::EnsureLocalTransformsExist()
+{
+	// first check validation
+	if (Hierarchy&& Hierarchy->GetNum() == GlobalTransforms.GetNum())
+	{
+		if (LocalTransforms.Num() != GlobalTransforms.GetNum())
+		{
+			LocalTransforms.SetNumUninitialized(GlobalTransforms.GetNum());
+			LocalTransformDirtyFlags.Init(true, GlobalTransforms.GetNum());
+		}
+	}
+}
+
+void FIKRigTransformModifier::RecalculateLocalTransform()
+{
+	// make sure to call EnsureLocalTransformsExist before getting here
+	check (Hierarchy && LocalTransforms.Num() == GlobalTransforms.GetNum());
+	
+	for (TConstSetBitIterator<> It(LocalTransformDirtyFlags); It; ++It)
+	{
+		const int32 BoneIndex = It.GetIndex();
+		const int32 ParentIndex = Hierarchy->GetParentIndex(BoneIndex);
+		LocalTransforms[BoneIndex] = GlobalTransforms.GetRelativeTransform(BoneIndex, ParentIndex);
+	}
+
+	// clear the bits
+	LocalTransformDirtyFlags.Init(false, GlobalTransforms.GetNum());
+}
+
+void FIKRigTransformModifier::UpdateLocalTransform(int32 Index)
+{
+	EnsureLocalTransformsExist();
+
+	if (Hierarchy && Hierarchy->IsValidIndex(Index) && LocalTransformDirtyFlags[Index])
+	{
+		int32 ParentIndex = Hierarchy->GetParentIndex(Index);
+		LocalTransforms[Index] = GlobalTransforms.GetRelativeTransform(Index, ParentIndex);
+		LocalTransformDirtyFlags[Index] = false;
+	}
+}
+
+void FIKRigTransformModifier::SetGlobalTransform(int32 Index, const FTransform& InTransform, bool bPropagate)
+{
+	if (GlobalTransforms.IsValidIndex(Index))
+	{
+		if (bPropagate)
+		{
+			// update all children
+			SetGlobalTransform_Internal(Index, InTransform);
+		}
+		else
+		{
+			GlobalTransforms.GlobalTransforms[Index] = InTransform;
+
+			// if we have local transform, we also update
+			if (LocalTransforms.IsValidIndex(Index))
+			{
+				LocalTransformDirtyFlags[Index] = true;
+				UpdateLocalTransform(Index);
+			}
+		}
+	}
+}
+
+void FIKRigTransformModifier::SetLocalTransform(int32 Index, const FTransform& InTransform, bool bPropagate)
+{
+	// we go through global 
+	// because we can't just update local only
+	if (Hierarchy && Hierarchy->IsValidIndex(Index))
+	{
+		const int32 ParentIndex = Hierarchy->GetParentIndex(Index);
+		// if parent is invalid, it will get identity
+		FTransform NewTransform = InTransform * GetGlobalTransform(ParentIndex);
+		SetGlobalTransform(Index, NewTransform, true);
+	}
+}
+
+const FTransform& FIKRigTransformModifier::GetLocalTransform(int32 Index) const
+{
+	if (Hierarchy && Hierarchy->IsValidIndex(Index))
+	{
+		// this is because this requires to modify LocalTransforms
+		const_cast<FIKRigTransformModifier*>(this)->UpdateLocalTransform(Index);
+		return LocalTransforms[Index];
+	}
+
+	return FTransform::Identity;
+}
+
+const FTransform& FIKRigTransformModifier::GetGlobalTransform(int32 Index) const
+{
+	return GlobalTransforms.GetGlobalTransform(Index);
+}
+/* This function does propagate through children
+	* of the current index and keeps last LocalTransform to up to date
+	*/
+void FIKRigTransformModifier::SetGlobalTransform_Internal(int32 Index, const FTransform& InTransform)
+{
+	// first calculate local transform if not found
+	EnsureLocalTransformsExist();
+	RecalculateLocalTransform();
+
+	// set global transform recursively without modifying Local Transform
+	SetGlobalTransform_Recursive(Index, InTransform);
+
+	// now update my local transform
+	LocalTransformDirtyFlags[Index] = true;
+	UpdateLocalTransform(Index);
+}
+
+/** Set Global Transform Recursive
+	*
+	*/
+void FIKRigTransformModifier::SetGlobalTransform_Recursive(int32 Index, const FTransform& InTransform)
+{
+	GlobalTransforms.GlobalTransforms[Index] = InTransform;
+	TArray<int32> Children = Hierarchy->FindChildren(Index);
+	for (int32 Child : Children)
+	{
+		// calculate new global based on local 
+		FTransform NewTransform = LocalTransforms[Child] * InTransform;
+		SetGlobalTransform_Recursive(Child, NewTransform);
+	}
+}
+
+void FIKRigTransformModifier::ResetGlobalTransform(const FIKRigTransform& InTransform)
+{
+	if (Hierarchy && Hierarchy->GetNum() == InTransform.GetNum())
+	{
+		GlobalTransforms = InTransform;
+
+		// we don't want previous transform data
+		LocalTransforms.Reset();
+		LocalTransformDirtyFlags.Reset();
+	}
+}
\ No newline at end of file
diff --git a/Engine/Plugins/Animation/IKRig/Source/IKRig/Private/IKRigDefinition.cpp b/Engine/Plugins/Animation/IKRig/Source/IKRig/Private/IKRigDefinition.cpp
new file mode 100644
--- /dev/null
+++ b/Engine/Plugins/Animation/IKRig/Source/IKRig/Private/IKRigDefinition.cpp
@@ -0,0 +1,383 @@
+// Copyright Epic Games, Inc. All Rights Reserved.
+
+/*=============================================================================
+	IKRigDefinition.cpp: Composite classes that contains sequence for each section
+=============================================================================*/
+
+#include "IKRigDefinition.h"
+#include "IKRigSolverDefinition.h"
+
+UIKRigDefinition::UIKRigDefinition()
+{
+	// we create constraint definition for them
+	//ConstraintDefinitions = NewObject<UIKRigConstraintDefinition>(this, TEXT("ConstraintDefinition"));
+}
+
+#if WITH_EDITOR
+void UIKRigDefinition::UpdateGoal()
+{
+	// collect goals from the solver definitions
+	// each solver definition collects their tasks to goals
+	// but the IKRigDefinition will maintain the default value and so on
+	Sanitize();
+
+	TArray<FName> ListOfGoals;
+	for (UIKRigSolverDefinition* SolverDef : SolverDefinitions)
+	{
+		if (SolverDef)
+		{
+			SolverDef->CollectGoals(ListOfGoals);
+		}
+	}
+
+	// we want to update goal list properly
+
+	// first if it's not used any more, remove
+	TArray<FName> GoalsToRemove;
+
+	// first we go through IKGoals and see if we should remove unused ones
+	for (auto Iter = IKGoals.CreateConstIterator(); Iter; ++Iter)
+	{
+		// if it's not used any more
+		if (!ListOfGoals.Contains(Iter->Key))
+		{
+			GoalsToRemove.Add(Iter->Key);
+		}
+	}
+
+	for (const FName& Goal : GoalsToRemove)
+	{
+		IKGoals.Remove(Goal);
+	}
+
+	// now add new ones if it doesn't exist
+	for (int32 Index=0; Index<ListOfGoals.Num(); ++Index)
+	{
+		const FName& Goal = ListOfGoals[Index];
+		if (!IKGoals.Contains(Goal))
+		{
+			// if we don't have it, add to new one
+			FIKRigGoal NewGoal(Goal);
+			IKGoals.Add(Goal, NewGoal);
+		}
+	}
+
+	// this doesn't keep the old list alive
+	// unsure if that's what we want yet. Sometimes you may lose
+	// your used variable, but that's something we can find it out 
+	// from UX
+}
+
+void UIKRigDefinition::Sanitize()
+{
+	// sanitize to clean up
+	SolverDefinitions.RemoveAll([this](const UIKRigSolverDefinition* IKDef) { return IKDef == nullptr; });
+}
+
+/////////////////////////////////////////////////////////
+//// Hierarchy modifiers
+// why this is here? 
+// delegate and user instancing 
+// can this be in controller?
+// 
+/////////////////////////////////////////////////////////
+
+// add a bone. Return false if it already has conflict name or parent is not found
+bool UIKRigDefinition::AddBone(const FName& InName, const FName& InParent, const FTransform& InGlobalTransform)
+{
+	// only add when it doesn't exist
+	if (Hierarchy.FindIndexFromBoneArray(InName) == INDEX_NONE)
+	{
+		// if no parent or we want to have parent BEFORE adding this joint
+		if (InParent == NAME_None || Hierarchy.FindIndexFromBoneArray(InParent) != INDEX_NONE)
+		{
+			Hierarchy.Bones.Add(FIKRigBone(InName, InParent));
+			ReferencePose.GlobalTransforms.Add(InGlobalTransform);
+
+			check (Hierarchy.Bones.Num() == ReferencePose.GetNum());
+
+			Hierarchy.RebuildCacheData();
+			BoneAddedDelegate.Broadcast(InName);
+			return true;
+		}
+	}
+
+	return false;
+}
+
+// remove a bone. Returns false if it is not found
+// if it has a children, it will remove all the children with it
+bool UIKRigDefinition::RemoveBone(const FName& InName)
+{
+	struct FRemoveChild
+	{
+		static void RemoveChild_Recursively(const FIKRigHierarchy& InHierarchy, TArray<FIKRigBone>& InOutBones, TArray<FTransform>& InOutTransform, const FName& InName)
+		{
+			// ensure I have this bone
+			int32 Index = InHierarchy.FindIndexFromBoneArray(InName);
+			if (Index != INDEX_NONE)
+			{
+				// find children
+				TArray<int32> ChildIndices = InHierarchy.FindIndicesByParentName(InName);
+				// iterate back from the forward
+				for (int32 ChildIndex = ChildIndices.Num() - 1; ChildIndex >= 0; --ChildIndex)
+				{
+					// remove children's first, since children is always behind of parent, this should work
+					RemoveChild_Recursively(InHierarchy, InOutBones, InOutTransform, InOutBones[ChildIndices[ChildIndex]].Name);
+					// do not shrink yet and remove all children first
+					// since children is behind of parent, this works
+				}
+
+				// now remove parent (this)
+				InOutBones.RemoveAt(Index, 1, false);
+				InOutTransform.RemoveAt(Index, 1, false);
+			}
+		}
+	};
+
+	// ensure I have this bone
+	int32 Index = Hierarchy.FindIndexFromBoneArray(InName);
+	if (Index != INDEX_NONE)
+	{
+		// remove all children
+		FRemoveChild::RemoveChild_Recursively(Hierarchy, Hierarchy.Bones, ReferencePose.GlobalTransforms, InName);
+
+		// shrink now
+		Hierarchy.Bones.Shrink();
+		ReferencePose.GlobalTransforms.Shrink();
+		Hierarchy.RebuildCacheData();
+		BoneRemovedDelegate.Broadcast(InName);
+		return true;
+	}
+
+	return false;
+}
+
+// rename just change the name
+// return false if not found
+bool UIKRigDefinition::RenameBone(const FName& InOldName, const FName& InNewName)
+{
+	// ensure we have old name and do not have new name
+	int32 Index = Hierarchy.FindIndexFromBoneArray(InOldName);
+	if (Index != INDEX_NONE)
+	{
+		// make sure nobody is using NewName
+		if (Hierarchy.FindIndexFromBoneArray(InNewName) == INDEX_NONE)
+		{
+			TArray<int32> ChildIndices = Hierarchy.FindIndicesByParentName(InOldName);
+
+			for (int32 ChildIndex = 0; ChildIndex < ChildIndices.Num(); ++ChildIndex)
+			{
+				Hierarchy.Bones[ChildIndices[ChildIndex]].ParentName = InNewName;
+			}
+
+			// now change me
+			Hierarchy.Bones[Index].Name = InNewName;
+			Hierarchy.RebuildCacheData();
+			BoneRenamedDelegate.Broadcast(InOldName, InNewName);
+			return true;
+		}
+	}
+
+	return false;
+}
+
+// reparent to InNewParent
+//	if not found - either InName or InParent
+//	Or if InNewParent is invalid (i.e. it's a child of InName)
+bool UIKRigDefinition::ReparentBone(const FName& InName, const FName& InNewParent)
+{
+	int32 Index = Hierarchy.FindIndexFromBoneArray(InName);
+	if (Index != INDEX_NONE)
+	{
+		int32 NewParentIndex = INDEX_NONE;
+		if (InNewParent != NAME_None)
+		{
+			NewParentIndex = Hierarchy.FindIndexFromBoneArray(InNewParent);
+			if (NewParentIndex == INDEX_NONE)
+			{
+				return false;
+			}
+
+			// we have to ensure we don't have cycling situation
+			int32 CurIndex = NewParentIndex;
+			while (CurIndex != INDEX_NONE)
+			{
+				// i hit myself
+				if (CurIndex == Index)
+				{
+					// we're cycling
+					return false;
+				}
+
+				// find parent
+				CurIndex = Hierarchy.FindIndexFromBoneArray(Hierarchy.Bones[CurIndex].ParentName);
+			}
+		}
+
+		// we don't have a cycling, now allow set it
+		Hierarchy.Bones[Index].ParentName = InNewParent;
+		EnsureSortedCorrectly(true);
+		BoneReparentedDelegate.Broadcast(InName, InNewParent);
+		return true;
+	}
+
+	return false;
+}
+
+
+// ensure it's sorted from parent to children
+void UIKRigDefinition::EnsureSortedCorrectly(bool bReSortIfNeeded)
+{
+	// we want to ensure it's sorted from parent to child
+	// we do this twice, we use this data to sort, so after sort, we still have to do this again
+	Hierarchy.RebuildCacheData();
+
+	// @TODO: FIX THE CODE, it's confusing
+	bool bNeedReSort = false;
+	const TArray<FIKRigBone>& Bones = Hierarchy.Bones;
+	const TArray<int32>& ParentIndices = Hierarchy.ParentIndices;
+	// first check to see if we have violation
+	for (int32 Index = 0; Index < ParentIndices.Num(); ++Index)
+	{
+		// if my parent is ahead, I need to resort
+		if (ParentIndices[Index] >= Index)
+		{
+			bNeedReSort = true;
+		}
+	}
+
+	// resort if needed is true, then we're good
+	// but if resort is false, we want to ensure it's sorted already, and not have to resort
+	ensureAlways(bReSortIfNeeded || !bNeedReSort);
+
+	if (bNeedReSort)
+	{
+		// if so we create a tree and recreate Bones array
+		struct FBoneTreeNode
+		{
+			FIKRigBone Bone;
+			FBoneTreeNode* Children;
+			FBoneTreeNode* Sibling;
+
+			FBoneTreeNode()
+				: Children(nullptr)
+				, Sibling(nullptr)
+			{}
+
+			FBoneTreeNode(FIKRigBone InBone)
+				: Bone(InBone)
+				, Children(nullptr)
+				, Sibling(nullptr)
+			{}
+
+			void AddToBones(TArray<FIKRigBone>& OutBones)
+			{
+				OutBones.Add(Bone);
+
+				FBoneTreeNode* Iter = Children;
+				while (Iter)
+				{
+					Iter->AddToBones(OutBones);
+					Iter = Iter->Sibling;
+				}
+			}
+
+			void AddChild(FBoneTreeNode* InChild)
+			{
+				if (Children == nullptr)
+				{
+					Children = InChild;
+				}
+				else
+				{
+					FBoneTreeNode* Iter = Children;
+					do
+					{
+						if (Iter->Sibling == nullptr)
+						{
+							Iter->Sibling = InChild;
+							break;
+						}
+
+						Iter = Iter->Sibling;
+					} while (Iter);
+				}
+			}
+		};
+
+		// the assumption is that first one is the root
+		TArray<FBoneTreeNode> Tree;
+		Tree.AddDefaulted(Bones.Num());
+
+		for (int32 Index = 0; Index < Bones.Num(); ++Index)
+		{
+			Tree[Index].Bone = Bones[Index];
+			if (ParentIndices[Index] != INDEX_NONE)
+			{
+				Tree[ParentIndices[Index]].AddChild(&Tree[Index]);
+			}
+		}
+
+		TArray<FIKRigBone> NewBones;
+		// BFS for tree
+		Tree[0].AddToBones(NewBones);
+
+		// create new transform 
+		check(NewBones.Num() == Bones.Num());
+
+		// we need to move transform of old Bones to new Bones, and their indices have changed
+		// move transform, very slow move
+		TArray<FTransform> NewTransforms;
+		NewTransforms.AddUninitialized(Bones.Num());
+
+		for (int32 Index = 0; Index < Bones.Num(); ++Index)
+		{
+			const FName& Name = Bones[Index].Name;
+			for (int32 Jndex = 0; Jndex < NewBones.Num(); ++Jndex)
+			{
+				if (Name == NewBones[Jndex].Name)
+				{
+					NewTransforms[Jndex] = ReferencePose.GlobalTransforms[Index];
+				}
+			}
+		}
+
+		// override back to it
+		Hierarchy.Bones = NewBones;
+		ReferencePose.GlobalTransforms = NewTransforms;
+		Hierarchy.RebuildCacheData();
+	}
+}
+
+void UIKRigDefinition::ResetHierarchy()
+{
+	Hierarchy.Bones.Reset();
+	Hierarchy.ParentIndices.Reset();
+	Hierarchy.RuntimeNameLookupTable.Reset();
+	ReferencePose.GlobalTransforms.Reset();
+}
+
+void UIKRigDefinition::EnsureCreateUniqueGoalName(FName& InOutGoal) const
+{
+	int32 Index = 1;
+
+	FString GoalNameString = InOutGoal.ToString();
+
+	// while contains
+	while(IKGoals.Find(FName(*GoalNameString)))
+	{
+		GoalNameString += FString::Format(TEXT("_{0}"), {Index++});
+	}
+
+	InOutGoal = FName(*GoalNameString);
+}
+
+#endif // WITH_EDITOR
+
+void UIKRigDefinition::PostLoad()
+{
+	Super::PostLoad();
+
+	Hierarchy.RebuildCacheData();
+}
\ No newline at end of file
diff --git a/Engine/Plugins/Animation/IKRig/Source/IKRig/Private/IKRigHierarchy.cpp b/Engine/Plugins/Animation/IKRig/Source/IKRig/Private/IKRigHierarchy.cpp
new file mode 100644
--- /dev/null
+++ b/Engine/Plugins/Animation/IKRig/Source/IKRig/Private/IKRigHierarchy.cpp
@@ -0,0 +1,24 @@
+// Copyright Epic Games, Inc. All Rights Reserved.
+
+/*=============================================================================
+	IKRigHierarchy.cpp: IKRig Data Type impelmentation file
+=============================================================================*/
+
+#include "IKRigHierarchy.h"
+
+
+// @todo: cache this for runtime
+TArray<int32> FIKRigHierarchy::FindChildren(int32 Index) const
+{
+	if (Bones.IsValidIndex(Index))
+	{
+#if WITH_EDITOR
+		return FindIndicesByParentName(Bones[Index].Name);
+#else
+		return Bones[Index].Children;
+#endif 
+	}
+
+	return TArray<int32>();
+}
+
diff --git a/Engine/Plugins/Animation/IKRig/Source/IKRig/Private/IKRigProcessor.cpp b/Engine/Plugins/Animation/IKRig/Source/IKRig/Private/IKRigProcessor.cpp
new file mode 100644
--- /dev/null
+++ b/Engine/Plugins/Animation/IKRig/Source/IKRig/Private/IKRigProcessor.cpp
@@ -0,0 +1,165 @@
+// Copyright Epic Games, Inc. All Rights Reserved.
+
+/**
+ * Contains IKRigProcessor - Runtime Implementation
+ *
+ */
+
+#include "IKRigProcessor.h"
+
+ // IKRigProcessor implementation functions
+void UIKRigProcessor::SetIKRigDefinition(UIKRigDefinition* InRigDefinition, bool bInitialize /*= false*/)
+{
+	if (RigDefinition)
+	{
+		Terminate();
+	}
+
+	RigDefinition = InRigDefinition;
+
+	TransformModifier = FIKRigTransformModifier(&RigDefinition->GetHierarchy());
+
+	if (RigDefinition && bInitialize)
+	{
+		Initialize(RigDefinition->GetReferencePose());
+	}
+}
+
+void UIKRigProcessor::Initialize(const FIKRigTransform& InRefTransform)
+{
+	if (RigDefinition)
+	{
+		ReferenceTransform = InRefTransform;
+		Reinitialize();
+	}
+}
+
+void UIKRigProcessor::Reinitialize()
+{
+	TransformModifier.ResetGlobalTransform(ReferenceTransform);
+
+	IKGoals = RigDefinition->GetGoals();
+
+	// i'm copying this here 
+	TArray<UIKRigSolverDefinition*> SolverDefinitions = RigDefinition->GetSolverDefinitions();
+
+	const int32 NumSolvers = SolverDefinitions.Num();
+	Solvers.Reset(NumSolvers);
+	for (int32 Index = 0; Index < NumSolvers; ++Index)
+	{
+		// it's possible SolverDefinitions.Num() != Solvers.Num()
+		if (SolverDefinitions[Index])
+		{
+			TSubclassOf<UIKRigSolver> ClassType = SolverDefinitions[Index]->GetExecutionClass();
+			if (ClassType != nullptr)
+			{
+				UIKRigSolver* NewSolver = NewObject<UIKRigSolver>(this, ClassType);
+				NewSolver->Init(SolverDefinitions[Index], UIKRigSolver::FIKRigTransformGetter::CreateUObject(this, &UIKRigProcessor::GetRefPoseGetter), 
+					UIKRigSolver::FIKRigGoalGetter::CreateUObject(this, &UIKRigProcessor::GoalGetter));
+				Solvers.Add(NewSolver);
+			}
+		}
+	}
+
+	bInitialized = true;
+}
+
+const FIKRigTransform& UIKRigProcessor::GetRefPoseGetter() 
+{
+	return ReferenceTransform;
+}
+
+bool UIKRigProcessor::GoalGetter(const FName& InGoalName, FIKRigTarget& OutTarget)
+{
+	FIKRigTarget* Target = FindGoal(InGoalName);
+	if (Target)
+	{
+		OutTarget = *Target;
+
+		return true;
+	}
+
+	return false;
+}
+
+void UIKRigProcessor::Terminate()
+{
+	bInitialized = false;
+}
+
+void UIKRigProcessor::Solve()
+{
+	if (bInitialized)
+	{
+		check (RigDefinition);
+
+		for (int32 Index = 0; Index < Solvers.Num(); ++Index)
+		{
+			Solvers[Index]->Solve(TransformModifier);
+		}
+	}
+}
+
+// update goal functions
+void UIKRigProcessor::SetGoalPosition(const FName& GoalName, const FVector& InPosition)
+{
+	FIKRigTarget* Target = FindGoal(GoalName);
+	if (Target)
+	{
+		Target->PositionTarget.Position = InPosition;
+	}
+}
+
+void UIKRigProcessor::SetGoalRotation(const FName& GoalName, const FRotator& InRotation)
+{
+	FIKRigTarget* Target = FindGoal(GoalName);
+	if (Target)
+	{
+		Target->RotationTarget.Rotation = InRotation;
+	}
+}
+
+void UIKRigProcessor::SetGoalTarget(const FName& GoalName, const FIKRigTarget& InTarget)
+{
+	FIKRigTarget* Target = FindGoal(GoalName);
+	if (Target)
+	{
+		*Target = InTarget;
+	}
+}
+
+FIKRigTarget* UIKRigProcessor::FindGoal(const FName& GoalName)
+{
+	FIKRigGoal* RigGoal = IKGoals.Find(GoalName);
+	if (RigGoal)
+	{
+		return &RigGoal->Target;
+	}
+
+	return nullptr;
+}
+
+void UIKRigProcessor::GetGoals(TArray<FName>& OutNames) const
+{
+	IKGoals.GenerateKeyArray(OutNames);
+}
+
+FIKRigTransformModifier& UIKRigProcessor::GetIKRigTransformModifier() 
+{
+	return TransformModifier;
+}
+
+const FIKRigHierarchy* UIKRigProcessor::GetHierarchy() const
+{
+	if (RigDefinition)
+	{
+		return &RigDefinition->GetHierarchy();
+	}
+
+	return nullptr;
+}
+
+void UIKRigProcessor::ResetToRefPose()
+{
+	TransformModifier.ResetGlobalTransform(ReferenceTransform);
+}
\ No newline at end of file
diff --git a/Engine/Plugins/Animation/IKRig/Source/IKRig/Private/IKRigSolver.cpp b/Engine/Plugins/Animation/IKRig/Source/IKRig/Private/IKRigSolver.cpp
new file mode 100644
--- /dev/null
+++ b/Engine/Plugins/Animation/IKRig/Source/IKRig/Private/IKRigSolver.cpp
@@ -0,0 +1,61 @@
+// Copyright Epic Games, Inc. All Rights Reserved.
+
+/*=============================================================================
+	IKRigSolver.cpp: Solver definition class
+=============================================================================*/
+
+#include "IKRigSolver.h"
+
+// input hierarchy and ref pose? 
+void UIKRigSolver::Init(UIKRigSolverDefinition* InSolverDefinition, FIKRigTransformGetter InRefPoseGetter, FIKRigGoalGetter InGoalGetter)
+{
+	SolverDefinition = InSolverDefinition;
+
+	RefPoseGetter = InRefPoseGetter;
+	ensure(RefPoseGetter.IsBound());
+
+	GoalGetter = InGoalGetter;
+	ensure(GoalGetter.IsBound());
+
+	InitInternal();
+}
+
+bool UIKRigSolver::IsSolverActive() const 
+{
+	return (SolverDefinition && bEnabled);
+}
+
+// input : goal getter or goals
+// output : modified pose - GlobalTransforms
+void UIKRigSolver::Solve(FIKRigTransformModifier& InOutGlobalTransform)
+{
+	if (IsSolverActive())
+	{
+		SolveInternal(InOutGlobalTransform);
+	}
+}
+
+bool UIKRigSolver::GetTaskTarget(const FName& TaskName, FIKRigTarget& OutTarget) const
+{
+	if (SolverDefinition && GoalGetter.IsBound())
+	{	
+		const FName* GoalName = SolverDefinition->GetTaskToGoal().Find(TaskName);
+		if (GoalName)
+		{
+			return GoalGetter.Execute(*GoalName, OutTarget);
+		}
+	}
+
+	return false;
+}
+
+const FIKRigTransform& UIKRigSolver::GetReferencePose() const
+{
+	if (RefPoseGetter.IsBound())
+	{
+		return RefPoseGetter.Execute();
+	}
+
+	static FIKRigTransform Dummy;
+	return Dummy;
+}
\ No newline at end of file
diff --git a/Engine/Plugins/Animation/IKRig/Source/IKRig/Private/IKRigSolverDefinition.cpp b/Engine/Plugins/Animation/IKRig/Source/IKRig/Private/IKRigSolverDefinition.cpp
new file mode 100644
--- /dev/null
+++ b/Engine/Plugins/Animation/IKRig/Source/IKRig/Private/IKRigSolverDefinition.cpp
@@ -0,0 +1,69 @@
+// Copyright Epic Games, Inc. All Rights Reserved.
+
+/**
+ * Contains IKRigSolverDefinition Implementation
+ *
+ */
+
+#include "IKRigSolverDefinition.h"
+
+void UIKRigSolverDefinition::CollectGoals(TArray<FName>& OutGoals)
+{
+	TArray<FName> LocalGoals;
+	// this just accumulate on the current array
+	// if you want to clear, clear before coming here
+	TaskToGoal.GenerateValueArray(LocalGoals);
+
+	OutGoals += LocalGoals;
+}
+
+#if WITH_EDITOR
+void UIKRigSolverDefinition::RenameGoal(const FName& OldName, const FName& NewName)
+{
+	for (auto Iter = TaskToGoal.CreateIterator(); Iter; ++Iter)
+	{
+		if (Iter.Value() == OldName)
+		{
+			Iter.Value() = NewName;
+		}
+	}
+}
+
+void UIKRigSolverDefinition::EnsureUniqueGoalName(FName& InOutUniqueGoalName) const
+{
+	// call delegate 
+	UIKRigDefinition* IKRigDef = CastChecked<UIKRigDefinition>(GetOuter());
+	IKRigDef->EnsureCreateUniqueGoalName(InOutUniqueGoalName);
+}
+
+FName UIKRigSolverDefinition::CreateUniqueGoalName(const TCHAR* Suffix) const
+{
+	if (Suffix)
+	{
+		FString NewGoalStr = DisplayName + TEXT("_") + Suffix;
+		// replace any whitespace with _
+		NewGoalStr = NewGoalStr.Replace(TEXT(" \t"), TEXT("_"));
+		FName NewGoalName = FName(*NewGoalStr, FNAME_Add);
+		EnsureUniqueGoalName(NewGoalName);
+
+		return NewGoalName;
+	}
+
+	return NAME_None;
+}
+
+void UIKRigSolverDefinition::OnGoalHasBeenUpdated()
+{
+	GoalNeedsUpdateDelegate.Broadcast();
+}
+
+#endif // WITH_EDITOR
+
+void UIKRigSolverDefinition::PostLoad()
+{
+	Super::PostLoad();
+
+#if WITH_EDITOR
+	UpdateTaskList();
+#endif // WITH_EDITOR
+}
diff --git a/Engine/Plugins/Animation/IKRig/Source/IKRig/Private/Solvers/TransformSolver.cpp b/Engine/Plugins/Animation/IKRig/Source/IKRig/Private/Solvers/TransformSolver.cpp
new file mode 100644
--- /dev/null
+++ b/Engine/Plugins/Animation/IKRig/Source/IKRig/Private/Solvers/TransformSolver.cpp
@@ -0,0 +1,53 @@
+// Copyright Epic Games, Inc. All Rights Reserved.
+
+/*=============================================================================
+	TransformSolver.cpp: Solver execution class for Transform
+=============================================================================*/
+
+#include "TransformSolver.h"
+#include "Solvers/TransformSolverDefinition.h"
+
+void UTransformSolver::InitInternal()
+{
+	
+}
+
+bool UTransformSolver::IsSolverActive() const
+{
+	if (Super::IsSolverActive())
+	{
+		if (UTransformSolverDefinition* SolverDef = Cast<UTransformSolverDefinition>(SolverDefinition))
+		{
+			return (SolverDef->bEnablePosition || SolverDef->bEnableRotation);
+		}
+	}
+
+	return false;
+}
+
+void UTransformSolver::SolveInternal(FIKRigTransformModifier& InOutGlobalTransform)
+{
+	if (UTransformSolverDefinition* SolverDef = Cast<UTransformSolverDefinition>(SolverDefinition))
+	{
+		FIKRigTarget Target;
+		if (GetTaskTarget(UTransformSolverDefinition::TransformTarget, Target))
+		{
+			int32 Index = InOutGlobalTransform.Hierarchy->GetIndex(Target.Bone);
+			if (Index != INDEX_NONE)
+			{
+				FTransform CurrentTransform = InOutGlobalTransform.GetGlobalTransform(Index);
+
+				if (SolverDef->bEnablePosition)
+				{
+					CurrentTransform.SetLocation(Target.PositionTarget.Position);
+				}
+				if (SolverDef->bEnableRotation)
+				{
+					CurrentTransform.SetRotation(Target.RotationTarget.Rotation.Quaternion());
+				}
+
+				InOutGlobalTransform.SetGlobalTransform(Index, CurrentTransform, true);
+			}
+		}
+	}
+}
\ No newline at end of file
diff --git a/Engine/Plugins/Animation/IKRig/Source/IKRig/Private/Solvers/TransformSolver.h b/Engine/Plugins/Animation/IKRig/Source/IKRig/Private/Solvers/TransformSolver.h
new file mode 100644
--- /dev/null
+++ b/Engine/Plugins/Animation/IKRig/Source/IKRig/Private/Solvers/TransformSolver.h
@@ -0,0 +1,25 @@
+// Copyright Epic Games, Inc. All Rights Reserved.
+
+/**
+ * Contains Transform Solver Execution
+ *
+ */
+
+#pragma once
+
+#include "IKRigSolver.h"
+#include "TransformSolver.generated.h"
+
+
+// run time for UTransformSolverDefinition
+UCLASS()
+class IKRIG_API UTransformSolver : public UIKRigSolver
+{
+	GENERATED_BODY()
+
+protected:
+	virtual void InitInternal() override;
+	virtual void SolveInternal(FIKRigTransformModifier& InOutGlobalTransform) override;
+	virtual bool IsSolverActive() const override;
+};
+
diff --git a/Engine/Plugins/Animation/IKRig/Source/IKRig/Private/Solvers/TransformSolverDefinition.cpp b/Engine/Plugins/Animation/IKRig/Source/IKRig/Private/Solvers/TransformSolverDefinition.cpp
new file mode 100644
--- /dev/null
+++ b/Engine/Plugins/Animation/IKRig/Source/IKRig/Private/Solvers/TransformSolverDefinition.cpp
@@ -0,0 +1,50 @@
+// Copyright Epic Games, Inc. All Rights Reserved.
+
+/**
+ * Contains TransformSolverDefinition Implementation
+ *
+ */
+
+#include "Solvers/TransformSolverDefinition.h"
+#include "TransformSolver.h"
+
+const FName UTransformSolverDefinition::TransformTarget = FName(TEXT("TransformTarget"));
+
+UTransformSolverDefinition::UTransformSolverDefinition()
+{
+	DisplayName = TEXT("Transform Solver");
+	ExecutionClass = UTransformSolver::StaticClass();
+}
+
+#if WITH_EDITOR
+void UTransformSolverDefinition::UpdateTaskList()
+{
+	// update tasks - you want to keep old name though if changed
+	FName* TaskFound= TaskToGoal.Find(TransformTarget);
+	bool bActiveTask = bEnablePosition || bEnableRotation;
+
+	// not found, we add new
+	if (bActiveTask && !TaskFound)
+	{
+		TaskToGoal.Add(TransformTarget) = CreateUniqueGoalName(*TransformTarget.ToString());
+	}
+	// if the opposite, we have to remove
+	else if (!bActiveTask && TaskFound)
+	{
+		TaskToGoal.FindAndRemoveChecked(TransformTarget);
+	}
+
+	// trigger a delegate for goal has been updated?
+	OnGoalHasBeenUpdated();
+}
+
+void UTransformSolverDefinition::PostEditChangeChainProperty(struct FPropertyChangedChainEvent& PropertyChangedEvent)
+{
+	if (PropertyChangedEvent.Property->GetFName() == GET_MEMBER_NAME_CHECKED(UTransformSolverDefinition, bEnablePosition) ||
+		PropertyChangedEvent.Property->GetFName() == GET_MEMBER_NAME_CHECKED(UTransformSolverDefinition, bEnableRotation))
+	{
+		UpdateTaskList();
+	}
+}
+
+#endif// WITH_EDITOR
\ No newline at end of file
diff --git a/Engine/Plugins/Animation/IKRig/Source/IKRig/Public/AnimNodes/AnimNode_IKRig.h b/Engine/Plugins/Animation/IKRig/Source/IKRig/Public/AnimNodes/AnimNode_IKRig.h
new file mode 100644
--- /dev/null
+++ b/Engine/Plugins/Animation/IKRig/Source/IKRig/Public/AnimNodes/AnimNode_IKRig.h
@@ -0,0 +1,68 @@
+// Copyright Epic Games, Inc. All Rights Reserved.
+
+#pragma once
+
+#include "CoreMinimal.h"
+#include "UObject/ObjectMacros.h"
+#include "Animation/AnimNodeBase.h"
+#include "AnimNode_IKRig.generated.h"
+
+class UIKRigDefinition;
+class UIKRigProcessor;
+
+/**
+*	Controller which implements the FABRIK IK approximation algorithm -  see http://www.academia.edu/9165835/FABRIK_A_fast_iterative_solver_for_the_Inverse_Kinematics_problem for details
+*/
+
+USTRUCT(BlueprintInternalUseOnly)
+struct IKRIG_API FAnimNode_IKRig : public FAnimNode_Base
+{
+	GENERATED_BODY()
+
+	UPROPERTY(EditAnywhere, Category = Links)
+	FPoseLink Source;
+
+	UPROPERTY(EditAnywhere, BlueprintReadWrite, Category = RigDefinition, meta = (NeverAsPin))
+	UIKRigDefinition* RigDefinitionAsset;
+
+	/** Coordinates for target location of tip bone - if EffectorLocationSpace is bone, this is the offset from Target Bone to use as target location*/
+	UPROPERTY(EditAnywhere, BlueprintReadWrite, EditFixedSize, Category = Goal, meta = (PinShownByDefault))
+	TArray<FTransform> GoalTransforms;
+
+	/** For every evaluation, we'll start with refpose, this is safer if you have joints that don't exists in the target */
+	UPROPERTY(EditAnywhere, Category = Solver)
+	bool bStartFromRefPose;
+
+#if WITH_EDITORONLY_DATA
+	/** Toggle drawing of axes to debug joint rotation*/
+	UPROPERTY(EditAnywhere, Category = Solver)
+	bool bEnableDebugDraw;
+#endif
+
+private: 
+	UPROPERTY(Transient)
+	UIKRigProcessor* RigProcessor = nullptr;
+
+	TArray<FName> GoalNames;
+
+	TMap<FCompactPoseBoneIndex, int32, FDefaultSetAllocator, TCompactPoseBoneIndexMapKeyFuncs<int32>> CompactPoseToRigIndices;
+
+public:
+	FAnimNode_IKRig();
+
+	// FAnimNode_Base interface
+	virtual void GatherDebugData(FNodeDebugData& DebugData) override;
+	virtual void Initialize_AnyThread(const FAnimationInitializeContext& Context) override;
+	virtual void OnInitializeAnimInstance(const FAnimInstanceProxy* InProxy, const UAnimInstance* InAnimInstance) override;
+	virtual void CacheBones_AnyThread(const FAnimationCacheBonesContext& Context)  override;
+	virtual void Evaluate_AnyThread(FPoseContext& Output) override;
+	virtual void Update_AnyThread(const FAnimationUpdateContext& Context) override;
+	virtual bool NeedsOnInitializeAnimInstance() const override { return true; }
+	// End of FAnimNode_Base interface
+
+private:
+	void RebuildGoalList();
+	FName GetGoalName(int32 Index) const;
+
+	friend class UAnimGraphNode_IKRig;
+};
diff --git a/Engine/Plugins/Animation/IKRig/Source/IKRig/Public/IKRig.h b/Engine/Plugins/Animation/IKRig/Source/IKRig/Public/IKRig.h
new file mode 100644
--- /dev/null
+++ b/Engine/Plugins/Animation/IKRig/Source/IKRig/Public/IKRig.h
@@ -0,0 +1,17 @@
+// Copyright Epic Games, Inc. All Rights Reserved.
+
+#pragma once
+
+#include "CoreMinimal.h"
+#include "Modules/ModuleManager.h"
+
+class FIKRigModule : public IModuleInterface
+{
+public:
+
+	/** IModuleInterface implementation */
+	virtual void StartupModule() override;
+	virtual void ShutdownModule() override;
+};
+
+DECLARE_LOG_CATEGORY_EXTERN(LogIKRig, Warning, All);
diff --git a/Engine/Plugins/Animation/IKRig/Source/IKRig/Public/IKRigConstraint.h b/Engine/Plugins/Animation/IKRig/Source/IKRig/Public/IKRigConstraint.h
new file mode 100644
--- /dev/null
+++ b/Engine/Plugins/Animation/IKRig/Source/IKRig/Public/IKRigConstraint.h
@@ -0,0 +1,25 @@
+// Copyright Epic Games, Inc. All Rights Reserved.
+
+/**
+ * Contains IKRig Definition 
+ *
+ */
+
+#pragma once
+
+#include "CoreMinimal.h"
+#include "UObject/ObjectMacros.h"
+#include "UObject/Object.h"
+#include "IKRigConstraint.generated.h"
+
+UCLASS(config = Engine, hidecategories = UObject, BlueprintType)
+class IKRIG_API UIKRigConstraint : public UObject
+{
+	GENERATED_BODY()
+
+private:
+
+
+public: 
+};
+
diff --git a/Engine/Plugins/Animation/IKRig/Source/IKRig/Public/IKRigConstraintDefinition.h b/Engine/Plugins/Animation/IKRig/Source/IKRig/Public/IKRigConstraintDefinition.h
new file mode 100644
--- /dev/null
+++ b/Engine/Plugins/Animation/IKRig/Source/IKRig/Public/IKRigConstraintDefinition.h
@@ -0,0 +1,41 @@
+// Copyright Epic Games, Inc. All Rights Reserved.
+
+/**
+ * Contains IKRig Definition 
+ *
+ */
+
+#pragma once
+
+#include "CoreMinimal.h"
+#include "UObject/ObjectMacros.h"
+#include "UObject/Object.h"
+#include "IKRigConstraintDefinition.generated.h"
+
+USTRUCT()
+struct IKRIG_API FIKRigConstraintProfile 
+{
+	GENERATED_BODY()
+
+	UPROPERTY(EditAnywhere, Category=FIKRigConstraintProfile)
+	FName Name;
+
+	UPROPERTY(EditAnywhere, Category=FIKRigConstraintProfile)
+	TMap<FName, UIKRigConstraint*> Constraints;
+};
+
+UCLASS(config = Engine, hidecategories = UObject, BlueprintType)
+class IKRIG_API UIKRigConstraintDefinition : public UObject
+{
+	GENERATED_BODY()
+
+public: 
+	UIKRigConstraintDefinition();
+
+private:
+	// at least one profile always exists - Default
+	UPROPERTY(EditAnywhere, Category = "Profile")
+	TMap<FName, FIKRigConstraintProfile> ConstraintProfiles;
+
+	static const FName DefaultProfileName;
+};
\ No newline at end of file
diff --git a/Engine/Plugins/Animation/IKRig/Source/IKRig/Public/IKRigConstraintSolver.h b/Engine/Plugins/Animation/IKRig/Source/IKRig/Public/IKRigConstraintSolver.h
new file mode 100644
--- /dev/null
+++ b/Engine/Plugins/Animation/IKRig/Source/IKRig/Public/IKRigConstraintSolver.h
@@ -0,0 +1,37 @@
+// Copyright Epic Games, Inc. All Rights Reserved.
+
+/**
+ * Contains IKRig Constraint Solver Definition 
+ *
+ */
+
+#pragma once
+
+#include "IKRigSolver.h"
+#include "IKRigConstraintSolver.generated.h"
+
+// run time processor 
+UCLASS(config = Engine, hidecategories = UObject)
+class IKRIG_API UIKRigConstraintSolver : public UIKRigSolver
+{
+	GENERATED_BODY()
+
+	DECLARE_DELEGATE_TwoParams(FIKRigQueryConstraint, const FName& InConstraintName, UIKRigConstraint& OutConstraint);
+
+public: 
+	// input hierarchy and ref pose? 
+	//virtual void Init(UIKRigSolverDefinition* InSolverDefinition, const FIKRigTransform& InRefPose, FIKRigGoalGetter& InGoalGetter, FSolveConstraint& InConstraintHandler) override;
+
+	// input : goal getter or goals
+	// output : modified pose - GlobalTransforms
+	//virtual void Solve(FIKRigTransform& OutGlobalTransform) override;
+
+	// register/unregister query function
+	void RegisterQueryConstraintHandler(const FIKRigQueryConstraint& InQueryConstraintHandler);
+	void UnregisterQueryConstraintHandler();
+
+private:
+	// delegate
+	FIKRigQueryConstraint QueryConstraintHandler;
+};
+
diff --git a/Engine/Plugins/Animation/IKRig/Source/IKRig/Public/IKRigDataTypes.h b/Engine/Plugins/Animation/IKRig/Source/IKRig/Public/IKRigDataTypes.h
new file mode 100644
--- /dev/null
+++ b/Engine/Plugins/Animation/IKRig/Source/IKRig/Public/IKRigDataTypes.h
@@ -0,0 +1,195 @@
+// Copyright Epic Games, Inc. All Rights Reserved.
+
+/**
+ * Contains IKRig Definition 
+ *
+ *  https://docs.google.com/document/d/1yd8GCfT2aufxSdb5jAzlNTr1SptxEFpS9pWdQY-8LIk/edit#
+ */
+
+#pragma once
+
+#include "CoreMinimal.h"
+#include "UObject/ObjectMacros.h"
+#include "UObject/Object.h"
+#include "IKRigDataTypes.generated.h"
+
+struct FIKRigHierarchy;
+
+USTRUCT()
+struct IKRIG_API FIKRigPosition
+{
+	GENERATED_BODY()
+
+	FIKRigPosition()
+		: Position(ForceInitToZero)
+	{
+
+	}
+	UPROPERTY(EditAnywhere, Category = FIKRigPosition)
+	FVector Position;
+};
+
+USTRUCT()
+struct IKRIG_API FIKRigRotation
+{
+	GENERATED_BODY()
+
+	FIKRigRotation()
+		: Rotation(ForceInitToZero)
+	{
+
+	}
+
+	UPROPERTY(EditAnywhere, Category = FIKRigRotation)
+	FRotator Rotation;
+};
+
+// list of targets for goal
+// this contains default value for the target
+USTRUCT()
+struct IKRIG_API FIKRigTarget
+{
+	GENERATED_BODY()
+
+	UPROPERTY(EditAnywhere, Category = FIKRigTarget)
+	FName Bone;
+
+	UPROPERTY(EditAnywhere, Category = FIKRigTarget)
+	FIKRigPosition PositionTarget;
+
+	UPROPERTY(EditAnywhere, Category = FIKRigTarget)
+	FIKRigRotation RotationTarget;
+};
+
+USTRUCT()
+struct IKRIG_API FIKRigGoal
+{
+	GENERATED_BODY()
+
+	FIKRigGoal()
+	{
+		Name = NAME_None;
+	}
+	FIKRigGoal(const FName& GoalName)
+	{
+		Name = GoalName;
+	}
+
+	UPROPERTY(VisibleAnywhere, Category=FIKRigGoal)
+	FName Name;
+
+	UPROPERTY(EditAnywhere, Category = FIKRigGoal)
+	FIKRigTarget Target;
+
+	void SetPositionTarget(const FVector& InPosition)
+	{
+		Target.PositionTarget.Position = InPosition;
+	}
+
+	void SetRotationTarget(const FRotator& InRotation)
+	{
+		Target.RotationTarget.Rotation = InRotation;
+	}
+};
+
+
+// If you want to modify this, use FIKRigTransformModifier
+USTRUCT()
+struct IKRIG_API FIKRigTransform
+{
+	GENERATED_BODY()
+
+private:
+	UPROPERTY(VisibleAnywhere, Category = FIKRigTransform)
+	TArray<FTransform> GlobalTransforms;
+
+public:
+	int32 GetNum() const 
+	{
+		return GlobalTransforms.Num();
+	}
+
+	bool IsValidIndex(int32 Index) const
+	{
+		return GlobalTransforms.IsValidIndex(Index);
+	}
+
+	const TArray<FTransform>& GetGlobalTransforms() const
+	{
+		return GlobalTransforms;
+	}
+
+	const FTransform& GetGlobalTransform(int32 Index) const
+	{
+		if (GlobalTransforms.IsValidIndex(Index))
+		{
+			return GlobalTransforms[Index];
+		}
+
+		return FTransform::Identity;
+	}
+
+	FTransform GetRelativeTransform(int32 Index, int32 BaseIndex) const
+	{
+		if (GlobalTransforms.IsValidIndex(Index))
+		{
+			if (GlobalTransforms.IsValidIndex(BaseIndex))
+			{
+				return (GlobalTransforms[Index].GetRelativeTransform(GlobalTransforms[BaseIndex]));
+			}
+			else
+			{
+				return GlobalTransforms[Index];
+			}
+		}
+
+		return FTransform::Identity;
+	}
+	friend struct FIKRigTransformModifier;
+	friend class UIKRigDefinition;
+	friend class UIKRigController;
+};
+
+// allows transform to be modified 
+// use this class to modify transform
+struct IKRIG_API FIKRigTransformModifier
+{
+	FIKRigTransformModifier()
+		: Hierarchy(nullptr)
+	{
+	}
+
+	FIKRigTransformModifier(const FIKRigHierarchy* InHierarchy);
+
+	void SetGlobalTransform(int32 Index, const FTransform& InTransform, bool bPropagate);
+	void SetLocalTransform(int32 Index, const FTransform& InTransform, bool bUpdate);
+	// this is mutating, but const_cast 
+	const FTransform& GetLocalTransform(int32 Index) const;
+	const FTransform& GetGlobalTransform(int32 Index) const;
+
+	void ResetGlobalTransform(const FIKRigTransform& InTransform);
+
+	const FIKRigHierarchy* Hierarchy;
+
+private:
+
+	// we use this info to modify 
+	// ensure memory cycle for these are valid through
+	FIKRigTransform GlobalTransforms;
+	TArray<FTransform> LocalTransforms;
+	TBitArray<>	LocalTransformDirtyFlags;
+
+	void EnsureLocalTransformsExist();
+	void RecalculateLocalTransform();
+	void UpdateLocalTransform(int32 Index);
+	/* This function does propagate through children
+	 * of the current index and keeps last LocalTransform to up to date
+	 */
+	void SetGlobalTransform_Internal(int32 Index, const FTransform& InTransform);
+
+	/** Set Global Transform Recursive 
+	 *	
+	 */
+	void SetGlobalTransform_Recursive(int32 Index, const FTransform& InTransform);
+};
+
diff --git a/Engine/Plugins/Animation/IKRig/Source/IKRig/Public/IKRigDefinition.h b/Engine/Plugins/Animation/IKRig/Source/IKRig/Public/IKRigDefinition.h
new file mode 100644
--- /dev/null
+++ b/Engine/Plugins/Animation/IKRig/Source/IKRig/Public/IKRigDefinition.h
@@ -0,0 +1,154 @@
+// Copyright Epic Games, Inc. All Rights Reserved.
+
+/**
+ * Contains IKRig Definition 
+ *
+ *  https://docs.google.com/document/d/1yd8GCfT2aufxSdb5jAzlNTr1SptxEFpS9pWdQY-8LIk/edit#
+ */
+
+#pragma once
+
+#include "CoreMinimal.h"
+#include "UObject/ObjectMacros.h"
+#include "UObject/Object.h"
+#include "IKRigDataTypes.h"
+#include "IKRigHierarchy.h"
+#include "IKRigDefinition.generated.h"
+
+class UIKRigSolverDefinition;
+class UIKRigConstraintDefinition;
+
+UCLASS(Blueprintable)
+class IKRIG_API UIKRigDefinition : public UObject
+{
+	GENERATED_BODY()
+
+public:
+	UIKRigDefinition();
+
+	// ensure the goal name is unique
+	void EnsureCreateUniqueGoalName(FName& InOutGoal) const;
+public: 
+	/** Source Asset imported */
+	UPROPERTY()
+	TSoftObjectPtr<UObject> SourceAsset;
+
+	// goal related APIs
+	const TMap<FName, FIKRigGoal>& GetGoals() const
+	{
+		return IKGoals;
+	}
+
+	TArray<FName> GetGoalsNames() const
+	{
+		TArray<FName> NameArray;
+		IKGoals.GenerateKeyArray(NameArray);
+		return NameArray;
+	}
+
+	const FIKRigTransform& GetReferencePose() const
+	{
+		return ReferencePose;
+	}
+
+	const TArray<UIKRigSolverDefinition*>& GetSolverDefinitions() const 
+	{
+		return SolverDefinitions;
+	}
+
+	const FIKRigHierarchy& GetHierarchy() const
+	{
+		return Hierarchy;
+	}
+private:
+	/*********** Skeleton *********************/
+
+	/* Hierarchy
+	 * 
+	 * This just is a relationship. Parent/child relationship
+	 * sorted by parent to children
+	*/
+	UPROPERTY(VisibleAnywhere, Category = "Hierarchy")
+	FIKRigHierarchy Hierarchy;
+
+	// this is default transform of reference pose
+	// the separate with Hierarchy is for retargeting
+	// during retargeting, we want to ensure it works with different transform
+	// (this can be replaced during runtime for a retarget)
+	// do not cache or rely on this value inside of solver
+	// use the incoming value when initialize
+	// this has to match with Hierarchy data index
+	// IKRigDefinition interface handles the integrity of (this and Hierarchy)
+	// Use IKRigController for APIs
+	UPROPERTY(VisibleAnywhere, Category = "Hierarchy")
+	FIKRigTransform ReferencePose;
+
+	/*********** Solver Definition *********************/
+	// List of solvers, in the order of operations
+	UPROPERTY(EditAnywhere, instanced, Category = "Solver")
+	TArray<UIKRigSolverDefinition*> SolverDefinitions;
+
+	/*********** Constraint Definition *********************/
+	// contains constraint data
+	UPROPERTY(EditAnywhere, instanced, Category = "Constraint")
+	UIKRigConstraintDefinition* ConstraintDefinitions;
+
+	/*********** Goals with Default Value******************/
+	// goals data
+	// this is cached by "Unique Internal Name" here
+	// where in the runtime, it will cache by "Display Name"
+	// as that is what's searched to set the value by game
+	// but during editing time, we use internal name 
+	// as identifier
+	UPROPERTY(EditAnywhere, Category = "Goals")
+	TMap<FName, FIKRigGoal> IKGoals;
+
+	// UObject functions
+	virtual void PostLoad() override;
+
+#if WITH_EDITOR
+	// https ://drive.google.com/file/d/1EO0Ijojx-0jommdHuZv1JYxfCdoFLd44/view
+
+	/* Goal operators */
+	void UpdateGoal();
+
+	/** Hierarchy modifires */
+	// add a bone. Return false if it already has conflict name or parent is not found
+	bool AddBone(const FName& InName, const FName& InParent, const FTransform& InGlobalTransform);
+	// remove a bone. Returns false if it is not found
+	// if it has a children, it will remove all the children with it
+	bool RemoveBone(const FName& InName);
+	// rename just change the name
+	// return false if not found
+	bool RenameBone(const FName& InOldName, const FName& InNewName);
+	// reparent to InNewParent
+	// returns false 
+	//	if not found - either InName or InParent
+	//	Or if InNewParent is invalid (i.e. it's a child of InName)
+	bool ReparentBone(const FName& InName, const FName& InNewParent);
+	// reset 
+	void ResetHierarchy();
+
+	// ensure it's sorted from parent to children
+	void EnsureSortedCorrectly(bool bReSortIfNeeded=false);
+	void Sanitize();
+
+	// TODO: delegate for hierarchy changes
+	DECLARE_MULTICAST_DELEGATE_OneParam(FBoneAdded, FName /*BoneName*/);
+	DECLARE_MULTICAST_DELEGATE_OneParam(FBoneRemoved, FName /*BoneName*/);
+	DECLARE_MULTICAST_DELEGATE_TwoParams(FBoneRenamed, FName /*OldName*/, FName /*NewName*/);
+	DECLARE_MULTICAST_DELEGATE_TwoParams(FBoneReparented, FName /*Name*/, FName /*ParentName*/);
+
+	DECLARE_MULTICAST_DELEGATE(FIKRigDefinitionModified);
+
+	FBoneAdded BoneAddedDelegate;
+	FBoneRemoved BoneRemovedDelegate;
+	FBoneRenamed BoneRenamedDelegate;
+	FBoneReparented BoneReparentedDelegate;
+	FIKRigDefinitionModified IKRigDefinitionModified;
+
+#endif // WITH_EDITOR
+
+	friend class UIKRigController;
+};
+
diff --git a/Engine/Plugins/Animation/IKRig/Source/IKRig/Public/IKRigHierarchy.h b/Engine/Plugins/Animation/IKRig/Source/IKRig/Public/IKRigHierarchy.h
new file mode 100644
--- /dev/null
+++ b/Engine/Plugins/Animation/IKRig/Source/IKRig/Public/IKRigHierarchy.h
@@ -0,0 +1,192 @@
+// Copyright Epic Games, Inc. All Rights Reserved.
+
+/**
+ * Contains IKRig Definition 
+ *
+ *  https://docs.google.com/document/d/1yd8GCfT2aufxSdb5jAzlNTr1SptxEFpS9pWdQY-8LIk/edit#
+ */
+
+#pragma once
+
+#include "CoreMinimal.h"
+#include "UObject/ObjectMacros.h"
+#include "UObject/Object.h"
+#include "IKRigHierarchy.generated.h"
+
+struct FIKRigTransform;
+
+USTRUCT()
+struct IKRIG_API FIKRigBone 
+{
+	GENERATED_BODY()
+
+	UPROPERTY(VisibleAnywhere, Category = FIKRigBone)
+	FName Name;
+
+	// more user friendly with name than Index, and easy to move to different location without losing who was the parent
+	UPROPERTY(VisibleAnywhere, Category = FIKRigBone)
+	FName ParentName;
+
+	FIKRigBone(FName InName = NAME_None, FName InParentName = NAME_None)
+		: Name(InName)
+		, ParentName(InParentName)
+	{
+	}
+
+private:
+	// these are built when rebuilt
+	// don't rely on them in editor
+	UPROPERTY(transient)
+	TArray<int32> Children;
+
+	friend FIKRigHierarchy;
+};
+
+USTRUCT()
+struct IKRIG_API FIKRigHierarchy
+{
+	GENERATED_BODY()
+
+private:
+	UPROPERTY(VisibleAnywhere, Category=FIKRigHierarchy)
+	TArray<FIKRigBone> Bones;
+
+	// cached data for runtime is here, so that it can be shared by all the instances
+	// this should be required to build only once during runtime
+	UPROPERTY(transient)
+	TArray<int32> ParentIndices;
+
+	// from name to index of Relationship
+	UPROPERTY(transient)
+	TMap<FName, int32> RuntimeNameLookupTable;
+
+	/***** runtime operators  ****/
+	// rebuild cache data if required
+	void RebuildCacheData()
+	{
+		RuntimeNameLookupTable.Reset();
+		ParentIndices.Reset();
+
+		if (Bones.Num() > 0)
+		{
+			for (int32 Index = 0; Index < Bones.Num(); ++Index)
+			{
+				RuntimeNameLookupTable.Add(Bones[Index].Name) = Index;
+				Bones[Index].Children = FindChildren(Index);
+			}
+
+			ParentIndices.AddUninitialized(Bones.Num());
+			for (int32 Index = 0; Index < Bones.Num(); ++Index)
+			{
+				if (Bones[Index].ParentName != NAME_None)
+				{
+					int32* ParentIndex = RuntimeNameLookupTable.Find(Bones[Index].ParentName);
+					if (ParentIndex)
+					{
+						ParentIndices[Index] = *ParentIndex;
+					}
+					else
+					{
+						// should not happen
+						ensureMsgf(false, TEXT("IKRig : [%s]'s parent [%s] not found. Suspect data issue."), *Bones[Index].Name.ToString(), *Bones[Index].ParentName.ToString());
+						ParentIndices[Index] = INDEX_NONE;
+					}
+				}
+				else
+				{
+					ParentIndices[Index] = INDEX_NONE;
+				}
+			}
+		}
+
+	}
+
+public: 
+	bool IsValidIndex(int32 Index) const
+	{
+		return Bones.IsValidIndex(Index);
+	}
+
+	int32 GetNum() const
+	{
+		return Bones.Num();
+	}
+
+	int32 GetParentIndex(int32 Index) const
+	{
+		ValidateRuntimeData();
+
+		if (ParentIndices.IsValidIndex(Index))
+		{
+			return ParentIndices[Index];
+		}
+
+		return INDEX_NONE;
+	}
+
+	int32 GetParentIndex(const FName& InName) const
+	{
+		ValidateRuntimeData();
+
+		if (InName != NAME_None)
+		{
+			const int32* Found = RuntimeNameLookupTable.Find(InName);
+			if (Found)
+			{
+				return GetParentIndex(*Found);
+			}
+		}
+
+		return INDEX_NONE;
+	}
+
+	int32 GetIndex(const FName& InName) const
+	{
+		ValidateRuntimeData();
+
+		if (InName != NAME_None)
+		{
+			const int32* Found = RuntimeNameLookupTable.Find(InName);
+			if (Found)
+			{
+				return *Found;
+			}
+		}
+
+		return INDEX_NONE;
+	}
+
+	TArray<int32> FindChildren(int32 Index) const;
+
+private:
+	void ValidateRuntimeData() const
+	{
+		ensure (RuntimeNameLookupTable.Num() == Bones.Num() && ParentIndices.Num() == Bones.Num());
+	}
+
+	//slow function, this is only used for editor or for cache only
+	//  don't use this for runtime purpose, it'll be cached with Bones
+	TArray<int32> FindIndicesByParentName(const FName& InParentName) const
+	{
+		TArray<int32> Indices;
+
+		for (int32 Index = 0; Index < Bones.Num(); ++Index)
+		{
+			if (Bones[Index].ParentName == InParentName)
+			{
+				Indices.Add(Index);
+			}
+		}
+		return Indices;
+	}
+
+#if WITH_EDITOR
+	int32 FindIndexFromBoneArray(const FName& InName) const
+	{
+		return Bones.IndexOfByPredicate([&](const FIKRigBone& Item) { return Item.Name == InName; });
+	}
+#endif// WITH_EDITOR
+
+	friend class UIKRigDefinition;
+};
+
diff --git a/Engine/Plugins/Animation/IKRig/Source/IKRig/Public/IKRigProcessor.h b/Engine/Plugins/Animation/IKRig/Source/IKRig/Public/IKRigProcessor.h
new file mode 100644
--- /dev/null
+++ b/Engine/Plugins/Animation/IKRig/Source/IKRig/Public/IKRigProcessor.h
@@ -0,0 +1,80 @@
+// Copyright Epic Games, Inc. All Rights Reserved.
+
+/**
+ * Contains IKRig Definition runtime code. This takes IKRigDefinition asset as an input
+ *
+ */
+
+#pragma once
+
+#include "CoreMinimal.h"
+#include "UObject/ObjectMacros.h"
+#include "UObject/Object.h"
+#include "IKRigDataTypes.h"
+#include "IKRigProcessor.generated.h"
+
+
+class UIKRigConstraintSolver;
+class UIKRigSolver;
+
+UCLASS(BlueprintType)
+class IKRIG_API UIKRigProcessor : public UObject
+{
+	GENERATED_BODY()
+
+private:
+	// skeleton
+	// this just is a relationship. Parent/child relationship
+	// sorted by parent to children
+	UPROPERTY(EditAnywhere, Category = "RigDefinition")
+	UIKRigDefinition* RigDefinition = nullptr;
+
+	// reference transform
+	UPROPERTY(transient)
+	FIKRigTransform ReferenceTransform;
+
+	// solvers
+	UPROPERTY(transient)
+	TArray<UIKRigSolver*> Solvers;
+
+	// constraint solver
+	UPROPERTY(transient)
+	UIKRigConstraintSolver* CostraintSolver;
+
+	// goals - copy from RigDefinition
+	// modified in runtime
+	UPROPERTY(transient)
+	TMap<FName, FIKRigGoal> IKGoals;
+
+	FIKRigTransformModifier TransformModifier;
+
+	bool bInitialized = false;
+
+public: 
+
+	// IKRigProcessor implementation functions
+	// Set IKRigDefinition 
+	// if bInitialize is true, it will use Default Ref transform of IKRigDefinition to initialize
+	void SetIKRigDefinition(UIKRigDefinition* InRigDefinition, bool bInitialize = false);
+	void Initialize(const FIKRigTransform& InRefTransform);
+	void Reinitialize();
+	void Terminate();
+	void Solve();
+
+	// update goal functions
+	void SetGoalPosition(const FName& GoalName, const FVector& InPosition);
+	void SetGoalRotation(const FName& GoalName, const FRotator& InRotation);
+	void SetGoalTarget(const FName& GoalName, const FIKRigTarget& InTarget);
+	void GetGoals(TArray<FName>& OutNames) const;
+	FIKRigTransformModifier& GetIKRigTransformModifier();
+	const FIKRigHierarchy* GetHierarchy() const;
+
+	void ResetToRefPose();
+private:
+	FIKRigTarget* FindGoal(const FName& GoalName);
+
+	// delegate functions
+	bool GoalGetter(const FName& InGoalName, FIKRigTarget& OutTarget);
+	const FIKRigTransform& GetRefPoseGetter();
+};
+
diff --git a/Engine/Plugins/Animation/IKRig/Source/IKRig/Public/IKRigSolver.h b/Engine/Plugins/Animation/IKRig/Source/IKRig/Public/IKRigSolver.h
new file mode 100644
--- /dev/null
+++ b/Engine/Plugins/Animation/IKRig/Source/IKRig/Public/IKRigSolver.h
@@ -0,0 +1,63 @@
+// Copyright Epic Games, Inc. All Rights Reserved.
+
+/**
+ * Contains IKRig Solver Definition 
+ *
+ */
+
+#pragma once
+
+#include "CoreMinimal.h"
+#include "UObject/ObjectMacros.h"
+#include "UObject/Object.h"
+#include "IKRigSolver.generated.h"
+
+struct FIKRigTarget;
+struct FIKRigTransform;
+struct FIKRigTransformModifier;
+class UIKRigSolverDefinition;
+
+// run time processor 
+UCLASS(Abstract, hidecategories = UObject)
+class IKRIG_API UIKRigSolver : public UObject
+{
+	GENERATED_BODY()
+
+public: 
+	/** Required delegates to run this solver */
+	DECLARE_DELEGATE_RetVal(const FIKRigTransform&, FIKRigTransformGetter);
+	DECLARE_DELEGATE_RetVal_TwoParams(bool, FIKRigGoalGetter, const FName& /*InGoalName*/, FIKRigTarget& /*OutTarget*/);
+//	DECLARE_DELEGATE_TwoParams(FSolveConstraint, TArray<FName>& /*Constraint*/, FIKRigTransform& /*CurrentTransform*/);
+
+protected:
+
+	UPROPERTY(VisibleAnywhere, Category = "Definition")
+	UIKRigSolverDefinition* SolverDefinition;
+
+public:
+	UPROPERTY(EditAnywhere, Category = "Definition")
+	bool bEnabled = true;
+
+	// input hierarchy and ref pose? 
+	void Init(UIKRigSolverDefinition* InSolverDefinition, FIKRigTransformGetter InRefPoseGetter, FIKRigGoalGetter InGoalGetter/*, FSolveConstraint& InConstraintHandler*/);
+
+	// input : goal getter or goals
+	// output : modified pose - GlobalTransforms
+	// or use SolverInternal function
+	void Solve(FIKRigTransformModifier& InOutGlobalTransform);
+
+protected:
+	virtual void InitInternal() {};
+	virtual void SolveInternal(FIKRigTransformModifier& InOutGlobalTransform) {};
+	virtual bool IsSolverActive() const;
+
+	bool GetTaskTarget(const FName& TaskName, FIKRigTarget& OutTarget) const;
+	const FIKRigTransform& GetReferencePose() const;
+private:
+	// delegate
+	FIKRigTransformGetter RefPoseGetter;
+	FIKRigGoalGetter GoalGetter;
+	
+//	FSolveConstraint ConstraintHandler;
+};
+
diff --git a/Engine/Plugins/Animation/IKRig/Source/IKRig/Public/IKRigSolverDefinition.h b/Engine/Plugins/Animation/IKRig/Source/IKRig/Public/IKRigSolverDefinition.h
new file mode 100644
--- /dev/null
+++ b/Engine/Plugins/Animation/IKRig/Source/IKRig/Public/IKRigSolverDefinition.h
@@ -0,0 +1,84 @@
+// Copyright Epic Games, Inc. All Rights Reserved.
+
+/**
+ * Contains IKRig Solver Definition 
+ *
+ */
+
+#pragma once
+
+#include "CoreMinimal.h"
+#include "UObject/ObjectMacros.h"
+#include "UObject/Object.h"
+#include "IKRigSolverDefinition.generated.h"
+
+class UIKRigSolver;
+
+// do we save data between Goal it cares?
+// data type
+UCLASS(Blueprintable, BlueprintType)
+class IKRIG_API UIKRigSolverDefinition : public UObject
+{
+	GENERATED_BODY()
+protected:
+	UPROPERTY(EditAnywhere, Category = "Definition")
+	FString DisplayName;
+
+	// effector name to goals name map
+	UPROPERTY(VisibleAnywhere, Category = "Definition")
+	TMap<FName, FName> TaskToGoal;
+
+	// this is name of constraints it's used
+	// use Delegate from IKRigDefinition to query and solve
+	UPROPERTY(EditAnywhere, Category = "Definition")
+	TArray<FName> Constraints;
+
+	UPROPERTY(VisibleAnywhere, Category = "Definition")
+	TSubclassOf<UIKRigSolver> ExecutionClass;
+
+public:
+	const TMap<FName, FName>& GetTaskToGoal() const
+	{
+		return TaskToGoal;
+	}
+
+	TSubclassOf<UIKRigSolver> GetExecutionClass() 
+	{		
+		return ExecutionClass;		
+	}
+
+	// if you have auto configuration option, return true;
+	virtual bool CanAutoConfigure() const
+	{
+		return false;
+	}
+
+	virtual void AutoConfigure()
+	{
+		ensureMsgf(false, TEXT("(%s)Auto Configuration needs to be impelmeneted."), *GetName());
+	}
+
+	void CollectGoals(TArray<FName>& OutGoals);
+
+private:
+#if WITH_EDITOR
+	void RenameGoal(const FName& OldName, const FName& NewName);
+	// get unique name delegate by IKRigDefinition
+	void EnsureUniqueGoalName(FName& InOutGoalName) const;
+
+	DECLARE_MULTICAST_DELEGATE(FGoalNeedsUpdate);
+	FGoalNeedsUpdate GoalNeedsUpdateDelegate;
+
+protected:
+	FName CreateUniqueGoalName(const TCHAR* Suffix) const;
+	void OnGoalHasBeenUpdated();
+
+	virtual void UpdateTaskList() {};
+#endif // WITH_EDITOR
+
+	/// BEGIN UObject
+	virtual void PostLoad() override;
+	/// END UObject
+	friend class UIKRigController;
+};
+
diff --git a/Engine/Plugins/Animation/IKRig/Source/IKRig/Public/Solvers/TransformSolverDefinition.h b/Engine/Plugins/Animation/IKRig/Source/IKRig/Public/Solvers/TransformSolverDefinition.h
new file mode 100644
--- /dev/null
+++ b/Engine/Plugins/Animation/IKRig/Source/IKRig/Public/Solvers/TransformSolverDefinition.h
@@ -0,0 +1,39 @@
+// Copyright Epic Games, Inc. All Rights Reserved.
+
+/**
+ * Contains IKRig Solver Definition 
+ *
+ */
+
+#pragma once
+
+#include "IKRigSolverDefinition.h"
+#include "TransformSolverDefinition.generated.h"
+
+UCLASS(BlueprintType)
+class IKRIG_API UTransformSolverDefinition : public UIKRigSolverDefinition
+{
+	GENERATED_BODY()
+
+public:
+	UTransformSolverDefinition();
+
+	// align is terrible on this
+	UPROPERTY(EditAnywhere, Category = "Solver")
+	bool	bEnablePosition = true;
+
+	UPROPERTY(EditAnywhere, Category = "Solver")
+	bool	bEnableRotation = true;
+
+private:
+#if WITH_EDITOR
+	virtual void UpdateTaskList() override;
+
+	// UObject interface
+	virtual void PostEditChangeChainProperty(struct FPropertyChangedChainEvent& PropertyChangedEvent) override;
+	// UObject interface
+#endif // WITH_EDITOR
+
+public:
+	static const FName TransformTarget;
+};
\ No newline at end of file
diff --git a/Engine/Plugins/Animation/IKRig/Source/IKRigDeveloper/IKRigDeveloper.Build.cs b/Engine/Plugins/Animation/IKRig/Source/IKRigDeveloper/IKRigDeveloper.Build.cs
new file mode 100644
--- /dev/null
+++ b/Engine/Plugins/Animation/IKRig/Source/IKRigDeveloper/IKRigDeveloper.Build.cs
@@ -0,0 +1,39 @@
+// Copyright Epic Games, Inc. All Rights Reserved.
+
+namespace UnrealBuildTool.Rules
+{
+    public class IKRigDeveloper : ModuleRules
+    {
+        public IKRigDeveloper(ReadOnlyTargetRules Target) : base(Target)
+        {
+            // Copying some these from ControlRig.build.cs, our deps are likely leaner
+            // and therefore these could be pruned if needed:
+            PrivateDependencyModuleNames.AddRange(
+                new string[]
+                {
+					"Core",
+					"CoreUObject",
+					"Engine",
+					"IKRig",
+				}
+            );
+
+            PublicDependencyModuleNames.AddRange(
+                new string[]
+                {
+                }
+            );
+
+            if (Target.bBuildEditor == true)
+            {
+                PrivateDependencyModuleNames.AddRange(
+                    new string[]
+                    {
+						"AnimGraph",
+                        "BlueprintGraph",
+                    }
+                );
+            }
+        }
+    }
+}
diff --git a/Engine/Plugins/Animation/IKRig/Source/IKRigDeveloper/Private/AnimGraphNode_IKRig.cpp b/Engine/Plugins/Animation/IKRig/Source/IKRigDeveloper/Private/AnimGraphNode_IKRig.cpp
new file mode 100644
--- /dev/null
+++ b/Engine/Plugins/Animation/IKRig/Source/IKRigDeveloper/Private/AnimGraphNode_IKRig.cpp
@@ -0,0 +1,73 @@
+// Copyright Epic Games, Inc. All Rights Reserved.
+
+#include "AnimGraphNode_IKRig.h"
+#include "Animation/AnimInstance.h"
+
+/////////////////////////////////////////////////////
+// UAnimGraphNode_IKRig 
+
+#define LOCTEXT_NAMESPACE "AnimGraphNode_IKRig"
+const FName UAnimGraphNode_IKRig::AnimModeName(TEXT("IKRig.IKRigEditor.IKRigEditMode"));
+
+void UAnimGraphNode_IKRig::Draw(FPrimitiveDrawInterface* PDI, USkeletalMeshComponent* PreviewSkelMeshComp) const
+{
+	if (PreviewSkelMeshComp)
+	{
+		if (FAnimNode_IKRig* ActiveNode = GetActiveInstanceNode<FAnimNode_IKRig>(PreviewSkelMeshComp->GetAnimInstance()))
+		{
+			//ActiveNode->ConditionalDebugDraw(PDI, PreviewSkelMeshComp);
+		}
+	}
+}
+
+FText UAnimGraphNode_IKRig::GetNodeTitle(ENodeTitleType::Type TitleType) const
+{
+	return LOCTEXT("AnimGraphNode_IKRig_Title", "IK Rig");
+}
+
+void UAnimGraphNode_IKRig::CopyNodeDataToPreviewNode(FAnimNode_Base* InPreviewNode)
+{
+	FAnimNode_IKRig* IKRigNode = static_cast<FAnimNode_IKRig*>(InPreviewNode);
+
+	// copies Pin values from the internal node to get data which are not compiled yet
+	//Fabrik->EffectorTransform = Node.EffectorTransform;
+}
+
+FEditorModeID UAnimGraphNode_IKRig::GetEditorMode() const
+{
+// we have a crash on shutting it down because it's not unregistering
+// https://epic.slack.com/archives/C024FSXHJ/p1606256645266900
+// until we know what is solution, I'm disabling it
+//	return AnimModeName;
+
+	return NAME_None;
+}
+
+void UAnimGraphNode_IKRig::PostEditChangeProperty(struct FPropertyChangedEvent& PropertyChangedEvent)
+{
+	const FName PropertyName = (PropertyChangedEvent.Property ? PropertyChangedEvent.Property->GetFName() : NAME_None);
+
+	// Reconstruct node to show updates to PinFriendlyNames.
+	if ((PropertyName == GET_MEMBER_NAME_STRING_CHECKED(FAnimNode_IKRig, RigDefinitionAsset)))
+	{
+		Node.RebuildGoalList();
+		ReconstructNode();
+	}
+}
+
+void UAnimGraphNode_IKRig::CustomizePinData(UEdGraphPin* Pin, FName SourcePropertyName, int32 ArrayIndex) const
+{
+	Super::CustomizePinData(Pin, SourcePropertyName, ArrayIndex);
+
+	FString PinString = Pin->PinName.ToString();
+	FString GoalPinString = GET_MEMBER_NAME_STRING_CHECKED(FAnimNode_IKRig, GoalTransforms);
+	const int32 GoalStringLen = GoalPinString.Len();
+	if (PinString.Len() > GoalStringLen && PinString.Left(GoalStringLen) == GoalPinString)
+	{
+		if (!Pin->bHidden)
+		{
+			Pin->PinFriendlyName = FText::FromName(Node.GetGoalName(ArrayIndex));
+		}
+	}
+}
+#undef LOCTEXT_NAMESPACE
diff --git a/Engine/Plugins/Animation/IKRig/Source/IKRigDeveloper/Private/IKRigDeveloper.cpp b/Engine/Plugins/Animation/IKRig/Source/IKRigDeveloper/Private/IKRigDeveloper.cpp
new file mode 100644
--- /dev/null
+++ b/Engine/Plugins/Animation/IKRig/Source/IKRigDeveloper/Private/IKRigDeveloper.cpp
@@ -0,0 +1,25 @@
+// Copyright Epic Games, Inc. All Rights Reserved.
+
+#include "IKRigDeveloper.h"
+#include "Modules/ModuleManager.h"
+
+#define LOCTEXT_NAMESPACE "IKRigDeveloperModule"
+
+class FIKRigDeveloperModule : public IIKRigDeveloperModule
+{
+public:
+	/** IModuleInterface implementation */
+	virtual void StartupModule() override;
+	virtual void ShutdownModule() override;
+};
+
+void FIKRigDeveloperModule::StartupModule()
+{
+}
+
+void FIKRigDeveloperModule::ShutdownModule()
+{
+}
+
+IMPLEMENT_MODULE(FIKRigDeveloperModule, IKRigDeveloper)
+#undef LOCTEXT_NAMESPACE
diff --git a/Engine/Plugins/Animation/IKRig/Source/IKRigDeveloper/Public/AnimGraphNode_IKRig.h b/Engine/Plugins/Animation/IKRig/Source/IKRigDeveloper/Public/AnimGraphNode_IKRig.h
new file mode 100644
--- /dev/null
+++ b/Engine/Plugins/Animation/IKRig/Source/IKRigDeveloper/Public/AnimGraphNode_IKRig.h
@@ -0,0 +1,37 @@
+// Copyright Epic Games, Inc. All Rights Reserved.
+
+#pragma once
+
+#include "CoreMinimal.h"
+#include "UObject/ObjectMacros.h"
+#include "AnimGraphNode_Base.h"
+#include "AnimNodes/AnimNode_IKRig.h"
+#include "AnimGraphNode_IKRig.generated.h"
+
+class FPrimitiveDrawInterface;
+class USkeletalMeshComponent;
+
+// Editor node for IKRig 
+UCLASS()
+class IKRIGDEVELOPER_API UAnimGraphNode_IKRig : public UAnimGraphNode_Base
+{
+	GENERATED_BODY()
+
+	UPROPERTY(EditAnywhere, Category = Settings)
+	FAnimNode_IKRig Node;
+
+public:
+	// UEdGraphNode interface
+	virtual FText GetNodeTitle(ENodeTitleType::Type TitleType) const override;
+	virtual void PostEditChangeProperty(struct FPropertyChangedEvent& PropertyChangedEvent) override;
+	// End of UEdGraphNode interface
+
+	// UAnimGraphNode_Base interface
+	virtual void CopyNodeDataToPreviewNode(FAnimNode_Base* AnimNode) override;
+	virtual FEditorModeID GetEditorMode() const override;
+	virtual void Draw(FPrimitiveDrawInterface* PDI, USkeletalMeshComponent* PreviewSkelMeshComp) const override;
+	virtual void CustomizePinData(UEdGraphPin* Pin, FName SourcePropertyName, int32 ArrayIndex) const override;
+	// End of UAnimGraphNode_Base interface
+
+	static const FName AnimModeName;
+};
diff --git a/Engine/Plugins/Animation/IKRig/Source/IKRigDeveloper/Public/IKRigDeveloper.h b/Engine/Plugins/Animation/IKRig/Source/IKRigDeveloper/Public/IKRigDeveloper.h
new file mode 100644
--- /dev/null
+++ b/Engine/Plugins/Animation/IKRig/Source/IKRigDeveloper/Public/IKRigDeveloper.h
@@ -0,0 +1,12 @@
+// Copyright Epic Games, Inc. All Rights Reserved.
+
+#pragma once
+
+#include "CoreMinimal.h"
+#include "Modules/ModuleInterface.h"
+
+class IIKRigDeveloperModule : public IModuleInterface
+{
+public:
+};
+
diff --git a/Engine/Plugins/Animation/IKRig/Source/IKRigEditor/IKRigEditor.Build.cs b/Engine/Plugins/Animation/IKRig/Source/IKRigEditor/IKRigEditor.Build.cs
new file mode 100644
--- /dev/null
+++ b/Engine/Plugins/Animation/IKRig/Source/IKRigEditor/IKRigEditor.Build.cs
@@ -0,0 +1,46 @@
+// Copyright Epic Games, Inc. All Rights Reserved.
+
+using System.IO;
+
+namespace UnrealBuildTool.Rules
+{
+    public class IKRigEditor : ModuleRules
+    {
+        public IKRigEditor(ReadOnlyTargetRules Target) : base(Target)
+        {
+            PublicDependencyModuleNames.AddRange(
+                new string[]
+                {
+                }
+            );
+
+			PrivateDependencyModuleNames.AddRange(
+                new string[]
+                {
+                    "Core",
+                    "CoreUObject",
+                    "Engine",
+					"UnrealEd",
+					"IKRig",
+					"IKRigDeveloper",
+					"AssetTools",
+					"SlateCore",
+					"Slate",
+					"InputCore",
+					"PropertyEditor",
+					"EditorStyle",
+					"EditorFramework",
+					"BlueprintGraph",
+					"AnimGraph",
+				}
+            );
+
+            PrivateIncludePathModuleNames.AddRange(
+                new string[]
+                {
+				}
+            );
+
+        }
+    }
+}
diff --git a/Engine/Plugins/Animation/IKRig/Source/IKRigEditor/Private/AnimMode/IKRigEditMode.cpp b/Engine/Plugins/Animation/IKRig/Source/IKRigEditor/Private/AnimMode/IKRigEditMode.cpp
new file mode 100644
--- /dev/null
+++ b/Engine/Plugins/Animation/IKRig/Source/IKRigEditor/Private/AnimMode/IKRigEditMode.cpp
@@ -0,0 +1,48 @@
+// Copyright Epic Games, Inc. All Rights Reserved.
+
+#include "IKRigEditMode.h"
+#include "AnimGraphNode_IKRig.h"
+#include "IPersonaPreviewScene.h"
+#include "Animation/DebugSkelMeshComponent.h"
+
+void FIKRigEditMode::EnterMode(class UAnimGraphNode_Base* InEditorNode, struct FAnimNode_Base* InRuntimeNode)
+{
+	RuntimeNode = static_cast<FAnimNode_IKRig*>(InRuntimeNode);
+	GraphNode = CastChecked<UAnimGraphNode_IKRig>(InEditorNode);
+
+	FAnimNodeEditMode::EnterMode(InEditorNode, InRuntimeNode);
+}
+
+void FIKRigEditMode::ExitMode()
+{
+	RuntimeNode = nullptr;
+	GraphNode = nullptr;
+
+	FAnimNodeEditMode::ExitMode();
+}
+
+FVector FIKRigEditMode::GetWidgetLocation() const
+{
+	USkeletalMeshComponent* SkelComp = GetAnimPreviewScene().GetPreviewMeshComponent();
+
+// 	FBoneSocketTarget& Target = RuntimeNode->EffectorTarget;
+// 	FVector Location = RuntimeNode->EffectorTransform.GetLocation();
+// 	EBoneControlSpace Space = RuntimeNode->EffectorTransformSpace;
+// 	FVector WidgetLoc = ConvertWidgetLocation(SkelComp, RuntimeNode->ForwardedPose, Target, Location, Space);
+	return SkelComp->GetComponentToWorld().GetLocation();
+}
+
+UE::Widget::EWidgetMode FIKRigEditMode::GetWidgetMode() const
+{
+	// allow translation all the time for effectot target
+	return UE::Widget::WM_Translate;
+}
+
+void FIKRigEditMode::DoTranslation(FVector& InTranslation)
+{
+// 	USkeletalMeshComponent* SkelComp = GetAnimPreviewScene().GetPreviewMeshComponent();
+// 	FVector Offset = ConvertCSVectorToBoneSpace(SkelComp, InTranslation, RuntimeNode->ForwardedPose, RuntimeNode->EffectorTarget, RuntimeNode->EffectorTransformSpace);
+// 
+// 	RuntimeNode->EffectorTransform.AddToTranslation(Offset);
+// 	GraphNode->Node.EffectorTransform.SetTranslation(RuntimeNode->EffectorTransform.GetTranslation());
+}
diff --git a/Engine/Plugins/Animation/IKRig/Source/IKRigEditor/Private/AnimMode/IKRigEditMode.h b/Engine/Plugins/Animation/IKRig/Source/IKRigEditor/Private/AnimMode/IKRigEditMode.h
new file mode 100644
--- /dev/null
+++ b/Engine/Plugins/Animation/IKRig/Source/IKRigEditor/Private/AnimMode/IKRigEditMode.h
@@ -0,0 +1,22 @@
+// Copyright Epic Games, Inc. All Rights Reserved.
+
+#pragma once
+
+#include "CoreMinimal.h"
+#include "UnrealWidgetFwd.h"
+#include "AnimNodeEditMode.h"
+
+class FIKRigEditMode : public FAnimNodeEditMode
+{
+public:
+	/** IAnimNodeEditMode interface */
+	virtual void EnterMode(class UAnimGraphNode_Base* InEditorNode, struct FAnimNode_Base* InRuntimeNode) override;
+	virtual void ExitMode() override;
+	virtual FVector GetWidgetLocation() const override;
+	virtual UE::Widget::EWidgetMode GetWidgetMode() const override;
+	virtual void DoTranslation(FVector& InTranslation) override;
+
+private:
+	struct FAnimNode_IKRig* RuntimeNode;
+	class UAnimGraphNode_IKRig* GraphNode;
+};
diff --git a/Engine/Plugins/Animation/IKRig/Source/IKRigEditor/Private/AssetTypeActions_IKRigDefinition.cpp b/Engine/Plugins/Animation/IKRig/Source/IKRigEditor/Private/AssetTypeActions_IKRigDefinition.cpp
new file mode 100644
--- /dev/null
+++ b/Engine/Plugins/Animation/IKRig/Source/IKRigEditor/Private/AssetTypeActions_IKRigDefinition.cpp
@@ -0,0 +1,31 @@
+// Copyright Epic Games, Inc. All Rights Reserved.
+
+#include "AssetTypeActions_IKRigDefinition.h"
+#include "IKRigDefinition.h"
+
+//#include "ThumbnailRendering/SceneThumbnailInfo.h"
+//#include "ToolMenus.h"
+
+#define LOCTEXT_NAMESPACE "AssetTypeActions"
+
+UClass* FAssetTypeActions_IKRigDefinition::GetSupportedClass() const
+{
+	return UIKRigDefinition::StaticClass();
+}
+
+// UThumbnailInfo* FAssetTypeActions_IKRigDefinition::GetThumbnailInfo(UObject* Asset) const
+// {
+// 	UChaosSolver* ChaosSolver = CastChecked<UChaosSolver>(Asset);
+// 	return NewObject<USceneThumbnailInfo>(ChaosSolver, NAME_None, RF_Transactional);
+// }
+
+void FAssetTypeActions_IKRigDefinition::GetActions(const TArray<UObject*>& InObjects, FToolMenuSection& Section)
+{
+	FAssetTypeActions_Base::GetActions(InObjects, Section);
+}
+
+// void FAssetTypeActions_IKRigDefinition::OpenAssetEditor(const TArray<UObject*>& InObjects, TSharedPtr<IToolkitHost> EditWithinLevelEditor)
+// {
+// }
+
+#undef LOCTEXT_NAMESPACE
diff --git a/Engine/Plugins/Animation/IKRig/Source/IKRigEditor/Private/AssetTypeActions_IKRigDefinition.h b/Engine/Plugins/Animation/IKRig/Source/IKRigEditor/Private/AssetTypeActions_IKRigDefinition.h
new file mode 100644
--- /dev/null
+++ b/Engine/Plugins/Animation/IKRig/Source/IKRigEditor/Private/AssetTypeActions_IKRigDefinition.h
@@ -0,0 +1,20 @@
+// Copyright Epic Games, Inc. All Rights Reserved.
+
+#pragma once
+
+#include "CoreMinimal.h"
+#include "Toolkits/IToolkitHost.h"
+#include "AssetTypeActions_Base.h"
+
+class FAssetTypeActions_IKRigDefinition : public FAssetTypeActions_Base
+{
+public:
+	// IAssetTypeActions Implementation
+	virtual FText GetName() const override { return NSLOCTEXT("AssetTypeActions", "AssetTypeActions_IKRigDefinition", "IK Rig Definition"); }
+	virtual void GetActions(const TArray<UObject*>& InObjects, struct FToolMenuSection& Section) override;
+	virtual FColor GetTypeColor() const override { return FColor(255, 192, 128); }
+	virtual UClass* GetSupportedClass() const override;
+	//virtual void OpenAssetEditor(const TArray<UObject*>& InObjects, TSharedPtr<class IToolkitHost> EditWithinLevelEditor = TSharedPtr<IToolkitHost>()) override;
+	virtual uint32 GetCategories() override { return EAssetTypeCategories::Animation; }
+	//virtual UThumbnailInfo* GetThumbnailInfo(UObject* Asset) const override;
+};
diff --git a/Engine/Plugins/Animation/IKRig/Source/IKRigEditor/Private/IKRigController.cpp b/Engine/Plugins/Animation/IKRig/Source/IKRigEditor/Private/IKRigController.cpp
new file mode 100644
--- /dev/null
+++ b/Engine/Plugins/Animation/IKRig/Source/IKRigEditor/Private/IKRigController.cpp
@@ -0,0 +1,306 @@
+// Copyright Epic Games, Inc. All Rights Reserved.
+
+/**
+ * IKController implementation class
+ *
+ */
+
+#include "IKRigController.h"
+#include "IKRigDefinition.h"
+#include "AnimationRuntime.h"
+#include "IKRigSolverDefinition.h"
+
+void UIKRigController::BeginDestroy() 
+{
+	if (IKRigDefinition)
+	{
+		for (int32 Index = 0; Index < IKRigDefinition->SolverDefinitions.Num(); ++Index)
+		{
+			UninitializeIKRigSolverDefinition(IKRigDefinition->SolverDefinitions[Index]);
+		}
+	}
+
+	Super::BeginDestroy();
+}
+
+ // IKRigDefinition set up
+void UIKRigController::SetIKRigDefinition(UIKRigDefinition* InIKRigDefinition)
+{
+	if (IKRigDefinition)
+	{
+		for (int32 Index = 0; Index < IKRigDefinition->SolverDefinitions.Num(); ++Index)
+		{
+			UninitializeIKRigSolverDefinition(IKRigDefinition->SolverDefinitions[Index]);
+		}
+	}
+
+	IKRigDefinition = InIKRigDefinition;
+
+	if (IKRigDefinition)
+	{
+		for (int32 Index = 0; Index < IKRigDefinition->SolverDefinitions.Num(); ++Index)
+		{
+			InitializeIKRigSolverDefinition(IKRigDefinition->SolverDefinitions[Index]);
+		}
+	}
+}
+
+void UIKRigController::InitializeIKRigSolverDefinition(UIKRigSolverDefinition* SolverDef)
+{
+	if (SolverDef)
+	{
+		// we want to register this controller delegate
+		// if there are multiple controllers managing one IKRigDefinition
+		// we want to ensure that works 
+		FDelegateHandle DelegateHandle = SolverDef->GoalNeedsUpdateDelegate.AddUObject(this, &UIKRigController::UpdateGoal);
+		SolverDelegateHandles.Add(SolverDef) = DelegateHandle;
+
+		SolverDef->UpdateTaskList();
+	}
+}
+
+void UIKRigController::UninitializeIKRigSolverDefinition(UIKRigSolverDefinition* SolverDef)
+{
+	if (SolverDef)
+	{
+		FDelegateHandle* DelegateHandle = SolverDelegateHandles.Find(SolverDef);
+		if (ensure(DelegateHandle))
+		{
+			SolverDef->GoalNeedsUpdateDelegate.Remove(*DelegateHandle);
+		}
+	}
+}
+
+// hierarchy operators
+void UIKRigController::SetSkeleton(const FReferenceSkeleton& InSkeleton)
+{
+	if (IKRigDefinition)
+	{
+		IKRigDefinition->ResetHierarchy();
+
+		const TArray<FMeshBoneInfo>& RefBoneInfo = InSkeleton.GetRefBoneInfo();
+		TArray<FTransform> RefPoseInCS;
+		FAnimationRuntime::FillUpComponentSpaceTransforms(InSkeleton, InSkeleton.GetRefBonePose(), RefPoseInCS);
+		ensure(RefPoseInCS.Num() == RefBoneInfo.Num());
+		for (int32 Index=0; Index<RefBoneInfo.Num(); ++Index)
+		{
+			int32 ParentIndex = RefBoneInfo[Index].ParentIndex;
+			ensure (IKRigDefinition->AddBone(RefBoneInfo[Index].Name, (ParentIndex != INDEX_NONE)? RefBoneInfo[ParentIndex].Name : NAME_None, RefPoseInCS[Index]));
+		}
+
+		ensure(IKRigDefinition->Hierarchy.GetNum() == IKRigDefinition->ReferencePose.GetNum());
+	}
+}
+
+bool UIKRigController::AddBone(const FName& InName, const FName& InParent, const FTransform& InGlobalTransform)
+{
+	if (IKRigDefinition)
+	{
+		return IKRigDefinition->AddBone(InName, InParent, InGlobalTransform);
+	}
+
+	return false;
+}
+
+bool UIKRigController::RemoveBone(const FName& InName)
+{
+	if (IKRigDefinition)
+	{
+		return IKRigDefinition->RemoveBone(InName);
+	}
+
+	return false;
+}
+bool UIKRigController::RenameBone(const FName& InOldName, const FName& InNewName)
+{
+	if (IKRigDefinition)
+	{
+		return IKRigDefinition->RenameBone(InOldName, InNewName);
+	}
+
+	return false;
+}
+bool UIKRigController::ReparentBone(const FName& InName, const FName& InNewParent)
+{
+	if (IKRigDefinition)
+	{
+		return IKRigDefinition->ReparentBone(InName, InNewParent);
+	}
+
+	return false;
+}
+void UIKRigController::ResetHierarchy()
+{
+	if (IKRigDefinition)
+	{
+		IKRigDefinition->ResetHierarchy();
+	}
+}
+
+// solver operators
+UIKRigSolverDefinition* UIKRigController::AddSolver(TSubclassOf<UIKRigSolverDefinition> InIKRigSolverDefinitionClass)
+{
+	if (IKRigDefinition)
+	{
+		UIKRigSolverDefinition* NewSolver = NewObject<UIKRigSolverDefinition>(IKRigDefinition, InIKRigSolverDefinitionClass);
+		check(NewSolver);
+
+		// todo: set delegate for the goal update
+		IKRigDefinition->SolverDefinitions.Add(NewSolver);
+		InitializeIKRigSolverDefinition(NewSolver);
+		IKRigDefinition->UpdateGoal();
+		return NewSolver;
+	}
+
+	return nullptr;
+}
+
+int32 UIKRigController::GetTotalSolverCount() const
+{
+	if (IKRigDefinition)
+	{
+		return IKRigDefinition->SolverDefinitions.Num();
+	}
+
+	return 0;
+}
+
+UIKRigSolverDefinition* UIKRigController::GetSolver(int32 Index) const
+{
+	if (IKRigDefinition && IKRigDefinition->SolverDefinitions.IsValidIndex(Index))
+	{
+		return IKRigDefinition->SolverDefinitions[Index];
+	}
+
+	return nullptr;
+}
+
+void UIKRigController::RemoveSolver(UIKRigSolverDefinition* SolverToDelete)
+{
+	if (IKRigDefinition && SolverToDelete)
+	{
+		UninitializeIKRigSolverDefinition(SolverToDelete);
+		IKRigDefinition->SolverDefinitions.Remove(SolverToDelete);
+		IKRigDefinition->UpdateGoal();
+	}
+}
+
+bool UIKRigController::ValidateSolver(UIKRigSolverDefinition* const SolverDef) const
+{
+	return (IKRigDefinition && SolverDef && IKRigDefinition->SolverDefinitions.Find(SolverDef) != INDEX_NONE);
+}
+
+void UIKRigController::UpdateGoal()
+{
+	if (IKRigDefinition)
+	{
+		IKRigDefinition->UpdateGoal();
+	}
+}
+
+void UIKRigController::AutoConfigure(UIKRigSolverDefinition* SolverDef)
+{
+	if (ValidateSolver(SolverDef))
+	{
+		SolverDef->AutoConfigure();
+	}
+}
+
+bool UIKRigController::CanAutoConfigure(UIKRigSolverDefinition* SolverDef) const
+{
+	if (ValidateSolver(SolverDef))
+	{
+		return SolverDef->CanAutoConfigure();
+	}
+
+	return false;
+}
+
+// constraint operators
+// create new profile
+void UIKRigController::CreateNewProfile(FName& InNewProfileName)
+{
+	
+}
+
+bool UIKRigController::RemoveConstraintProfile(const FName& InProfileName)
+{
+	return false;
+}
+
+void UIKRigController::RenameProfile(FName InCurrentProfileName, FName& InNewProfileName)
+{
+
+}
+
+UIKRigConstraint* UIKRigController::AddConstraint(TSubclassOf<UIKRigConstraint> NewConstraintType, FName& InOutNewName, FName InProfile /*= NAME_None*/)
+{
+	return nullptr;	
+}
+
+UIKRigConstraint* UIKRigController::GetConstraint(const FName& InProfileName, const FName& InName) const
+{
+	return nullptr;
+}
+
+bool UIKRigController::RemoveConstraint(const FName& InConstraintName)
+{
+	return false;
+}
+
+void UIKRigController::GetConstraintProfileNames(TArray<FName>& OutProfileNames) const
+{
+
+}
+
+void UIKRigController::GetConstraintNames(TArray<FName>& OutConstraintNames) const
+{
+
+}
+
+// goal operators
+void UIKRigController::QueryGoals(TArray<FName>& OutGoals) const
+{
+	if (IKRigDefinition)
+	{
+		IKRigDefinition->GetGoals().GenerateKeyArray(OutGoals);
+	}
+}
+
+void UIKRigController::RenameGoal(const FName& OldName, const FName& NewName)
+{
+	// ensure we don't have used NewDisplayName 
+	if (IKRigDefinition)
+	{
+		for (UIKRigSolverDefinition* SolverDef : IKRigDefinition->SolverDefinitions)
+		{
+			if (SolverDef)
+			{
+				SolverDef->RenameGoal(OldName, NewName);
+			}
+		}
+		
+		// update goal list
+		IKRigDefinition->UpdateGoal();
+	}
+}
+
+FIKRigGoal* UIKRigController::GetGoal(const FName& InGoalName) 
+{
+	if (IKRigDefinition)
+	{
+		return IKRigDefinition->IKGoals.Find(InGoalName);
+	}
+
+	return nullptr;
+}
+
+const FIKRigGoal* UIKRigController::GetGoal(const FName& InGoalName) const
+{
+	if (IKRigDefinition)
+	{
+		return IKRigDefinition->IKGoals.Find(InGoalName);
+	}
+
+	return nullptr;
+}
\ No newline at end of file
diff --git a/Engine/Plugins/Animation/IKRig/Source/IKRigEditor/Private/IKRigDefinitionDetails.cpp b/Engine/Plugins/Animation/IKRig/Source/IKRigEditor/Private/IKRigDefinitionDetails.cpp
new file mode 100644
--- /dev/null
+++ b/Engine/Plugins/Animation/IKRig/Source/IKRigEditor/Private/IKRigDefinitionDetails.cpp
@@ -0,0 +1,344 @@
+// Copyright Epic Games, Inc. All Rights Reservekd.
+
+#include "IKRigDefinitionDetails.h"
+#include "Widgets/Input/SButton.h"
+#include "AssetData.h"
+#include "DetailLayoutBuilder.h"
+#include "DetailCategoryBuilder.h"
+#include "IDetailsView.h"
+
+#include "IKRigDefinition.h"
+#include "IKRigController.h"
+#include "IKRigSolverDefinition.h"
+
+#include "ScopedTransaction.h"
+#include "PropertyCustomizationHelpers.h"
+#include "Kismet2/SClassPickerDialog.h"
+#include "ClassViewerFilter.h"
+#include "Widgets/Input/SEditableTextBox.h"
+
+#define LOCTEXT_NAMESPACE	"IKRigDefinitionDetails"
+
+TSharedRef<IDetailCustomization> FIKRigDefinitionDetails::MakeInstance()
+{
+	return MakeShareable(new FIKRigDefinitionDetails);
+}
+
+FIKRigDefinitionDetails::~FIKRigDefinitionDetails()
+{
+	if (ObjectChangedDelegate.IsValid())
+	{
+		FCoreUObjectDelegates::OnObjectPropertyChanged.Remove(ObjectChangedDelegate);
+	}
+}
+void FIKRigDefinitionDetails::CustomizeDetails(const TSharedPtr<IDetailLayoutBuilder>& DetailBuilder)
+{
+	DetailBuilderWeakPtr = DetailBuilder;
+	CustomizeDetails(*DetailBuilder);
+}
+
+void FIKRigDefinitionDetails::CustomizeDetails(IDetailLayoutBuilder& DetailBuilder)
+{
+	const TArray< TWeakObjectPtr<UObject> >& SelectedObjectsList = DetailBuilder.GetSelectedObjects();
+	TArray< TWeakObjectPtr<UIKRigDefinition> > SelectedIKRigDefinitions;
+
+	for (auto SelectionIt = SelectedObjectsList.CreateConstIterator(); SelectionIt; ++SelectionIt)
+	{
+		if (UIKRigDefinition* TestIKRigDefinition = Cast<UIKRigDefinition>(SelectionIt->Get()))
+		{
+			SelectedIKRigDefinitions.Add(TestIKRigDefinition);
+		}
+	}
+
+	// we only support 1 asset for now
+	if (SelectedIKRigDefinitions.Num() > 1)
+	{
+		return;
+	}
+
+	IKRigDefinition = SelectedIKRigDefinitions[0];
+
+	if (!IKRigDefinition.IsValid())
+	{
+		return;
+	}
+
+	// create controller
+	if (!IKRigController.IsValid())
+	{
+		IKRigController = TStrongObjectPtr<UIKRigController>(NewObject<UIKRigController>());
+	}
+	
+	IKRigController->SetIKRigDefinition(IKRigDefinition.Get());
+
+	ObjectChangedDelegate = FCoreUObjectDelegates::OnObjectPropertyChanged.AddRaw(this, &FIKRigDefinitionDetails::OnObjectPostEditChange);
+	/////////////////////////////////////////////////////////////////////////////////
+	// skeleton set up
+	/////////////////////////////////////////////////////////////////////////////////
+	IDetailCategoryBuilder& HierarchyCategory = DetailBuilder.EditCategory("Hierarchy");
+ 
+	SelectedAsset = IKRigDefinition->SourceAsset.Get();
+
+	HierarchyCategory.AddCustomRow(FText::FromString("ChangeSkeleton"))
+	.NameContent()
+	[
+		SNew(STextBlock)
+		.Font(IDetailLayoutBuilder::GetDetailFont())
+		.Text(LOCTEXT("SelectSourceSkeleton", "Souce Skeleton"))
+	]
+	.ValueContent()
+	[
+		SNew(SBorder)
+		.BorderImage(FCoreStyle::Get().GetBrush("ToolPanel.GroupBorder"))
+		.BorderBackgroundColor(FLinearColor::Gray) // Darken the outer border
+		[
+			SNew(SVerticalBox)
+			+ SVerticalBox::Slot()
+			.AutoHeight()
+			.HAlign(HAlign_Center)
+			.Padding(2, 2)
+			[
+				SNew(SBox)
+				.WidthOverride(300)
+				[
+					SNew(SObjectPropertyEntryBox)
+					.ObjectPath(this, &FIKRigDefinitionDetails::GetCurrentSourceAsset)
+					.OnShouldFilterAsset(this, &FIKRigDefinitionDetails::ShouldFilterAsset)
+					.OnObjectChanged(this, &FIKRigDefinitionDetails::OnAssetSelected)
+					.AllowClear(false)
+					.DisplayUseSelected(true)
+					.DisplayBrowse(true)
+				]
+			]
+			+ SVerticalBox::Slot()
+			.AutoHeight()
+			.HAlign(HAlign_Center)
+			.Padding(2, 2)
+			[
+				SNew(SButton)
+				.ContentPadding(3)
+				.IsEnabled(this, &FIKRigDefinitionDetails::CanImport)
+				.OnClicked(this, &FIKRigDefinitionDetails::OnImportHierarchy)
+				.ToolTipText(LOCTEXT("OnImportHierarchyTooltip", "Change Skeleton Data with Selected Asset. This replaces existing skeleton."))
+				[
+					SNew(STextBlock)
+					.Font(IDetailLayoutBuilder::GetDetailFont())
+					.Text(LOCTEXT("UpdateHierarchyTitle", "Update"))
+				]
+			]
+		]
+	];
+
+ 	IDetailCategoryBuilder& SolverCategory = DetailBuilder.EditCategory("Solver");
+
+	SolverCategory.AddCustomRow(FText::FromString("AddSolver"))
+	.NameContent()
+	[
+		SNullWidget::NullWidget
+	]
+	.ValueContent()
+	[
+		SNew(SButton)
+		.ContentPadding(3)
+		.OnClicked(this, &FIKRigDefinitionDetails::OnShowClassPicker)
+		.ToolTipText(LOCTEXT("OnShowSolverListTooltip", "Select Solver to Add"))
+		[
+			SNew(STextBlock)
+			.Font(IDetailLayoutBuilder::GetDetailFont())
+			.Text(LOCTEXT("ShowSolverList", "Add Solver"))
+		]
+	];
+
+	GoalPropertyHandle = DetailBuilder.GetProperty(TEXT("IKGoals"));
+
+	TArray<FName> GoalNames;
+	IKRigController->QueryGoals(GoalNames);
+
+	GoalListNames.SetNumZeroed(GoalNames.Num());
+
+	for (int32 Index = 0; Index < GoalNames.Num(); ++Index)
+	{
+		GoalListNames[Index] = MakeShareable(new FGoalNameListItem(GoalNames[Index]));
+	}
+
+	// I need to think about goal modified event OR just IKRigAssetMOdified event to update this
+	// for now i'm commenting it out
+// 	IDetailCategoryBuilder& GoalPropertyGroup = DetailBuilder.EditCategory("Goals");
+// 	GoalPropertyGroup.AddCustomRow(LOCTEXT("GoalsTitleLabel", "Goals"))
+// 	.NameContent()
+// 	[
+// 		GoalPropertyHandle->CreatePropertyNameWidget()
+// 	]
+// 	.ValueContent()
+// 	[
+// 		SAssignNew(GoalListView, SListView<FGoalNameListItemPtr>)
+// 		.ListItemsSource(&GoalListNames)
+// 		.OnGenerateRow(this, &FIKRigDefinitionDetails::OnGenerateWidgetForGoals)
+// 	];
+
+//	GoalPropertyHandle->MarkHiddenByCustomization();
+}
+
+class FIKSolverClassFilter : public IClassViewerFilter
+{
+public:
+	/** All children of these classes will be included unless filtered out by another setting. */
+	TSet< const UClass* > AllowedChildrenOfClasses;
+
+	TSet< const UClass* > DisallowedClasses;
+
+	/** Disallowed class flags. */
+	EClassFlags DisallowedClassFlags;
+
+	virtual bool IsClassAllowed(const FClassViewerInitializationOptions& InInitOptions, const UClass* InClass, TSharedRef< FClassViewerFilterFuncs > InFilterFuncs) override
+	{
+		return !InClass->HasAnyClassFlags(DisallowedClassFlags) && InFilterFuncs->IfInClassesSet(DisallowedClasses, InClass) == EFilterReturn::Failed
+			&& InFilterFuncs->IfInChildOfClassesSet(AllowedChildrenOfClasses, InClass) != EFilterReturn::Failed;
+	}
+
+	virtual bool IsUnloadedClassAllowed(const FClassViewerInitializationOptions& InInitOptions, const TSharedRef< const IUnloadedBlueprintData > InUnloadedClassData, TSharedRef< FClassViewerFilterFuncs > InFilterFuncs) override
+	{
+		return !InUnloadedClassData->HasAnyClassFlags(DisallowedClassFlags) && InFilterFuncs->IfInClassesSet(DisallowedClasses, InUnloadedClassData) == EFilterReturn::Failed
+			&& InFilterFuncs->IfInChildOfClassesSet(AllowedChildrenOfClasses, InUnloadedClassData) != EFilterReturn::Failed;
+	}
+};
+
+FReply FIKRigDefinitionDetails::OnShowClassPicker()
+{
+	FClassViewerModule& ClassViewerModule = FModuleManager::LoadModuleChecked<FClassViewerModule>("ClassViewer");
+	FClassViewerInitializationOptions Options;
+	Options.Mode = EClassViewerMode::ClassPicker;
+	Options.DisplayMode = EClassViewerDisplayMode::TreeView;
+	Options.bShowObjectRootClass = false;
+	Options.bExpandRootNodes = true;
+	Options.bShowUnloadedBlueprints = true;
+	TSharedPtr<FIKSolverClassFilter> Filter = MakeShareable(new FIKSolverClassFilter);
+	Options.ClassFilter = Filter;
+
+	Filter->DisallowedClassFlags = CLASS_Abstract | CLASS_Deprecated | CLASS_NewerVersionExists | CLASS_Transient;
+	Filter->AllowedChildrenOfClasses.Add(UIKRigSolverDefinition::StaticClass());
+	Filter->DisallowedClasses.Add(UIKRigSolverDefinition::StaticClass());
+
+	const FText TitleText = LOCTEXT("SelectSolverClass", "Select Solver Class");
+	UClass* ChosenClass = nullptr;
+	const bool bPressedOk = SClassPickerDialog::PickClass(TitleText, Options, ChosenClass, UIKRigSolverDefinition::StaticClass());
+
+	if (bPressedOk)
+	{
+		IKRigController->AddSolver(ChosenClass);
+	}
+
+	return FReply::Handled();
+}
+
+bool FIKRigDefinitionDetails::CanImport() const
+{
+	return (SelectedAsset.IsValid());
+}
+
+FString FIKRigDefinitionDetails::GetCurrentSourceAsset() const
+{
+	return GetPathNameSafe(SelectedAsset.IsValid()? SelectedAsset.Get() : nullptr);
+}
+
+bool FIKRigDefinitionDetails::ShouldFilterAsset(const FAssetData& AssetData)
+{
+	return (AssetData.AssetClass != USkeletalMesh::StaticClass()->GetFName() && AssetData.AssetClass != USkeleton::StaticClass()->GetFName());
+}
+
+void FIKRigDefinitionDetails::OnAssetSelected(const FAssetData& AssetData)
+{
+	SelectedAsset = AssetData.GetAsset();
+}
+
+FReply FIKRigDefinitionDetails::OnImportHierarchy()
+{
+	if (SelectedAsset.IsValid())
+	{
+		FScopedTransaction Transaction(LOCTEXT("UpdateSkeleton", "Update Skeleton"));
+		IKRigDefinition->Modify();
+
+		const FReferenceSkeleton* RefSkeleton = nullptr;
+		if (SelectedAsset->IsA(USkeleton::StaticClass()))
+		{
+			IKRigDefinition->SourceAsset = SelectedAsset.Get();
+			RefSkeleton = &(CastChecked<USkeleton>(SelectedAsset)->GetReferenceSkeleton());
+		}
+		else if (SelectedAsset->IsA(USkeletalMesh::StaticClass()))
+		{
+			IKRigDefinition->SourceAsset = SelectedAsset.Get();
+			RefSkeleton = &(CastChecked<USkeletalMesh>(SelectedAsset)->RefSkeleton);
+		}
+
+		if (RefSkeleton)
+		{
+			IKRigController->SetSkeleton(*RefSkeleton);
+		}
+
+		// Raw because we don't want to keep alive the details builder when calling the force refresh details
+		IDetailLayoutBuilder* DetailLayoutBuilder = DetailBuilderWeakPtr.Pin().Get();
+		if (DetailLayoutBuilder)
+		{
+			DetailLayoutBuilder->ForceRefreshDetails();
+		}
+	}
+
+	return FReply::Handled();
+}
+
+void FIKRigDefinitionDetails::OnObjectPostEditChange(UObject* Object, FPropertyChangedEvent& InPropertyChangedEvent)
+{
+// 	if (Object == IKRigDefinition || Object->GetOuter())
+// 	{
+// 		IDetailLayoutBuilder* DetailLayoutBuilder = DetailBuilderWeakPtr.Pin().Get();
+// 		if (DetailLayoutBuilder)
+// 		{
+// 			DetailLayoutBuilder->ForceRefreshDetails();
+// 		}
+// 	}
+}
+
+TSharedRef<ITableRow> FIKRigDefinitionDetails::OnGenerateWidgetForGoals(FGoalNameListItemPtr InItem, const TSharedRef<STableViewBase>& OwnerTable)
+{
+	return SNew(STableRow<FGoalNameListItemPtr>, OwnerTable)
+		[
+			SNew(SHorizontalBox)
+			+SHorizontalBox::Slot()
+			.AutoWidth()
+			[
+				SNew(SEditableTextBox)
+				.Text(this, &FIKRigDefinitionDetails::GetGoalNameText, InItem)
+				.OnTextCommitted(this, &FIKRigDefinitionDetails::HandleGoalNameChanged, InItem)
+				.SelectAllTextWhenFocused(true)
+				.RevertTextOnEscape(true)
+				.Font(IDetailLayoutBuilder::GetDetailFont())
+			]
+		];
+}
+
+void FIKRigDefinitionDetails::HandleGoalNameChanged(const FText& NewName, ETextCommit::Type CommitType, FGoalNameListItemPtr InItem)
+{
+	//if (CommitType == ETextCommit::OnEnter)
+	if (IKRigController.IsValid())
+	{
+		if (!NewName.IsEmptyOrWhitespace())
+		{
+			const FName NewFName = FName(*NewName.ToString());
+			if (InItem->DisplayName != NewFName)
+			{
+				IKRigController->RenameGoal(InItem->GoalName, NewFName);
+				InItem->GoalName = NewFName; // if you rename to the same as others, you'll reduce the number of goals
+				InItem->DisplayName = NewFName;
+
+				// refresh?
+
+			}
+		}
+	}
+}
+
+FText FIKRigDefinitionDetails::GetGoalNameText(FGoalNameListItemPtr InItem) const
+{
+	return FText::FromName(InItem->DisplayName);
+}
+#undef LOCTEXT_NAMESPACE
diff --git a/Engine/Plugins/Animation/IKRig/Source/IKRigEditor/Private/IKRigDefinitionDetails.h b/Engine/Plugins/Animation/IKRig/Source/IKRigEditor/Private/IKRigDefinitionDetails.h
new file mode 100644
--- /dev/null
+++ b/Engine/Plugins/Animation/IKRig/Source/IKRigEditor/Private/IKRigDefinitionDetails.h
@@ -0,0 +1,67 @@
+// Copyright Epic Games, Inc. All Rights Reserved.
+
+#pragma once
+
+#include "CoreMinimal.h"
+#include "UObject/StrongObjectPtr.h"
+#include "IDetailCustomization.h"
+
+class IDetailLayoutBuilder;
+class UIKRigDefinition;
+class UIKRigController;
+
+// Utility class to build combo boxes out of arrays of names.
+struct FGoalNameListItem : public TSharedFromThis<FGoalNameListItem>
+{
+	FName GoalName;
+	FName DisplayName;
+
+	FGoalNameListItem(const FName& InName)
+	{
+		GoalName = InName;
+		DisplayName = InName;
+	}
+};
+
+class FIKRigDefinitionDetails : public IDetailCustomization
+{
+public:
+	/** Makes a new instance of this detail layout class for a specific detail view requesting it */
+	static TSharedRef<IDetailCustomization> MakeInstance();
+
+	virtual ~FIKRigDefinitionDetails();
+
+	/** IDetailCustomization interface */
+	virtual void CustomizeDetails( IDetailLayoutBuilder& DetailBuilder ) override;
+	virtual void CustomizeDetails(const TSharedPtr<IDetailLayoutBuilder>& DetailBuilder) override;
+private:
+	TWeakObjectPtr<UIKRigDefinition>		IKRigDefinition;
+	TStrongObjectPtr<UIKRigController>		IKRigController;
+
+	// source asset
+	FString GetCurrentSourceAsset() const;
+	bool ShouldFilterAsset(const FAssetData& AssetData);
+	void OnAssetSelected(const FAssetData& AssetData);
+	TWeakObjectPtr<UObject>					SelectedAsset;
+
+	FReply OnImportHierarchy();
+	bool CanImport() const;
+
+	TWeakPtr< IDetailLayoutBuilder > DetailBuilderWeakPtr;
+
+	FReply OnShowClassPicker();
+
+	typedef TSharedPtr<FGoalNameListItem> FGoalNameListItemPtr;
+
+	void OnObjectPostEditChange(UObject* Object, FPropertyChangedEvent& InPropertyChangedEvent);
+	TSharedRef<ITableRow> OnGenerateWidgetForGoals(FGoalNameListItemPtr InItem, const TSharedRef<STableViewBase>& OwnerTable);
+	void HandleGoalNameChanged(const FText& NewName, ETextCommit::Type CommitType, FGoalNameListItemPtr InItem);
+	FText GetGoalNameText(FGoalNameListItemPtr InName) const;
+
+	TSharedPtr<IPropertyHandle> GoalPropertyHandle;
+	TSharedPtr<SListView<FGoalNameListItemPtr>> GoalListView;
+	TArray<FGoalNameListItemPtr> GoalListNames;
+
+	FDelegateHandle ObjectChangedDelegate;
+};
+
diff --git a/Engine/Plugins/Animation/IKRig/Source/IKRigEditor/Private/IKRigDefinitionFactory.cpp b/Engine/Plugins/Animation/IKRig/Source/IKRigEditor/Private/IKRigDefinitionFactory.cpp
new file mode 100644
--- /dev/null
+++ b/Engine/Plugins/Animation/IKRig/Source/IKRigEditor/Private/IKRigDefinitionFactory.cpp
@@ -0,0 +1,53 @@
+// Copyright Epic Games, Inc. All Rights Reserved.
+
+#include "IKRigDefinitionFactory.h"
+#include "IKRigDefinition.h"
+
+#define LOCTEXT_NAMESPACE "IKRigDefinitionFactory"
+
+
+/* UFactory interface
+ *****************************************************************************/
+
+UIKRigDefinitionFactory::UIKRigDefinitionFactory()
+{
+	bCreateNew = true;
+	bEditAfterNew = true;
+	SupportedClass = UIKRigDefinition::StaticClass();
+}
+
+UObject* UIKRigDefinitionFactory::FactoryCreateNew(UClass* Class, UObject* InParent, FName Name, EObjectFlags Flags, UObject* Context, FFeedbackContext* Warn)
+{
+	return NewObject<UIKRigDefinition>(InParent, Name, Flags | RF_Transactional);;
+}
+
+bool UIKRigDefinitionFactory::ShouldShowInNewMenu() const
+{
+	return true;
+}
+
+bool UIKRigDefinitionFactory::ConfigureProperties()
+{
+	return true;
+}
+
+FText UIKRigDefinitionFactory::GetDisplayName() const
+{
+	return LOCTEXT("IKRigDefinition_DisplayName", "IK Rig Definition");
+}
+
+uint32 UIKRigDefinitionFactory::GetMenuCategories() const
+{
+	return EAssetTypeCategories::Animation;
+}
+
+FText UIKRigDefinitionFactory::GetToolTip() const
+{
+	return LOCTEXT("IKRigDefinition_Tooltip", "Create IK Rig simply to animate or to use in runtime.");
+}
+
+FString UIKRigDefinitionFactory::GetDefaultNewAssetName() const
+{
+	return FString(TEXT("NewIKRigDefinition"));
+}
+#undef LOCTEXT_NAMESPACE
diff --git a/Engine/Plugins/Animation/IKRig/Source/IKRigEditor/Private/IKRigDefinitionFactory.h b/Engine/Plugins/Animation/IKRig/Source/IKRigEditor/Private/IKRigDefinitionFactory.h
new file mode 100644
--- /dev/null
+++ b/Engine/Plugins/Animation/IKRig/Source/IKRigEditor/Private/IKRigDefinitionFactory.h
@@ -0,0 +1,30 @@
+// Copyright Epic Games, Inc. All Rights Reserved.
+
+#pragma once
+
+#include "CoreMinimal.h"
+#include "UObject/ObjectMacros.h"
+#include "Factories/Factory.h"
+#include "IKRigDefinitionFactory.generated.h"
+
+/**
+ * Implements a factory for UIKRigDefinitionFactory objects.
+ */
+UCLASS(BlueprintType, hidecategories=Object)
+class UIKRigDefinitionFactory : public UFactory
+{
+	GENERATED_BODY()
+
+public:
+
+	UIKRigDefinitionFactory();
+
+	// UFactory Interface
+	virtual FText GetDisplayName() const override;
+	virtual uint32 GetMenuCategories() const override;
+	virtual FText GetToolTip() const override;
+	virtual FString GetDefaultNewAssetName() const override;
+	virtual UObject* FactoryCreateNew(UClass* Class, UObject* InParent, FName Name, EObjectFlags Flags, UObject* Context, FFeedbackContext* Warn) override;
+	virtual bool ConfigureProperties() override;
+	virtual bool ShouldShowInNewMenu() const override;
+};
diff --git a/Engine/Plugins/Animation/IKRig/Source/IKRigEditor/Private/IKRigEditor.cpp b/Engine/Plugins/Animation/IKRig/Source/IKRigEditor/Private/IKRigEditor.cpp
new file mode 100644
--- /dev/null
+++ b/Engine/Plugins/Animation/IKRig/Source/IKRigEditor/Private/IKRigEditor.cpp
@@ -0,0 +1,53 @@
+// Copyright Epic Games, Inc. All Rights Reserved.
+
+#pragma once
+
+#include "IKRigEditor.h"
+#include "Features/IModularFeatures.h"
+#include "Engine/Engine.h"
+#include "Misc/Paths.h"
+#include "Modules/ModuleManager.h"
+#include "Logging/LogMacros.h"
+#include "AssetToolsModule.h"
+#include "PropertyEditorModule.h"
+#include "IKRigDefinitionDetails.h"
+#include "PropertyEditorDelegates.h"
+#include "EditorModeRegistry.h"
+#include "AnimGraphNode_IKRig.h"
+
+IMPLEMENT_MODULE(FIKRigEditor, IKRigEditor)
+
+DEFINE_LOG_CATEGORY_STATIC(LogIKRigEditor, Log, All);
+
+#define LOCTEXT_NAMESPACE "IKRigEditor"
+
+void FIKRigEditor::StartupModule()
+{
+	IKRigDefinitionAssetAction = MakeShareable(new FAssetTypeActions_IKRigDefinition);
+	FModuleManager::LoadModuleChecked<FAssetToolsModule>("AssetTools").Get().RegisterAssetTypeActions(IKRigDefinitionAssetAction.ToSharedRef());
+
+	FPropertyEditorModule& PropertyModule = FModuleManager::GetModuleChecked<FPropertyEditorModule>("PropertyEditor");
+	PropertyModule.RegisterCustomClassLayout(UIKRigDefinition::StaticClass()->GetFName(), FOnGetDetailCustomizationInstance::CreateStatic(&FIKRigDefinitionDetails::MakeInstance));
+
+//	FEditorModeRegistry::Get().RegisterMode<FIKRigEditMode>(UAnimGraphNode_IKRig::AnimModeName, LOCTEXT("IKRigEditMode", "IKRig"), FSlateIcon(), false);
+}
+
+void FIKRigEditor::ShutdownModule()
+{
+	//	FEditorModeRegistry::Get().UnregisterMode(UAnimGraphNode_IKRig::AnimModeName);
+
+	FPropertyEditorModule& PropertyModule = FModuleManager::GetModuleChecked<FPropertyEditorModule>("PropertyEditor");
+	PropertyModule.UnregisterCustomClassLayout(UIKRigDefinition::StaticClass()->GetFName());
+
+	if (IKRigDefinitionAssetAction.IsValid())
+	{
+		if (FModuleManager::Get().IsModuleLoaded("AssetTools"))
+		{
+			FModuleManager::GetModuleChecked<FAssetToolsModule>("AssetTools").Get().UnregisterAssetTypeActions(IKRigDefinitionAssetAction.ToSharedRef());
+		}
+		IKRigDefinitionAssetAction.Reset();
+	}
+
+}
+
+#undef LOCTEXT_NAMESPACE
\ No newline at end of file
diff --git a/Engine/Plugins/Animation/IKRig/Source/IKRigEditor/Public/IKRigController.h b/Engine/Plugins/Animation/IKRig/Source/IKRigEditor/Public/IKRigController.h
new file mode 100644
--- /dev/null
+++ b/Engine/Plugins/Animation/IKRig/Source/IKRigEditor/Public/IKRigController.h
@@ -0,0 +1,102 @@
+// Copyright Epic Games, Inc. All Rights Reserved.
+
+/**
+ * IKRigController
+ *
+ * Mutates IKRig asset side
+ */
+
+#pragma once
+
+#include "CoreMinimal.h"
+#include "UObject/ObjectMacros.h"
+#include "UObject/Object.h"
+#include "IKRigDefinition.h"
+#include "IKRigController.generated.h"
+
+// DELEGATE for solver chnage
+// DELEGATE for hierarchy changes
+// DELEGATE for profile changes
+// DELEGATE for constraint changes
+// DELEGATE for goals
+
+class UIKRigSolverDefinition;
+class UIKRigConstraint;
+struct FIKRigHierarchy;
+struct FReferenceSkeleton;
+
+
+UCLASS(config = Engine, hidecategories = UObject, BlueprintType)
+class IKRIGEDITOR_API UIKRigController : public UObject
+{
+	GENERATED_BODY()
+
+public:
+	// IKRigDefinition set up
+	void SetIKRigDefinition(UIKRigDefinition* InIKRigDefinition);
+
+	// hierarchy operators
+	void SetSkeleton(const FReferenceSkeleton& InSkeleton);
+	// const hierarchy getter
+	const FIKRigHierarchy* GetHierarchy() const
+	{
+		return (IKRigDefinition)? &IKRigDefinition->Hierarchy : nullptr;
+	}
+
+	const FIKRigTransform* GetReferenceTransform() const
+	{
+		return (IKRigDefinition) ? &IKRigDefinition->ReferencePose : nullptr;
+	}
+
+	bool AddBone(const FName& InName, const FName& InParent, const FTransform& InGlobalTransform);
+	bool RemoveBone(const FName& InName);
+	bool RenameBone(const FName& InOldName, const FName& InNewName);
+	bool ReparentBone(const FName& InName, const FName& InNewParent);
+	void ResetHierarchy();
+
+	// todo: more APIs for hierarchy operators
+
+	// solver operators
+	UIKRigSolverDefinition* AddSolver(TSubclassOf<UIKRigSolverDefinition> InIKRigSolverDefinitionClass);
+	int32 GetTotalSolverCount() const;
+	UIKRigSolverDefinition* GetSolver(int32 Index) const;
+	void RemoveSolver(UIKRigSolverDefinition* SolverToDelete);
+	void AutoConfigure(UIKRigSolverDefinition* SolverDef);
+	bool CanAutoConfigure(UIKRigSolverDefinition* SolverDef) const;
+
+	// constraint operators
+	// create new profile
+	void CreateNewProfile(FName& InNewProfileName);
+	bool RemoveConstraintProfile(const FName& InProfileName);
+	void RenameProfile(FName InCurrentProfileName, FName& InNewProfileName);
+
+	UIKRigConstraint* AddConstraint(TSubclassOf<UIKRigConstraint> NewConstraintType, FName& InOutNewName, FName InProfile=NAME_None);
+	UIKRigConstraint* GetConstraint(const FName& InProfileName, const FName& InName) const;
+	bool RemoveConstraint(const FName& InConstraintName);
+
+	void GetConstraintProfileNames(TArray<FName>& OutProfileNames) const;
+	void GetConstraintNames(TArray<FName>& OutConstraintNames) const;
+
+	// goal operators
+	void QueryGoals(TArray<FName>& OutGoals) const;
+	void RenameGoal(const FName& OldName, const FName& NewName);
+
+	// this is to modify default value
+	FIKRigGoal* GetGoal(const FName& InGoalName);
+	const FIKRigGoal* GetGoal(const FName& InGoalName) const;
+
+	// BEGIN UObject
+	virtual void BeginDestroy() override;
+	// END UObject
+private:
+	UIKRigDefinition* IKRigDefinition = nullptr;
+
+	bool ValidateSolver(UIKRigSolverDefinition* const SolverDef) const;
+	void UpdateGoal();
+
+	void InitializeIKRigSolverDefinition(UIKRigSolverDefinition* SolverDef);
+	void UninitializeIKRigSolverDefinition(UIKRigSolverDefinition* SolverDef);
+
+	TMap<UIKRigSolverDefinition*, FDelegateHandle> SolverDelegateHandles;
+};
+
diff --git a/Engine/Plugins/Animation/IKRig/Source/IKRigEditor/Public/IKRigEditor.h b/Engine/Plugins/Animation/IKRig/Source/IKRigEditor/Public/IKRigEditor.h
new file mode 100644
--- /dev/null
+++ b/Engine/Plugins/Animation/IKRig/Source/IKRigEditor/Public/IKRigEditor.h
@@ -0,0 +1,16 @@
+// Copyright Epic Games, Inc. All Rights Reserved.
+
+#pragma once
+
+#include "CoreMinimal.h"
+#include "Modules/ModuleInterface.h"
+
+class IKRIGEDITOR_API FIKRigEditor : public IModuleInterface 
+{
+public:
+	void StartupModule() override;
+	void ShutdownModule() override;
+
+private:
+	TSharedPtr<class FAssetTypeActions_IKRigDefinition> IKRigDefinitionAssetAction;
+};
diff --git a/Engine/Source/Editor/AnimGraph/Public/AnimNodeEditMode.h b/Engine/Source/Editor/AnimGraph/Public/AnimNodeEditMode.h
--- a/Engine/Source/Editor/AnimGraph/Public/AnimNodeEditMode.h
+++ b/Engine/Source/Editor/AnimGraph/Public/AnimNodeEditMode.h
@@ -16,7 +16,7 @@ struct FViewportClick;
 struct FBoneSocketTarget;
 
 /** Base implementation for anim node edit modes */
-class FAnimNodeEditMode : public IAnimNodeEditMode
+class ANIMGRAPH_API FAnimNodeEditMode : public IAnimNodeEditMode
 {
 public:
 	FAnimNodeEditMode();
diff --git a/Engine/Source/Runtime/AnimationCore/Public/BoneIndices.h b/Engine/Source/Runtime/AnimationCore/Public/BoneIndices.h
--- a/Engine/Source/Runtime/AnimationCore/Public/BoneIndices.h
+++ b/Engine/Source/Runtime/AnimationCore/Public/BoneIndices.h
@@ -132,3 +132,11 @@ struct FSkeletonPoseBoneIndex : public FBoneIndexWithOperators < FSkeletonPoseBo
 public:
 	explicit FSkeletonPoseBoneIndex(int32 InBoneIndex) { BoneIndex = InBoneIndex; }
 };
+
+template <typename ValueType>
+struct TCompactPoseBoneIndexMapKeyFuncs : public TDefaultMapKeyFuncs<const FCompactPoseBoneIndex, ValueType, false>
+{
+	static FORCEINLINE FCompactPoseBoneIndex			GetSetKey(TPair<FCompactPoseBoneIndex, ValueType> const& Element) { return Element.Key; }
+	static FORCEINLINE uint32							GetKeyHash(FCompactPoseBoneIndex const& Key) { return GetTypeHash(Key.GetInt()); }
+	static FORCEINLINE bool								Matches(FCompactPoseBoneIndex const& A, FCompactPoseBoneIndex const& B) { return (A.GetInt() == B.GetInt()); }
+};
\ No newline at end of file
