Fix nonunity build issues - had preflight https://horde.devtools.epicgames.com/job/5fbebdab0c6d82000115bea3 - and didn't catch it. I think it didn't run nonunity.

diff --git a/Engine/Plugins/Animation/IKRig/Source/IKRig/Private/IKRigProcessor.cpp b/Engine/Plugins/Animation/IKRig/Source/IKRig/Private/IKRigProcessor.cpp
--- a/Engine/Plugins/Animation/IKRig/Source/IKRig/Private/IKRigProcessor.cpp
+++ b/Engine/Plugins/Animation/IKRig/Source/IKRig/Private/IKRigProcessor.cpp
@@ -6,6 +6,9 @@
  */
 
 #include "IKRigProcessor.h"
+#include "IKRigDefinition.h"
+#include "IKRigSolverDefinition.h"
+#include "IKRigSolver.h"
 
  // IKRigProcessor implementation functions
 void UIKRigProcessor::SetIKRigDefinition(UIKRigDefinition* InRigDefinition, bool bInitialize /*= false*/)
diff --git a/Engine/Plugins/Animation/IKRig/Source/IKRig/Private/IKRigSolver.cpp b/Engine/Plugins/Animation/IKRig/Source/IKRig/Private/IKRigSolver.cpp
--- a/Engine/Plugins/Animation/IKRig/Source/IKRig/Private/IKRigSolver.cpp
+++ b/Engine/Plugins/Animation/IKRig/Source/IKRig/Private/IKRigSolver.cpp
@@ -5,6 +5,8 @@
 =============================================================================*/
 
 #include "IKRigSolver.h"
+#include "IKRigSolverDefinition.h"
+#include "IKRigDataTypes.h"
 
 // input hierarchy and ref pose? 
 void UIKRigSolver::Init(UIKRigSolverDefinition* InSolverDefinition, FIKRigTransformGetter InRefPoseGetter, FIKRigGoalGetter InGoalGetter)
diff --git a/Engine/Plugins/Animation/IKRig/Source/IKRig/Private/IKRigSolverDefinition.cpp b/Engine/Plugins/Animation/IKRig/Source/IKRig/Private/IKRigSolverDefinition.cpp
--- a/Engine/Plugins/Animation/IKRig/Source/IKRig/Private/IKRigSolverDefinition.cpp
+++ b/Engine/Plugins/Animation/IKRig/Source/IKRig/Private/IKRigSolverDefinition.cpp
@@ -6,6 +6,7 @@
  */
 
 #include "IKRigSolverDefinition.h"
+#include "IKRigDefinition.h"
 
 void UIKRigSolverDefinition::CollectGoals(TArray<FName>& OutGoals)
 {
diff --git a/Engine/Plugins/Animation/IKRig/Source/IKRig/Private/Solvers/TransformSolver.cpp b/Engine/Plugins/Animation/IKRig/Source/IKRig/Private/Solvers/TransformSolver.cpp
--- a/Engine/Plugins/Animation/IKRig/Source/IKRig/Private/Solvers/TransformSolver.cpp
+++ b/Engine/Plugins/Animation/IKRig/Source/IKRig/Private/Solvers/TransformSolver.cpp
@@ -5,6 +5,8 @@
 =============================================================================*/
 
 #include "TransformSolver.h"
+#include "IKRigDataTypes.h"
+#include "IKRigHierarchy.h"
 #include "Solvers/TransformSolverDefinition.h"
 
 void UTransformSolver::InitInternal()
diff --git a/Engine/Plugins/Animation/IKRig/Source/IKRig/Public/IKRigConstraintDefinition.h b/Engine/Plugins/Animation/IKRig/Source/IKRig/Public/IKRigConstraintDefinition.h
--- a/Engine/Plugins/Animation/IKRig/Source/IKRig/Public/IKRigConstraintDefinition.h
+++ b/Engine/Plugins/Animation/IKRig/Source/IKRig/Public/IKRigConstraintDefinition.h
@@ -12,6 +12,8 @@
 #include "UObject/Object.h"
 #include "IKRigConstraintDefinition.generated.h"
 
+class UIKRigConstraint;
+
 USTRUCT()
 struct IKRIG_API FIKRigConstraintProfile 
 {
diff --git a/Engine/Plugins/Animation/IKRig/Source/IKRig/Public/IKRigConstraintSolver.h b/Engine/Plugins/Animation/IKRig/Source/IKRig/Public/IKRigConstraintSolver.h
--- a/Engine/Plugins/Animation/IKRig/Source/IKRig/Public/IKRigConstraintSolver.h
+++ b/Engine/Plugins/Animation/IKRig/Source/IKRig/Public/IKRigConstraintSolver.h
@@ -10,6 +10,8 @@
 #include "IKRigSolver.h"
 #include "IKRigConstraintSolver.generated.h"
 
+class UIKRigConstraint;
+
 // run time processor 
 UCLASS(config = Engine, hidecategories = UObject)
 class IKRIG_API UIKRigConstraintSolver : public UIKRigSolver
diff --git a/Engine/Plugins/Animation/IKRig/Source/IKRigEditor/Private/IKRigDefinitionDetails.cpp b/Engine/Plugins/Animation/IKRig/Source/IKRigEditor/Private/IKRigDefinitionDetails.cpp
--- a/Engine/Plugins/Animation/IKRig/Source/IKRigEditor/Private/IKRigDefinitionDetails.cpp
+++ b/Engine/Plugins/Animation/IKRig/Source/IKRigEditor/Private/IKRigDefinitionDetails.cpp
@@ -267,7 +267,7 @@ FReply FIKRigDefinitionDetails::OnImportHierarchy()
 		else if (SelectedAsset->IsA(USkeletalMesh::StaticClass()))
 		{
 			IKRigDefinition->SourceAsset = SelectedAsset.Get();
-			RefSkeleton = &(CastChecked<USkeletalMesh>(SelectedAsset)->RefSkeleton);
+			RefSkeleton = &(CastChecked<USkeletalMesh>(SelectedAsset)->GetRefSkeleton());
 		}
 
 		if (RefSkeleton)
diff --git a/Engine/Plugins/Animation/IKRig/Source/IKRigEditor/Private/IKRigDefinitionFactory.cpp b/Engine/Plugins/Animation/IKRig/Source/IKRigEditor/Private/IKRigDefinitionFactory.cpp
--- a/Engine/Plugins/Animation/IKRig/Source/IKRigEditor/Private/IKRigDefinitionFactory.cpp
+++ b/Engine/Plugins/Animation/IKRig/Source/IKRigEditor/Private/IKRigDefinitionFactory.cpp
@@ -2,6 +2,7 @@
 
 #include "IKRigDefinitionFactory.h"
 #include "IKRigDefinition.h"
+#include "AssetTypeCategories.h"
 
 #define LOCTEXT_NAMESPACE "IKRigDefinitionFactory"
 
diff --git a/Engine/Plugins/Animation/IKRig/Source/IKRigEditor/Private/IKRigEditor.cpp b/Engine/Plugins/Animation/IKRig/Source/IKRigEditor/Private/IKRigEditor.cpp
--- a/Engine/Plugins/Animation/IKRig/Source/IKRigEditor/Private/IKRigEditor.cpp
+++ b/Engine/Plugins/Animation/IKRig/Source/IKRigEditor/Private/IKRigEditor.cpp
@@ -14,6 +14,8 @@
 #include "PropertyEditorDelegates.h"
 #include "EditorModeRegistry.h"
 #include "AnimGraphNode_IKRig.h"
+#include "AssetTypeActions_IKRigDefinition.h"
+#include "IKRigDefinition.h"
 
 IMPLEMENT_MODULE(FIKRigEditor, IKRigEditor)
 
