Delay the construction of UObjects until the Automation tests start, to avoid creating them during dll initialization which can be too early to construct UObjects.
Fix a unity compile error.
#rnx

diff --git a/Engine/Plugins/Animation/IKRig/Source/IKRig/Public/IKRigSolverDefinition.h b/Engine/Plugins/Animation/IKRig/Source/IKRig/Public/IKRigSolverDefinition.h
--- a/Engine/Plugins/Animation/IKRig/Source/IKRig/Public/IKRigSolverDefinition.h
+++ b/Engine/Plugins/Animation/IKRig/Source/IKRig/Public/IKRigSolverDefinition.h
@@ -8,8 +8,11 @@
 #pragma once
 
 #include "CoreMinimal.h"
+
+#include "Templates/SubclassOf.h"
 #include "UObject/ObjectMacros.h"
 #include "UObject/Object.h"
+
 #include "IKRigSolverDefinition.generated.h"
 
 class UIKRigSolver;
