Move bone from Goals to Effector.
Changed code accordingly.

diff --git a/Engine/Plugins/Animation/IKRig/Source/IKRig/Private/Solvers/TransformSolver.cpp b/Engine/Plugins/Animation/IKRig/Source/IKRig/Private/Solvers/TransformSolver.cpp
--- a/Engine/Plugins/Animation/IKRig/Source/IKRig/Private/Solvers/TransformSolver.cpp
+++ b/Engine/Plugins/Animation/IKRig/Source/IKRig/Private/Solvers/TransformSolver.cpp
@@ -34,7 +34,7 @@ void UTransformSolver::SolveInternal(FIKRigTransformModifier& InOutGlobalTransfo
 		FIKRigTarget Target;
 		if (GetEffectorTarget(SolverDef->TransformTarget, Target))
 		{
-			int32 Index = InOutGlobalTransform.Hierarchy->GetIndex(Target.Bone);
+			int32 Index = InOutGlobalTransform.Hierarchy->GetIndex(SolverDef->TransformTarget.Bone);
 			if (Index != INDEX_NONE)
 			{
 				FTransform CurrentTransform = InOutGlobalTransform.GetGlobalTransform(Index);
diff --git a/Engine/Plugins/Animation/IKRig/Source/IKRig/Public/IKRigDataTypes.h b/Engine/Plugins/Animation/IKRig/Source/IKRig/Public/IKRigDataTypes.h
--- a/Engine/Plugins/Animation/IKRig/Source/IKRig/Public/IKRigDataTypes.h
+++ b/Engine/Plugins/Animation/IKRig/Source/IKRig/Public/IKRigDataTypes.h
@@ -52,9 +52,6 @@ struct IKRIG_API FIKRigTarget
 {
 	GENERATED_BODY()
 
-	UPROPERTY(EditAnywhere, Category = FIKRigTarget)
-	FName Bone;
-
 	UPROPERTY(EditAnywhere, Category = FIKRigTarget)
 	FIKRigPosition PositionTarget;
 
@@ -101,6 +98,9 @@ struct IKRIG_API FIKRigEffector
 	UPROPERTY()
 	FGuid Guid;
 
+	UPROPERTY(EditAnywhere, Category = FIKRigEffector)
+	FName Bone;
+
 	FIKRigEffector()
 	{
 		Guid = FGuid::NewGuid();
