- add customization for effector to goal
  - from user, it shouldn't look anything different, internally they're mapped in different ways, but when users look at it, it should look to deal with the same way.
- a couple of different APIs for goals/effector, still need more clean ups

diff --git a/Engine/Plugins/Animation/IKRig/Source/IKRig/Private/AnimNodes/AnimNode_IKRig.cpp b/Engine/Plugins/Animation/IKRig/Source/IKRig/Private/AnimNodes/AnimNode_IKRig.cpp
--- a/Engine/Plugins/Animation/IKRig/Source/IKRig/Private/AnimNodes/AnimNode_IKRig.cpp
+++ b/Engine/Plugins/Animation/IKRig/Source/IKRig/Private/AnimNodes/AnimNode_IKRig.cpp
@@ -124,12 +124,19 @@ void FAnimNode_IKRig::OnInitializeAnimInstance(const FAnimInstanceProxy* InProxy
 	FAnimNode_Base::OnInitializeAnimInstance(InProxy, InAnimInstance);
 }
 
-void FAnimNode_IKRig::RebuildGoalList()
+bool FAnimNode_IKRig::RebuildGoalList()
 {
 	if (RigDefinitionAsset)
 	{
-		GoalTransforms.SetNum(RigDefinitionAsset->GetGoals().Num());
+		const int32 GoalNum = RigDefinitionAsset->GetGoals().Num();
+		if (GoalTransforms.Num() != GoalNum)
+		{
+			GoalTransforms.SetNum(GoalNum);
+			return true;
+		}
 	}
+
+	return false;
 }
 
 FName FAnimNode_IKRig::GetGoalName(int32 Index) const
diff --git a/Engine/Plugins/Animation/IKRig/Source/IKRig/Public/AnimNodes/AnimNode_IKRig.h b/Engine/Plugins/Animation/IKRig/Source/IKRig/Public/AnimNodes/AnimNode_IKRig.h
--- a/Engine/Plugins/Animation/IKRig/Source/IKRig/Public/AnimNodes/AnimNode_IKRig.h
+++ b/Engine/Plugins/Animation/IKRig/Source/IKRig/Public/AnimNodes/AnimNode_IKRig.h
@@ -61,7 +61,7 @@ public:
 	// End of FAnimNode_Base interface
 
 private:
-	void RebuildGoalList();
+	bool RebuildGoalList();
 	FName GetGoalName(int32 Index) const;
 
 	friend class UAnimGraphNode_IKRig;
diff --git a/Engine/Plugins/Animation/IKRig/Source/IKRigDeveloper/Private/AnimGraphNode_IKRig.cpp b/Engine/Plugins/Animation/IKRig/Source/IKRigDeveloper/Private/AnimGraphNode_IKRig.cpp
--- a/Engine/Plugins/Animation/IKRig/Source/IKRigDeveloper/Private/AnimGraphNode_IKRig.cpp
+++ b/Engine/Plugins/Animation/IKRig/Source/IKRigDeveloper/Private/AnimGraphNode_IKRig.cpp
@@ -2,6 +2,7 @@
 
 #include "AnimGraphNode_IKRig.h"
 #include "Animation/AnimInstance.h"
+#include "IKRigDefinition.h"
 
 /////////////////////////////////////////////////////
 // UAnimGraphNode_IKRig 
@@ -35,12 +36,7 @@ void UAnimGraphNode_IKRig::CopyNodeDataToPreviewNode(FAnimNode_Base* InPreviewNo
 
 FEditorModeID UAnimGraphNode_IKRig::GetEditorMode() const
 {
-// we have a crash on shutting it down because it's not unregistering
-// https://epic.slack.com/archives/C024FSXHJ/p1606256645266900
-// until we know what is solution, I'm disabling it
-//	return AnimModeName;
-
-	return NAME_None;
+	return AnimModeName;
 }
 
 void UAnimGraphNode_IKRig::PostEditChangeProperty(struct FPropertyChangedEvent& PropertyChangedEvent)
@@ -50,8 +46,10 @@ void UAnimGraphNode_IKRig::PostEditChangeProperty(struct FPropertyChangedEvent&
 	// Reconstruct node to show updates to PinFriendlyNames.
 	if ((PropertyName == GET_MEMBER_NAME_STRING_CHECKED(FAnimNode_IKRig, RigDefinitionAsset)))
 	{
-		Node.RebuildGoalList();
-		ReconstructNode();
+		if (Node.RebuildGoalList())
+		{
+			ReconstructNode();
+		}
 	}
 }
 
@@ -70,4 +68,23 @@ void UAnimGraphNode_IKRig::CustomizePinData(UEdGraphPin* Pin, FName SourceProper
 		}
 	}
 }
+
+void UAnimGraphNode_IKRig::PostLoad()
+{
+	Super::PostLoad();
+
+	// resize global transforms
+	Node.RebuildGoalList();
+	// I can't reconstruct here
+}
+
+void UAnimGraphNode_IKRig::PreloadRequiredAssets()
+{
+	if (Node.RigDefinitionAsset)
+	{
+		PreloadObject(Node.RigDefinitionAsset);
+	}
+
+	Super::PreloadRequiredAssets();
+}
 #undef LOCTEXT_NAMESPACE
diff --git a/Engine/Plugins/Animation/IKRig/Source/IKRigDeveloper/Public/AnimGraphNode_IKRig.h b/Engine/Plugins/Animation/IKRig/Source/IKRigDeveloper/Public/AnimGraphNode_IKRig.h
--- a/Engine/Plugins/Animation/IKRig/Source/IKRigDeveloper/Public/AnimGraphNode_IKRig.h
+++ b/Engine/Plugins/Animation/IKRig/Source/IKRigDeveloper/Public/AnimGraphNode_IKRig.h
@@ -24,9 +24,11 @@ public:
 	// UEdGraphNode interface
 	virtual FText GetNodeTitle(ENodeTitleType::Type TitleType) const override;
 	virtual void PostEditChangeProperty(struct FPropertyChangedEvent& PropertyChangedEvent) override;
+	virtual void PostLoad() override;
 	// End of UEdGraphNode interface
 
 	// UAnimGraphNode_Base interface
+	virtual void PreloadRequiredAssets() override;
 	virtual void CopyNodeDataToPreviewNode(FAnimNode_Base* AnimNode) override;
 	virtual FEditorModeID GetEditorMode() const override;
 	virtual void Draw(FPrimitiveDrawInterface* PDI, USkeletalMeshComponent* PreviewSkelMeshComp) const override;
diff --git a/Engine/Plugins/Animation/IKRig/Source/IKRigEditor/IKRigEditor.Build.cs b/Engine/Plugins/Animation/IKRig/Source/IKRigEditor/IKRigEditor.Build.cs
--- a/Engine/Plugins/Animation/IKRig/Source/IKRigEditor/IKRigEditor.Build.cs
+++ b/Engine/Plugins/Animation/IKRig/Source/IKRigEditor/IKRigEditor.Build.cs
@@ -32,6 +32,7 @@ namespace UnrealBuildTool.Rules
 					"EditorFramework",
 					"BlueprintGraph",
 					"AnimGraph",
+					"EditorWidgets"
 				}
             );
 
diff --git a/Engine/Plugins/Animation/IKRig/Source/IKRigEditor/Private/IKRigController.cpp b/Engine/Plugins/Animation/IKRig/Source/IKRigEditor/Private/IKRigController.cpp
--- a/Engine/Plugins/Animation/IKRig/Source/IKRigEditor/Private/IKRigController.cpp
+++ b/Engine/Plugins/Animation/IKRig/Source/IKRigEditor/Private/IKRigController.cpp
@@ -242,6 +242,35 @@ void UIKRigController::UpdateGoal()
 	if (IKRigDefinition)
 	{
 		IKRigDefinition->UpdateGoal();
+		OnGoalModified.Broadcast();
+	}
+}
+
+FName UIKRigController::GetGoalName(UIKRigSolverDefinition* InSolverDefinition, const FIKRigEffector& InEffector)
+{
+	if (InSolverDefinition)
+	{
+		FName* GoalName = InSolverDefinition->EffectorToGoal.Find(InEffector);
+		if (GoalName)
+		{
+			return *GoalName;
+		}
+	}
+
+	return NAME_None;
+}
+
+void UIKRigController::SetGoalName(UIKRigSolverDefinition* InSolverDefinition, const FIKRigEffector& InEffector, const FName& NewGoalName)
+{
+	if (InSolverDefinition)
+	{
+		FName* GoalName = InSolverDefinition->EffectorToGoal.Find(InEffector);
+		if (GoalName)
+		{
+			*GoalName = NewGoalName;
+
+			UpdateGoal();
+		}
 	}
 }
 
diff --git a/Engine/Plugins/Animation/IKRig/Source/IKRigEditor/Private/IKRigEditor.cpp b/Engine/Plugins/Animation/IKRig/Source/IKRigEditor/Private/IKRigEditor.cpp
--- a/Engine/Plugins/Animation/IKRig/Source/IKRigEditor/Private/IKRigEditor.cpp
+++ b/Engine/Plugins/Animation/IKRig/Source/IKRigEditor/Private/IKRigEditor.cpp
@@ -17,6 +17,7 @@
 #include "AssetTypeActions_IKRigDefinition.h"
 #include "IKRigDefinition.h"
 #include "AnimMode/IKRigEditMode.h"
+#include "IKRigEffectorCustomization.h"
 
 IMPLEMENT_MODULE(FIKRigEditor, IKRigEditor)
 
@@ -32,7 +33,7 @@ void FIKRigEditor::StartupModule()
 	FPropertyEditorModule& PropertyModule = FModuleManager::GetModuleChecked<FPropertyEditorModule>("PropertyEditor");
 	PropertyModule.RegisterCustomClassLayout(UIKRigDefinition::StaticClass()->GetFName(), FOnGetDetailCustomizationInstance::CreateStatic(&FIKRigDefinitionDetails::MakeInstance));
 
-//	PropertyModule.RegisterCustomPropertyTypeLayout("IKRigEffector", FOnGetPropertyTypeCustomizationInstance::CreateStatic(&FIKRigEffectorCustomization::MakeInstance));
+	PropertyModule.RegisterCustomPropertyTypeLayout("IKRigEffector", FOnGetPropertyTypeCustomizationInstance::CreateStatic(&FIKRigEffectorCustomization::MakeInstance));
 
 	FEditorModeRegistry::Get().RegisterMode<FIKRigEditMode>(UAnimGraphNode_IKRig::AnimModeName, LOCTEXT("IKRigEditMode", "IKRig"), FSlateIcon(), false);
 }
@@ -41,9 +42,9 @@ void FIKRigEditor::ShutdownModule()
 {
 	FEditorModeRegistry::Get().UnregisterMode(UAnimGraphNode_IKRig::AnimModeName);
 
-//	PropertyModule.UnregisterCustomPropertyTypeLayout("IKRigEffector");
-
 	FPropertyEditorModule& PropertyModule = FModuleManager::GetModuleChecked<FPropertyEditorModule>("PropertyEditor");
+	PropertyModule.UnregisterCustomPropertyTypeLayout("IKRigEffector");
+
 	PropertyModule.UnregisterCustomClassLayout(UIKRigDefinition::StaticClass()->GetFName());
 
 	if (IKRigDefinitionAssetAction.IsValid())
diff --git a/Engine/Plugins/Animation/IKRig/Source/IKRigEditor/Private/IKRigEffectorCustomization.cpp b/Engine/Plugins/Animation/IKRig/Source/IKRigEditor/Private/IKRigEffectorCustomization.cpp
new file mode 100644
--- /dev/null
+++ b/Engine/Plugins/Animation/IKRig/Source/IKRigEditor/Private/IKRigEffectorCustomization.cpp
@@ -0,0 +1,270 @@
+// Copyright Epic Games, Inc. All Rights Reservekd.
+
+#include "IKRigEffectorCustomization.h"
+#include "DetailWidgetRow.h"
+#include "DetailLayoutBuilder.h"
+#include "IDetailChildrenBuilder.h"
+
+#include "IKRigController.h"
+#include "IKRigSolverDefinition.h"
+ 
+#include "ScopedTransaction.h"
+#include "PropertyEditorModule.h"
+#include "PropertyHandle.h"
+#include "SSearchableComboBox.h"
+
+#define LOCTEXT_NAMESPACE	"IKRigDefinitionDetails"
+
+/////////////////////////////////////////////////////////////////////////////////////////////
+//  FIKRigEffectorCustomization
+/////////////////////////////////////////////////////////////////////////////////////////////
+
+TSharedRef<IPropertyTypeCustomization> FIKRigEffectorCustomization::MakeInstance()
+{
+	return MakeShareable(new FIKRigEffectorCustomization());
+}
+
+void FIKRigEffectorCustomization::CustomizeHeader(TSharedRef<IPropertyHandle> StructPropertyHandle, class FDetailWidgetRow& HeaderRow, IPropertyTypeCustomizationUtils& StructCustomizationUtils)
+{
+	// find actual IKRigEffector - this is to query the data itself
+	void* Data;
+	FPropertyAccess::Result Result = StructPropertyHandle->GetValueData(Data);
+	// if not a single value, we return, then we return
+	if (Result != FPropertyAccess::Success)
+	{
+		// failed
+		return;
+	}
+
+	// set editable skeleton info from struct
+	SetSolverDefinition(StructPropertyHandle);
+	SetPropertyHandle(StructPropertyHandle);
+
+	// list out available goals
+	// set the one currently selected
+	if (IKRigController.IsValid())
+	{
+		TSharedPtr<FString> InitialSelected;
+		
+		TArray<FName> Goals;
+		IKRigController->QueryGoals(Goals);
+
+		Effector= reinterpret_cast<FIKRigEffector*>(Data);
+		// give me the name of goal by effecto
+		FName CurrentGoalName = IKRigController->GetGoalName(IKRigSolverDefinition.Get(), *Effector);
+
+		// go through profile and see if it has mine
+		for (const FName& Goal : Goals)
+		{
+			EffectorGoalComboList.Add(MakeShareable(new FString(Goal.ToString())));
+
+			if (CurrentGoalName == Goal)
+			{
+				InitialSelected = EffectorGoalComboList.Last();
+			}
+		}
+
+		HeaderRow
+		.NameContent()
+			[
+				StructPropertyHandle->CreatePropertyNameWidget()
+			]
+		.ValueContent()
+			[
+				SAssignNew(EffectorGoalComboBox, SSearchableComboBox)
+				.OptionsSource(&EffectorGoalComboList)
+				.OnGenerateWidget(this, &FIKRigEffectorCustomization::MakeEffectorGoalComboWidget)
+				.OnSelectionChanged(this, &FIKRigEffectorCustomization::OnEffectorGoalSelectionChanged)
+				.OnComboBoxOpening(this, &FIKRigEffectorCustomization::OnEffectorGoalComboOpening)
+				.InitiallySelectedItem(InitialSelected)
+				.IsEnabled(FSlateApplication::Get().GetNormalExecutionAttribute())
+				.ContentPadding(0)
+				.Content()
+				[
+					SNew(SEditableTextBox)
+					.Text(this, &FIKRigEffectorCustomization::GetEffectorGoalComboBoxContent)
+					.Font(IDetailLayoutBuilder::GetDetailFont())
+					.OnTextCommitted(this, &FIKRigEffectorCustomization::OnEffectorGoalChanged)
+					.ToolTipText(this, &FIKRigEffectorCustomization::GetEffectorGoalComboBoxToolTip)
+				]
+			];
+
+			// I need to see bone name, where is it?
+	}
+	else
+	{
+		// if this FIKRigEffector is used by some other Outers, this will fail	
+		// should warn programmers instead of silent fail
+		ensureAlways(false);
+	}
+}
+
+void FIKRigEffectorCustomization::CustomizeChildren(TSharedRef<IPropertyHandle> StructPropertyHandle, class IDetailChildrenBuilder& ChildBuilder, IPropertyTypeCustomizationUtils& StructCustomizationUtils)
+{
+	if (EffectorBoneProperty)
+	{
+		ChildBuilder.AddProperty(EffectorBoneProperty.ToSharedRef());
+	}
+}
+
+TSharedRef<SWidget> FIKRigEffectorCustomization::MakeEffectorGoalComboWidget(TSharedPtr<FString> InItem)
+{
+	return 	SNew(STextBlock).Text(FText::FromString(*InItem)).Font(IDetailLayoutBuilder::GetDetailFont());
+}
+
+void FIKRigEffectorCustomization::SetSolverDefinition(TSharedRef<IPropertyHandle> StructPropertyHandle)
+{
+	TArray<UObject*> Objects;
+	StructPropertyHandle->GetOuterObjects(Objects);
+
+	IKRigSolverDefinition.Reset();
+	IKRigController.Reset();
+
+	UIKRigSolverDefinition* SelectedSolverDefinition = nullptr;
+	
+	// currently not allowing multi selection
+	// or if you select different IKRigSolvers, this will break
+	for (UObject* Outer : Objects)
+	{
+		if (UIKRigSolverDefinition* SolverDefinition = Cast<UIKRigSolverDefinition>(Outer))
+		{
+			SelectedSolverDefinition = SolverDefinition;
+			break;
+		}
+	}
+
+	IKRigSolverDefinition = SelectedSolverDefinition;
+
+	// once we know solver definion, we know outer is IKRigDefinition
+	if (SelectedSolverDefinition)
+	{
+		UIKRigDefinition* IKRigDefinition = Cast<UIKRigDefinition>(SelectedSolverDefinition->GetOuter());
+		if (IKRigDefinition)
+		{
+			IKRigController = UIKRigController::GetControllerByRigDefinition(IKRigDefinition);
+
+			IKRigController->OnGoalModified.AddSP(this, &FIKRigEffectorCustomization::RefreshEffectorGoals);
+		}
+	}
+}
+
+void FIKRigEffectorCustomization::RefreshEffectorGoals()
+{
+	TArray<FName> Goals;
+	IKRigController->QueryGoals(Goals);
+
+	EffectorGoalComboList.Reset();
+
+	// go through profile and see if it has mine
+	for (const FName& Goal : Goals)
+	{
+		EffectorGoalComboList.Add(MakeShareable(new FString(Goal.ToString())));
+	}
+
+	EffectorGoalComboBox->ClearSelection();
+	EffectorGoalComboBox->RefreshOptions();
+}
+
+TSharedPtr<IPropertyHandle> FIKRigEffectorCustomization::FindStructMemberProperty(TSharedRef<IPropertyHandle> PropertyHandle, const FName& PropertyName)
+{
+	uint32 NumChildren = 0;
+	PropertyHandle->GetNumChildren(NumChildren);
+	for (uint32 ChildIdx = 0; ChildIdx < NumChildren; ++ChildIdx)
+	{
+		TSharedPtr<IPropertyHandle> ChildHandle = PropertyHandle->GetChildHandle(ChildIdx);
+		if (ChildHandle->GetProperty()->GetFName() == PropertyName)
+		{
+			return ChildHandle;
+		}
+	}
+
+	return TSharedPtr<IPropertyHandle>();
+}
+
+void FIKRigEffectorCustomization::SetPropertyHandle(TSharedRef<IPropertyHandle> StructPropertyHandle)
+{
+	EffectorBoneProperty = FindStructMemberProperty(StructPropertyHandle, GET_MEMBER_NAME_CHECKED(FIKRigEffector, Bone));
+	check(EffectorBoneProperty->IsValidHandle());
+}
+
+// this is only selection change
+void FIKRigEffectorCustomization::SetNewGoalName(FName Name)
+{
+	// I have to change the mapping of SolverDefinition
+	if (UIKRigSolverDefinition* SolverDef = IKRigSolverDefinition.Get())
+	{
+		FScopedTransaction Transaction(LOCTEXT("ApplyAdditiveSetting_Transaciton", "Apply Additive Setting"));
+		SolverDef->Modify();
+		IKRigController->SetGoalName(SolverDef, *Effector, Name);
+	}
+}
+
+FName FIKRigEffectorCustomization::GetSelectedEffectorGoal() const
+{
+	if (UIKRigSolverDefinition* SolverDef = IKRigSolverDefinition.Get())
+	{
+		return IKRigController->GetGoalName(SolverDef, *Effector);
+	}
+
+	return NAME_None;
+}
+
+void FIKRigEffectorCustomization::OnEffectorGoalSelectionChanged(TSharedPtr<FString> NewSelection, ESelectInfo::Type SelectInfo)
+{
+	if (SelectInfo != ESelectInfo::Direct)
+	{
+		FString NewValue = *NewSelection.Get();
+
+		SetNewGoalName(FName(*NewValue, FNAME_Find));
+	}
+}
+
+void FIKRigEffectorCustomization::OnEffectorGoalChanged(const FText& InText, ETextCommit::Type CommitType)
+{
+	if (UIKRigSolverDefinition* SolverDef = IKRigSolverDefinition.Get())
+	{
+		FString NewGoalName = InText.ToString();
+		NewGoalName.TrimStartAndEndInline();
+		if (!NewGoalName.IsEmpty())
+		{
+			IKRigController->SetGoalName(SolverDef, *Effector, FName(*NewGoalName));
+		}
+	}
+}
+
+void FIKRigEffectorCustomization::OnEffectorGoalComboOpening()
+{
+	TSharedPtr<FString> ComboStringPtr = GetEffectorGoalString();
+	if (ComboStringPtr.IsValid())
+	{
+		EffectorGoalComboBox->SetSelectedItem(ComboStringPtr);
+	}
+}
+
+FText FIKRigEffectorCustomization::GetEffectorGoalComboBoxContent() const
+{
+	return FText::FromName(GetSelectedEffectorGoal());
+}
+
+FText FIKRigEffectorCustomization::GetEffectorGoalComboBoxToolTip() const
+{
+	return LOCTEXT("EffectorGoalComboToolTip", "This will be the Goal name used externally to access this effector.");
+}
+
+TSharedPtr<FString> FIKRigEffectorCustomization::GetEffectorGoalString() const
+{
+	FName EffectorGoalName = GetSelectedEffectorGoal();
+
+	// go through profile and see if it has mine
+	for (int32 Index = 0; Index < EffectorGoalComboList.Num(); ++Index)
+	{
+		if (EffectorGoalName.ToString() == *EffectorGoalComboList[Index])
+		{
+			return EffectorGoalComboList[Index];
+		}
+	}
+
+	return TSharedPtr<FString>();
+}
+
+#undef LOCTEXT_NAMESPACE
diff --git a/Engine/Plugins/Animation/IKRig/Source/IKRigEditor/Private/IKRigEffectorCustomization.h b/Engine/Plugins/Animation/IKRig/Source/IKRigEditor/Private/IKRigEffectorCustomization.h
new file mode 100644
--- /dev/null
+++ b/Engine/Plugins/Animation/IKRig/Source/IKRigEditor/Private/IKRigEffectorCustomization.h
@@ -0,0 +1,63 @@
+// Copyright Epic Games, Inc. All Rights Reserved.
+
+#pragma once
+
+#include "CoreMinimal.h"
+#include "IPropertyTypeCustomization.h"
+#include "IKRigDataTypes.h"
+
+class UIKRigSolverDefinition;
+class UIKRigController;
+class SSearchableComboBox;
+
+//////////////////////////////////////////////////////////////////////////
+// FIKRigEffectorCustomization
+
+class FIKRigEffectorCustomization : public IPropertyTypeCustomization
+{
+public:
+	static TSharedRef<IPropertyTypeCustomization> MakeInstance();
+
+	// IPropertyTypeCustomization interface
+	virtual void CustomizeHeader(TSharedRef<IPropertyHandle> StructPropertyHandle, class FDetailWidgetRow& HeaderRow, IPropertyTypeCustomizationUtils& StructCustomizationUtils) override;
+	virtual void CustomizeChildren(TSharedRef<IPropertyHandle> StructPropertyHandle, class IDetailChildrenBuilder& ChildBuilder, IPropertyTypeCustomizationUtils& StructCustomizationUtils) override;
+
+private:
+	void SetSolverDefinition(TSharedRef<IPropertyHandle> StructPropertyHandle);
+	virtual void SetPropertyHandle(TSharedRef<IPropertyHandle> StructPropertyHandle);
+	TSharedPtr<IPropertyHandle> FindStructMemberProperty(TSharedRef<IPropertyHandle> PropertyHandle, const FName& PropertyName);
+
+	// Property to change after bone has been picked
+	TSharedPtr<IPropertyHandle> EffectorBoneProperty;
+
+	// Target Skeleton this widget is referencing
+	TWeakObjectPtr<UIKRigSolverDefinition>	IKRigSolverDefinition;
+	TWeakObjectPtr<UIKRigController>		IKRigController;
+
+	// Goal update option
+	TSharedPtr< SSearchableComboBox >			EffectorGoalComboBox;
+	TArray< TSharedPtr< FString > >				EffectorGoalComboList;
+	FIKRigEffector*								Effector;								
+
+	TSharedRef<SWidget> MakeEffectorGoalComboWidget(TSharedPtr<FString> InItem);
+	// combo box
+	void OnEffectorGoalSelectionChanged(TSharedPtr<FString> NewSelection, ESelectInfo::Type SelectInfo);
+	// edit box
+	void OnEffectorGoalChanged(const FText& InText, ETextCommit::Type CommitType);
+
+	FText GetEffectorGoalComboBoxContent() const;
+	// tool tip test
+	FText GetEffectorGoalComboBoxToolTip() const;
+	// when combo box is opened
+	void OnEffectorGoalComboOpening();
+	TSharedPtr<FString> GetEffectorGoalString() const;
+	// this is possibly renamed OR reusing same name
+	FName GetSelectedEffectorGoal() const;
+
+	void SetNewGoalName(FName Name);
+
+	const TIKRigEffectorMap<FName>& GetEffectorMap() const;
+	const TArray<FName>& GetAvailableGoals() const;
+
+	void RefreshEffectorGoals();
+};
\ No newline at end of file
diff --git a/Engine/Plugins/Animation/IKRig/Source/IKRigEditor/Public/IKRigController.h b/Engine/Plugins/Animation/IKRig/Source/IKRigEditor/Public/IKRigController.h
--- a/Engine/Plugins/Animation/IKRig/Source/IKRigEditor/Public/IKRigController.h
+++ b/Engine/Plugins/Animation/IKRig/Source/IKRigEditor/Public/IKRigController.h
@@ -76,8 +76,13 @@ public:
 
 	// goal operators
 	void QueryGoals(TArray<FName>& OutGoals) const;
+	// why do we need rename? @todo: remove - maybe in the fuutre when we want a tool that does rename automatically this could be useful, but for now it's not really
+	// goals are collected from solvers, so they could change name, and if it doesn't exists, it doesn't exist
 	void RenameGoal(const FName& OldName, const FName& NewName);
 
+	FName GetGoalName(UIKRigSolverDefinition* InSolverDefinition, const FIKRigEffector& InEffector);
+	void SetGoalName(UIKRigSolverDefinition* InSolverDefinition, const FIKRigEffector& InEffector, const FName& NewGoalName);
+
 	// this is to modify default value
 	FIKRigGoal* GetGoal(const FName& InGoalName);
 	const FIKRigGoal* GetGoal(const FName& InGoalName) const;
@@ -90,6 +95,11 @@ public:
 	// Controller getter
 	// Use this interface to get controller 
 	static UIKRigController* GetControllerByRigDefinition(UIKRigDefinition* InIKRigDefinition);
+
+	// delegates
+	DECLARE_MULTICAST_DELEGATE(FGoalModified);
+
+	FGoalModified OnGoalModified;
 private:
 	UIKRigDefinition* IKRigDefinition = nullptr;
 
