Fix undo/redo on controller

diff --git a/Engine/Plugins/Animation/IKRig/Source/IKRigEditor/Private/IKRigController.cpp b/Engine/Plugins/Animation/IKRig/Source/IKRigEditor/Private/IKRigController.cpp
--- a/Engine/Plugins/Animation/IKRig/Source/IKRigEditor/Private/IKRigController.cpp
+++ b/Engine/Plugins/Animation/IKRig/Source/IKRigEditor/Private/IKRigController.cpp
@@ -9,6 +9,9 @@
 #include "IKRigDefinition.h"
 #include "AnimationRuntime.h"
 #include "IKRigSolverDefinition.h"
+#include "ScopedTransaction.h"
+
+#define LOCTEXT_NAMESPACE	"IKRigController"
 
 TMap<UIKRigDefinition*, UIKRigController*> UIKRigController::DefinitionToControllerMap;
 
@@ -123,6 +126,9 @@ void UIKRigController::SetSkeleton(const FReferenceSkeleton& InSkeleton)
 {
 	if (IKRigDefinition)
 	{
+		FScopedTransaction Transaction(LOCTEXT("SetSkeleton_Label", "Set Skeleton"));
+		IKRigDefinition->Modify();
+
 		IKRigDefinition->ResetHierarchy();
 
 		const TArray<FMeshBoneInfo>& RefBoneInfo = InSkeleton.GetRefBoneInfo();
@@ -143,6 +149,9 @@ bool UIKRigController::AddBone(const FName& InName, const FName& InParent, const
 {
 	if (IKRigDefinition)
 	{
+		FScopedTransaction Transaction(LOCTEXT("AddBone_Label", "Add Bone"));
+		IKRigDefinition->Modify();
+
 		return IKRigDefinition->AddBone(InName, InParent, InGlobalTransform);
 	}
 
@@ -153,6 +162,9 @@ bool UIKRigController::RemoveBone(const FName& InName)
 {
 	if (IKRigDefinition)
 	{
+		FScopedTransaction Transaction(LOCTEXT("RemoveBone_Label", "Remove Bone"));
+		IKRigDefinition->Modify();
+
 		return IKRigDefinition->RemoveBone(InName);
 	}
 
@@ -162,6 +174,9 @@ bool UIKRigController::RenameBone(const FName& InOldName, const FName& InNewName
 {
 	if (IKRigDefinition)
 	{
+		FScopedTransaction Transaction(LOCTEXT("RenameBone_Label", "Rename Bone"));
+		IKRigDefinition->Modify();
+
 		return IKRigDefinition->RenameBone(InOldName, InNewName);
 	}
 
@@ -171,6 +186,9 @@ bool UIKRigController::ReparentBone(const FName& InName, const FName& InNewParen
 {
 	if (IKRigDefinition)
 	{
+		FScopedTransaction Transaction(LOCTEXT("ReparentBone_Label", "Reparent Bone"));
+		IKRigDefinition->Modify();
+
 		return IKRigDefinition->ReparentBone(InName, InNewParent);
 	}
 
@@ -180,6 +198,9 @@ void UIKRigController::ResetHierarchy()
 {
 	if (IKRigDefinition)
 	{
+		FScopedTransaction Transaction(LOCTEXT("ResetHierarchy_Label", "Reset Hierarchy"));
+		IKRigDefinition->Modify();
+
 		IKRigDefinition->ResetHierarchy();
 	}
 }
@@ -189,6 +210,9 @@ UIKRigSolverDefinition* UIKRigController::AddSolver(TSubclassOf<UIKRigSolverDefi
 {
 	if (IKRigDefinition)
 	{
+		FScopedTransaction Transaction(LOCTEXT("AddSolver_Label", "Add Solver"));
+		IKRigDefinition->Modify();
+
 		UIKRigSolverDefinition* NewSolver = NewObject<UIKRigSolverDefinition>(IKRigDefinition, InIKRigSolverDefinitionClass);
 		check(NewSolver);
 
@@ -226,6 +250,9 @@ void UIKRigController::RemoveSolver(UIKRigSolverDefinition* SolverToDelete)
 {
 	if (IKRigDefinition && SolverToDelete)
 	{
+		FScopedTransaction Transaction(LOCTEXT("RemoveSolver_Label", "Remove Solver"));
+		IKRigDefinition->Modify();
+
 		UninitializeIKRigSolverDefinition(SolverToDelete);
 		IKRigDefinition->SolverDefinitions.Remove(SolverToDelete);
 		IKRigDefinition->UpdateGoal();
@@ -241,6 +268,9 @@ void UIKRigController::UpdateGoal()
 {
 	if (IKRigDefinition)
 	{
+		FScopedTransaction Transaction(LOCTEXT("UpdateGoal_Label", "Update Goal"));
+		IKRigDefinition->Modify();
+
 		IKRigDefinition->UpdateGoal();
 		OnGoalModified.Broadcast();
 	}
@@ -264,6 +294,9 @@ void UIKRigController::SetGoalName(UIKRigSolverDefinition* InSolverDefinition, c
 {
 	if (InSolverDefinition)
 	{
+		FScopedTransaction Transaction(LOCTEXT("SetGoalName_Label", "Set Goal Name"));
+		InSolverDefinition->Modify();
+
 		FName* GoalName = InSolverDefinition->EffectorToGoal.Find(InEffector);
 		if (GoalName)
 		{
@@ -278,6 +311,9 @@ void UIKRigController::AutoConfigure(UIKRigSolverDefinition* SolverDef)
 {
 	if (ValidateSolver(SolverDef))
 	{
+		FScopedTransaction Transaction(LOCTEXT("AutoConfigure_Label", "Auto Congirure"));
+		SolverDef->Modify();
+
 		SolverDef->AutoConfigure();
 	}
 }
@@ -348,6 +384,9 @@ void UIKRigController::RenameGoal(const FName& OldName, const FName& NewName)
 	// ensure we don't have used NewDisplayName 
 	if (IKRigDefinition)
 	{
+		FScopedTransaction Transaction(LOCTEXT("RenameGoal_Label", "Rename Goal"));
+		IKRigDefinition->Modify();
+
 		for (UIKRigSolverDefinition* SolverDef : IKRigDefinition->SolverDefinitions)
 		{
 			if (SolverDef)
@@ -379,4 +418,6 @@ const FIKRigGoal* UIKRigController::GetGoal(const FName& InGoalName) const
 	}
 
 	return nullptr;
-}
\ No newline at end of file
+}
+
+#undef LOCTEXT_NAMESPACE
\ No newline at end of file
diff --git a/Engine/Plugins/Animation/IKRig/Source/IKRigEditor/Private/IKRigEffectorCustomization.cpp b/Engine/Plugins/Animation/IKRig/Source/IKRigEditor/Private/IKRigEffectorCustomization.cpp
--- a/Engine/Plugins/Animation/IKRig/Source/IKRigEditor/Private/IKRigEffectorCustomization.cpp
+++ b/Engine/Plugins/Animation/IKRig/Source/IKRigEditor/Private/IKRigEffectorCustomization.cpp
@@ -8,7 +8,6 @@
 #include "IKRigController.h"
 #include "IKRigSolverDefinition.h"
  
-#include "ScopedTransaction.h"
 #include "PropertyEditorModule.h"
 #include "PropertyHandle.h"
 #include "SSearchableComboBox.h"
@@ -193,8 +192,6 @@ void FIKRigEffectorCustomization::SetNewGoalName(FName Name)
 	// I have to change the mapping of SolverDefinition
 	if (UIKRigSolverDefinition* SolverDef = IKRigSolverDefinition.Get())
 	{
-		FScopedTransaction Transaction(LOCTEXT("ApplyAdditiveSetting_Transaciton", "Apply Additive Setting"));
-		SolverDef->Modify();
 		IKRigController->SetGoalName(SolverDef, *Effector, Name);
 	}
 }
