- Constraint first work. It's not working, but it compiles, and want to check in for compile issue
- Fix static analysis warning

diff --git a/Engine/Plugins/Animation/IKRig/Source/IKRig/Private/Constraints/RotationLimitConstraint.h b/Engine/Plugins/Animation/IKRig/Source/IKRig/Private/Constraints/RotationLimitConstraint.h
new file mode 100644
--- /dev/null
+++ b/Engine/Plugins/Animation/IKRig/Source/IKRig/Private/Constraints/RotationLimitConstraint.h
@@ -0,0 +1,38 @@
+// Copyright Epic Games, Inc. All Rights Reserved.
+
+/**
+ * Contains RotationLimit Definition 
+ *
+ */
+
+#pragma once
+
+#include "IKRigConstraint.h"
+#include "RotationLimitConstraint.generated.h"
+
+UCLASS(config = Engine, hidecategories = UObject, BlueprintType)
+class IKRIG_API URotationLimitConstraint : public UIKRigConstraint
+{
+	GENERATED_BODY()
+
+	int32 BaseIndex;
+	int32 ConstrainedIndex;
+
+	// frame of reference for this
+	FQuat		BaseFrameOffset; // include any offset and in their local space
+	FTransform	RelativelRefPose;
+
+	// xyz - twist/swing1/swing2
+	UPROPERTY(EditAnywhere, Category = "Limit")
+	uint8 bXLimitSet : 1;
+	UPROPERTY(EditAnywhere, Category = "Limit")
+	uint8 bYLimitSet : 1;
+	UPROPERTY(EditAnywhere, Category = "Limit")
+	uint8 bZLimitSet : 1;
+	UPROPERTY(EditAnywhere, Category = "Limit")
+	FVector Limit;
+public: 
+	virtual void SetupInternal(FIKRigTransformModifier& InOutTransformModifier) override {};
+	virtual void Apply(FIKRigTransformModifier& InOutTransformModifier, FControlRigDrawInterface* InOutDrawInterface) override {};
+};
+
diff --git a/Engine/Plugins/Animation/IKRig/Source/IKRig/Private/IKRigConstraint.cpp b/Engine/Plugins/Animation/IKRig/Source/IKRig/Private/IKRigConstraint.cpp
--- a/Engine/Plugins/Animation/IKRig/Source/IKRig/Private/IKRigConstraint.cpp
+++ b/Engine/Plugins/Animation/IKRig/Source/IKRig/Private/IKRigConstraint.cpp
@@ -7,8 +7,18 @@
 
 #include "IKRigConstraint.h"
 
-// UIKRigConstraint::UIKRigConstraint(const FObjectInitializer& ObjectInitializer)
-// 	: Super(ObjectInitializer)
-// {
-// 
-// }
+void UIKRigConstraint::Setup(FIKRigTransformModifier& InOutTransformModifier)
+{
+	SetupInternal(InOutTransformModifier);
+	bInitialized = true;
+}
+
+void UIKRigConstraint::SetAndApplyConstraint(FIKRigTransformModifier& InOutTransformModifier)
+{
+	if (!bInitialized)
+	{
+		Setup(InOutTransformModifier);
+	}
+
+	Apply(InOutTransformModifier, nullptr);
+}
diff --git a/Engine/Plugins/Animation/IKRig/Source/IKRig/Private/IKRigConstraintDefinition.cpp b/Engine/Plugins/Animation/IKRig/Source/IKRig/Private/IKRigConstraintDefinition.cpp
--- a/Engine/Plugins/Animation/IKRig/Source/IKRig/Private/IKRigConstraintDefinition.cpp
+++ b/Engine/Plugins/Animation/IKRig/Source/IKRig/Private/IKRigConstraintDefinition.cpp
@@ -12,5 +12,4 @@ const FName UIKRigConstraintDefinition::DefaultProfileName = FName(TEXT("Default
 UIKRigConstraintDefinition::UIKRigConstraintDefinition()
 {
 	FIKRigConstraintProfile& DefaultProfile = ConstraintProfiles.Add(DefaultProfileName);
-	DefaultProfile.Name = DefaultProfileName;
 }
diff --git a/Engine/Plugins/Animation/IKRig/Source/IKRig/Private/IKRigConstraintSolver.cpp b/Engine/Plugins/Animation/IKRig/Source/IKRig/Private/IKRigConstraintSolver.cpp
--- a/Engine/Plugins/Animation/IKRig/Source/IKRig/Private/IKRigConstraintSolver.cpp
+++ b/Engine/Plugins/Animation/IKRig/Source/IKRig/Private/IKRigConstraintSolver.cpp
@@ -5,6 +5,7 @@
  *
  */
 #include "IKRigConstraintSolver.h"
+#include "IKRigConstraint.h"
 
 // register/unregister query function
 void UIKRigConstraintSolver::RegisterQueryConstraintHandler(const FIKRigQueryConstraint& InQueryConstraintHandler)
@@ -15,4 +16,78 @@ void UIKRigConstraintSolver::RegisterQueryConstraintHandler(const FIKRigQueryCon
 void UIKRigConstraintSolver::UnregisterQueryConstraintHandler()
 {
 
+}
+
+void UIKRigConstraintSolver::InitInternal(const FIKRigTransformModifier& InGlobalTransform)
+{
+	if (ConstraintDefinition)
+	{
+		// we duplicate and create a copy of the objects, so that we can mutate
+		ConstraintProfiles = ConstraintDefinition->ConstraintProfiles;
+		for (auto Iter = ConstraintProfiles.CreateIterator(); Iter; ++Iter)
+		{
+			FIKRigConstraintProfile& CurrentProfile = Iter.Value();
+			for (auto InnerIter = CurrentProfile.Constraints.CreateIterator(); InnerIter; ++InnerIter)
+			{
+				UIKRigConstraint*& Template = InnerIter.Value();
+				UIKRigConstraint* CopyConstraint = NewObject<UIKRigConstraint>(this, Template->GetClass(), TEXT("IKRigConstraint"), RF_NoFlags, Template);
+				// replace value
+				Template = CopyConstraint;
+			}
+		}
+
+		// go through all UIKRigConstraint and create copy
+		ActiveProfile = UIKRigConstraintDefinition::DefaultProfileName;
+	}
+}
+
+void UIKRigConstraintSolver::SolveInternal(FIKRigTransformModifier& InOutGlobalTransform, FControlRigDrawInterface* InOutDrawInterface)
+{
+	FIKRigConstraintProfile* Current = ConstraintProfiles.Find(ActiveProfile);
+
+	if (Current)
+	{
+		for (auto Iter = Current->Constraints.CreateIterator(); Iter; ++Iter)
+		{
+			UIKRigConstraint* Constraint = Iter.Value();
+			if (Constraint)
+			{
+				// @todo: later add inoutdrawinterface
+				Constraint->Apply(InOutGlobalTransform, InOutDrawInterface);
+			}
+		}
+	}
+}
+
+void UIKRigConstraintSolver::SolverConstraints(const TArray<FName>& ConstraintsList, FIKRigTransformModifier& InOutGlobalTransform, FControlRigDrawInterface* InOutDrawInterface)
+{
+	FIKRigConstraintProfile* Current = ConstraintProfiles.Find(ActiveProfile);
+
+	if (Current)
+	{
+		for (const FName& ConstraintName : ConstraintsList)
+		{
+			UIKRigConstraint** Constraint = Current->Constraints.Find(ConstraintName);
+			if (Constraint)
+			{
+				// @todo: later add inoutdrawinterface
+				(*Constraint)->Apply(InOutGlobalTransform, InOutDrawInterface);
+			}
+		}
+	}
+}
+
+bool UIKRigConstraintSolver::IsSolverActive() const
+{	
+// later we can do more elaborated version
+	return Super::IsSolverActive();
+}
+
+void UIKRigConstraintSolver::SetConstraintDefinition(UIKRigConstraintDefinition* InConstraintDefinition)
+{
+	ConstraintDefinition = InConstraintDefinition;
+	// reset instanced data
+
+	ConstraintProfiles.Reset();
+	ActiveProfile = UIKRigConstraintDefinition::DefaultProfileName;
 }
\ No newline at end of file
diff --git a/Engine/Plugins/Animation/IKRig/Source/IKRig/Private/IKRigDefinition.cpp b/Engine/Plugins/Animation/IKRig/Source/IKRig/Private/IKRigDefinition.cpp
--- a/Engine/Plugins/Animation/IKRig/Source/IKRig/Private/IKRigDefinition.cpp
+++ b/Engine/Plugins/Animation/IKRig/Source/IKRig/Private/IKRigDefinition.cpp
@@ -6,11 +6,12 @@
 
 #include "IKRigDefinition.h"
 #include "IKRigSolverDefinition.h"
+#include "IKRigConstraintDefinition.h"
 
 UIKRigDefinition::UIKRigDefinition()
 {
 	// we create constraint definition for them
-	//ConstraintDefinitions = NewObject<UIKRigConstraintDefinition>(this, TEXT("ConstraintDefinition"));
+//	ConstraintDefinitions = CreateDefaultSubobject<UIKRigConstraintDefinition>(TEXT("ConstraintDefinition"));
 }
 
 #if WITH_EDITOR
@@ -301,7 +302,7 @@ void UIKRigDefinition::EnsureSortedCorrectly(bool bReSortIfNeeded)
 						}
 
 						Iter = Iter->Sibling;
-					} while (Iter != nullptr);
+					} while (true);
 				}
 			}
 		};
diff --git a/Engine/Plugins/Animation/IKRig/Source/IKRig/Public/IKRigConstraint.h b/Engine/Plugins/Animation/IKRig/Source/IKRig/Public/IKRigConstraint.h
--- a/Engine/Plugins/Animation/IKRig/Source/IKRig/Public/IKRigConstraint.h
+++ b/Engine/Plugins/Animation/IKRig/Source/IKRig/Public/IKRigConstraint.h
@@ -12,14 +12,21 @@
 #include "UObject/Object.h"
 #include "IKRigConstraint.generated.h"
 
+struct FIKRigTransformModifier;
+struct FControlRigDrawInterface;
+
 UCLASS(config = Engine, hidecategories = UObject, BlueprintType)
 class IKRIG_API UIKRigConstraint : public UObject
 {
 	GENERATED_BODY()
 
-private:
-
+	bool bInitialized = false;
 
 public: 
+	void Setup(FIKRigTransformModifier& InOutTransformModifier);
+	void SetAndApplyConstraint(FIKRigTransformModifier& InOutTransformModifier);
+
+	virtual void Apply(FIKRigTransformModifier& InOutTransformModifier, FControlRigDrawInterface* InOutDrawInterface) {};
+	virtual void SetupInternal(FIKRigTransformModifier& InOutTransformModifier) {};
 };
 
diff --git a/Engine/Plugins/Animation/IKRig/Source/IKRig/Public/IKRigConstraintDefinition.h b/Engine/Plugins/Animation/IKRig/Source/IKRig/Public/IKRigConstraintDefinition.h
--- a/Engine/Plugins/Animation/IKRig/Source/IKRig/Public/IKRigConstraintDefinition.h
+++ b/Engine/Plugins/Animation/IKRig/Source/IKRig/Public/IKRigConstraintDefinition.h
@@ -19,10 +19,7 @@ struct IKRIG_API FIKRigConstraintProfile
 {
 	GENERATED_BODY()
 
-	UPROPERTY(EditAnywhere, Category=FIKRigConstraintProfile)
-	FName Name;
-
-	UPROPERTY(EditAnywhere, Category=FIKRigConstraintProfile)
+	UPROPERTY(EditAnywhere, instanced, Category=FIKRigConstraintProfile)
 	TMap<FName, UIKRigConstraint*> Constraints;
 };
 
@@ -38,6 +35,7 @@ private:
 	// at least one profile always exists - Default
 	UPROPERTY(EditAnywhere, Category = "Profile")
 	TMap<FName, FIKRigConstraintProfile> ConstraintProfiles;
-
+	friend class UIKRigConstraintSolver;
+public:
 	static const FName DefaultProfileName;
 };
\ No newline at end of file
diff --git a/Engine/Plugins/Animation/IKRig/Source/IKRig/Public/IKRigConstraintSolver.h b/Engine/Plugins/Animation/IKRig/Source/IKRig/Public/IKRigConstraintSolver.h
--- a/Engine/Plugins/Animation/IKRig/Source/IKRig/Public/IKRigConstraintSolver.h
+++ b/Engine/Plugins/Animation/IKRig/Source/IKRig/Public/IKRigConstraintSolver.h
@@ -8,32 +8,49 @@
 #pragma once
 
 #include "IKRigSolver.h"
+#include "IKRigConstraintDefinition.h"
 #include "IKRigConstraintSolver.generated.h"
 
 class UIKRigConstraint;
 
+
 // run time processor 
-UCLASS(config = Engine, hidecategories = UObject)
+UCLASS(config = Engine, Transient, hidecategories = UObject)
 class IKRIG_API UIKRigConstraintSolver : public UIKRigSolver
 {
 	GENERATED_BODY()
 
 	DECLARE_DELEGATE_TwoParams(FIKRigQueryConstraint, const FName& InConstraintName, UIKRigConstraint& OutConstraint);
 
+	UPROPERTY(transient)
+	UIKRigConstraintDefinition* ConstraintDefinition;
+
 public: 
-	// input hierarchy and ref pose? 
-	//virtual void Init(UIKRigSolverDefinition* InSolverDefinition, const FIKRigTransform& InRefPose, FIKRigGoalGetter& InGoalGetter, FSolveConstraint& InConstraintHandler) override;
+	// during we don't mutate bone transform, but you can read it
+	virtual void InitInternal(const FIKRigTransformModifier& InGlobalTransform) override;
+	virtual void SolveInternal(FIKRigTransformModifier& InOutGlobalTransform, FControlRigDrawInterface* InOutDrawInterface) override;
+	virtual bool IsSolverActive() const override;
 
-	// input : goal getter or goals
-	// output : modified pose - GlobalTransforms
-	//virtual void Solve(FIKRigTransform& OutGlobalTransform) override;
+	// this can be utilized for just subset of constraints
+	void SolverConstraints(const TArray<FName>& ConstraintsList, FIKRigTransformModifier& InOutGlobalTransform, FControlRigDrawInterface* InOutDrawInterface);
 
 	// register/unregister query function
 	void RegisterQueryConstraintHandler(const FIKRigQueryConstraint& InQueryConstraintHandler);
 	void UnregisterQueryConstraintHandler();
 
+	void SetConstraintDefinition(UIKRigConstraintDefinition* InConstraintDefinition);
+
 private:
 	// delegate
 	FIKRigQueryConstraint QueryConstraintHandler;
+
+	// current active profile
+	UPROPERTY(transient)
+	FName ActiveProfile;
+
+	// instanced information you can mutate from ConstraintDefinition
+	// we want uproperty, so that it doesn't GC-ed
+	UPROPERTY(transient)
+	TMap<FName, FIKRigConstraintProfile> ConstraintProfiles;
 };
 
diff --git a/Engine/Plugins/Animation/IKRig/Source/IKRig/Public/IKRigDataTypes.h b/Engine/Plugins/Animation/IKRig/Source/IKRig/Public/IKRigDataTypes.h
--- a/Engine/Plugins/Animation/IKRig/Source/IKRig/Public/IKRigDataTypes.h
+++ b/Engine/Plugins/Animation/IKRig/Source/IKRig/Public/IKRigDataTypes.h
@@ -187,6 +187,7 @@ public:
 
 // allows transform to be modified 
 // use this class to modify transform
+// @todo: ref pose getter 
 struct IKRIG_API FIKRigTransformModifier
 {
 	FIKRigTransformModifier()
