Non-unity fixes

diff --git a/Engine/Plugins/Animation/IKRig/Source/IKRig/Public/IKRigDefinition.h b/Engine/Plugins/Animation/IKRig/Source/IKRig/Public/IKRigDefinition.h
--- a/Engine/Plugins/Animation/IKRig/Source/IKRig/Public/IKRigDefinition.h
+++ b/Engine/Plugins/Animation/IKRig/Source/IKRig/Public/IKRigDefinition.h
@@ -11,6 +11,7 @@
 #include "CoreMinimal.h"
 #include "UObject/ObjectMacros.h"
 #include "UObject/Object.h"
+//#include "UObject/WeakObjectPtr.h"
 #include "IKRigDataTypes.h"
 #include "IKRigHierarchy.h"
 #include "IKRigDefinition.generated.h"
diff --git a/Engine/Plugins/Animation/IKRig/Source/IKRig/Public/IKRigProcessor.h b/Engine/Plugins/Animation/IKRig/Source/IKRig/Public/IKRigProcessor.h
--- a/Engine/Plugins/Animation/IKRig/Source/IKRig/Public/IKRigProcessor.h
+++ b/Engine/Plugins/Animation/IKRig/Source/IKRig/Public/IKRigProcessor.h
@@ -14,8 +14,8 @@
 #include "Drawing/ControlRigDrawInterface.h"
 #include "IKRigProcessor.generated.h"
 
-
 class UIKRigConstraintSolver;
+class UIKRigDefinition;
 class UIKRigSolver;
 struct FControlRigDrawInterface;
 
diff --git a/Engine/Plugins/Animation/IKRig/Source/IKRigEditor/Private/IKRigDefinitionDetails.cpp b/Engine/Plugins/Animation/IKRig/Source/IKRigEditor/Private/IKRigDefinitionDetails.cpp
--- a/Engine/Plugins/Animation/IKRig/Source/IKRigEditor/Private/IKRigDefinitionDetails.cpp
+++ b/Engine/Plugins/Animation/IKRig/Source/IKRigEditor/Private/IKRigDefinitionDetails.cpp
@@ -14,7 +14,10 @@
 #include "ScopedTransaction.h"
 #include "PropertyCustomizationHelpers.h"
 #include "Kismet2/SClassPickerDialog.h"
+#include "Animation/Skeleton.h"
 #include "ClassViewerFilter.h"
+#include "Engine/SkeletalMesh.h"
+#include "Modules/ModuleManager.h"
 #include "Widgets/Input/SEditableTextBox.h"
 
 #define LOCTEXT_NAMESPACE	"IKRigDefinitionDetails"
diff --git a/Engine/Plugins/Animation/IKRig/Source/IKRigEditor/Private/IKRigDefinitionDetails.h b/Engine/Plugins/Animation/IKRig/Source/IKRigEditor/Private/IKRigDefinitionDetails.h
--- a/Engine/Plugins/Animation/IKRig/Source/IKRigEditor/Private/IKRigDefinitionDetails.h
+++ b/Engine/Plugins/Animation/IKRig/Source/IKRigEditor/Private/IKRigDefinitionDetails.h
@@ -3,12 +3,19 @@
 #pragma once
 
 #include "CoreMinimal.h"
+#include "UObject/WeakObjectPtr.h"
 #include "UObject/WeakObjectPtrTemplates.h"
 #include "IDetailCustomization.h"
+#include "Input/Reply.h"
+#include "Widgets/Views/SListView.h"
 
 class IDetailLayoutBuilder;
+class IPropertyHandle;
+class ITableRow;
+class STableViewBase;
 class UIKRigDefinition;
 class UIKRigController;
+struct FAssetData;
 
 // Utility class to build combo boxes out of arrays of names.
 struct FGoalNameListItem : public TSharedFromThis<FGoalNameListItem>
diff --git a/Engine/Plugins/Experimental/ControlRig/Source/ControlRigEditor/Private/EditMode/SControlRigControlViews.h b/Engine/Plugins/Experimental/ControlRig/Source/ControlRigEditor/Private/EditMode/SControlRigControlViews.h
--- a/Engine/Plugins/Experimental/ControlRig/Source/ControlRigEditor/Private/EditMode/SControlRigControlViews.h
+++ b/Engine/Plugins/Experimental/ControlRig/Source/ControlRigEditor/Private/EditMode/SControlRigControlViews.h
@@ -15,6 +15,8 @@
 #include "Widgets/Text/STextBlock.h"
 #include "Widgets/Views/STableRow.h"
 #include "Tools/ControlRigPose.h"
+
+class FAssetThumbnail;
 class UControlRig;
 
 //Class to Hold Statics that are shared and externally callable
diff --git a/Engine/Plugins/Runtime/Metasound/Source/MetasoundEditor/Private/MetasoundEditorGraphBuilder.h b/Engine/Plugins/Runtime/Metasound/Source/MetasoundEditor/Private/MetasoundEditorGraphBuilder.h
--- a/Engine/Plugins/Runtime/Metasound/Source/MetasoundEditor/Private/MetasoundEditorGraphBuilder.h
+++ b/Engine/Plugins/Runtime/Metasound/Source/MetasoundEditor/Private/MetasoundEditorGraphBuilder.h
@@ -6,6 +6,7 @@
 
 // Forward Declarations
 class UEdGraphNode;
+class UEdGraphPin;
 class UMetasound;
 class UMetasoundEditorGraphNode;
 
diff --git a/Engine/Plugins/Runtime/Metasound/Source/MetasoundEngine/Private/MetasoundUObjectRegistry.cpp b/Engine/Plugins/Runtime/Metasound/Source/MetasoundEngine/Private/MetasoundUObjectRegistry.cpp
--- a/Engine/Plugins/Runtime/Metasound/Source/MetasoundEngine/Private/MetasoundUObjectRegistry.cpp
+++ b/Engine/Plugins/Runtime/Metasound/Source/MetasoundEngine/Private/MetasoundUObjectRegistry.cpp
@@ -4,6 +4,7 @@
 
 #include "Algo/Copy.h"
 #include "AssetRegistry/AssetRegistryModule.h"
+#include "Async/Async.h"
 #include "CoreMinimal.h"
 #include "UObject/Object.h"
 
diff --git a/Engine/Source/Editor/AssetPlacementEdMode/Private/AssetPlacementEdModeStyle.h b/Engine/Source/Editor/AssetPlacementEdMode/Private/AssetPlacementEdModeStyle.h
--- a/Engine/Source/Editor/AssetPlacementEdMode/Private/AssetPlacementEdModeStyle.h
+++ b/Engine/Source/Editor/AssetPlacementEdMode/Private/AssetPlacementEdModeStyle.h
@@ -5,6 +5,7 @@
 #include "CoreMinimal.h"
 
 class FSlateStyleSet;
+class ISlateStyle;
 
 class FAssetPlacementEdModeStyle
 {
diff --git a/Engine/Source/Editor/MovieSceneTools/Private/MovieSceneToolHelpers.cpp b/Engine/Source/Editor/MovieSceneTools/Private/MovieSceneToolHelpers.cpp
--- a/Engine/Source/Editor/MovieSceneTools/Private/MovieSceneToolHelpers.cpp
+++ b/Engine/Source/Editor/MovieSceneTools/Private/MovieSceneToolHelpers.cpp
@@ -12,6 +12,7 @@
 #include "ISequencer.h"
 #include "Modules/ModuleManager.h"
 #include "Framework/Application/SlateApplication.h"
+#include "Framework/MultiBox/MultiBoxBuilder.h"
 #include "Widgets/Input/SComboBox.h"
 #include "ScopedTransaction.h"
 #include "EditorStyleSet.h"
diff --git a/Engine/Source/Editor/Persona/Public/IAssetFamily.h b/Engine/Source/Editor/Persona/Public/IAssetFamily.h
--- a/Engine/Source/Editor/Persona/Public/IAssetFamily.h
+++ b/Engine/Source/Editor/Persona/Public/IAssetFamily.h
@@ -5,6 +5,8 @@
 #include "CoreMinimal.h"
 #include "AssetData.h"
 
+struct FSlateBrush;
+
 /** Represents a group of related assets, e.g. a skeleton, its animations and skeletal meshes */
 class IAssetFamily
 {
