Integrating Lina's final contributions to Pretzel with recent changes to fix Unity build.

diff --git a/Engine/Plugins/Animation/IKRig/Source/IKRig/Private/Constraints/RotationLimitConstraint.cpp b/Engine/Plugins/Animation/IKRig/Source/IKRig/Private/Constraints/RotationLimitConstraint.cpp
new file mode 100644
--- /dev/null
+++ b/Engine/Plugins/Animation/IKRig/Source/IKRig/Private/Constraints/RotationLimitConstraint.cpp
@@ -0,0 +1,100 @@
+// Copyright Epic Games, Inc. All Rights Reserved.
+
+/**
+ * Contains IKRigConstraint Runtime Implementation
+ *
+ */
+
+#include "RotationLimitConstraint.h"
+#include "IKRigHierarchy.h"
+#include "IKRigDataTypes.h"
+#include "IKRig.h"
+
+void URotationLimitConstraint::SetupInternal(const FIKRigTransformModifier& InOutTransformModifier)
+{
+	if (TargetBone != NAME_None)
+	{
+		const int32 TargetBoneIndex = InOutTransformModifier.Hierarchy->GetIndex(TargetBone);
+		if (TargetBoneIndex != INDEX_NONE)
+		{
+			// looking for local transform
+			const int32 TargetParentBoneIndex = InOutTransformModifier.Hierarchy->GetParentIndex(TargetBone);
+			// for now we only constraint to parent
+			BaseIndex = TargetParentBoneIndex;
+			ConstrainedIndex = TargetBoneIndex;
+
+			// if we support constraint to different joint, this function has to change to get relative transform
+			FTransform LocalTransform = InOutTransformModifier.GetLocalTransform(TargetBoneIndex);
+
+			// set rotation frame
+			BaseFrameOffset = FQuat(Offset);
+			RelativelRefPose = LocalTransform;
+		}
+	}
+}
+
+void URotationLimitConstraint::Apply(FIKRigTransformModifier& InOutTransformModifier, FControlRigDrawInterface* InOutDrawInterface)
+{
+	if (ConstrainedIndex != INDEX_NONE)
+	{
+		FTransform LocalTransform = InOutTransformModifier.GetLocalTransform(ConstrainedIndex);
+		// ref pose?
+		// just do rotation for now
+		FQuat LocalRotation = BaseFrameOffset.Inverse() * LocalTransform.GetRotation(); // child rotation
+		LocalRotation.Normalize();
+		FQuat LocalRefRotation = BaseFrameOffset.Inverse() * RelativelRefPose.GetRotation(); // later, maybe we'll make it more generic, so it doesn't always have to be just local transform
+		LocalRefRotation.Normalize();
+		bool bRotationChanged = false;
+
+		auto SetLimitedQuat = [&](EAxis::Type InAxis, float InLimtAngle)
+		{
+			FQuat DeltaQuat = LocalRefRotation.Inverse() * LocalRotation;
+
+			FVector RefTwistAxis = FMatrix::Identity.GetUnitAxis(InAxis);
+			FQuat Twist, Swing;
+			DeltaQuat.ToSwingTwist(RefTwistAxis, Swing, Twist);
+			float SwingAngle = Swing.GetAngle();
+			float TwistAngle = Twist.GetAngle();
+
+			UE_LOG(LogIKRig, Log, TEXT("Delta Decomposition : Swing %s (%f), Twist %s (%f)"),
+				*Swing.GetRotationAxis().ToString(), FMath::RadiansToDegrees(SwingAngle),
+				*Twist.GetRotationAxis().ToString(), FMath::RadiansToDegrees(TwistAngle));
+
+			// we deal with the extra ones left
+			LocalRotation = LocalRefRotation * Swing * FQuat(Twist.GetRotationAxis(), FMath::Clamp(TwistAngle, -InLimtAngle, InLimtAngle));
+			LocalRotation.Normalize();
+			bRotationChanged = true;
+		};
+
+		if (bXLimitSet)
+		{
+			SetLimitedQuat(EAxis::X, Limit.X);
+		}
+
+		if (bYLimitSet)
+		{
+			SetLimitedQuat(EAxis::Y, Limit.Y);
+		}
+
+		if (bZLimitSet)
+		{
+			SetLimitedQuat(EAxis::Z, Limit.Z);
+		}
+
+
+		if (bRotationChanged)
+		{
+			LocalRotation = BaseFrameOffset * LocalRotation;
+
+			if ((LocalRotation | LocalTransform.GetRotation()) < 0.f)
+			{
+				LocalRotation *= -1.f;
+			}
+
+			LocalTransform.SetRotation(LocalRotation);
+			LocalRotation.Normalize();
+
+			InOutTransformModifier.SetLocalTransform(ConstrainedIndex, LocalTransform, true);
+		}
+	}
+}
\ No newline at end of file
diff --git a/Engine/Plugins/Animation/IKRig/Source/IKRig/Private/Constraints/RotationLimitConstraint.h b/Engine/Plugins/Animation/IKRig/Source/IKRig/Private/Constraints/RotationLimitConstraint.h
--- a/Engine/Plugins/Animation/IKRig/Source/IKRig/Private/Constraints/RotationLimitConstraint.h
+++ b/Engine/Plugins/Animation/IKRig/Source/IKRig/Private/Constraints/RotationLimitConstraint.h
@@ -15,6 +15,18 @@ class IKRIG_API URotationLimitConstraint : public UIKRigConstraint
 {
 	GENERATED_BODY()
 
+	URotationLimitConstraint()
+		: BaseIndex(INDEX_NONE)
+		, ConstrainedIndex(INDEX_NONE)
+		, BaseFrameOffset(EForceInit::ForceInitToZero)
+		, RelativelRefPose(FVector::ZeroVector)
+		, bXLimitSet(false)
+		, bYLimitSet(false)
+		, bZLimitSet(false)
+		, Limit(0.f)
+	{
+	}
+
 	int32 BaseIndex;
 	int32 ConstrainedIndex;
 
@@ -22,6 +34,9 @@ class IKRIG_API URotationLimitConstraint : public UIKRigConstraint
 	FQuat		BaseFrameOffset; // include any offset and in their local space
 	FTransform	RelativelRefPose;
 
+	UPROPERTY(EditAnywhere, Category = "Limit")
+	FName TargetBone;
+
 	// xyz - twist/swing1/swing2
 	UPROPERTY(EditAnywhere, Category = "Limit")
 	uint8 bXLimitSet : 1;
@@ -29,10 +44,15 @@ class IKRIG_API URotationLimitConstraint : public UIKRigConstraint
 	uint8 bYLimitSet : 1;
 	UPROPERTY(EditAnywhere, Category = "Limit")
 	uint8 bZLimitSet : 1;
+
 	UPROPERTY(EditAnywhere, Category = "Limit")
 	FVector Limit;
+
+	// apply to local rotation if you want to modify offset of limit - i.e. knee
+	UPROPERTY(EditAnywhere, Category = "Offset")
+	FRotator Offset;
 public: 
-	virtual void SetupInternal(FIKRigTransformModifier& InOutTransformModifier) override {};
-	virtual void Apply(FIKRigTransformModifier& InOutTransformModifier, FControlRigDrawInterface* InOutDrawInterface) override {};
+	virtual void SetupInternal(const FIKRigTransformModifier& InOutTransformModifier) override;
+	virtual void Apply(FIKRigTransformModifier& InOutTransformModifier, FControlRigDrawInterface* InOutDrawInterface) override;
 };
 
diff --git a/Engine/Plugins/Animation/IKRig/Source/IKRig/Private/IKRigConstraint.cpp b/Engine/Plugins/Animation/IKRig/Source/IKRig/Private/IKRigConstraint.cpp
--- a/Engine/Plugins/Animation/IKRig/Source/IKRig/Private/IKRigConstraint.cpp
+++ b/Engine/Plugins/Animation/IKRig/Source/IKRig/Private/IKRigConstraint.cpp
@@ -7,7 +7,7 @@
 
 #include "IKRigConstraint.h"
 
-void UIKRigConstraint::Setup(FIKRigTransformModifier& InOutTransformModifier)
+void UIKRigConstraint::Setup(const FIKRigTransformModifier& InOutTransformModifier)
 {
 	SetupInternal(InOutTransformModifier);
 	bInitialized = true;
@@ -17,7 +17,7 @@ void UIKRigConstraint::SetAndApplyConstraint(FIKRigTransformModifier& InOutTrans
 {
 	if (!bInitialized)
 	{
-		Setup(InOutTransformModifier);
+		Setup(const_cast<FIKRigTransformModifier&>(InOutTransformModifier));
 	}
 
 	Apply(InOutTransformModifier, nullptr);
diff --git a/Engine/Plugins/Animation/IKRig/Source/IKRig/Private/IKRigConstraintDefinition.cpp b/Engine/Plugins/Animation/IKRig/Source/IKRig/Private/IKRigConstraintDefinition.cpp
--- a/Engine/Plugins/Animation/IKRig/Source/IKRig/Private/IKRigConstraintDefinition.cpp
+++ b/Engine/Plugins/Animation/IKRig/Source/IKRig/Private/IKRigConstraintDefinition.cpp
@@ -6,10 +6,14 @@
  */
 
 #include "IKRigConstraintDefinition.h"
+#include "IKRigConstraintSolver.h"
 
 const FName UIKRigConstraintDefinition::DefaultProfileName = FName(TEXT("Default"));
 
 UIKRigConstraintDefinition::UIKRigConstraintDefinition()
 {
 	FIKRigConstraintProfile& DefaultProfile = ConstraintProfiles.Add(DefaultProfileName);
+
+	DisplayName = TEXT("Constraint Solver");
+	ExecutionClass = UIKRigConstraintSolver::StaticClass();
 }
diff --git a/Engine/Plugins/Animation/IKRig/Source/IKRig/Private/IKRigConstraintSolver.cpp b/Engine/Plugins/Animation/IKRig/Source/IKRig/Private/IKRigConstraintSolver.cpp
--- a/Engine/Plugins/Animation/IKRig/Source/IKRig/Private/IKRigConstraintSolver.cpp
+++ b/Engine/Plugins/Animation/IKRig/Source/IKRig/Private/IKRigConstraintSolver.cpp
@@ -31,6 +31,7 @@ void UIKRigConstraintSolver::InitInternal(const FIKRigTransformModifier& InGloba
 			{
 				UIKRigConstraint*& Template = InnerIter.Value();
 				UIKRigConstraint* CopyConstraint = NewObject<UIKRigConstraint>(this, Template->GetClass(), TEXT("IKRigConstraint"), RF_NoFlags, Template);
+				CopyConstraint->Setup(InGlobalTransform);
 				// replace value
 				Template = CopyConstraint;
 			}
diff --git a/Engine/Plugins/Animation/IKRig/Source/IKRig/Private/IKRigDefinition.cpp b/Engine/Plugins/Animation/IKRig/Source/IKRig/Private/IKRigDefinition.cpp
--- a/Engine/Plugins/Animation/IKRig/Source/IKRig/Private/IKRigDefinition.cpp
+++ b/Engine/Plugins/Animation/IKRig/Source/IKRig/Private/IKRigDefinition.cpp
@@ -11,7 +11,7 @@
 UIKRigDefinition::UIKRigDefinition()
 {
 	// we create constraint definition for them
-//	ConstraintDefinitions = CreateDefaultSubobject<UIKRigConstraintDefinition>(TEXT("ConstraintDefinition"));
+	ConstraintDefinitions = CreateDefaultSubobject<UIKRigConstraintDefinition>(TEXT("ConstraintDefinition"));
 }
 
 #if WITH_EDITOR
@@ -368,7 +368,7 @@ void UIKRigDefinition::EnsureCreateUniqueGoalName(FName& InOutGoal) const
 	// while contains
 	while(IKGoals.Find(FName(*GoalNameString)))
 	{
-		GoalNameString += FString::Format(TEXT("_{0}"), {Index++});
+		GoalNameString = FString::Format(TEXT("{0}_{1}"), {InOutGoal.ToString(), Index++});
 	}
 
 	InOutGoal = FName(*GoalNameString);
diff --git a/Engine/Plugins/Animation/IKRig/Source/IKRig/Private/IKRigProcessor.cpp b/Engine/Plugins/Animation/IKRig/Source/IKRig/Private/IKRigProcessor.cpp
--- a/Engine/Plugins/Animation/IKRig/Source/IKRig/Private/IKRigProcessor.cpp
+++ b/Engine/Plugins/Animation/IKRig/Source/IKRig/Private/IKRigProcessor.cpp
@@ -9,6 +9,7 @@
 #include "IKRigDefinition.h"
 #include "IKRigSolverDefinition.h"
 #include "IKRigSolver.h"
+#include "IKRigConstraintSolver.h"
 
  // IKRigProcessor implementation functions
 void UIKRigProcessor::SetIKRigDefinition(UIKRigDefinition* InRigDefinition, bool bInitialize /*= false*/)
@@ -45,8 +46,12 @@ void UIKRigProcessor::Reinitialize()
 	// i'm copying this here 
 	TArray<UIKRigSolverDefinition*> SolverDefinitions = RigDefinition->GetSolverDefinitions();
 
+	UIKRigSolver::FIKRigTransformGetter RefTransformGetter = UIKRigSolver::FIKRigTransformGetter::CreateUObject(this, &UIKRigProcessor::GetRefPoseGetter);
+	UIKRigSolver::FIKRigGoalGetter GoalGetter = UIKRigSolver::FIKRigGoalGetter::CreateUObject(this, &UIKRigProcessor::GoalGetter);
+
 	const int32 NumSolvers = SolverDefinitions.Num();
 	Solvers.Reset(NumSolvers);
+
 	for (int32 Index = 0; Index < NumSolvers; ++Index)
 	{
 		// it's possible SolverDefinitions.Num() != Solvers.Num()
@@ -57,13 +62,18 @@ void UIKRigProcessor::Reinitialize()
 			{
 				UIKRigSolver* NewSolver = NewObject<UIKRigSolver>(this, ClassType);
 				// const transform modifier for hierarchy query
-				NewSolver->Init(SolverDefinitions[Index], TransformModifier, UIKRigSolver::FIKRigTransformGetter::CreateUObject(this, &UIKRigProcessor::GetRefPoseGetter),
-					UIKRigSolver::FIKRigGoalGetter::CreateUObject(this, &UIKRigProcessor::GoalGetter));
+				NewSolver->Init(SolverDefinitions[Index], TransformModifier, RefTransformGetter, GoalGetter);
 				Solvers.Add(NewSolver);
 			}
 		}
 	}
 
+	// initialize constraint solver
+	CostraintSolver = NewObject<UIKRigConstraintSolver>(this);
+	check (CostraintSolver);
+	CostraintSolver->SetConstraintDefinition(RigDefinition->ConstraintDefinitions);
+	CostraintSolver->Init(RigDefinition->ConstraintDefinitions, TransformModifier, RefTransformGetter, GoalGetter);
+
 	bInitialized = true;
 }
 
@@ -102,6 +112,9 @@ void UIKRigProcessor::Solve()
 			// Draw interface is pointer in case we want to nullify for optimization
 			Solvers[Index]->Solve(TransformModifier, &DrawInterface);
 		}
+
+		// solve constraint
+		CostraintSolver->Solve(TransformModifier, &DrawInterface);
 	}
 }
 
diff --git a/Engine/Plugins/Animation/IKRig/Source/IKRig/Public/IKRigConstraint.h b/Engine/Plugins/Animation/IKRig/Source/IKRig/Public/IKRigConstraint.h
--- a/Engine/Plugins/Animation/IKRig/Source/IKRig/Public/IKRigConstraint.h
+++ b/Engine/Plugins/Animation/IKRig/Source/IKRig/Public/IKRigConstraint.h
@@ -23,10 +23,10 @@ class IKRIG_API UIKRigConstraint : public UObject
 	bool bInitialized = false;
 
 public: 
-	void Setup(FIKRigTransformModifier& InOutTransformModifier);
+	void Setup(const FIKRigTransformModifier& InOutTransformModifier);
 	void SetAndApplyConstraint(FIKRigTransformModifier& InOutTransformModifier);
 
 	virtual void Apply(FIKRigTransformModifier& InOutTransformModifier, FControlRigDrawInterface* InOutDrawInterface) {};
-	virtual void SetupInternal(FIKRigTransformModifier& InOutTransformModifier) {};
+	virtual void SetupInternal(const FIKRigTransformModifier& InOutTransformModifier) {};
 };
 
diff --git a/Engine/Plugins/Animation/IKRig/Source/IKRig/Public/IKRigConstraintDefinition.h b/Engine/Plugins/Animation/IKRig/Source/IKRig/Public/IKRigConstraintDefinition.h
--- a/Engine/Plugins/Animation/IKRig/Source/IKRig/Public/IKRigConstraintDefinition.h
+++ b/Engine/Plugins/Animation/IKRig/Source/IKRig/Public/IKRigConstraintDefinition.h
@@ -7,9 +7,7 @@
 
 #pragma once
 
-#include "CoreMinimal.h"
-#include "UObject/ObjectMacros.h"
-#include "UObject/Object.h"
+#include "IKRigSolverDefinition.h"
 #include "IKRigConstraintDefinition.generated.h"
 
 class UIKRigConstraint;
@@ -24,7 +22,7 @@ struct IKRIG_API FIKRigConstraintProfile
 };
 
 UCLASS(config = Engine, hidecategories = UObject, BlueprintType)
-class IKRIG_API UIKRigConstraintDefinition : public UObject
+class IKRIG_API UIKRigConstraintDefinition : public UIKRigSolverDefinition
 {
 	GENERATED_BODY()
 
@@ -36,6 +34,7 @@ private:
 	UPROPERTY(EditAnywhere, Category = "Profile")
 	TMap<FName, FIKRigConstraintProfile> ConstraintProfiles;
 	friend class UIKRigConstraintSolver;
+	friend class UIKRigController;
 public:
 	static const FName DefaultProfileName;
 };
\ No newline at end of file
diff --git a/Engine/Plugins/Animation/IKRig/Source/IKRig/Public/IKRigConstraintSolver.h b/Engine/Plugins/Animation/IKRig/Source/IKRig/Public/IKRigConstraintSolver.h
--- a/Engine/Plugins/Animation/IKRig/Source/IKRig/Public/IKRigConstraintSolver.h
+++ b/Engine/Plugins/Animation/IKRig/Source/IKRig/Public/IKRigConstraintSolver.h
@@ -15,7 +15,7 @@ class UIKRigConstraint;
 
 
 // run time processor 
-UCLASS(config = Engine, Transient, hidecategories = UObject)
+UCLASS(config = Engine, hidecategories = UObject)
 class IKRIG_API UIKRigConstraintSolver : public UIKRigSolver
 {
 	GENERATED_BODY()
diff --git a/Engine/Plugins/Animation/IKRig/Source/IKRig/Public/IKRigDefinition.h b/Engine/Plugins/Animation/IKRig/Source/IKRig/Public/IKRigDefinition.h
--- a/Engine/Plugins/Animation/IKRig/Source/IKRig/Public/IKRigDefinition.h
+++ b/Engine/Plugins/Animation/IKRig/Source/IKRig/Public/IKRigDefinition.h
@@ -89,11 +89,13 @@ private:
 	UPROPERTY(EditAnywhere, instanced, Category = "Solver")
 	TArray<UIKRigSolverDefinition*> SolverDefinitions;
 
+public:
 	/*********** Constraint Definition *********************/
 	// contains constraint data
 	UPROPERTY(EditAnywhere, instanced, Category = "Constraint")
 	UIKRigConstraintDefinition* ConstraintDefinitions;
 
+private:
 	/*********** Goals with Default Value******************/
 	// goals data
 	// this is cached by "Unique Internal Name" here
diff --git a/Engine/Plugins/Animation/IKRig/Source/IKRigEditor/Private/IKRigController.cpp b/Engine/Plugins/Animation/IKRig/Source/IKRigEditor/Private/IKRigController.cpp
--- a/Engine/Plugins/Animation/IKRig/Source/IKRigEditor/Private/IKRigController.cpp
+++ b/Engine/Plugins/Animation/IKRig/Source/IKRigEditor/Private/IKRigController.cpp
@@ -10,6 +10,8 @@
 #include "AnimationRuntime.h"
 #include "IKRigSolverDefinition.h"
 #include "ScopedTransaction.h"
+#include "IKRigConstraintDefinition.h"
+#include "IKRigConstraint.h"
 
 #define LOCTEXT_NAMESPACE	"IKRigController"
 
@@ -332,42 +334,142 @@ bool UIKRigController::CanAutoConfigure(UIKRigSolverDefinition* SolverDef) const
 // create new profile
 void UIKRigController::CreateNewProfile(FName& InNewProfileName)
 {
-	
+	if (IKRigDefinition && InNewProfileName != NAME_None)
+	{
+		FScopedTransaction Transaction(LOCTEXT("CreateNewProfile_Label", "Create New Constraint Profile"));
+		IKRigDefinition->Modify();
+		// we just find or add
+		IKRigDefinition->ConstraintDefinitions->ConstraintProfiles.FindOrAdd(InNewProfileName);
+	}
 }
 
 bool UIKRigController::RemoveConstraintProfile(const FName& InProfileName)
 {
+	if (IKRigDefinition && InProfileName != NAME_None)
+	{
+		FScopedTransaction Transaction(LOCTEXT("RemoveProfile_Label", "Remove Constraint Profile"));
+		IKRigDefinition->Modify();
+		// we just find or add
+		return IKRigDefinition->ConstraintDefinitions->ConstraintProfiles.Remove(InProfileName) > 0;
+
+	}
+
 	return false;
 }
 
 void UIKRigController::RenameProfile(FName InCurrentProfileName, FName& InNewProfileName)
 {
+	// write this 
+	ensure (false);
+	// we don't allow change default profile until we support saving the default profile
+}
+
+FIKRigConstraintProfile* UIKRigController::GetConstraintProfile(const FName& InProfileName) const
+{
+	if (IKRigDefinition)
+	{
+		FName SearchProfileName = (InProfileName == NAME_None)? UIKRigConstraintDefinition::DefaultProfileName : InProfileName;
+		return IKRigDefinition->ConstraintDefinitions->ConstraintProfiles.Find(SearchProfileName);
+	}
 
+	return nullptr;
+}
+
+void UIKRigController::EnsureUniqueConstraintName(FName& InOutName)
+{
+	int32 Index = 1;
+
+	TArray<FName> ConstraintNames;
+
+	GetConstraintNames(ConstraintNames);
+
+	FString NewName = InOutName.ToString();
+	while (ConstraintNames.Contains(FName(*NewName)))
+	{
+		NewName = FString::Format(TEXT("{0}_{1}"), { InOutName.ToString(), Index++ });
+	}
+
+	InOutName = FName(*NewName);
 }
 
 UIKRigConstraint* UIKRigController::AddConstraint(TSubclassOf<UIKRigConstraint> NewConstraintType, FName& InOutNewName, FName InProfile /*= NAME_None*/)
 {
+	FIKRigConstraintProfile* Profile = GetConstraintProfile(InProfile);
+	if (Profile)
+	{
+		FScopedTransaction Transaction(LOCTEXT("AddConstraint_Label", "Add Constraint"));
+		IKRigDefinition->Modify();
+
+		EnsureUniqueConstraintName(InOutNewName);
+		UIKRigConstraint* NewRigConstraint = NewObject<UIKRigConstraint>(IKRigDefinition, NewConstraintType, InOutNewName);
+		if (NewRigConstraint)
+		{
+		
+			Profile->Constraints.Add(NewRigConstraint->GetFName(), NewRigConstraint);
+			return NewRigConstraint;
+		}
+	}
 	return nullptr;	
 }
 
 UIKRigConstraint* UIKRigController::GetConstraint(const FName& InProfileName, const FName& InName) const
 {
+	FIKRigConstraintProfile* Profile = GetConstraintProfile(InProfileName);
+	if (Profile)
+	{
+		UIKRigConstraint** RigConstraint = Profile->Constraints.Find(InName);
+		if (RigConstraint)
+		{
+			return *RigConstraint;
+		}
+	}
+
 	return nullptr;
 }
 
 bool UIKRigController::RemoveConstraint(const FName& InConstraintName)
 {
+	// remove this in all profile
+	if (IKRigDefinition)
+	{
+		FScopedTransaction Transaction(LOCTEXT("AddConstraint_Label", "Add Constraint"));
+		IKRigDefinition->Modify();
+
+		for (auto Iter = IKRigDefinition->ConstraintDefinitions->ConstraintProfiles.CreateIterator(); Iter; ++Iter)
+		{
+			FIKRigConstraintProfile& Profile = Iter.Value();
+			for (auto InnerIter = Profile.Constraints.CreateIterator(); InnerIter; ++InnerIter)
+			{
+				// possible the constraint name may not be there
+				// what this means, when we rename, we have to rename all of them
+				Profile.Constraints.Remove(InConstraintName);
+			}
+		}
+		// remove return?
+		return true;
+	}
+
 	return false;
 }
 
 void UIKRigController::GetConstraintProfileNames(TArray<FName>& OutProfileNames) const
 {
-
+	if (IKRigDefinition)
+	{
+		OutProfileNames.Reset();
+		IKRigDefinition->ConstraintDefinitions->ConstraintProfiles.GenerateKeyArray(OutProfileNames);
+	}
 }
 
 void UIKRigController::GetConstraintNames(TArray<FName>& OutConstraintNames) const
 {
-
+	// still can be null if IKRigDefinition is nullptr
+	FIKRigConstraintProfile* DefaultProfile = GetConstraintProfile(NAME_None);
+	if (DefaultProfile)
+	{
+		OutConstraintNames.Reset();
+		DefaultProfile->Constraints.GenerateKeyArray(OutConstraintNames);
+	}
 }
 
 // goal operators
diff --git a/Engine/Plugins/Animation/IKRig/Source/IKRigEditor/Private/IKRigDefinitionDetails.cpp b/Engine/Plugins/Animation/IKRig/Source/IKRigEditor/Private/IKRigDefinitionDetails.cpp
--- a/Engine/Plugins/Animation/IKRig/Source/IKRigEditor/Private/IKRigDefinitionDetails.cpp
+++ b/Engine/Plugins/Animation/IKRig/Source/IKRigEditor/Private/IKRigDefinitionDetails.cpp
@@ -10,6 +10,7 @@
 #include "IKRigDefinition.h"
 #include "IKRigController.h"
 #include "IKRigSolverDefinition.h"
+#include "IKRigConstraint.h"
 
 #include "ScopedTransaction.h"
 #include "PropertyCustomizationHelpers.h"
@@ -138,7 +139,7 @@ void FIKRigDefinitionDetails::CustomizeDetails(IDetailLayoutBuilder& DetailBuild
 	[
 		SNew(SButton)
 		.ContentPadding(3)
-		.OnClicked(this, &FIKRigDefinitionDetails::OnShowClassPicker)
+		.OnClicked(this, &FIKRigDefinitionDetails::OnShowSolverClassPicker)
 		.ToolTipText(LOCTEXT("OnShowSolverListTooltip", "Select Solver to Add"))
 		[
 			SNew(STextBlock)
@@ -147,6 +148,25 @@ void FIKRigDefinitionDetails::CustomizeDetails(IDetailLayoutBuilder& DetailBuild
 		]
 	];
 
+	IDetailCategoryBuilder& ConstraintCategory = DetailBuilder.EditCategory("Constraint");
+	ConstraintCategory.AddCustomRow(FText::FromString("AddConstraint"))
+	.NameContent()
+	[
+		SNullWidget::NullWidget
+	]
+	.ValueContent()
+	[
+		SNew(SButton)
+		.ContentPadding(3)
+		.OnClicked(this, &FIKRigDefinitionDetails::OnShowConstraintClassPicker)
+		.ToolTipText(LOCTEXT("OnShowConstraintListTooltip", "Select Constraint to Add"))
+		[
+			SNew(STextBlock)
+			.Font(IDetailLayoutBuilder::GetDetailFont())
+			.Text(LOCTEXT("ShowConstraintList", "Add Constraint"))
+		]
+	];
+
 	GoalPropertyHandle = DetailBuilder.GetProperty(TEXT("IKGoals"));
 
 	TArray<FName> GoalNames;
@@ -177,7 +197,7 @@ void FIKRigDefinitionDetails::CustomizeDetails(IDetailLayoutBuilder& DetailBuild
 //	GoalPropertyHandle->MarkHiddenByCustomization();
 }
 
-class FIKSolverClassFilter : public IClassViewerFilter
+class FIKRigClassFilter : public IClassViewerFilter
 {
 public:
 	/** All children of these classes will be included unless filtered out by another setting. */
@@ -201,7 +221,8 @@ public:
 	}
 };
 
-FReply FIKRigDefinitionDetails::OnShowClassPicker()
+// choose class for them
+UClass* SelectClass(UClass* ClassType, const FText& TitleText)
 {
 	FClassViewerModule& ClassViewerModule = FModuleManager::LoadModuleChecked<FClassViewerModule>("ClassViewer");
 	FClassViewerInitializationOptions Options;
@@ -210,18 +231,30 @@ FReply FIKRigDefinitionDetails::OnShowClassPicker()
 	Options.bShowObjectRootClass = false;
 	Options.bExpandRootNodes = true;
 	Options.bShowUnloadedBlueprints = true;
-	TSharedPtr<FIKSolverClassFilter> Filter = MakeShareable(new FIKSolverClassFilter);
+	TSharedPtr<FIKRigClassFilter> Filter = MakeShareable(new FIKRigClassFilter);
 	Options.ClassFilter = Filter;
 
 	Filter->DisallowedClassFlags = CLASS_Abstract | CLASS_Deprecated | CLASS_NewerVersionExists | CLASS_Transient;
-	Filter->AllowedChildrenOfClasses.Add(UIKRigSolverDefinition::StaticClass());
-	Filter->DisallowedClasses.Add(UIKRigSolverDefinition::StaticClass());
+	Filter->AllowedChildrenOfClasses.Add(ClassType);
+	Filter->DisallowedClasses.Add(ClassType);
 
-	const FText TitleText = LOCTEXT("SelectSolverClass", "Select Solver Class");
 	UClass* ChosenClass = nullptr;
-	const bool bPressedOk = SClassPickerDialog::PickClass(TitleText, Options, ChosenClass, UIKRigSolverDefinition::StaticClass());
+	const bool bPressedOk = SClassPickerDialog::PickClass(TitleText, Options, ChosenClass, ClassType);
 
 	if (bPressedOk)
+	{
+		return ChosenClass;
+	}
+
+	return nullptr;
+}
+
+
+
+FReply FIKRigDefinitionDetails::OnShowSolverClassPicker()
+{
+	UClass* ChosenClass = SelectClass(UIKRigSolverDefinition::StaticClass(), LOCTEXT("SelectSolverClass", "Select Solver Class"));
+	if (ChosenClass)
 	{
 		IKRigController->AddSolver(ChosenClass);
 	}
@@ -229,6 +262,19 @@ FReply FIKRigDefinitionDetails::OnShowClassPicker()
 	return FReply::Handled();
 }
 
+FReply FIKRigDefinitionDetails::OnShowConstraintClassPicker()
+{
+	UClass* ChosenClass = SelectClass(UIKRigConstraint::StaticClass(), LOCTEXT("SelectConstraintClass", "Select Constraint Class"));
+	if (ChosenClass)
+	{
+		// ask for name of constraint
+		FName NewName(TEXT("NewConstraint"));
+		IKRigController->AddConstraint(ChosenClass, NewName);
+	}
+
+	return FReply::Handled();
+}
+
 bool FIKRigDefinitionDetails::CanImport() const
 {
 	return (SelectedAsset.IsValid());
diff --git a/Engine/Plugins/Animation/IKRig/Source/IKRigEditor/Private/IKRigDefinitionDetails.h b/Engine/Plugins/Animation/IKRig/Source/IKRigEditor/Private/IKRigDefinitionDetails.h
--- a/Engine/Plugins/Animation/IKRig/Source/IKRigEditor/Private/IKRigDefinitionDetails.h
+++ b/Engine/Plugins/Animation/IKRig/Source/IKRigEditor/Private/IKRigDefinitionDetails.h
@@ -6,16 +6,18 @@
 #include "UObject/WeakObjectPtr.h"
 #include "UObject/WeakObjectPtrTemplates.h"
 #include "IDetailCustomization.h"
-#include "Input/Reply.h"
+#include "Types/SlateEnums.h"
 #include "Widgets/Views/SListView.h"
+#include "Input/Reply.h"
 
 class IDetailLayoutBuilder;
-class IPropertyHandle;
-class ITableRow;
-class STableViewBase;
 class UIKRigDefinition;
 class UIKRigController;
 struct FAssetData;
+class FReply;
+class ITableRow;
+class STableViewBase;
+class IPropertyHandle;
 
 // Utility class to build combo boxes out of arrays of names.
 struct FGoalNameListItem : public TSharedFromThis<FGoalNameListItem>
@@ -56,7 +58,8 @@ private:
 
 	TWeakPtr< IDetailLayoutBuilder > DetailBuilderWeakPtr;
 
-	FReply OnShowClassPicker();
+	FReply OnShowSolverClassPicker();
+	FReply OnShowConstraintClassPicker();
 
 	typedef TSharedPtr<FGoalNameListItem> FGoalNameListItemPtr;
 
diff --git a/Engine/Plugins/Animation/IKRig/Source/IKRigEditor/Public/IKRigController.h b/Engine/Plugins/Animation/IKRig/Source/IKRigEditor/Public/IKRigController.h
--- a/Engine/Plugins/Animation/IKRig/Source/IKRigEditor/Public/IKRigController.h
+++ b/Engine/Plugins/Animation/IKRig/Source/IKRigEditor/Public/IKRigController.h
@@ -24,6 +24,7 @@ class UIKRigSolverDefinition;
 class UIKRigConstraint;
 struct FIKRigHierarchy;
 struct FReferenceSkeleton;
+struct FIKRigConstraintProfile;
 
 
 UCLASS(config = Engine, hidecategories = UObject, BlueprintType)
@@ -106,6 +107,8 @@ private:
 	bool ValidateSolver(UIKRigSolverDefinition* const SolverDef) const;
 	void UpdateGoal();
 
+	FIKRigConstraintProfile* GetConstraintProfile(const FName& InProfileName) const;
+
 	void InitializeIKRigSolverDefinition(UIKRigSolverDefinition* SolverDef);
 	void UninitializeIKRigSolverDefinition(UIKRigSolverDefinition* SolverDef);
 
@@ -116,6 +119,8 @@ private:
 	// IKRigDefinition set up
 	void SetIKRigDefinition(UIKRigDefinition* InIKRigDefinition);
 
+	void EnsureUniqueConstraintName(FName& InOutName);
+
 	// called by IKRigDefinition when it's begin destroy
 	static void RemoveControllerByRigDefinition(UIKRigDefinition* InIKRigDefinition);
 
