Fixed issue where IKRigDefinition was garbage collected while IKRigController was still referencing it.

diff --git a/Engine/Plugins/Animation/IKRig/Source/IKRigEditor/Public/IKRigController.h b/Engine/Plugins/Animation/IKRig/Source/IKRigEditor/Public/IKRigController.h
--- a/Engine/Plugins/Animation/IKRig/Source/IKRigEditor/Public/IKRigController.h
+++ b/Engine/Plugins/Animation/IKRig/Source/IKRigEditor/Public/IKRigController.h
@@ -102,6 +102,7 @@ public:
 
 	FGoalModified OnGoalModified;
 private:
+	UPROPERTY(transient)
 	UIKRigDefinition* IKRigDefinition = nullptr;
 
 	bool ValidateSolver(UIKRigSolverDefinition* const SolverDef) const;
