Renamed and moved constraint solver and transform solver files.

diff --git a/Engine/Plugins/Animation/IKRig/Source/IKRig/Private/IKRigConstraintDefinition.cpp b/Engine/Plugins/Animation/IKRig/Source/IKRig/Private/IKRigConstraintDefinition.cpp
--- a/Engine/Plugins/Animation/IKRig/Source/IKRig/Private/IKRigConstraintDefinition.cpp
+++ b/Engine/Plugins/Animation/IKRig/Source/IKRig/Private/IKRigConstraintDefinition.cpp
@@ -6,7 +6,7 @@
  */
 
 #include "IKRigConstraintDefinition.h"
-#include "IKRigConstraintSolver.h"
+#include "Solvers/ConstraintSolver.h"
 
 const FName UIKRigConstraintDefinition::DefaultProfileName = FName(TEXT("Default"));
 
diff --git a/Engine/Plugins/Animation/IKRig/Source/IKRig/Private/IKRigConstraintSolver.cpp b/Engine/Plugins/Animation/IKRig/Source/IKRig/Private/IKRigConstraintSolver.cpp
deleted file mode 100644
diff --git a/Engine/Plugins/Animation/IKRig/Source/IKRig/Private/IKRigProcessor.cpp b/Engine/Plugins/Animation/IKRig/Source/IKRig/Private/IKRigProcessor.cpp
--- a/Engine/Plugins/Animation/IKRig/Source/IKRig/Private/IKRigProcessor.cpp
+++ b/Engine/Plugins/Animation/IKRig/Source/IKRig/Private/IKRigProcessor.cpp
@@ -9,7 +9,7 @@
 #include "IKRigDefinition.h"
 #include "IKRigSolverDefinition.h"
 #include "IKRigSolver.h"
-#include "IKRigConstraintSolver.h"
+#include "Solvers/ConstraintSolver.h"
 
  // IKRigProcessor implementation functions
 void UIKRigProcessor::SetIKRigDefinition(UIKRigDefinition* InRigDefinition, bool bInitialize /*= false*/)
diff --git a/Engine/Plugins/Animation/IKRig/Source/IKRig/Private/Solvers/ConstraintSolver.cpp b/Engine/Plugins/Animation/IKRig/Source/IKRig/Private/Solvers/ConstraintSolver.cpp
new file mode 100644
--- /dev/null
+++ b/Engine/Plugins/Animation/IKRig/Source/IKRig/Private/Solvers/ConstraintSolver.cpp
@@ -0,0 +1,94 @@
+// Copyright Epic Games, Inc. All Rights Reserved.
+
+/**
+ * Contains IKRig Constraint Solver Definition 
+ *
+ */
+#include "Solvers/ConstraintSolver.h"
+#include "IKRigConstraint.h"
+
+// register/unregister query function
+void UIKRigConstraintSolver::RegisterQueryConstraintHandler(const FIKRigQueryConstraint& InQueryConstraintHandler)
+{
+
+}
+
+void UIKRigConstraintSolver::UnregisterQueryConstraintHandler()
+{
+
+}
+
+void UIKRigConstraintSolver::InitInternal(const FIKRigTransformModifier& InGlobalTransform)
+{
+	if (ConstraintDefinition)
+	{
+		// we duplicate and create a copy of the objects, so that we can mutate
+		ConstraintProfiles = ConstraintDefinition->ConstraintProfiles;
+		for (auto Iter = ConstraintProfiles.CreateIterator(); Iter; ++Iter)
+		{
+			FIKRigConstraintProfile& CurrentProfile = Iter.Value();
+			for (auto InnerIter = CurrentProfile.Constraints.CreateIterator(); InnerIter; ++InnerIter)
+			{
+				UIKRigConstraint*& Template = InnerIter.Value();
+				UIKRigConstraint* CopyConstraint = NewObject<UIKRigConstraint>(this, Template->GetClass(), TEXT("IKRigConstraint"), RF_NoFlags, Template);
+				CopyConstraint->Setup(InGlobalTransform);
+				// replace value
+				Template = CopyConstraint;
+			}
+		}
+
+		// go through all UIKRigConstraint and create copy
+		ActiveProfile = UIKRigConstraintDefinition::DefaultProfileName;
+	}
+}
+
+void UIKRigConstraintSolver::SolveInternal(FIKRigTransformModifier& InOutGlobalTransform, FControlRigDrawInterface* InOutDrawInterface)
+{
+	FIKRigConstraintProfile* Current = ConstraintProfiles.Find(ActiveProfile);
+
+	if (Current)
+	{
+		for (auto Iter = Current->Constraints.CreateIterator(); Iter; ++Iter)
+		{
+			UIKRigConstraint* Constraint = Iter.Value();
+			if (Constraint)
+			{
+				// @todo: later add inoutdrawinterface
+				Constraint->Apply(InOutGlobalTransform, InOutDrawInterface);
+			}
+		}
+	}
+}
+
+void UIKRigConstraintSolver::SolverConstraints(const TArray<FName>& ConstraintsList, FIKRigTransformModifier& InOutGlobalTransform, FControlRigDrawInterface* InOutDrawInterface)
+{
+	FIKRigConstraintProfile* Current = ConstraintProfiles.Find(ActiveProfile);
+
+	if (Current)
+	{
+		for (const FName& ConstraintName : ConstraintsList)
+		{
+			UIKRigConstraint** Constraint = Current->Constraints.Find(ConstraintName);
+			if (Constraint)
+			{
+				// @todo: later add inoutdrawinterface
+				(*Constraint)->Apply(InOutGlobalTransform, InOutDrawInterface);
+			}
+		}
+	}
+}
+
+bool UIKRigConstraintSolver::IsSolverActive() const
+{	
+// later we can do more elaborated version
+	return Super::IsSolverActive();
+}
+
+void UIKRigConstraintSolver::SetConstraintDefinition(UIKRigConstraintDefinition* InConstraintDefinition)
+{
+	ConstraintDefinition = InConstraintDefinition;
+	// reset instanced data
+
+	ConstraintProfiles.Reset();
+	ActiveProfile = UIKRigConstraintDefinition::DefaultProfileName;
+}
\ No newline at end of file
diff --git a/Engine/Plugins/Animation/IKRig/Source/IKRig/Private/Solvers/TransformSolver.cpp b/Engine/Plugins/Animation/IKRig/Source/IKRig/Private/Solvers/TransformSolver.cpp
--- a/Engine/Plugins/Animation/IKRig/Source/IKRig/Private/Solvers/TransformSolver.cpp
+++ b/Engine/Plugins/Animation/IKRig/Source/IKRig/Private/Solvers/TransformSolver.cpp
@@ -4,7 +4,7 @@
 	TransformSolver.cpp: Solver execution class for Transform
 =============================================================================*/
 
-#include "TransformSolver.h"
+#include "Solvers/TransformSolver.h"
 #include "IKRigDataTypes.h"
 #include "IKRigHierarchy.h"
 #include "Solvers/TransformSolverDefinition.h"
diff --git a/Engine/Plugins/Animation/IKRig/Source/IKRig/Private/Solvers/TransformSolver.h b/Engine/Plugins/Animation/IKRig/Source/IKRig/Private/Solvers/TransformSolver.h
deleted file mode 100644
diff --git a/Engine/Plugins/Animation/IKRig/Source/IKRig/Private/Solvers/TransformSolverDefinition.cpp b/Engine/Plugins/Animation/IKRig/Source/IKRig/Private/Solvers/TransformSolverDefinition.cpp
--- a/Engine/Plugins/Animation/IKRig/Source/IKRig/Private/Solvers/TransformSolverDefinition.cpp
+++ b/Engine/Plugins/Animation/IKRig/Source/IKRig/Private/Solvers/TransformSolverDefinition.cpp
@@ -6,7 +6,7 @@
  */
 
 #include "Solvers/TransformSolverDefinition.h"
-#include "TransformSolver.h"
+#include "Solvers/TransformSolver.h"
 
 UTransformSolverDefinition::UTransformSolverDefinition()
 	: TransformTargetName(TEXT("TransformTarget"))
diff --git a/Engine/Plugins/Animation/IKRig/Source/IKRig/Public/IKRigConstraintSolver.h b/Engine/Plugins/Animation/IKRig/Source/IKRig/Public/IKRigConstraintSolver.h
deleted file mode 100644
diff --git a/Engine/Plugins/Animation/IKRig/Source/IKRig/Public/Solvers/ConstraintSolver.h b/Engine/Plugins/Animation/IKRig/Source/IKRig/Public/Solvers/ConstraintSolver.h
new file mode 100644
--- /dev/null
+++ b/Engine/Plugins/Animation/IKRig/Source/IKRig/Public/Solvers/ConstraintSolver.h
@@ -0,0 +1,56 @@
+// Copyright Epic Games, Inc. All Rights Reserved.
+
+/**
+ * Contains IKRig Constraint Solver Definition 
+ *
+ */
+
+#pragma once
+
+#include "IKRigSolver.h"
+#include "IKRigConstraintDefinition.h"
+#include "ConstraintSolver.generated.h"
+
+class UIKRigConstraint;
+
+
+// run time processor 
+UCLASS(config = Engine, hidecategories = UObject)
+class IKRIG_API UIKRigConstraintSolver : public UIKRigSolver
+{
+	GENERATED_BODY()
+
+	DECLARE_DELEGATE_TwoParams(FIKRigQueryConstraint, const FName& InConstraintName, UIKRigConstraint& OutConstraint);
+
+	UPROPERTY(transient)
+	UIKRigConstraintDefinition* ConstraintDefinition;
+
+public: 
+	// during we don't mutate bone transform, but you can read it
+	virtual void InitInternal(const FIKRigTransformModifier& InGlobalTransform) override;
+	virtual void SolveInternal(FIKRigTransformModifier& InOutGlobalTransform, FControlRigDrawInterface* InOutDrawInterface) override;
+	virtual bool IsSolverActive() const override;
+
+	// this can be utilized for just subset of constraints
+	void SolverConstraints(const TArray<FName>& ConstraintsList, FIKRigTransformModifier& InOutGlobalTransform, FControlRigDrawInterface* InOutDrawInterface);
+
+	// register/unregister query function
+	void RegisterQueryConstraintHandler(const FIKRigQueryConstraint& InQueryConstraintHandler);
+	void UnregisterQueryConstraintHandler();
+
+	void SetConstraintDefinition(UIKRigConstraintDefinition* InConstraintDefinition);
+
+private:
+	// delegate
+	FIKRigQueryConstraint QueryConstraintHandler;
+
+	// current active profile
+	UPROPERTY(transient)
+	FName ActiveProfile;
+
+	// instanced information you can mutate from ConstraintDefinition
+	// we want uproperty, so that it doesn't GC-ed
+	UPROPERTY(transient)
+	TMap<FName, FIKRigConstraintProfile> ConstraintProfiles;
+};
+
diff --git a/Engine/Plugins/Animation/IKRig/Source/IKRig/Public/Solvers/TransformSolver.h b/Engine/Plugins/Animation/IKRig/Source/IKRig/Public/Solvers/TransformSolver.h
new file mode 100644
--- /dev/null
+++ b/Engine/Plugins/Animation/IKRig/Source/IKRig/Public/Solvers/TransformSolver.h
@@ -0,0 +1,25 @@
+// Copyright Epic Games, Inc. All Rights Reserved.
+
+/**
+ * Contains Transform Solver Execution
+ *
+ */
+
+#pragma once
+
+#include "IKRigSolver.h"
+#include "TransformSolver.generated.h"
+
+
+// run time for UTransformSolverDefinition
+UCLASS()
+class IKRIG_API UTransformSolver : public UIKRigSolver
+{
+	GENERATED_BODY()
+
+protected:
+	virtual void InitInternal(const FIKRigTransformModifier& InGlobalTransform) override;
+	virtual void SolveInternal(FIKRigTransformModifier& InOutGlobalTransform, FControlRigDrawInterface* InOutDrawInterface) override;
+	virtual bool IsSolverActive() const override;
+};
+
