Simplifying transform data in IKRig. Merged Transform and TransformModifier classes. General renaming for clarity.

diff --git a/Engine/Plugins/Animation/IKRig/Source/IKRig/Private/AnimNodes/AnimNode_IKRig.cpp b/Engine/Plugins/Animation/IKRig/Source/IKRig/Private/AnimNodes/AnimNode_IKRig.cpp
--- a/Engine/Plugins/Animation/IKRig/Source/IKRig/Private/AnimNodes/AnimNode_IKRig.cpp
+++ b/Engine/Plugins/Animation/IKRig/Source/IKRig/Private/AnimNodes/AnimNode_IKRig.cpp
@@ -32,7 +32,7 @@ void FAnimNode_IKRig::Evaluate_AnyThread(FPoseContext& Output)
 
 	if (RigProcessor)
 	{
-		FIKRigTransformModifier& TransformModifier = RigProcessor->GetIKRigTransformModifier();
+		FIKRigTransforms& TransformModifier = RigProcessor->GetIKRigTransformModifier();
 
 		if (bStartFromRefPose)
 		{
diff --git a/Engine/Plugins/Animation/IKRig/Source/IKRig/Private/Constraints/RotationLimitConstraint.cpp b/Engine/Plugins/Animation/IKRig/Source/IKRig/Private/Constraints/RotationLimitConstraint.cpp
--- a/Engine/Plugins/Animation/IKRig/Source/IKRig/Private/Constraints/RotationLimitConstraint.cpp
+++ b/Engine/Plugins/Animation/IKRig/Source/IKRig/Private/Constraints/RotationLimitConstraint.cpp
@@ -5,7 +5,7 @@
 #include "IKRigDataTypes.h"
 #include "IKRig.h"
 
-void URotationLimitConstraint::SetupInternal(const FIKRigTransformModifier& InOutTransformModifier)
+void URotationLimitConstraint::SetupInternal(const FIKRigTransforms& InOutTransformModifier)
 {
 	if (TargetBone != NAME_None)
 	{
@@ -28,7 +28,7 @@ void URotationLimitConstraint::SetupInternal(const FIKRigTransformModifier& InOu
 	}
 }
 
-void URotationLimitConstraint::Apply(FIKRigTransformModifier& InOutTransformModifier, FControlRigDrawInterface* InOutDrawInterface)
+void URotationLimitConstraint::Apply(FIKRigTransforms& InOutTransformModifier, FControlRigDrawInterface* InOutDrawInterface)
 {
 	if (ConstrainedIndex != INDEX_NONE)
 	{
diff --git a/Engine/Plugins/Animation/IKRig/Source/IKRig/Private/Constraints/RotationLimitConstraint.h b/Engine/Plugins/Animation/IKRig/Source/IKRig/Private/Constraints/RotationLimitConstraint.h
--- a/Engine/Plugins/Animation/IKRig/Source/IKRig/Private/Constraints/RotationLimitConstraint.h
+++ b/Engine/Plugins/Animation/IKRig/Source/IKRig/Private/Constraints/RotationLimitConstraint.h
@@ -52,7 +52,7 @@ class IKRIG_API URotationLimitConstraint : public UIKRigConstraint
 	UPROPERTY(EditAnywhere, Category = "Offset")
 	FRotator Offset;
 public: 
-	virtual void SetupInternal(const FIKRigTransformModifier& InOutTransformModifier) override;
-	virtual void Apply(FIKRigTransformModifier& InOutTransformModifier, FControlRigDrawInterface* InOutDrawInterface) override;
+	virtual void SetupInternal(const FIKRigTransforms& InOutTransformModifier) override;
+	virtual void Apply(FIKRigTransforms& InOutTransformModifier, FControlRigDrawInterface* InOutDrawInterface) override;
 };
 
diff --git a/Engine/Plugins/Animation/IKRig/Source/IKRig/Private/IKRigConstraint.cpp b/Engine/Plugins/Animation/IKRig/Source/IKRig/Private/IKRigConstraint.cpp
--- a/Engine/Plugins/Animation/IKRig/Source/IKRig/Private/IKRigConstraint.cpp
+++ b/Engine/Plugins/Animation/IKRig/Source/IKRig/Private/IKRigConstraint.cpp
@@ -7,17 +7,17 @@
 
 #include "IKRigConstraint.h"
 
-void UIKRigConstraint::Setup(const FIKRigTransformModifier& InOutTransformModifier)
+void UIKRigConstraint::Setup(const FIKRigTransforms& InOutTransformModifier)
 {
 	SetupInternal(InOutTransformModifier);
 	bInitialized = true;
 }
 
-void UIKRigConstraint::SetAndApplyConstraint(FIKRigTransformModifier& InOutTransformModifier)
+void UIKRigConstraint::SetAndApplyConstraint(FIKRigTransforms& InOutTransformModifier)
 {
 	if (!bInitialized)
 	{
-		Setup(const_cast<FIKRigTransformModifier&>(InOutTransformModifier));
+		Setup(const_cast<FIKRigTransforms&>(InOutTransformModifier));
 	}
 
 	Apply(InOutTransformModifier, nullptr);
diff --git a/Engine/Plugins/Animation/IKRig/Source/IKRig/Private/IKRigDataTypes.cpp b/Engine/Plugins/Animation/IKRig/Source/IKRig/Private/IKRigDataTypes.cpp
--- a/Engine/Plugins/Animation/IKRig/Source/IKRig/Private/IKRigDataTypes.cpp
+++ b/Engine/Plugins/Animation/IKRig/Source/IKRig/Private/IKRigDataTypes.cpp
@@ -1,13 +1,9 @@
 // Copyright Epic Games, Inc. All Rights Reserved.
 
-/*=============================================================================
-	IKRigDataTypes.cpp: IKRig Data Type impelmentation file
-=============================================================================*/
-
 #include "IKRigDataTypes.h"
 #include "IKRigHierarchy.h"
 
-FIKRigTransformModifier::FIKRigTransformModifier(const FIKRigHierarchy* InHierarchy)
+FIKRigTransforms::FIKRigTransforms(const FIKRigHierarchy* InHierarchy)
 	: Hierarchy(InHierarchy)
 {
 	ensure(Hierarchy != nullptr);
@@ -15,48 +11,48 @@ FIKRigTransformModifier::FIKRigTransformModifier(const FIKRigHierarchy* InHierar
 
 // we don't want to allocate local transforms until required
 // so we don't do this in the constructor but on demand when to be used
-void FIKRigTransformModifier::EnsureLocalTransformsExist()
+void FIKRigTransforms::EnsureLocalTransformsExist()
 {
 	// first check validation
-	if (Hierarchy&& Hierarchy->GetNum() == GlobalTransforms.GetNum())
+	if (Hierarchy&& Hierarchy->GetNum() == GlobalTransforms.Num())
 	{
-		if (LocalTransforms.Num() != GlobalTransforms.GetNum())
+		if (LocalTransforms.Num() != GlobalTransforms.Num())
 		{
-			LocalTransforms.SetNumUninitialized(GlobalTransforms.GetNum());
-			LocalTransformDirtyFlags.Init(true, GlobalTransforms.GetNum());
+			LocalTransforms.SetNumUninitialized(GlobalTransforms.Num());
+			LocalTransformDirtyFlags.Init(true, GlobalTransforms.Num());
 		}
 	}
 }
 
-void FIKRigTransformModifier::RecalculateLocalTransform()
+void FIKRigTransforms::RecalculateLocalTransform()
 {
 	// make sure to call EnsureLocalTransformsExist before getting here
-	check (Hierarchy && LocalTransforms.Num() == GlobalTransforms.GetNum());
+	check (Hierarchy && LocalTransforms.Num() == GlobalTransforms.Num());
 	
 	for (TConstSetBitIterator<> It(LocalTransformDirtyFlags); It; ++It)
 	{
 		const int32 BoneIndex = It.GetIndex();
 		const int32 ParentIndex = Hierarchy->GetParentIndex(BoneIndex);
-		LocalTransforms[BoneIndex] = GlobalTransforms.GetRelativeTransform(BoneIndex, ParentIndex);
+		LocalTransforms[BoneIndex] = GetRelativeTransform(BoneIndex, ParentIndex);
 	}
 
 	// clear the bits
-	LocalTransformDirtyFlags.Init(false, GlobalTransforms.GetNum());
+	LocalTransformDirtyFlags.Init(false, GlobalTransforms.Num());
 }
 
-void FIKRigTransformModifier::UpdateLocalTransform(int32 Index)
+void FIKRigTransforms::UpdateLocalTransform(int32 Index)
 {
 	EnsureLocalTransformsExist();
 
 	if (Hierarchy && Hierarchy->IsValidIndex(Index) && LocalTransformDirtyFlags[Index])
 	{
 		int32 ParentIndex = Hierarchy->GetParentIndex(Index);
-		LocalTransforms[Index] = GlobalTransforms.GetRelativeTransform(Index, ParentIndex);
+		LocalTransforms[Index] = GetRelativeTransform(Index, ParentIndex);
 		LocalTransformDirtyFlags[Index] = false;
 	}
 }
 
-void FIKRigTransformModifier::SetGlobalTransform(int32 Index, const FTransform& InTransform, bool bPropagate)
+void FIKRigTransforms::SetGlobalTransform(int32 Index, const FTransform& InTransform, bool bPropagate)
 {
 	if (GlobalTransforms.IsValidIndex(Index))
 	{
@@ -67,7 +63,7 @@ void FIKRigTransformModifier::SetGlobalTransform(int32 Index, const FTransform&
 		}
 		else
 		{
-			GlobalTransforms.GlobalTransforms[Index] = InTransform;
+			GlobalTransforms[Index] = InTransform;
 
 			// if we have local transform, we also update
 			if (LocalTransforms.IsValidIndex(Index))
@@ -79,7 +75,7 @@ void FIKRigTransformModifier::SetGlobalTransform(int32 Index, const FTransform&
 	}
 }
 
-void FIKRigTransformModifier::SetLocalTransform(int32 Index, const FTransform& InTransform, bool bPropagate)
+void FIKRigTransforms::SetLocalTransform(int32 Index, const FTransform& InTransform, bool bPropagate)
 {
 	// we go through global 
 	// because we can't just update local only
@@ -92,26 +88,31 @@ void FIKRigTransformModifier::SetLocalTransform(int32 Index, const FTransform& I
 	}
 }
 
-const FTransform& FIKRigTransformModifier::GetLocalTransform(int32 Index) const
+const FTransform& FIKRigTransforms::GetLocalTransform(int32 Index) const
 {
 	if (Hierarchy && Hierarchy->IsValidIndex(Index))
 	{
 		// this is because this requires to modify LocalTransforms
-		const_cast<FIKRigTransformModifier*>(this)->UpdateLocalTransform(Index);
+		const_cast<FIKRigTransforms*>(this)->UpdateLocalTransform(Index);
 		return LocalTransforms[Index];
 	}
 
 	return FTransform::Identity;
 }
 
-const FTransform& FIKRigTransformModifier::GetGlobalTransform(int32 Index) const
+const FTransform& FIKRigTransforms::GetGlobalTransform(int32 Index) const
 {
-	return GlobalTransforms.GetGlobalTransform(Index);
+	if (GlobalTransforms.IsValidIndex(Index))
+	{
+		return GlobalTransforms[Index];
+	}
+
+	return FTransform::Identity;
 }
 /* This function does propagate through children
 	* of the current index and keeps last LocalTransform to up to date
 	*/
-void FIKRigTransformModifier::SetGlobalTransform_Internal(int32 Index, const FTransform& InTransform)
+void FIKRigTransforms::SetGlobalTransform_Internal(int32 Index, const FTransform& InTransform)
 {
 	// first calculate local transform if not found
 	EnsureLocalTransformsExist();
@@ -128,9 +129,9 @@ void FIKRigTransformModifier::SetGlobalTransform_Internal(int32 Index, const FTr
 /** Set Global Transform Recursive
 	*
 	*/
-void FIKRigTransformModifier::SetGlobalTransform_Recursive(int32 Index, const FTransform& InTransform)
+void FIKRigTransforms::SetGlobalTransform_Recursive(int32 Index, const FTransform& InTransform)
 {
-	GlobalTransforms.GlobalTransforms[Index] = InTransform;
+	GlobalTransforms[Index] = InTransform;
 	TArray<int32> Children = Hierarchy->FindChildren(Index);
 	for (int32 Child : Children)
 	{
@@ -140,14 +141,31 @@ void FIKRigTransformModifier::SetGlobalTransform_Recursive(int32 Index, const FT
 	}
 }
 
-void FIKRigTransformModifier::ResetGlobalTransform(const FIKRigTransform& InTransform)
+void FIKRigTransforms::SetAllGlobalTransforms(const TArray<FTransform>& InTransforms)
 {
-	if (Hierarchy && Hierarchy->GetNum() == InTransform.GetNum())
+	if (Hierarchy && Hierarchy->GetNum() == InTransforms.Num())
 	{
-		GlobalTransforms = InTransform;
+		GlobalTransforms = InTransforms;
 
 		// we don't want previous transform data
 		LocalTransforms.Reset();
 		LocalTransformDirtyFlags.Reset();
 	}
+}
+
+FTransform FIKRigTransforms::GetRelativeTransform(int32 ChildIndex, int32 ParentIndex) const
+{
+	if (GlobalTransforms.IsValidIndex(ChildIndex))
+	{
+		if (GlobalTransforms.IsValidIndex(ParentIndex))
+		{
+			return (GlobalTransforms[ChildIndex].GetRelativeTransform(GlobalTransforms[ParentIndex]));
+		}
+		else
+		{
+			return GlobalTransforms[ChildIndex];
+		}
+	}
+
+	return FTransform::Identity;
 }
\ No newline at end of file
diff --git a/Engine/Plugins/Animation/IKRig/Source/IKRig/Private/IKRigDefinition.cpp b/Engine/Plugins/Animation/IKRig/Source/IKRig/Private/IKRigDefinition.cpp
--- a/Engine/Plugins/Animation/IKRig/Source/IKRig/Private/IKRigDefinition.cpp
+++ b/Engine/Plugins/Animation/IKRig/Source/IKRig/Private/IKRigDefinition.cpp
@@ -86,9 +86,9 @@ bool UIKRigDefinition::AddBone(const FName& InName, const FName& InParent, const
 		if (InParent == NAME_None || Hierarchy.FindIndexFromBoneArray(InParent) != INDEX_NONE)
 		{
 			Hierarchy.Bones.Add(FIKRigBone(InName, InParent));
-			ReferencePose.GlobalTransforms.Add(InGlobalTransform);
+			RefPoseTransforms.Add(InGlobalTransform);
 
-			check (Hierarchy.Bones.Num() == ReferencePose.GetNum());
+			check (Hierarchy.Bones.Num() == RefPoseTransforms.Num());
 
 			Hierarchy.RebuildCacheData();
 			BoneAddedDelegate.Broadcast(InName);
@@ -134,11 +134,11 @@ bool UIKRigDefinition::RemoveBone(const FName& InName)
 	if (Index != INDEX_NONE)
 	{
 		// remove all children
-		FRemoveChild::RemoveChild_Recursively(Hierarchy, Hierarchy.Bones, ReferencePose.GlobalTransforms, InName);
+		FRemoveChild::RemoveChild_Recursively(Hierarchy, Hierarchy.Bones, RefPoseTransforms, InName);
 
 		// shrink now
 		Hierarchy.Bones.Shrink();
-		ReferencePose.GlobalTransforms.Shrink();
+		RefPoseTransforms.Shrink();
 		Hierarchy.RebuildCacheData();
 		BoneRemovedDelegate.Broadcast(InName);
 		return true;
@@ -332,14 +332,14 @@ void UIKRigDefinition::EnsureSortedCorrectly(bool bReSortIfNeeded)
 			{
 				if (Name == NewBones[Jndex].Name)
 				{
-					NewTransforms[Jndex] = ReferencePose.GlobalTransforms[Index];
+					NewTransforms[Jndex] = RefPoseTransforms[Index];
 				}
 			}
 		}
 
 		// override back to it
 		Hierarchy.Bones = NewBones;
-		ReferencePose.GlobalTransforms = NewTransforms;
+		RefPoseTransforms = NewTransforms;
 		Hierarchy.RebuildCacheData();
 	}
 }
@@ -349,7 +349,7 @@ void UIKRigDefinition::ResetHierarchy()
 	Hierarchy.Bones.Reset();
 	Hierarchy.ParentIndices.Reset();
 	Hierarchy.RuntimeNameLookupTable.Reset();
-	ReferencePose.GlobalTransforms.Reset();
+	RefPoseTransforms.Reset();
 }
 
 void UIKRigDefinition::EnsureCreateUniqueGoalName(FName& InOutGoal) const
diff --git a/Engine/Plugins/Animation/IKRig/Source/IKRig/Private/IKRigProcessor.cpp b/Engine/Plugins/Animation/IKRig/Source/IKRig/Private/IKRigProcessor.cpp
--- a/Engine/Plugins/Animation/IKRig/Source/IKRig/Private/IKRigProcessor.cpp
+++ b/Engine/Plugins/Animation/IKRig/Source/IKRig/Private/IKRigProcessor.cpp
@@ -22,23 +22,23 @@ void UIKRigProcessor::SetIKRigDefinition(UIKRigDefinition* InRigDefinition, bool
 
 	if (RigDefinition && bInitialize)
 	{
-		TransformModifier = FIKRigTransformModifier(&RigDefinition->GetHierarchy());
+		Transforms = FIKRigTransforms(&RigDefinition->GetHierarchy());
 		Initialize(RigDefinition->GetReferencePose());
 	}
 }
 
-void UIKRigProcessor::Initialize(const FIKRigTransform& InRefTransform)
+void UIKRigProcessor::Initialize(const TArray<FTransform>& InRefPoseTransforms)
 {
 	if (RigDefinition)
 	{
-		ReferenceTransform = InRefTransform;
+		RefPoseTransforms = InRefPoseTransforms;
 		Reinitialize();
 	}
 }
 
 void UIKRigProcessor::Reinitialize()
 {
-	TransformModifier.ResetGlobalTransform(ReferenceTransform);
+	Transforms.SetAllGlobalTransforms(RefPoseTransforms);
 	IKGoals = RigDefinition->GetGoals();
 	UIKRigSolver::FIKRigTransformGetter RefTransformGetter = UIKRigSolver::FIKRigTransformGetter::CreateUObject(this, &UIKRigProcessor::GetRefPoseGetter);
 	UIKRigSolver::FIKRigGoalGetter GoalGetter = UIKRigSolver::FIKRigGoalGetter::CreateUObject(this, &UIKRigProcessor::GoalGetter);
@@ -53,7 +53,7 @@ void UIKRigProcessor::Reinitialize()
 		if (ClassType != nullptr)
 		{
 			UIKRigSolver* NewSolver = DuplicateObject(RigSolvers[Index], this);
-			NewSolver->Init(TransformModifier, RefTransformGetter, GoalGetter);
+			NewSolver->Init(Transforms, RefTransformGetter, GoalGetter);
 			Solvers.Add(NewSolver);
 		}
 	}
@@ -61,9 +61,9 @@ void UIKRigProcessor::Reinitialize()
 	bInitialized = true;
 }
 
-const FIKRigTransform& UIKRigProcessor::GetRefPoseGetter() 
+const TArray<FTransform>& UIKRigProcessor::GetRefPoseGetter()
 {
-	return ReferenceTransform;
+	return RefPoseTransforms;
 }
 
 bool UIKRigProcessor::GoalGetter(const FName& InGoalName, FIKRigTarget& OutTarget)
@@ -94,7 +94,7 @@ void UIKRigProcessor::Solve()
 		for (int32 Index = 0; Index < Solvers.Num(); ++Index)
 		{
 			// Draw interface is pointer in case we want to nullify for optimization
-			Solvers[Index]->Solve(TransformModifier, &DrawInterface);
+			Solvers[Index]->Solve(Transforms, &DrawInterface);
 		}
 	}
 }
@@ -143,9 +143,9 @@ void UIKRigProcessor::GetGoals(TArray<FName>& OutNames) const
 	IKGoals.GenerateKeyArray(OutNames);
 }
 
-FIKRigTransformModifier& UIKRigProcessor::GetIKRigTransformModifier() 
+FIKRigTransforms& UIKRigProcessor::GetIKRigTransformModifier() 
 {
-	return TransformModifier;
+	return Transforms;
 }
 
 const FIKRigHierarchy* UIKRigProcessor::GetHierarchy() const
@@ -160,5 +160,5 @@ const FIKRigHierarchy* UIKRigProcessor::GetHierarchy() const
 
 void UIKRigProcessor::ResetToRefPose()
 {
-	TransformModifier.ResetGlobalTransform(ReferenceTransform);
+	Transforms.SetAllGlobalTransforms(RefPoseTransforms);
 }
\ No newline at end of file
diff --git a/Engine/Plugins/Animation/IKRig/Source/IKRig/Private/IKRigSolver.cpp b/Engine/Plugins/Animation/IKRig/Source/IKRig/Private/IKRigSolver.cpp
--- a/Engine/Plugins/Animation/IKRig/Source/IKRig/Private/IKRigSolver.cpp
+++ b/Engine/Plugins/Animation/IKRig/Source/IKRig/Private/IKRigSolver.cpp
@@ -6,7 +6,7 @@
 #include "IKRigDefinition.h"
 
 // input hierarchy and ref pose? 
-void UIKRigSolver::Init(const FIKRigTransformModifier& TransformModifier, FIKRigTransformGetter InRefPoseGetter, FIKRigGoalGetter InGoalGetter)
+void UIKRigSolver::Init(const FIKRigTransforms& TransformModifier, FIKRigTransformGetter InRefPoseGetter, FIKRigGoalGetter InGoalGetter)
 {
 	//SolverDefinition = InSolverDefinition;
 
@@ -27,7 +27,7 @@ bool UIKRigSolver::IsSolverActive() const
 
 // input : goal getter or goals
 // output : modified pose - GlobalTransforms
-void UIKRigSolver::Solve(FIKRigTransformModifier& InOutGlobalTransform, FControlRigDrawInterface* InOutDrawInterface)
+void UIKRigSolver::Solve(FIKRigTransforms& InOutGlobalTransform, FControlRigDrawInterface* InOutDrawInterface)
 {
 	if (IsSolverActive())
 	{
@@ -50,14 +50,14 @@ bool UIKRigSolver::GetEffectorTarget(const FIKRigEffector& InEffector, FIKRigTar
 	return false;
 }
 
-const FIKRigTransform& UIKRigSolver::GetReferencePose() const
+const TArray<FTransform>& UIKRigSolver::GetRefPoseTransforms() const
 {
 	if (RefPoseGetter.IsBound())
 	{
 		return RefPoseGetter.Execute();
 	}
 
-	static FIKRigTransform Dummy;
+	static TArray<FTransform> Dummy;
 	return Dummy;
 }
 
diff --git a/Engine/Plugins/Animation/IKRig/Source/IKRig/Private/Solvers/ConstraintSolver.cpp b/Engine/Plugins/Animation/IKRig/Source/IKRig/Private/Solvers/ConstraintSolver.cpp
--- a/Engine/Plugins/Animation/IKRig/Source/IKRig/Private/Solvers/ConstraintSolver.cpp
+++ b/Engine/Plugins/Animation/IKRig/Source/IKRig/Private/Solvers/ConstraintSolver.cpp
@@ -16,7 +16,7 @@ void UIKRigConstraintSolver::UnregisterQueryConstraintHandler()
 }
 */
 
-void UIKRigConstraintSolver::InitInternal(const FIKRigTransformModifier& InGlobalTransform)
+void UIKRigConstraintSolver::InitInternal(const FIKRigTransforms& InGlobalTransform)
 {
 	/*
 	if (ConstraintDefinition)
@@ -41,7 +41,7 @@ void UIKRigConstraintSolver::InitInternal(const FIKRigTransformModifier& InGloba
 	}*/
 }
 
-void UIKRigConstraintSolver::SolveInternal(FIKRigTransformModifier& InOutGlobalTransform, FControlRigDrawInterface* InOutDrawInterface)
+void UIKRigConstraintSolver::SolveInternal(FIKRigTransforms& InOutGlobalTransform, FControlRigDrawInterface* InOutDrawInterface)
 {
 	/*
 	FIKRigConstraintProfile* Current = ConstraintProfiles.Find(ActiveProfile);
diff --git a/Engine/Plugins/Animation/IKRig/Source/IKRig/Private/Solvers/TransformSolver.cpp b/Engine/Plugins/Animation/IKRig/Source/IKRig/Private/Solvers/TransformSolver.cpp
--- a/Engine/Plugins/Animation/IKRig/Source/IKRig/Private/Solvers/TransformSolver.cpp
+++ b/Engine/Plugins/Animation/IKRig/Source/IKRig/Private/Solvers/TransformSolver.cpp
@@ -11,7 +11,7 @@ UTransformSolver::UTransformSolver()
 {
 }
 
-void UTransformSolver::InitInternal(const FIKRigTransformModifier& InGlobalTransform)
+void UTransformSolver::InitInternal(const FIKRigTransforms& InGlobalTransform)
 {
 	
 }
@@ -21,7 +21,7 @@ bool UTransformSolver::IsSolverActive() const
 	return Super::IsSolverActive() && (bEnablePosition || bEnableRotation);
 }
 
-void UTransformSolver::SolveInternal(FIKRigTransformModifier& InOutGlobalTransform, FControlRigDrawInterface* InOutDrawInterface)
+void UTransformSolver::SolveInternal(FIKRigTransforms& InOutGlobalTransform, FControlRigDrawInterface* InOutDrawInterface)
 {
 	FIKRigTarget Target;
 	if (!GetEffectorTarget(TransformTarget, Target))
diff --git a/Engine/Plugins/Animation/IKRig/Source/IKRig/Public/IKRigConstraint.h b/Engine/Plugins/Animation/IKRig/Source/IKRig/Public/IKRigConstraint.h
--- a/Engine/Plugins/Animation/IKRig/Source/IKRig/Public/IKRigConstraint.h
+++ b/Engine/Plugins/Animation/IKRig/Source/IKRig/Public/IKRigConstraint.h
@@ -12,7 +12,7 @@
 #include "UObject/Object.h"
 #include "IKRigConstraint.generated.h"
 
-struct FIKRigTransformModifier;
+struct FIKRigTransforms;
 struct FControlRigDrawInterface;
 
 UCLASS(config = Engine, hidecategories = UObject, BlueprintType)
@@ -23,10 +23,10 @@ class IKRIG_API UIKRigConstraint : public UObject
 	bool bInitialized = false;
 
 public: 
-	void Setup(const FIKRigTransformModifier& InOutTransformModifier);
-	void SetAndApplyConstraint(FIKRigTransformModifier& InOutTransformModifier);
+	void Setup(const FIKRigTransforms& InOutTransformModifier);
+	void SetAndApplyConstraint(FIKRigTransforms& InOutTransformModifier);
 
-	virtual void Apply(FIKRigTransformModifier& InOutTransformModifier, FControlRigDrawInterface* InOutDrawInterface) {};
-	virtual void SetupInternal(const FIKRigTransformModifier& InOutTransformModifier) {};
+	virtual void Apply(FIKRigTransforms& InOutTransformModifier, FControlRigDrawInterface* InOutDrawInterface) {};
+	virtual void SetupInternal(const FIKRigTransforms& InOutTransformModifier) {};
 };
 
diff --git a/Engine/Plugins/Animation/IKRig/Source/IKRig/Public/IKRigDataTypes.h b/Engine/Plugins/Animation/IKRig/Source/IKRig/Public/IKRigDataTypes.h
--- a/Engine/Plugins/Animation/IKRig/Source/IKRig/Public/IKRigDataTypes.h
+++ b/Engine/Plugins/Animation/IKRig/Source/IKRig/Public/IKRigDataTypes.h
@@ -128,74 +128,17 @@ struct TIKRigEffectorMapKeyFuncs : public TDefaultMapKeyFuncs<const FIKRigEffect
 template <typename ValueType>
 using TIKRigEffectorMap = TMap<FIKRigEffector, ValueType, FDefaultSetAllocator, TIKRigEffectorMapKeyFuncs<ValueType>>;
 
-// If you want to modify this, use FIKRigTransformModifier
-USTRUCT()
-struct IKRIG_API FIKRigTransform
-{
-	GENERATED_BODY()
-
-private:
-	UPROPERTY(VisibleAnywhere, Category = FIKRigTransform)
-	TArray<FTransform> GlobalTransforms;
-
-public:
-	int32 GetNum() const 
-	{
-		return GlobalTransforms.Num();
-	}
-
-	bool IsValidIndex(int32 Index) const
-	{
-		return GlobalTransforms.IsValidIndex(Index);
-	}
-
-	const TArray<FTransform>& GetGlobalTransforms() const
-	{
-		return GlobalTransforms;
-	}
-
-	const FTransform& GetGlobalTransform(int32 Index) const
-	{
-		if (GlobalTransforms.IsValidIndex(Index))
-		{
-			return GlobalTransforms[Index];
-		}
-
-		return FTransform::Identity;
-	}
-
-	FTransform GetRelativeTransform(int32 Index, int32 BaseIndex) const
-	{
-		if (GlobalTransforms.IsValidIndex(Index))
-		{
-			if (GlobalTransforms.IsValidIndex(BaseIndex))
-			{
-				return (GlobalTransforms[Index].GetRelativeTransform(GlobalTransforms[BaseIndex]));
-			}
-			else
-			{
-				return GlobalTransforms[Index];
-			}
-		}
-
-		return FTransform::Identity;
-	}
-	friend struct FIKRigTransformModifier;
-	friend class UIKRigDefinition;
-	friend class UIKRigController;
-};
-
 // allows transform to be modified 
 // use this class to modify transform
 // @todo: ref pose getter 
-struct IKRIG_API FIKRigTransformModifier
+struct IKRIG_API FIKRigTransforms
 {
-	FIKRigTransformModifier()
+	FIKRigTransforms()
 		: Hierarchy(nullptr)
 	{
 	}
 
-	FIKRigTransformModifier(const FIKRigHierarchy* InHierarchy);
+	FIKRigTransforms(const FIKRigHierarchy* InHierarchy);
 
 	void SetGlobalTransform(int32 Index, const FTransform& InTransform, bool bPropagate);
 	void SetLocalTransform(int32 Index, const FTransform& InTransform, bool bUpdate);
@@ -203,29 +146,23 @@ struct IKRIG_API FIKRigTransformModifier
 	const FTransform& GetLocalTransform(int32 Index) const;
 	const FTransform& GetGlobalTransform(int32 Index) const;
 
-	void ResetGlobalTransform(const FIKRigTransform& InTransform);
+	void SetAllGlobalTransforms(const TArray<FTransform>& InTransforms);
 
 	const FIKRigHierarchy* Hierarchy;
 
 private:
 
-	// we use this info to modify 
-	// ensure memory cycle for these are valid through
-	FIKRigTransform GlobalTransforms;
+	TArray<FTransform> GlobalTransforms;
 	TArray<FTransform> LocalTransforms;
 	TBitArray<>	LocalTransformDirtyFlags;
 
 	void EnsureLocalTransformsExist();
 	void RecalculateLocalTransform();
 	void UpdateLocalTransform(int32 Index);
-	/* This function does propagate through children
-	 * of the current index and keeps last LocalTransform to up to date
-	 */
+	/* this function propagates to children and updates LocalTransforms*/
 	void SetGlobalTransform_Internal(int32 Index, const FTransform& InTransform);
-
-	/** Set Global Transform Recursive 
-	 *	
-	 */
 	void SetGlobalTransform_Recursive(int32 Index, const FTransform& InTransform);
+
+	FTransform GetRelativeTransform(int32 Index, int32 BaseIndex) const;
 };
 
diff --git a/Engine/Plugins/Animation/IKRig/Source/IKRig/Public/IKRigDefinition.h b/Engine/Plugins/Animation/IKRig/Source/IKRig/Public/IKRigDefinition.h
--- a/Engine/Plugins/Animation/IKRig/Source/IKRig/Public/IKRigDefinition.h
+++ b/Engine/Plugins/Animation/IKRig/Source/IKRig/Public/IKRigDefinition.h
@@ -46,9 +46,9 @@ public:
 		return NameArray;
 	}
 
-	const FIKRigTransform& GetReferencePose() const
+	const TArray<FTransform>& GetReferencePose() const
 	{
-		return ReferencePose;
+		return RefPoseTransforms;
 	}
 
 	const TArray<UIKRigSolver*>& GetSolvers() const 
@@ -81,7 +81,7 @@ private:
 	// IKRigDefinition interface handles the integrity of (this and Hierarchy)
 	// Use IKRigController for APIs
 	UPROPERTY(VisibleAnywhere, Category = "Hierarchy")
-	FIKRigTransform ReferencePose;
+	TArray<FTransform> RefPoseTransforms;
 
 	// stack of solvers, executed in order
 	UPROPERTY(EditAnywhere, instanced, Category = "Solver")
diff --git a/Engine/Plugins/Animation/IKRig/Source/IKRig/Public/IKRigProcessor.h b/Engine/Plugins/Animation/IKRig/Source/IKRig/Public/IKRigProcessor.h
--- a/Engine/Plugins/Animation/IKRig/Source/IKRig/Public/IKRigProcessor.h
+++ b/Engine/Plugins/Animation/IKRig/Source/IKRig/Public/IKRigProcessor.h
@@ -1,10 +1,5 @@
 // Copyright Epic Games, Inc. All Rights Reserved.
 
-/**
- * Contains IKRig Definition runtime code. This takes IKRigDefinition asset as an input
- *
- */
-
 #pragma once
 
 #include "CoreMinimal.h"
@@ -33,7 +28,7 @@ private:
 
 	// reference transform
 	UPROPERTY(transient)
-	FIKRigTransform ReferenceTransform;
+	TArray<FTransform> RefPoseTransforms;
 
 	// solvers
 	UPROPERTY(transient)
@@ -48,18 +43,16 @@ private:
 	UPROPERTY(transient)
 	TMap<FName, FIKRigGoal> IKGoals;
 
-	FIKRigTransformModifier TransformModifier;
+	FIKRigTransforms Transforms;
 	FControlRigDrawInterface DrawInterface;
 
 	bool bInitialized = false;
 
 public: 
 
-	// IKRigProcessor implementation functions
-	// Set IKRigDefinition 
 	// if bInitialize is true, it will use Default Ref transform of IKRigDefinition to initialize
 	void SetIKRigDefinition(UIKRigDefinition* InRigDefinition, bool bInitialize = false);
-	void Initialize(const FIKRigTransform& InRefTransform);
+	void Initialize(const TArray<FTransform>& InRefPoseTransforms);
 	void Reinitialize();
 	void Terminate();
 	void Solve();
@@ -69,7 +62,7 @@ public:
 	void SetGoalRotation(const FName& GoalName, const FRotator& InRotation);
 	void SetGoalTarget(const FName& GoalName, const FIKRigTarget& InTarget);
 	void GetGoals(TArray<FName>& OutNames) const;
-	FIKRigTransformModifier& GetIKRigTransformModifier();
+	FIKRigTransforms& GetIKRigTransformModifier();
 	const FIKRigHierarchy* GetHierarchy() const;
 
 	void ResetToRefPose();
@@ -80,6 +73,6 @@ private:
 
 	// delegate functions
 	bool GoalGetter(const FName& InGoalName, FIKRigTarget& OutTarget);
-	const FIKRigTransform& GetRefPoseGetter();
+	const TArray<FTransform>& GetRefPoseGetter();
 };
 
diff --git a/Engine/Plugins/Animation/IKRig/Source/IKRig/Public/IKRigSolver.h b/Engine/Plugins/Animation/IKRig/Source/IKRig/Public/IKRigSolver.h
--- a/Engine/Plugins/Animation/IKRig/Source/IKRig/Public/IKRigSolver.h
+++ b/Engine/Plugins/Animation/IKRig/Source/IKRig/Public/IKRigSolver.h
@@ -11,7 +11,7 @@
 struct FIKRigEffector;
 struct FIKRigTarget;
 struct FIKRigTransform;
-struct FIKRigTransformModifier;
+struct FIKRigTransforms;
 struct FControlRigDrawInterface;
 
 // run time processor 
@@ -22,16 +22,16 @@ class IKRIG_API UIKRigSolver : public UObject
 
 public: 
 	/** Required delegates to run this solver */
-	DECLARE_DELEGATE_RetVal(const FIKRigTransform&, FIKRigTransformGetter);
+	DECLARE_DELEGATE_RetVal(const TArray<FTransform>&, FIKRigTransformGetter);
 	DECLARE_DELEGATE_RetVal_TwoParams(bool, FIKRigGoalGetter, const FName& /*InGoalName*/, FIKRigTarget& /*OutTarget*/);
 
 	// input hierarchy and ref pose? 
-	void Init(const FIKRigTransformModifier& TransformModifier, FIKRigTransformGetter InRefPoseGetter, FIKRigGoalGetter InGoalGetter/*, FSolveConstraint& InConstraintHandler*/);
+	void Init(const FIKRigTransforms& TransformModifier, FIKRigTransformGetter InRefPoseGetter, FIKRigGoalGetter InGoalGetter/*, FSolveConstraint& InConstraintHandler*/);
 
 	// input : goal getter or goals
 	// output : modified pose - GlobalTransforms
 	// or use SolverInternal function
-	void Solve(FIKRigTransformModifier& InOutGlobalTransform, FControlRigDrawInterface* InOutDrawInterface);
+	void Solve(FIKRigTransforms& InOutGlobalTransform, FControlRigDrawInterface* InOutDrawInterface);
 
 	const TIKRigEffectorMap<FName>& GetEffectorToGoal() const
 	{
@@ -42,12 +42,12 @@ public:
 
 protected:
 
-	virtual void InitInternal(const FIKRigTransformModifier& InGlobalTransform) {};
-	virtual void SolveInternal(FIKRigTransformModifier& InOutGlobalTransform, FControlRigDrawInterface* InOutDrawInterface) {};
+	virtual void InitInternal(const FIKRigTransforms& InGlobalTransform) {};
+	virtual void SolveInternal(FIKRigTransforms& InOutGlobalTransform, FControlRigDrawInterface* InOutDrawInterface) {};
 	virtual bool IsSolverActive() const;
 
 	bool GetEffectorTarget(const FIKRigEffector& InEffector, FIKRigTarget& OutTarget) const;
-	const FIKRigTransform& GetReferencePose() const;
+	const TArray<FTransform>& GetRefPoseTransforms() const;
 
 	UPROPERTY(EditAnywhere, Category = "Definition")
 	bool bEnabled = true;
diff --git a/Engine/Plugins/Animation/IKRig/Source/IKRig/Public/Solvers/ConstraintSolver.h b/Engine/Plugins/Animation/IKRig/Source/IKRig/Public/Solvers/ConstraintSolver.h
--- a/Engine/Plugins/Animation/IKRig/Source/IKRig/Public/Solvers/ConstraintSolver.h
+++ b/Engine/Plugins/Animation/IKRig/Source/IKRig/Public/Solvers/ConstraintSolver.h
@@ -24,8 +24,8 @@ class IKRIG_API UIKRigConstraintSolver : public UIKRigSolver
 
 public: 
 	// during we don't mutate bone transform, but you can read it
-	virtual void InitInternal(const FIKRigTransformModifier& InGlobalTransform) override;
-	virtual void SolveInternal(FIKRigTransformModifier& InOutGlobalTransform, FControlRigDrawInterface* InOutDrawInterface) override;
+	virtual void InitInternal(const FIKRigTransforms& InGlobalTransform) override;
+	virtual void SolveInternal(FIKRigTransforms& InOutGlobalTransform, FControlRigDrawInterface* InOutDrawInterface) override;
 
 	// this can be utilized for just subset of constraints
 	//void SolverConstraints(const TArray<FName>& ConstraintsList, FIKRigTransformModifier& InOutGlobalTransform, FControlRigDrawInterface* InOutDrawInterface);
diff --git a/Engine/Plugins/Animation/IKRig/Source/IKRig/Public/Solvers/TransformSolver.h b/Engine/Plugins/Animation/IKRig/Source/IKRig/Public/Solvers/TransformSolver.h
--- a/Engine/Plugins/Animation/IKRig/Source/IKRig/Public/Solvers/TransformSolver.h
+++ b/Engine/Plugins/Animation/IKRig/Source/IKRig/Public/Solvers/TransformSolver.h
@@ -26,8 +26,8 @@ public:
 
 protected:
 	
-	virtual void InitInternal(const FIKRigTransformModifier& InGlobalTransform) override;
-	virtual void SolveInternal(FIKRigTransformModifier& InOutGlobalTransform, FControlRigDrawInterface* InOutDrawInterface) override;
+	virtual void InitInternal(const FIKRigTransforms& InGlobalTransform) override;
+	virtual void SolveInternal(FIKRigTransforms& InOutGlobalTransform, FControlRigDrawInterface* InOutDrawInterface) override;
 	virtual bool IsSolverActive() const override;
 
 private:
diff --git a/Engine/Plugins/Animation/IKRig/Source/IKRigEditor/Private/IKRigController.cpp b/Engine/Plugins/Animation/IKRig/Source/IKRigEditor/Private/IKRigController.cpp
--- a/Engine/Plugins/Animation/IKRig/Source/IKRigEditor/Private/IKRigController.cpp
+++ b/Engine/Plugins/Animation/IKRig/Source/IKRigEditor/Private/IKRigController.cpp
@@ -137,7 +137,7 @@ void UIKRigController::SetSkeleton(const FReferenceSkeleton& InSkeleton)
 			ensure (IKRigDefinition->AddBone(RefBoneInfo[Index].Name, (ParentIndex != INDEX_NONE)? RefBoneInfo[ParentIndex].Name : NAME_None, RefPoseInCS[Index]));
 		}
 
-		ensure(IKRigDefinition->Hierarchy.GetNum() == IKRigDefinition->ReferencePose.GetNum());
+		ensure(IKRigDefinition->Hierarchy.GetNum() == IKRigDefinition->RefPoseTransforms.Num());
 	}
 }
 
diff --git a/Engine/Plugins/Animation/IKRig/Source/IKRigEditor/Public/IKRigController.h b/Engine/Plugins/Animation/IKRig/Source/IKRigEditor/Public/IKRigController.h
--- a/Engine/Plugins/Animation/IKRig/Source/IKRigEditor/Public/IKRigController.h
+++ b/Engine/Plugins/Animation/IKRig/Source/IKRigEditor/Public/IKRigController.h
@@ -27,17 +27,22 @@ class IKRIGEDITOR_API UIKRigController : public UObject
 	GENERATED_BODY()
 
 public:
-	// hierarchy operators
+
 	void SetSkeleton(const FReferenceSkeleton& InSkeleton);
-	// const hierarchy getter
 	const FIKRigHierarchy* GetHierarchy() const
 	{
 		return (IKRigDefinition)? &IKRigDefinition->Hierarchy : nullptr;
 	}
 
-	const FIKRigTransform* GetReferenceTransform() const
+	const TArray<FTransform>& GetRefPoseTransforms() const
 	{
-		return (IKRigDefinition) ? &IKRigDefinition->ReferencePose : nullptr;
+		if (IKRigDefinition)
+		{
+			return IKRigDefinition->RefPoseTransforms;
+		}
+
+		static TArray<FTransform> Dummy;
+		return Dummy;
 	}
 
 	bool AddBone(const FName& InName, const FName& InParent, const FTransform& InGlobalTransform);
diff --git a/Engine/Plugins/Experimental/FullBodyIK/Source/FullBodyIK/Private/FullBodyIKSolver.cpp b/Engine/Plugins/Experimental/FullBodyIK/Source/FullBodyIK/Private/FullBodyIKSolver.cpp
--- a/Engine/Plugins/Experimental/FullBodyIK/Source/FullBodyIK/Private/FullBodyIKSolver.cpp
+++ b/Engine/Plugins/Experimental/FullBodyIK/Source/FullBodyIK/Private/FullBodyIKSolver.cpp
@@ -18,7 +18,7 @@ const FName UFullBodyIKSolver::EffectorTargetPrefix = FName(TEXT("FullBodyIKTarg
 //////////////////////////////////////////////
 // utility functions
 ////////////////////////////////////////////////
-static float EnsureToAddBoneToLinkData(const FIKRigTransformModifier& TransformModifier, const int32 CurrentItem, TArray<FFBIKLinkData>& LinkData,
+static float EnsureToAddBoneToLinkData(const FIKRigTransforms& TransformModifier, const int32 CurrentItem, TArray<FFBIKLinkData>& LinkData,
 	TMap<int32, int32>& HierarchyToLinkDataMap, TMap<int32, int32>& LinkDataToHierarchyIndices)
 {
 	float ChainLength = 0;
@@ -84,7 +84,7 @@ static void AddToEffectorTarget(int32 EffectorIndex, const int32 Effector, TMap<
 }
 
 // the output should be from current index -> to root parent 
-static bool GetBoneChain(const FIKRigTransformModifier& TransformModifier, const FName& Root, const FName& Current, TArray<int32>& ChainIndices)
+static bool GetBoneChain(const FIKRigTransforms& TransformModifier, const FName& Root, const FName& Current, TArray<int32>& ChainIndices)
 {
 	ChainIndices.Reset();
 
@@ -110,7 +110,7 @@ static bool GetBoneChain(const FIKRigTransformModifier& TransformModifier, const
 	return ChainIndices.Num() > 0;
 }
 
-static void AddEffectors(const FIKRigTransformModifier& TransformModifier, const FName Root, const TArray<FFBIKRigEffector>& Effectors,
+static void AddEffectors(const FIKRigTransforms& TransformModifier, const FName Root, const TArray<FFBIKRigEffector>& Effectors,
 	TArray<FFBIKLinkData>& LinkData, TMap<int32, FFBIKEffectorTarget>& EffectorTargets, TArray<int32>& EffectorLinkIndices,
 	TMap<int32, int32>& LinkDataToHierarchyIndices, TMap<int32, int32>& HierarchyToLinkDataMap, const FSolverInput& SolverProperty)
 {
@@ -204,7 +204,7 @@ UFullBodyIKSolver::UFullBodyIKSolver()
 {
 }
 
-void UFullBodyIKSolver::InitInternal(const FIKRigTransformModifier& InGlobalTransform)
+void UFullBodyIKSolver::InitInternal(const FIKRigTransforms& InGlobalTransform)
 {
 
 	LinkData.Reset();
@@ -229,7 +229,7 @@ bool UFullBodyIKSolver::IsSolverActive() const
 }
 
 void UFullBodyIKSolver::SolveInternal(
-	FIKRigTransformModifier& InOutGlobalTransform, 
+	FIKRigTransforms& InOutGlobalTransform, 
 	FControlRigDrawInterface* InOutDrawInterface)
 {
 	if (LinkDataToHierarchyIndices.Num() > 0)
diff --git a/Engine/Plugins/Experimental/FullBodyIK/Source/FullBodyIK/Private/FullBodyIKSolver.h b/Engine/Plugins/Experimental/FullBodyIK/Source/FullBodyIK/Private/FullBodyIKSolver.h
--- a/Engine/Plugins/Experimental/FullBodyIK/Source/FullBodyIK/Private/FullBodyIKSolver.h
+++ b/Engine/Plugins/Experimental/FullBodyIK/Source/FullBodyIK/Private/FullBodyIKSolver.h
@@ -99,8 +99,8 @@ private:
 	TArray<FJacobianDebugData> DebugData;
 
 protected:
-	virtual void InitInternal(const FIKRigTransformModifier& InGlobalTransform) override;
-	virtual void SolveInternal(FIKRigTransformModifier& InOutGlobalTransform, FControlRigDrawInterface* InOutDrawInterface) override;
+	virtual void InitInternal(const FIKRigTransforms& InGlobalTransforms) override;
+	virtual void SolveInternal(FIKRigTransforms& InOutGlobalTransforms, FControlRigDrawInterface* InOutDrawInterface) override;
 	virtual bool IsSolverActive() const override;
 
 public:
