Fix uninitialized struct members

diff --git a/Engine/Plugins/Animation/IKRig/Source/IKRig/Public/IKRigDataTypes.h b/Engine/Plugins/Animation/IKRig/Source/IKRig/Public/IKRigDataTypes.h
--- a/Engine/Plugins/Animation/IKRig/Source/IKRig/Public/IKRigDataTypes.h
+++ b/Engine/Plugins/Animation/IKRig/Source/IKRig/Public/IKRigDataTypes.h
@@ -96,14 +96,14 @@ struct IKRIG_API FIKRigEffector
 	GENERATED_BODY()
 
 	UPROPERTY()
-	FGuid Guid = FGuid();
+	FGuid Guid;
 
 	UPROPERTY(EditAnywhere, Category = FIKRigEffector)
 	FName Bone;
 
 	FIKRigEffector()
+		: Guid(FGuid::NewGuid())
 	{
-		Guid = FGuid::NewGuid();
 	}
 
 	friend FArchive& operator<<(FArchive& Ar, FIKRigEffector& Effector)
