Refactoring goals and constraints in IKRig.

diff --git a/Engine/Plugins/Animation/IKRig/Source/IKRig/Private/AnimNodes/AnimNode_IKRig.cpp b/Engine/Plugins/Animation/IKRig/Source/IKRig/Private/AnimNodes/AnimNode_IKRig.cpp
--- a/Engine/Plugins/Animation/IKRig/Source/IKRig/Private/AnimNodes/AnimNode_IKRig.cpp
+++ b/Engine/Plugins/Animation/IKRig/Source/IKRig/Private/AnimNodes/AnimNode_IKRig.cpp
@@ -26,65 +26,66 @@ void FAnimNode_IKRig::Evaluate_AnyThread(FPoseContext& Output)
 	}
 	else
 	{
-		// apply refpose
 		SourcePose.ResetToRefPose();
 	}
 
-	if (RigProcessor)
+	if (!RigProcessor)
 	{
-		FIKRigTransforms& TransformModifier = RigProcessor->GetIKRigTransformModifier();
+		Output = SourcePose;
+		return;
+	}
 
-		if (bStartFromRefPose)
-		{
-			RigProcessor->ResetToRefPose();
-		}
+	if (bStartFromRefPose)
+	{
+		RigProcessor->ResetToRefPose();
+	}
 		
-		FCompactPose& OutPose = SourcePose.Pose;
+	FCompactPose& OutPose = SourcePose.Pose;
+	FIKRigTransforms& Transforms = RigProcessor->GetTransforms();
 
-		for (FCompactPoseBoneIndex CPIndex : OutPose.ForEachBoneIndex())
+	// copy input pose
+	for (FCompactPoseBoneIndex CPIndex : OutPose.ForEachBoneIndex())
+	{
+		int32* Index = CompactPoseToRigIndices.Find(CPIndex);
+		if (Index)
 		{
-			int32* Index = CompactPoseToRigIndices.Find(CPIndex);
-			if (Index)
-			{
-				// Todo: this is again slow. 
-				// the tricky part is that if we go start from child to parent
-				// we may still missing joints that don't exists
-				// so this is slow and for safety but needs revisit
-				TransformModifier.SetLocalTransform(*Index, OutPose[CPIndex], true);
-			}
+			// Todo: this is again slow. 
+			// the tricky part is that if we go start from child to parent
+			// we may still missing joints that don't exists
+			// so this is slow and for safety but needs revisit
+			Transforms.SetLocalTransform(*Index, OutPose[CPIndex], true);
 		}
+	}
 
-		for (int32 GoalIndex = 0; GoalIndex < GoalNames.Num(); ++GoalIndex)
+	// update goal transforms before solve
+	for (int32 GoalIndex = 0; GoalIndex < GoalNames.Num(); ++GoalIndex)
+	{
+		if (ensure(GoalTransforms.IsValidIndex(GoalIndex)))
 		{
-			if (ensure(GoalTransforms.IsValidIndex(GoalIndex)))
-			{
-				RigProcessor->SetGoalPosition( GoalNames[GoalIndex], GoalTransforms[GoalIndex].GetLocation());
-				RigProcessor->SetGoalRotation(GoalNames[GoalIndex], GoalTransforms[GoalIndex].GetRotation().Rotator());
-			}
+			RigProcessor->Goals.SetGoalTransform(
+				GoalNames[GoalIndex], 
+				GoalTransforms[GoalIndex].GetLocation(),
+				GoalTransforms[GoalIndex].GetRotation().Rotator());
 		}
+	}
 
-		// fill up current pose
-		RigProcessor->Solve();
+	// run stack of solvers, generate new pose
+	RigProcessor->Solve();
 
-		// now we copy back to output pose
-		for (FCompactPoseBoneIndex CPIndex : OutPose.ForEachBoneIndex())
+	// copy output pose
+	for (FCompactPoseBoneIndex CPIndex : OutPose.ForEachBoneIndex())
+	{
+		int32* Index = CompactPoseToRigIndices.Find(CPIndex);
+		if (Index)
 		{
-			int32* Index = CompactPoseToRigIndices.Find(CPIndex);
-			if (Index)
-			{
-				Output.Pose[CPIndex] = TransformModifier.GetLocalTransform(*Index);
-			}
+			Output.Pose[CPIndex] = Transforms.GetLocalTransform(*Index);
 		}
+	}
 
-		Output.Pose.NormalizeRotations();
+	Output.Pose.NormalizeRotations();
 
-		// draw the interface
-		QueueDrawInterface(Output.AnimInstanceProxy, Output.AnimInstanceProxy->GetComponentTransform());
-	}
-	else
-	{
-		Output = SourcePose;
-	}
+	// draw the interface
+	QueueDrawInterface(Output.AnimInstanceProxy, Output.AnimInstanceProxy->GetComponentTransform());
 }
 
 void FAnimNode_IKRig::GatherDebugData(FNodeDebugData& DebugData)
@@ -103,9 +104,9 @@ void FAnimNode_IKRig::Initialize_AnyThread(const FAnimationInitializeContext& Co
 
 	if (RigProcessor)
 	{
-		RigProcessor->SetIKRigDefinition(RigDefinitionAsset, true);
+		RigProcessor->Initialize(RigDefinitionAsset);
 		GoalNames.Reset();
-		RigProcessor->GetGoals(GoalNames);
+		RigProcessor->Goals.GetNames(GoalNames);
 	}
 }
 
@@ -142,7 +143,7 @@ FName FAnimNode_IKRig::GetGoalName(int32 Index) const
 	{
 		// TODO: fix  this vs RigProcessor goals. I think RigProcessor
 		// only issue with RigProcessor is you don't know until it compiles
-		TArray<FName> Names = RigDefinitionAsset->GetGoalsNames();
+		TArray<FName> Names = RigDefinitionAsset->GetGoalNames();
 		if (Names.IsValidIndex(Index))
 		{
 			return Names[Index];
diff --git a/Engine/Plugins/Animation/IKRig/Source/IKRig/Private/Constraints/RotationLimitConstraint.cpp b/Engine/Plugins/Animation/IKRig/Source/IKRig/Private/Constraints/RotationLimitConstraint.cpp
deleted file mode 100644
diff --git a/Engine/Plugins/Animation/IKRig/Source/IKRig/Private/Constraints/RotationLimitConstraint.h b/Engine/Plugins/Animation/IKRig/Source/IKRig/Private/Constraints/RotationLimitConstraint.h
deleted file mode 100644
diff --git a/Engine/Plugins/Animation/IKRig/Source/IKRig/Private/IKRigBoneSetting.cpp b/Engine/Plugins/Animation/IKRig/Source/IKRig/Private/IKRigBoneSetting.cpp
new file mode 100644
--- /dev/null
+++ b/Engine/Plugins/Animation/IKRig/Source/IKRig/Private/IKRigBoneSetting.cpp
@@ -0,0 +1,20 @@
+// Copyright Epic Games, Inc. All Rights Reserved.
+
+#include "IKRigBoneSetting.h"
+
+/*
+void UIKRigConstraint::Setup(const FIKRigTransforms& InOutTransformModifier)
+{
+	SetupInternal(InOutTransformModifier);
+	bInitialized = true;
+}
+
+void UIKRigConstraint::SetAndApplyConstraint(FIKRigTransforms& InOutTransformModifier)
+{
+	if (!bInitialized)
+	{
+		Setup(const_cast<FIKRigTransforms&>(InOutTransformModifier));
+	}
+
+	Apply(InOutTransformModifier, nullptr);
+}*/
diff --git a/Engine/Plugins/Animation/IKRig/Source/IKRig/Private/IKRigConstraint.cpp b/Engine/Plugins/Animation/IKRig/Source/IKRig/Private/IKRigConstraint.cpp
deleted file mode 100644
diff --git a/Engine/Plugins/Animation/IKRig/Source/IKRig/Private/IKRigDefinition.cpp b/Engine/Plugins/Animation/IKRig/Source/IKRig/Private/IKRigDefinition.cpp
--- a/Engine/Plugins/Animation/IKRig/Source/IKRig/Private/IKRigDefinition.cpp
+++ b/Engine/Plugins/Animation/IKRig/Source/IKRig/Private/IKRigDefinition.cpp
@@ -20,7 +20,7 @@ void UIKRigDefinition::UpdateGoal()
 	{
 		if (Solver)
 		{
-			Solver->CollectGoals(ListOfGoals);
+			Solver->AppendGoalNamesToArray(ListOfGoals);
 		}
 	}
 
@@ -30,7 +30,7 @@ void UIKRigDefinition::UpdateGoal()
 	TArray<FName> GoalsToRemove;
 
 	// first we go through IKGoals and see if we should remove unused ones
-	for (auto Iter = IKGoals.CreateConstIterator(); Iter; ++Iter)
+	for (auto Iter = Goals.CreateConstIterator(); Iter; ++Iter)
 	{
 		// if it's not used any more
 		if (!ListOfGoals.Contains(Iter->Key))
@@ -41,18 +41,18 @@ void UIKRigDefinition::UpdateGoal()
 
 	for (const FName& Goal : GoalsToRemove)
 	{
-		IKGoals.Remove(Goal);
+		Goals.Remove(Goal);
 	}
 
 	// now add new ones if it doesn't exist
 	for (int32 Index=0; Index<ListOfGoals.Num(); ++Index)
 	{
 		const FName& Goal = ListOfGoals[Index];
-		if (!IKGoals.Contains(Goal))
+		if (!Goals.Contains(Goal))
 		{
 			// if we don't have it, add to new one
 			FIKRigGoal NewGoal(Goal);
-			IKGoals.Add(Goal, NewGoal);
+			Goals.Add(Goal, NewGoal);
 		}
 	}
 
@@ -359,7 +359,7 @@ void UIKRigDefinition::EnsureCreateUniqueGoalName(FName& InOutGoal) const
 	FString GoalNameString = InOutGoal.ToString();
 
 	// while contains
-	while(IKGoals.Find(FName(*GoalNameString)))
+	while(Goals.Find(FName(*GoalNameString)))
 	{
 		GoalNameString = FString::Format(TEXT("{0}_{1}"), {InOutGoal.ToString(), Index++});
 	}
diff --git a/Engine/Plugins/Animation/IKRig/Source/IKRig/Private/IKRigProcessor.cpp b/Engine/Plugins/Animation/IKRig/Source/IKRig/Private/IKRigProcessor.cpp
--- a/Engine/Plugins/Animation/IKRig/Source/IKRig/Private/IKRigProcessor.cpp
+++ b/Engine/Plugins/Animation/IKRig/Source/IKRig/Private/IKRigProcessor.cpp
@@ -1,149 +1,58 @@
 // Copyright Epic Games, Inc. All Rights Reserved.
 
-/**
- * Contains IKRigProcessor - Runtime Implementation
- *
- */
-
 #include "IKRigProcessor.h"
 #include "IKRigDefinition.h"
 #include "IKRigSolver.h"
 #include "Solvers/ConstraintSolver.h"
 
- // IKRigProcessor implementation functions
-void UIKRigProcessor::SetIKRigDefinition(UIKRigDefinition* InRigDefinition, bool bInitialize /*= false*/)
+
+void UIKRigProcessor::Initialize(UIKRigDefinition* InRigDefinition)
 {
-	if (RigDefinition)
-	{
-		Terminate();
-	}
+	bInitialized = false;
 
 	RigDefinition = InRigDefinition;
-
-	if (RigDefinition && bInitialize)
-	{
-		Transforms = FIKRigTransforms(&RigDefinition->GetHierarchy());
-		Initialize(RigDefinition->GetReferencePose());
-	}
-}
-
-void UIKRigProcessor::Initialize(const TArray<FTransform>& InRefPoseTransforms)
-{
-	if (RigDefinition)
+	if (!RigDefinition)
 	{
-		RefPoseTransforms = InRefPoseTransforms;
-		Reinitialize();
+		return;
 	}
-}
-
-void UIKRigProcessor::Reinitialize()
-{
+	
+	RefPoseTransforms = RigDefinition->GetReferencePose();
+	Transforms = FIKRigTransforms(&RigDefinition->GetHierarchy());
 	Transforms.SetAllGlobalTransforms(RefPoseTransforms);
-	IKGoals = RigDefinition->GetGoals();
-	UIKRigSolver::FIKRigTransformGetter RefTransformGetter = UIKRigSolver::FIKRigTransformGetter::CreateUObject(this, &UIKRigProcessor::GetRefPoseGetter);
-	UIKRigSolver::FIKRigGoalGetter GoalGetter = UIKRigSolver::FIKRigGoalGetter::CreateUObject(this, &UIKRigProcessor::GoalGetter);
+	
+	Goals.SetAllGoals(RigDefinition->GetGoals());
 
 	TArray<UIKRigSolver*> RigSolvers = RigDefinition->GetSolvers();
 	const int32 NumSolvers = RigSolvers.Num();
 	Solvers.Reset(NumSolvers);
-
 	for (int32 Index = 0; Index < NumSolvers; ++Index)
 	{
-		TSubclassOf<UIKRigSolver> ClassType = RigSolvers[Index]->GetClass();
-		if (ClassType != nullptr)
-		{
-			UIKRigSolver* NewSolver = DuplicateObject(RigSolvers[Index], this);
-			NewSolver->Init(Transforms, RefTransformGetter, GoalGetter);
-			Solvers.Add(NewSolver);
-		}
+		UIKRigSolver* Solver = DuplicateObject(RigSolvers[Index], this);
+		Solver->Init(Transforms);
+		Solvers.Add(Solver);
 	}
 
 	bInitialized = true;
 }
 
-const TArray<FTransform>& UIKRigProcessor::GetRefPoseGetter()
-{
-	return RefPoseTransforms;
-}
-
-bool UIKRigProcessor::GoalGetter(const FName& InGoalName, FIKRigTarget& OutTarget)
-{
-	FIKRigTarget* Target = FindGoal(InGoalName);
-	if (Target)
-	{
-		OutTarget = *Target;
-
-		return true;
-	}
-
-	return false;
-}
-
-void UIKRigProcessor::Terminate()
-{
-	bInitialized = false;
-}
-
 void UIKRigProcessor::Solve()
 {
-	if (bInitialized)
+	if (!bInitialized)
 	{
-		check (RigDefinition);
-
-		DrawInterface.Reset();
-		for (int32 Index = 0; Index < Solvers.Num(); ++Index)
-		{
-			// Draw interface is pointer in case we want to nullify for optimization
-			Solvers[Index]->Solve(Transforms, &DrawInterface);
-		}
+		return;
 	}
-}
 
-// update goal functions
-void UIKRigProcessor::SetGoalPosition(const FName& GoalName, const FVector& InPosition)
-{
-	FIKRigTarget* Target = FindGoal(GoalName);
-	if (Target)
-	{
-		Target->PositionTarget.Position = InPosition;
-	}
-}
-
-void UIKRigProcessor::SetGoalRotation(const FName& GoalName, const FRotator& InRotation)
-{
-	FIKRigTarget* Target = FindGoal(GoalName);
-	if (Target)
-	{
-		Target->RotationTarget.Rotation = InRotation;
-	}
-}
+	check (RigDefinition);
 
-void UIKRigProcessor::SetGoalTarget(const FName& GoalName, const FIKRigTarget& InTarget)
-{
-	FIKRigTarget* Target = FindGoal(GoalName);
-	if (Target)
-	{
-		*Target = InTarget;
-	}
-}
+	DrawInterface.Reset();
 
-FIKRigTarget* UIKRigProcessor::FindGoal(const FName& GoalName)
-{
-	FIKRigGoal* RigGoal = IKGoals.Find(GoalName);
-	if (RigGoal)
+	for (int32 Index = 0; Index < Solvers.Num(); ++Index)
 	{
-		return &RigGoal->Target;
+		Solvers[Index]->Solve(Transforms, Goals, &DrawInterface);
 	}
-
-	return nullptr;
-}
-
-void UIKRigProcessor::GetGoals(TArray<FName>& OutNames) const
-{
-	IKGoals.GenerateKeyArray(OutNames);
 }
 
-FIKRigTransforms& UIKRigProcessor::GetIKRigTransformModifier() 
+FIKRigTransforms& UIKRigProcessor::GetTransforms() 
 {
 	return Transforms;
 }
diff --git a/Engine/Plugins/Animation/IKRig/Source/IKRig/Private/IKRigSolver.cpp b/Engine/Plugins/Animation/IKRig/Source/IKRig/Private/IKRigSolver.cpp
--- a/Engine/Plugins/Animation/IKRig/Source/IKRig/Private/IKRigSolver.cpp
+++ b/Engine/Plugins/Animation/IKRig/Source/IKRig/Private/IKRigSolver.cpp
@@ -5,76 +5,48 @@
 #include "IKRigDataTypes.h"
 #include "IKRigDefinition.h"
 
-// input hierarchy and ref pose? 
-void UIKRigSolver::Init(const FIKRigTransforms& TransformModifier, FIKRigTransformGetter InRefPoseGetter, FIKRigGoalGetter InGoalGetter)
-{
-	//SolverDefinition = InSolverDefinition;
-
-	RefPoseGetter = InRefPoseGetter;
-	ensure(RefPoseGetter.IsBound());
-
-	GoalGetter = InGoalGetter;
-	ensure(GoalGetter.IsBound());
-
-	InitInternal(TransformModifier);
-}
-
-bool UIKRigSolver::IsSolverActive() const 
-{
-	//return (SolverDefinition && bEnabled);
-	return bEnabled;
-}
 
-// input : goal getter or goals
-// output : modified pose - GlobalTransforms
-void UIKRigSolver::Solve(FIKRigTransforms& InOutGlobalTransform, FControlRigDrawInterface* InOutDrawInterface)
+void UIKRigSolver::SolveInternal(
+	FIKRigTransforms& InOutGlobalTransform,
+	const FIKRigGoalContainer& Goals,
+	FControlRigDrawInterface* InOutDrawInterface)
 {
 	if (IsSolverActive())
 	{
-		SolveInternal(InOutGlobalTransform, InOutDrawInterface);
+		Solve(InOutGlobalTransform, Goals, InOutDrawInterface);
 	}
 }
 
-bool UIKRigSolver::GetEffectorTarget(const FIKRigEffector& InEffector, FIKRigTarget& OutTarget) const
+bool UIKRigSolver::IsSolverActive() const
 {
-	//if (SolverDefinition && GoalGetter.IsBound())
-	if (GoalGetter.IsBound())
-	{	
-		const FName* GoalName = GetEffectorToGoal().Find(InEffector);
-		if (GoalName)
-		{
-			return GoalGetter.Execute(*GoalName, OutTarget);
-		}
-	}
-
-	return false;
+	return bEnabled;
 }
 
-const TArray<FTransform>& UIKRigSolver::GetRefPoseTransforms() const
+bool UIKRigSolver::GetGoalForEffector(
+	const FIKRigEffector& InEffector,
+	const FIKRigGoalContainer& Goals,
+	FIKRigGoal& OutGoal) const
 {
-	if (RefPoseGetter.IsBound())
+	const FName* GoalName = EffectorToGoalName.Find(InEffector);
+	if (!GoalName)
 	{
-		return RefPoseGetter.Execute();
+		return false;
 	}
 
-	static TArray<FTransform> Dummy;
-	return Dummy;
+	return Goals.GetGoalByName(*GoalName, OutGoal);
 }
 
-void UIKRigSolver::CollectGoals(TArray<FName>& OutGoals)
+void UIKRigSolver::AppendGoalNamesToArray(TArray<FName>& OutGoals)
 {
 	TArray<FName> LocalGoals;
-	// this just accumulate on the current array
-	// if you want to clear, clear before coming here
-	EffectorToGoal.GenerateValueArray(LocalGoals);
-
+	EffectorToGoalName.GenerateValueArray(LocalGoals);
 	OutGoals += LocalGoals;
 }
 
 #if WITH_EDITOR
 void UIKRigSolver::RenameGoal(const FName& OldName, const FName& NewName)
 {
-	for (auto Iter = EffectorToGoal.CreateIterator(); Iter; ++Iter)
+	for (auto Iter = EffectorToGoalName.CreateIterator(); Iter; ++Iter)
 	{
 		if (Iter.Value() == OldName)
 		{
@@ -113,16 +85,16 @@ void UIKRigSolver::OnGoalHasBeenUpdated()
 
 void UIKRigSolver::EnsureToAddEffector(const FIKRigEffector& InEffector, const FString& InPrefix)
 {
-	FName* GoalName = EffectorToGoal.Find(InEffector);
+	FName* GoalName = EffectorToGoalName.Find(InEffector);
 	if (!GoalName)
 	{
-		EffectorToGoal.Add(InEffector) = CreateUniqueGoalName(*InPrefix);
+		EffectorToGoalName.Add(InEffector) = CreateUniqueGoalName(*InPrefix);
 	}
 }
 
 void UIKRigSolver::EnsureToRemoveEffector(const FIKRigEffector& InEffector)
 {
-	EffectorToGoal.Remove(InEffector);
+	EffectorToGoalName.Remove(InEffector);
 }
 
 #endif // WITH_EDITOR
@@ -141,5 +113,5 @@ void UIKRigSolver::Serialize(FArchive& Ar)
 	Super::Serialize(Ar);
 
 	// serialize these manually because that's the custom types
-	Ar << EffectorToGoal;
+	Ar << EffectorToGoalName;
 }
diff --git a/Engine/Plugins/Animation/IKRig/Source/IKRig/Private/Solvers/ConstraintSolver.cpp b/Engine/Plugins/Animation/IKRig/Source/IKRig/Private/Solvers/ConstraintSolver.cpp
--- a/Engine/Plugins/Animation/IKRig/Source/IKRig/Private/Solvers/ConstraintSolver.cpp
+++ b/Engine/Plugins/Animation/IKRig/Source/IKRig/Private/Solvers/ConstraintSolver.cpp
@@ -1,7 +1,7 @@
 // Copyright Epic Games, Inc. All Rights Reserved.
 
 #include "Solvers/ConstraintSolver.h"
-#include "IKRigConstraint.h"
+#include "IKRigBoneSetting.h"
 
 /*
 // register/unregister query function
@@ -16,7 +16,7 @@ void UIKRigConstraintSolver::UnregisterQueryConstraintHandler()
 }
 */
 
-void UIKRigConstraintSolver::InitInternal(const FIKRigTransforms& InGlobalTransform)
+void UIKRigConstraintSolver::Init(const FIKRigTransforms& InGlobalTransform)
 {
 	/*
 	if (ConstraintDefinition)
@@ -41,7 +41,10 @@ void UIKRigConstraintSolver::InitInternal(const FIKRigTransforms& InGlobalTransf
 	}*/
 }
 
-void UIKRigConstraintSolver::SolveInternal(FIKRigTransforms& InOutGlobalTransform, FControlRigDrawInterface* InOutDrawInterface)
+void UIKRigConstraintSolver::Solve(
+	FIKRigTransforms& InOutGlobalTransform,
+	const FIKRigGoalContainer& Goals,
+	FControlRigDrawInterface* InOutDrawInterface)
 {
 	/*
 	FIKRigConstraintProfile* Current = ConstraintProfiles.Find(ActiveProfile);
diff --git a/Engine/Plugins/Animation/IKRig/Source/IKRig/Private/Solvers/TransformSolver.cpp b/Engine/Plugins/Animation/IKRig/Source/IKRig/Private/Solvers/TransformSolver.cpp
--- a/Engine/Plugins/Animation/IKRig/Source/IKRig/Private/Solvers/TransformSolver.cpp
+++ b/Engine/Plugins/Animation/IKRig/Source/IKRig/Private/Solvers/TransformSolver.cpp
@@ -11,7 +11,7 @@ UTransformSolver::UTransformSolver()
 {
 }
 
-void UTransformSolver::InitInternal(const FIKRigTransforms& InGlobalTransform)
+void UTransformSolver::Init(const FIKRigTransforms& InGlobalTransform)
 {
 	
 }
@@ -21,15 +21,18 @@ bool UTransformSolver::IsSolverActive() const
 	return Super::IsSolverActive() && (bEnablePosition || bEnableRotation);
 }
 
-void UTransformSolver::SolveInternal(FIKRigTransforms& InOutGlobalTransform, FControlRigDrawInterface* InOutDrawInterface)
+void UTransformSolver::Solve(
+	FIKRigTransforms& InOutGlobalTransform,
+	const FIKRigGoalContainer& Goals,
+	FControlRigDrawInterface* InOutDrawInterface)
 {
-	FIKRigTarget Target;
-	if (!GetEffectorTarget(TransformTarget, Target))
+	FIKRigGoal Goal;
+	if (!GetGoalForEffector(Effector, Goals, Goal))
 	{
 		return;
 	}
 	
-	int32 Index = InOutGlobalTransform.Hierarchy->GetIndex(TransformTarget.Bone);
+	int32 Index = InOutGlobalTransform.Hierarchy->GetIndex(Effector.Bone);
 	if (Index == INDEX_NONE)
 	{
 		return;
@@ -39,11 +42,11 @@ void UTransformSolver::SolveInternal(FIKRigTransforms& InOutGlobalTransform, FCo
 
 	if (bEnablePosition)
 	{
-		CurrentTransform.SetLocation(Target.PositionTarget.Position);
+		CurrentTransform.SetLocation(Goal.Position);
 	}
 	if (bEnableRotation)
 	{
-		CurrentTransform.SetRotation(Target.RotationTarget.Rotation.Quaternion());
+		CurrentTransform.SetRotation(Goal.Rotation.Quaternion());
 	}
 
 	InOutGlobalTransform.SetGlobalTransform(Index, CurrentTransform, true);
@@ -62,12 +65,12 @@ void UTransformSolver::UpdateEffectors()
 	// not found, we add new
 	if (bActiveTask)
 	{
-		EnsureToAddEffector(TransformTarget, *TransformTargetName);
+		EnsureToAddEffector(Effector, *TransformTargetName);
 	}
 	// if the opposite, we have to remove
 	else if (!bActiveTask)
 	{
-		EnsureToRemoveEffector(TransformTarget);
+		EnsureToRemoveEffector(Effector);
 	}
 
 	// trigger a delegate for goal has been updated?
diff --git a/Engine/Plugins/Animation/IKRig/Source/IKRig/Public/BoneSettings/RotationLimits.h b/Engine/Plugins/Animation/IKRig/Source/IKRig/Public/BoneSettings/RotationLimits.h
new file mode 100644
--- /dev/null
+++ b/Engine/Plugins/Animation/IKRig/Source/IKRig/Public/BoneSettings/RotationLimits.h
@@ -0,0 +1,44 @@
+// Copyright Epic Games, Inc. All Rights Reserved.
+
+#pragma once
+
+#include "IKRigBoneSetting.h"
+#include "RotationLimits.generated.h"
+
+UENUM(BlueprintType)
+enum class ERotationLimitMode : uint8
+{
+	Free	UMETA(DisplayName = "Free"),
+	Limited	UMETA(DisplayName = "Limited"),
+	Locked	UMETA(DisplayName = "Locked"),
+};
+
+UCLASS(config = Engine, hidecategories = UObject, BlueprintType)
+class IKRIG_API URotationLimits : public UIKRigBoneSetting
+{
+	GENERATED_BODY()
+
+	URotationLimits() {}
+
+	UPROPERTY(EditAnywhere, Category = "Limit")
+	ERotationLimitMode RotateX;
+	UPROPERTY(EditAnywhere, Category = "Limit")
+	float MinX = -10.0f;
+	UPROPERTY(EditAnywhere, Category = "Limit")
+	float MaxX = 10.0f;
+
+	UPROPERTY(EditAnywhere, Category = "Limit")
+	ERotationLimitMode RotateY;
+	UPROPERTY(EditAnywhere, Category = "Limit")
+	float MinY = -10.0f;
+	UPROPERTY(EditAnywhere, Category = "Limit")
+	float MaxY = 10.0f;
+
+	UPROPERTY(EditAnywhere, Category = "Limit")
+	ERotationLimitMode RotateZ;
+	UPROPERTY(EditAnywhere, Category = "Limit")
+	float MinZ = -10.0f;
+	UPROPERTY(EditAnywhere, Category = "Limit")
+	float MaxZ = 10.0f;
+};
+
diff --git a/Engine/Plugins/Animation/IKRig/Source/IKRig/Public/IKRigBoneSetting.h b/Engine/Plugins/Animation/IKRig/Source/IKRig/Public/IKRigBoneSetting.h
new file mode 100644
--- /dev/null
+++ b/Engine/Plugins/Animation/IKRig/Source/IKRig/Public/IKRigBoneSetting.h
@@ -0,0 +1,20 @@
+// Copyright Epic Games, Inc. All Rights Reserved.
+
+#pragma once
+
+#include "CoreMinimal.h"
+#include "UObject/ObjectMacros.h"
+#include "UObject/Object.h"
+#include "IKRigBoneSetting.generated.h"
+
+UCLASS(config = Engine, hidecategories = UObject, DontCollapseCategories, BlueprintType)
+class IKRIG_API UIKRigBoneSetting : public UObject
+{
+	GENERATED_BODY()
+
+public:
+
+	UPROPERTY(EditAnywhere, Category = Bone)
+	FName Bone;
+};
+
diff --git a/Engine/Plugins/Animation/IKRig/Source/IKRig/Public/IKRigConstraint.h b/Engine/Plugins/Animation/IKRig/Source/IKRig/Public/IKRigConstraint.h
deleted file mode 100644
diff --git a/Engine/Plugins/Animation/IKRig/Source/IKRig/Public/IKRigDataTypes.h b/Engine/Plugins/Animation/IKRig/Source/IKRig/Public/IKRigDataTypes.h
--- a/Engine/Plugins/Animation/IKRig/Source/IKRig/Public/IKRigDataTypes.h
+++ b/Engine/Plugins/Animation/IKRig/Source/IKRig/Public/IKRigDataTypes.h
@@ -1,10 +1,5 @@
 // Copyright Epic Games, Inc. All Rights Reserved.
 
-/**
- * Contains IKRig Definition 
- *
- *  https://docs.google.com/document/d/1yd8GCfT2aufxSdb5jAzlNTr1SptxEFpS9pWdQY-8LIk/edit#
- */
 
 #pragma once
 
@@ -16,48 +11,6 @@
 
 struct FIKRigHierarchy;
 
-USTRUCT()
-struct IKRIG_API FIKRigPosition
-{
-	GENERATED_BODY()
-
-	FIKRigPosition()
-		: Position(ForceInitToZero)
-	{
-
-	}
-	UPROPERTY(EditAnywhere, Category = FIKRigPosition)
-	FVector Position;
-};
-
-USTRUCT()
-struct IKRIG_API FIKRigRotation
-{
-	GENERATED_BODY()
-
-	FIKRigRotation()
-		: Rotation(ForceInitToZero)
-	{
-
-	}
-
-	UPROPERTY(EditAnywhere, Category = FIKRigRotation)
-	FRotator Rotation;
-};
-
-// list of targets for goal
-// this contains default value for the target
-USTRUCT()
-struct IKRIG_API FIKRigTarget
-{
-	GENERATED_BODY()
-
-	UPROPERTY(EditAnywhere, Category = FIKRigTarget)
-	FIKRigPosition PositionTarget;
-
-	UPROPERTY(EditAnywhere, Category = FIKRigTarget)
-	FIKRigRotation RotationTarget;
-};
 
 USTRUCT()
 struct IKRIG_API FIKRigGoal
@@ -65,29 +18,26 @@ struct IKRIG_API FIKRigGoal
 	GENERATED_BODY()
 
 	FIKRigGoal()
+		: Position(ForceInitToZero),
+		Rotation(ForceInitToZero)
 	{
 		Name = NAME_None;
 	}
 	FIKRigGoal(const FName& GoalName)
+		: Position(ForceInitToZero),
+		Rotation(ForceInitToZero)
 	{
 		Name = GoalName;
 	}
 
-	UPROPERTY(VisibleAnywhere, Category=FIKRigGoal)
+	UPROPERTY(VisibleAnywhere, Category = FIKRigGoal)
 	FName Name;
 
 	UPROPERTY(EditAnywhere, Category = FIKRigGoal)
-	FIKRigTarget Target;
-
-	void SetPositionTarget(const FVector& InPosition)
-	{
-		Target.PositionTarget.Position = InPosition;
-	}
+	FVector Position;
 
-	void SetRotationTarget(const FRotator& InRotation)
-	{
-		Target.RotationTarget.Rotation = InRotation;
-	}
+	UPROPERTY(EditAnywhere, Category = FIKRigGoal)
+	FRotator Rotation;
 };
 
 USTRUCT()
@@ -117,20 +67,64 @@ struct IKRIG_API FIKRigEffector
 	}
 };
 
+USTRUCT()
+struct IKRIG_API FIKRigGoalContainer
+{
+	GENERATED_BODY()
+
+private:
+
+	TMap<FName, FIKRigGoal> Goals;
+
+public:
+
+	void SetAllGoals(const TMap<FName, FIKRigGoal> &InGoals)
+	{
+		Goals = InGoals;
+	}
+
+	void SetGoalTransform(
+		const FName& GoalName,
+		const FVector& InPosition,
+		const FRotator& InRotation)
+	{
+		if (FIKRigGoal* Goal = Goals.Find(GoalName))
+		{
+			Goal->Position = InPosition;
+			Goal->Rotation = InRotation;
+		}
+	}
+
+	bool GetGoalByName(const FName& InGoalName, FIKRigGoal& OutGoal) const
+	{
+		const FIKRigGoal* Goal = Goals.Find(InGoalName);
+		if (Goal)
+		{
+			OutGoal = *Goal;
+			return true;
+		}
+
+		return false;
+	}
+
+	void GetNames(TArray<FName>& OutNames) const
+	{
+		Goals.GenerateKeyArray(OutNames);
+	}
+};
+
 template <typename ValueType>
 struct TIKRigEffectorMapKeyFuncs : public TDefaultMapKeyFuncs<const FIKRigEffector, ValueType, false>
 {
-	static FORCEINLINE FIKRigEffector					GetSetKey(TPair<FIKRigEffector, ValueType> const& Element) { return Element.Key; }
-	static FORCEINLINE uint32							GetKeyHash(FIKRigEffector const& Key) { return GetTypeHash(Key.Guid); }
-	static FORCEINLINE bool								Matches(FIKRigEffector const& A, FIKRigEffector const& B) { return (A.Guid == B.Guid); }
+	static FORCEINLINE FIKRigEffector	GetSetKey(TPair<FIKRigEffector, ValueType> const& Element) { return Element.Key; }
+	static FORCEINLINE uint32			GetKeyHash(FIKRigEffector const& Key) { return GetTypeHash(Key.Guid); }
+	static FORCEINLINE bool				Matches(FIKRigEffector const& A, FIKRigEffector const& B) { return (A.Guid == B.Guid); }
 };
 
 template <typename ValueType>
 using TIKRigEffectorMap = TMap<FIKRigEffector, ValueType, FDefaultSetAllocator, TIKRigEffectorMapKeyFuncs<ValueType>>;
 
-// allows transform to be modified 
-// use this class to modify transform
-// @todo: ref pose getter 
+
 struct IKRIG_API FIKRigTransforms
 {
 	FIKRigTransforms()
@@ -142,8 +136,8 @@ struct IKRIG_API FIKRigTransforms
 
 	void SetGlobalTransform(int32 Index, const FTransform& InTransform, bool bPropagate);
 	void SetLocalTransform(int32 Index, const FTransform& InTransform, bool bUpdate);
-	// this is mutating, but const_cast 
-	const FTransform& GetLocalTransform(int32 Index) const;
+	
+	const FTransform& GetLocalTransform(int32 Index) const; // this is mutating, but const_cast 
 	const FTransform& GetGlobalTransform(int32 Index) const;
 
 	void SetAllGlobalTransforms(const TArray<FTransform>& InTransforms);
@@ -159,7 +153,8 @@ private:
 	void EnsureLocalTransformsExist();
 	void RecalculateLocalTransform();
 	void UpdateLocalTransform(int32 Index);
-	/* this function propagates to children and updates LocalTransforms*/
+
+	/** this function propagates to children and updates LocalTransforms */
 	void SetGlobalTransform_Internal(int32 Index, const FTransform& InTransform);
 	void SetGlobalTransform_Recursive(int32 Index, const FTransform& InTransform);
 
diff --git a/Engine/Plugins/Animation/IKRig/Source/IKRig/Public/IKRigDefinition.h b/Engine/Plugins/Animation/IKRig/Source/IKRig/Public/IKRigDefinition.h
--- a/Engine/Plugins/Animation/IKRig/Source/IKRig/Public/IKRigDefinition.h
+++ b/Engine/Plugins/Animation/IKRig/Source/IKRig/Public/IKRigDefinition.h
@@ -5,7 +5,7 @@
 #include "CoreMinimal.h"
 #include "UObject/ObjectMacros.h"
 #include "UObject/Object.h"
-#include "IKRigConstraint.h"
+#include "IKRigBoneSetting.h"
 #include "IKRigDataTypes.h"
 #include "IKRigHierarchy.h"
 #include "IKRigDefinition.generated.h"
@@ -21,28 +21,28 @@ public:
 
 	UIKRigDefinition();
 
-	/** Source Asset imported */
 	UPROPERTY()
 	TSoftObjectPtr<UObject> SourceAsset;
 
-	UPROPERTY(EditAnywhere, instanced, Category = Constraints)
-	TArray<UIKRigConstraint*> Constraints;
+	UPROPERTY(EditAnywhere, instanced, Category = BoneSettings)
+	TArray<UIKRigBoneSetting*> BoneSettings;
+
+	UPROPERTY(EditAnywhere, Category = BoneSettings)
+	TArray<FName> ExcludedBones;
 
 public:
 
-	// ensure the goal name is unique
 	void EnsureCreateUniqueGoalName(FName& InOutGoal) const;
 
-	// goal related APIs
 	const TMap<FName, FIKRigGoal>& GetGoals() const
 	{
-		return IKGoals;
+		return Goals;
 	}
 
-	TArray<FName> GetGoalsNames() const
+	TArray<FName> GetGoalNames() const
 	{
 		TArray<FName> NameArray;
-		IKGoals.GenerateKeyArray(NameArray);
+		Goals.GenerateKeyArray(NameArray);
 		return NameArray;
 	}
 
@@ -60,44 +60,39 @@ public:
 	{
 		return Hierarchy;
 	}
-private:
-	/*********** Skeleton *********************/
 
-	/* Hierarchy
-	 * 
-	 * This just is a relationship. Parent/child relationship
-	 * sorted by parent to children
+private:
+	
+	/**
+	* Parent/child relationships stored separately from skeleton-specific transforms
 	*/
-	UPROPERTY(VisibleAnywhere, Category = "Hierarchy")
+	UPROPERTY(VisibleAnywhere, Category = "Skeleton")
 	FIKRigHierarchy Hierarchy;
 
-	// this is default transform of reference pose
-	// the separate with Hierarchy is for retargeting
-	// during retargeting, we want to ensure it works with different transform
-	// (this can be replaced during runtime for a retarget)
-	// do not cache or rely on this value inside of solver
-	// use the incoming value when initialize
-	// this has to match with Hierarchy data index
-	// IKRigDefinition interface handles the integrity of (this and Hierarchy)
-	// Use IKRigController for APIs
-	UPROPERTY(VisibleAnywhere, Category = "Hierarchy")
+	/**
+	* These are the default transforms of the reference pose.
+	* 
+	* Transforms are separated from the hierarchical representation
+	* 
+	* These transforms are treated as the common base, from which all
+	* retargeting operations will be compared against.
+	* 
+	* Do no cache these values inside a Solver. Rely instead, only on the 
+	* transforms passed to the Solver::Init()
+	*/
+	UPROPERTY(VisibleAnywhere, Category = "Skeleton")
 	TArray<FTransform> RefPoseTransforms;
 
-	// stack of solvers, executed in order
-	UPROPERTY(EditAnywhere, instanced, Category = "Solver")
+	/**
+	* Stack of solvers, of varying types, executed in serial fashion where
+	* output of prior solve is input to the next.
+	*/
+	UPROPERTY(EditAnywhere, instanced, Category = "Solvers")
 	TArray<UIKRigSolver*> Solvers;
 
-private:
-
-	/*********** Goals with Default Value******************/
-	// goals data
-	// this is cached by "Unique Internal Name" here
-	// where in the runtime, it will cache by "Display Name"
-	// as that is what's searched to set the value by game
-	// but during editing time, we use internal name 
-	// as identifier
+	/** Goals (ie, named transforms) are mapped to solver effectors by name. */
 	UPROPERTY(EditAnywhere, Category = "Goals")
-	TMap<FName, FIKRigGoal> IKGoals;
+	TMap<FName, FIKRigGoal> Goals;
 
 	// BEGIN UObject functions 
 	virtual void PostLoad() override;
@@ -110,7 +105,7 @@ private:
 	/* Goal operators */
 	void UpdateGoal();
 
-	/** Hierarchy modifires */
+	/** hierarchy modifiers */
 	// add a bone. Return false if it already has conflict name or parent is not found
 	bool AddBone(const FName& InName, const FName& InParent, const FTransform& InGlobalTransform);
 	// remove a bone. Returns false if it is not found
diff --git a/Engine/Plugins/Animation/IKRig/Source/IKRig/Public/IKRigProcessor.h b/Engine/Plugins/Animation/IKRig/Source/IKRig/Public/IKRigProcessor.h
--- a/Engine/Plugins/Animation/IKRig/Source/IKRig/Public/IKRigProcessor.h
+++ b/Engine/Plugins/Animation/IKRig/Source/IKRig/Public/IKRigProcessor.h
@@ -14,65 +14,39 @@ class UIKRigDefinition;
 class UIKRigSolver;
 struct FControlRigDrawInterface;
 
+
 UCLASS(BlueprintType)
 class IKRIG_API UIKRigProcessor : public UObject
 {
 	GENERATED_BODY()
 
 private:
-	// skeleton
-	// this just is a relationship. Parent/child relationship
-	// sorted by parent to children
+
 	UPROPERTY(EditAnywhere, Category = "RigDefinition")
 	UIKRigDefinition* RigDefinition = nullptr;
 
-	// reference transform
 	UPROPERTY(transient)
 	TArray<FTransform> RefPoseTransforms;
+	FIKRigTransforms Transforms;
 
-	// solvers
 	UPROPERTY(transient)
 	TArray<UIKRigSolver*> Solvers;
 
-	// constraint solver
-	UPROPERTY(transient)
-	UIKRigConstraintSolver* CostraintSolver;
-
-	// goals - copy from RigDefinition
-	// modified in runtime
-	UPROPERTY(transient)
-	TMap<FName, FIKRigGoal> IKGoals;
-
-	FIKRigTransforms Transforms;
 	FControlRigDrawInterface DrawInterface;
-
 	bool bInitialized = false;
 
 public: 
 
-	// if bInitialize is true, it will use Default Ref transform of IKRigDefinition to initialize
-	void SetIKRigDefinition(UIKRigDefinition* InRigDefinition, bool bInitialize = false);
-	void Initialize(const TArray<FTransform>& InRefPoseTransforms);
-	void Reinitialize();
-	void Terminate();
+	UPROPERTY(transient)
+	FIKRigGoalContainer Goals;
+
+	/** setup and run */
+	void Initialize(UIKRigDefinition* InRigDefinition);
 	void Solve();
 
-	// update goal functions
-	void SetGoalPosition(const FName& GoalName, const FVector& InPosition);
-	void SetGoalRotation(const FName& GoalName, const FRotator& InRotation);
-	void SetGoalTarget(const FName& GoalName, const FIKRigTarget& InTarget);
-	void GetGoals(TArray<FName>& OutNames) const;
-	FIKRigTransforms& GetIKRigTransformModifier();
+	FIKRigTransforms& GetTransforms();
 	const FIKRigHierarchy* GetHierarchy() const;
-
 	void ResetToRefPose();
 	const FControlRigDrawInterface& GetDrawInterface() { return DrawInterface; }
-
-private:
-	FIKRigTarget* FindGoal(const FName& GoalName);
-
-	// delegate functions
-	bool GoalGetter(const FName& InGoalName, FIKRigTarget& OutTarget);
-	const TArray<FTransform>& GetRefPoseGetter();
 };
 
diff --git a/Engine/Plugins/Animation/IKRig/Source/IKRig/Public/IKRigSolver.h b/Engine/Plugins/Animation/IKRig/Source/IKRig/Public/IKRigSolver.h
--- a/Engine/Plugins/Animation/IKRig/Source/IKRig/Public/IKRigSolver.h
+++ b/Engine/Plugins/Animation/IKRig/Source/IKRig/Public/IKRigSolver.h
@@ -20,70 +20,65 @@ class IKRIG_API UIKRigSolver : public UObject
 {
 	GENERATED_BODY()
 
-public: 
-	/** Required delegates to run this solver */
-	DECLARE_DELEGATE_RetVal(const TArray<FTransform>&, FIKRigTransformGetter);
-	DECLARE_DELEGATE_RetVal_TwoParams(bool, FIKRigGoalGetter, const FName& /*InGoalName*/, FIKRigTarget& /*OutTarget*/);
+public:
 
-	// input hierarchy and ref pose? 
-	void Init(const FIKRigTransforms& TransformModifier, FIKRigTransformGetter InRefPoseGetter, FIKRigGoalGetter InGoalGetter/*, FSolveConstraint& InConstraintHandler*/);
-
-	// input : goal getter or goals
-	// output : modified pose - GlobalTransforms
-	// or use SolverInternal function
-	void Solve(FIKRigTransforms& InOutGlobalTransform, FControlRigDrawInterface* InOutDrawInterface);
-
-	const TIKRigEffectorMap<FName>& GetEffectorToGoal() const
-	{
-		return EffectorToGoal;
-	}
-
-	void CollectGoals(TArray<FName>& OutGoals);
+	/** wraps Solver::Solve() */
+	void SolveInternal(
+		FIKRigTransforms& InOutGlobalTransform, 
+		const FIKRigGoalContainer& Goals,
+		FControlRigDrawInterface* InOutDrawInterface);
+	void AppendGoalNamesToArray(TArray<FName>& OutGoals);
 
 protected:
 
-	virtual void InitInternal(const FIKRigTransforms& InGlobalTransform) {};
-	virtual void SolveInternal(FIKRigTransforms& InOutGlobalTransform, FControlRigDrawInterface* InOutDrawInterface) {};
+	/** override Init(), Solve() and IsSolverActive() in subclasses */
+	virtual void Init(const FIKRigTransforms& InGlobalTransform) {};
+	virtual void Solve(
+		FIKRigTransforms& InOutGlobalTransform,
+		const FIKRigGoalContainer& Goals,
+		FControlRigDrawInterface* InOutDrawInterface) {};
 	virtual bool IsSolverActive() const;
 
-	bool GetEffectorTarget(const FIKRigEffector& InEffector, FIKRigTarget& OutTarget) const;
-	const TArray<FTransform>& GetRefPoseTransforms() const;
+	bool GetGoalForEffector(
+		const FIKRigEffector& InEffector, 
+		const FIKRigGoalContainer &Goals, 
+		FIKRigGoal& OutGoal) const;
 
 	UPROPERTY(EditAnywhere, Category = "Definition")
 	bool bEnabled = true;
 
 	// effector name to goals name map
-	TIKRigEffectorMap<FName> EffectorToGoal;
+	TIKRigEffectorMap<FName> EffectorToGoalName;
 
 private:
 
-	// delegate
-	FIKRigTransformGetter RefPoseGetter;
-	FIKRigGoalGetter GoalGetter;
-
-	/// BEGIN UObject
+	// BEGIN UObject
 	virtual void PostLoad() override;
 	virtual void Serialize(FArchive& Ar) override;
-	/// END UObject
+	// END UObject
 
 	friend class UIKRigController;
+	friend class UIKRigProcessor;
 
 #if WITH_EDITOR
+
 private:
+
 	void RenameGoal(const FName& OldName, const FName& NewName);
-	// get unique name delegate by IKRigDefinition
 	void EnsureUniqueGoalName(FName& InOutGoalName) const;
 
 	DECLARE_MULTICAST_DELEGATE(FGoalNeedsUpdate);
 	FGoalNeedsUpdate GoalNeedsUpdateDelegate;
 
 protected:
+
 	FName CreateUniqueGoalName(const TCHAR* Suffix) const;
 	void OnGoalHasBeenUpdated();
 	void EnsureToAddEffector(const FIKRigEffector& InEffector, const FString& InPrefix);
 	void EnsureToRemoveEffector(const FIKRigEffector& InEffector);
 
 	virtual void UpdateEffectors() {};
+
 #endif // WITH_EDITOR
 };
 
diff --git a/Engine/Plugins/Animation/IKRig/Source/IKRig/Public/Solvers/ConstraintSolver.h b/Engine/Plugins/Animation/IKRig/Source/IKRig/Public/Solvers/ConstraintSolver.h
--- a/Engine/Plugins/Animation/IKRig/Source/IKRig/Public/Solvers/ConstraintSolver.h
+++ b/Engine/Plugins/Animation/IKRig/Source/IKRig/Public/Solvers/ConstraintSolver.h
@@ -1,31 +1,28 @@
 // Copyright Epic Games, Inc. All Rights Reserved.
 
-/**
- * Contains IKRig Constraint Solver Definition 
- *
- */
-
 #pragma once
 
 #include "IKRigSolver.h"
 #include "ConstraintSolver.generated.h"
 
-class UIKRigConstraint;
+class UIKRigBoneSettings;
 
 
-// run time processor 
 UCLASS(config = Engine, hidecategories = UObject)
 class IKRIG_API UIKRigConstraintSolver : public UIKRigSolver
 {
 	GENERATED_BODY()
 
-	DECLARE_DELEGATE_TwoParams(FIKRigQueryConstraint, const FName& InConstraintName, UIKRigConstraint& OutConstraint);
+	//DECLARE_DELEGATE_TwoParams(FIKRigQueryConstraint, const FName& InConstraintName, UIKRigConstraint& OutConstraint);
 
 
 public: 
-	// during we don't mutate bone transform, but you can read it
-	virtual void InitInternal(const FIKRigTransforms& InGlobalTransform) override;
-	virtual void SolveInternal(FIKRigTransforms& InOutGlobalTransform, FControlRigDrawInterface* InOutDrawInterface) override;
+
+	virtual void Init(const FIKRigTransforms& InGlobalTransform) override;
+	virtual void Solve(
+		FIKRigTransforms& InOutGlobalTransform,
+		const FIKRigGoalContainer& Goals,
+		FControlRigDrawInterface* InOutDrawInterface) override;
 
 	// this can be utilized for just subset of constraints
 	//void SolverConstraints(const TArray<FName>& ConstraintsList, FIKRigTransformModifier& InOutGlobalTransform, FControlRigDrawInterface* InOutDrawInterface);
diff --git a/Engine/Plugins/Animation/IKRig/Source/IKRig/Public/Solvers/TransformSolver.h b/Engine/Plugins/Animation/IKRig/Source/IKRig/Public/Solvers/TransformSolver.h
--- a/Engine/Plugins/Animation/IKRig/Source/IKRig/Public/Solvers/TransformSolver.h
+++ b/Engine/Plugins/Animation/IKRig/Source/IKRig/Public/Solvers/TransformSolver.h
@@ -22,12 +22,15 @@ public:
 	bool bEnableRotation = true;
 
 	UPROPERTY(EditAnywhere, Category = "Solver")
-	FIKRigEffector TransformTarget;
+	FIKRigEffector Effector;
 
 protected:
 	
-	virtual void InitInternal(const FIKRigTransforms& InGlobalTransform) override;
-	virtual void SolveInternal(FIKRigTransforms& InOutGlobalTransform, FControlRigDrawInterface* InOutDrawInterface) override;
+	virtual void Init(const FIKRigTransforms& InGlobalTransform) override;
+	virtual void Solve(
+		FIKRigTransforms& InOutGlobalTransform, 
+		const FIKRigGoalContainer& Goals,
+		FControlRigDrawInterface* InOutDrawInterface) override;
 	virtual bool IsSolverActive() const override;
 
 private:
diff --git a/Engine/Plugins/Animation/IKRig/Source/IKRigEditor/Private/IKRigController.cpp b/Engine/Plugins/Animation/IKRig/Source/IKRigEditor/Private/IKRigController.cpp
--- a/Engine/Plugins/Animation/IKRig/Source/IKRigEditor/Private/IKRigController.cpp
+++ b/Engine/Plugins/Animation/IKRig/Source/IKRigEditor/Private/IKRigController.cpp
@@ -5,7 +5,7 @@
 #include "AnimationRuntime.h"
 #include "IKRigSolver.h"
 #include "ScopedTransaction.h"
-#include "IKRigConstraint.h"
+#include "IKRigBoneSetting.h"
 
 #define LOCTEXT_NAMESPACE	"IKRigController"
 
@@ -276,7 +276,7 @@ FName UIKRigController::GetGoalName(UIKRigSolver* InSolver, const FIKRigEffector
 {
 	if (InSolver)
 	{
-		FName* GoalName = InSolver->EffectorToGoal.Find(InEffector);
+		FName* GoalName = InSolver->EffectorToGoalName.Find(InEffector);
 		if (GoalName)
 		{
 			return *GoalName;
@@ -293,7 +293,7 @@ void UIKRigController::SetGoalName(UIKRigSolver* InSolver, const FIKRigEffector&
 		FScopedTransaction Transaction(LOCTEXT("SetGoalName_Label", "Set Goal Name"));
 		InSolver->Modify();
 
-		FName* GoalName = InSolver->EffectorToGoal.Find(InEffector);
+		FName* GoalName = InSolver->EffectorToGoalName.Find(InEffector);
 		if (GoalName)
 		{
 			*GoalName = NewGoalName;
@@ -303,17 +303,17 @@ void UIKRigController::SetGoalName(UIKRigSolver* InSolver, const FIKRigEffector&
 	}
 }
 
-UIKRigConstraint* UIKRigController::AddConstraint(TSubclassOf<UIKRigConstraint> NewConstraintType)
+UIKRigBoneSetting* UIKRigController::AddBoneSetting(TSubclassOf<UIKRigBoneSetting> NewSettingType)
 {
 	
-	FScopedTransaction Transaction(LOCTEXT("AddConstraint_Label", "Add Constraint"));
+	FScopedTransaction Transaction(LOCTEXT("AddSetting_Label", "Add Bone Setting"));
 	IKRigDefinition->Modify();
 
-	UIKRigConstraint* NewRigConstraint = NewObject<UIKRigConstraint>(IKRigDefinition, NewConstraintType);
-	if (NewRigConstraint)
+	UIKRigBoneSetting* NewBoneSetting = NewObject<UIKRigBoneSetting>(IKRigDefinition, NewSettingType);
+	if (NewBoneSetting)
 	{
-		IKRigDefinition->Constraints.Add(NewRigConstraint);
-		return NewRigConstraint;
+		IKRigDefinition->BoneSettings.Add(NewBoneSetting);
+		return NewBoneSetting;
 	}
 	
 	return nullptr;	
@@ -353,7 +353,7 @@ FIKRigGoal* UIKRigController::GetGoal(const FName& InGoalName)
 {
 	if (IKRigDefinition)
 	{
-		return IKRigDefinition->IKGoals.Find(InGoalName);
+		return IKRigDefinition->Goals.Find(InGoalName);
 	}
 
 	return nullptr;
@@ -363,7 +363,7 @@ const FIKRigGoal* UIKRigController::GetGoal(const FName& InGoalName) const
 {
 	if (IKRigDefinition)
 	{
-		return IKRigDefinition->IKGoals.Find(InGoalName);
+		return IKRigDefinition->Goals.Find(InGoalName);
 	}
 
 	return nullptr;
diff --git a/Engine/Plugins/Animation/IKRig/Source/IKRigEditor/Private/IKRigDefinitionDetails.cpp b/Engine/Plugins/Animation/IKRig/Source/IKRigEditor/Private/IKRigDefinitionDetails.cpp
--- a/Engine/Plugins/Animation/IKRig/Source/IKRigEditor/Private/IKRigDefinitionDetails.cpp
+++ b/Engine/Plugins/Animation/IKRig/Source/IKRigEditor/Private/IKRigDefinitionDetails.cpp
@@ -10,7 +10,7 @@
 #include "IKRigDefinition.h"
 #include "IKRigController.h"
 #include "IKRigSolver.h"
-#include "IKRigConstraint.h"
+#include "IKRigBoneSetting.h"
 
 #include "ScopedTransaction.h"
 #include "PropertyCustomizationHelpers.h"
@@ -74,7 +74,7 @@ void FIKRigDefinitionDetails::CustomizeDetails(IDetailLayoutBuilder& DetailBuild
 	/////////////////////////////////////////////////////////////////////////////////
 	// skeleton set up
 	/////////////////////////////////////////////////////////////////////////////////
-	IDetailCategoryBuilder& HierarchyCategory = DetailBuilder.EditCategory("Hierarchy");
+	IDetailCategoryBuilder& HierarchyCategory = DetailBuilder.EditCategory("Skeleton");
  
 	SelectedAsset = IKRigDefinition->SourceAsset.Get();
 
@@ -83,7 +83,7 @@ void FIKRigDefinitionDetails::CustomizeDetails(IDetailLayoutBuilder& DetailBuild
 	[
 		SNew(STextBlock)
 		.Font(IDetailLayoutBuilder::GetDetailFont())
-		.Text(LOCTEXT("SelectSourceSkeleton", "Souce Skeleton"))
+		.Text(LOCTEXT("SelectSourceSkeleton", "Source Skeleton"))
 	]
 	.ValueContent()
 	[
@@ -118,7 +118,7 @@ void FIKRigDefinitionDetails::CustomizeDetails(IDetailLayoutBuilder& DetailBuild
 				.ContentPadding(3)
 				.IsEnabled(this, &FIKRigDefinitionDetails::CanImport)
 				.OnClicked(this, &FIKRigDefinitionDetails::OnImportHierarchy)
-				.ToolTipText(LOCTEXT("OnImportHierarchyTooltip", "Change Skeleton Data with Selected Asset. This replaces existing skeleton."))
+				.ToolTipText(LOCTEXT("OnImportHierarchyTooltip", "Set skeleton to selected asset. This replaces existing skeleton."))
 				[
 					SNew(STextBlock)
 					.Font(IDetailLayoutBuilder::GetDetailFont())
@@ -128,7 +128,7 @@ void FIKRigDefinitionDetails::CustomizeDetails(IDetailLayoutBuilder& DetailBuild
 		]
 	];
 
- 	IDetailCategoryBuilder& SolverCategory = DetailBuilder.EditCategory("Solver");
+ 	IDetailCategoryBuilder& SolverCategory = DetailBuilder.EditCategory("Solvers");
 
 	SolverCategory.AddCustomRow(FText::FromString("AddSolver"))
 	.NameContent()
@@ -148,8 +148,8 @@ void FIKRigDefinitionDetails::CustomizeDetails(IDetailLayoutBuilder& DetailBuild
 		]
 	];
 
-	IDetailCategoryBuilder& ConstraintCategory = DetailBuilder.EditCategory("Constraint");
-	ConstraintCategory.AddCustomRow(FText::FromString("AddConstraint"))
+	IDetailCategoryBuilder& BoneSettingCategory = DetailBuilder.EditCategory("BoneSettings");
+	BoneSettingCategory.AddCustomRow(FText::FromString("AddBoneSetting"))
 	.NameContent()
 	[
 		SNullWidget::NullWidget
@@ -158,16 +158,16 @@ void FIKRigDefinitionDetails::CustomizeDetails(IDetailLayoutBuilder& DetailBuild
 	[
 		SNew(SButton)
 		.ContentPadding(3)
-		.OnClicked(this, &FIKRigDefinitionDetails::OnShowConstraintClassPicker)
-		.ToolTipText(LOCTEXT("OnShowConstraintListTooltip", "Select Constraint to Add"))
+		.OnClicked(this, &FIKRigDefinitionDetails::OnShowBoneSettingsClassPicker)
+		.ToolTipText(LOCTEXT("OnShowBoneSettingListTooltip", "Select Bone Setting to Add"))
 		[
 			SNew(STextBlock)
 			.Font(IDetailLayoutBuilder::GetDetailFont())
-			.Text(LOCTEXT("ShowConstraintList", "Add Constraint"))
+			.Text(LOCTEXT("ShowBoneSettingList", "Add Bone Setting"))
 		]
 	];
 
-	GoalPropertyHandle = DetailBuilder.GetProperty(TEXT("IKGoals"));
+	GoalPropertyHandle = DetailBuilder.GetProperty(TEXT("Goals"));
 
 	TArray<FName> GoalNames;
 	IKRigController->QueryGoals(GoalNames);
@@ -262,12 +262,12 @@ FReply FIKRigDefinitionDetails::OnShowSolverClassPicker()
 	return FReply::Handled();
 }
 
-FReply FIKRigDefinitionDetails::OnShowConstraintClassPicker()
+FReply FIKRigDefinitionDetails::OnShowBoneSettingsClassPicker()
 {
-	UClass* ChosenClass = SelectClass(UIKRigConstraint::StaticClass(), LOCTEXT("SelectConstraintClass", "Select Constraint Class"));
+	UClass* ChosenClass = SelectClass(UIKRigBoneSetting::StaticClass(), LOCTEXT("SelectBoneSettingClass", "Select Bone Setting Class"));
 	if (ChosenClass)
 	{
-		IKRigController->AddConstraint(ChosenClass);
+		IKRigController->AddBoneSetting(ChosenClass);
 	}
 
 	return FReply::Handled();
diff --git a/Engine/Plugins/Animation/IKRig/Source/IKRigEditor/Private/IKRigDefinitionDetails.h b/Engine/Plugins/Animation/IKRig/Source/IKRigEditor/Private/IKRigDefinitionDetails.h
--- a/Engine/Plugins/Animation/IKRig/Source/IKRigEditor/Private/IKRigDefinitionDetails.h
+++ b/Engine/Plugins/Animation/IKRig/Source/IKRigEditor/Private/IKRigDefinitionDetails.h
@@ -59,7 +59,7 @@ private:
 	TWeakPtr< IDetailLayoutBuilder > DetailBuilderWeakPtr;
 
 	FReply OnShowSolverClassPicker();
-	FReply OnShowConstraintClassPicker();
+	FReply OnShowBoneSettingsClassPicker();
 
 	typedef TSharedPtr<FGoalNameListItem> FGoalNameListItemPtr;
 
diff --git a/Engine/Plugins/Animation/IKRig/Source/IKRigEditor/Private/IKRigEffectorCustomization.cpp b/Engine/Plugins/Animation/IKRig/Source/IKRigEditor/Private/IKRigEffectorCustomization.cpp
--- a/Engine/Plugins/Animation/IKRig/Source/IKRigEditor/Private/IKRigEffectorCustomization.cpp
+++ b/Engine/Plugins/Animation/IKRig/Source/IKRigEditor/Private/IKRigEffectorCustomization.cpp
@@ -80,11 +80,18 @@ void FIKRigEffectorCustomization::CustomizeHeader(TSharedRef<IPropertyHandle> St
 				.ContentPadding(0)
 				.Content()
 				[
+					/* KIARAN - Refactor this to non-editable field that does not delete goals when modified
+					SNew(STextBlock)
+					.Text(this, &FIKRigEffectorCustomization::GetEffectorGoalComboBoxContent)
+					*/
+					
 					SNew(SEditableTextBox)
 					.Text(this, &FIKRigEffectorCustomization::GetEffectorGoalComboBoxContent)
 					.Font(IDetailLayoutBuilder::GetDetailFont())
 					.OnTextCommitted(this, &FIKRigEffectorCustomization::OnEffectorGoalChanged)
 					.ToolTipText(this, &FIKRigEffectorCustomization::GetEffectorGoalComboBoxToolTip)
+					
+
 				]
 			];
 
diff --git a/Engine/Plugins/Animation/IKRig/Source/IKRigEditor/Public/IKRigController.h b/Engine/Plugins/Animation/IKRig/Source/IKRigEditor/Public/IKRigController.h
--- a/Engine/Plugins/Animation/IKRig/Source/IKRigEditor/Public/IKRigController.h
+++ b/Engine/Plugins/Animation/IKRig/Source/IKRigEditor/Public/IKRigController.h
@@ -15,7 +15,7 @@
 // DELEGATE for goals
 
 class UIKRigSolverDefinition;
-class UIKRigConstraint;
+class UIKRigBoneSetting;
 struct FIKRigHierarchy;
 struct FReferenceSkeleton;
 struct FIKRigConstraintProfile;
@@ -59,8 +59,8 @@ public:
 	UIKRigSolver* GetSolver(int32 Index) const;
 	int32 GetNumSolvers() const;
 
-	// constraints
-	UIKRigConstraint* AddConstraint(TSubclassOf<UIKRigConstraint> NewConstraintType);
+	// per-bone settings and constraints
+	UIKRigBoneSetting* AddBoneSetting(TSubclassOf<UIKRigBoneSetting> NewBoneSettingType);
 
 	// goal operators
 	void QueryGoals(TArray<FName>& OutGoals) const;
diff --git a/Engine/Plugins/Experimental/FullBodyIK/Source/FullBodyIK/Private/FullBodyIKSolver.cpp b/Engine/Plugins/Experimental/FullBodyIK/Source/FullBodyIK/Private/FullBodyIKSolver.cpp
--- a/Engine/Plugins/Experimental/FullBodyIK/Source/FullBodyIK/Private/FullBodyIKSolver.cpp
+++ b/Engine/Plugins/Experimental/FullBodyIK/Source/FullBodyIK/Private/FullBodyIKSolver.cpp
@@ -204,7 +204,7 @@ UFullBodyIKSolver::UFullBodyIKSolver()
 {
 }
 
-void UFullBodyIKSolver::InitInternal(const FIKRigTransforms& InGlobalTransform)
+void UFullBodyIKSolver::Init(const FIKRigTransforms& InGlobalTransform)
 {
 
 	LinkData.Reset();
@@ -228,8 +228,9 @@ bool UFullBodyIKSolver::IsSolverActive() const
 	return false;
 }
 
-void UFullBodyIKSolver::SolveInternal(
-	FIKRigTransforms& InOutGlobalTransform, 
+void UFullBodyIKSolver::Solve(
+	FIKRigTransforms& InOutGlobalTransform,
+	const FIKRigGoalContainer& Goals,
 	FControlRigDrawInterface* InOutDrawInterface)
 {
 	if (LinkDataToHierarchyIndices.Num() > 0)
@@ -289,10 +290,10 @@ void UFullBodyIKSolver::SolveInternal(
 					const FFBIKRigEffector& CurEffector = Effectors[EffectorIndex];
 					const FVector CurrentLinkLocation = LinkData[EffectorLinkIndex].GetTransform().GetLocation();
 					const FQuat CurrentLinkRotation = LinkData[EffectorLinkIndex].GetTransform().GetRotation();
-					FIKRigTarget RigTarget;
-					ensure(GetEffectorTarget(CurEffector.Target, RigTarget));
-					EffectorTarget->Position = RigTarget.PositionTarget.Position; // FMath::Lerp(CurrentLinkLocation, EffectorLocation, RigTarget.PositionTarget.);
-					EffectorTarget->Rotation = RigTarget.RotationTarget.Rotation.Quaternion(); // FMath::Lerp(CurrentLinkRotation, EffectorRotation, CurEffector.RotationAlpha);
+					FIKRigGoal Goal;
+					ensure(GetGoalForEffector(CurEffector.Target, Goals, Goal));
+					EffectorTarget->Position = Goal.Position; // FMath::Lerp(CurrentLinkLocation, EffectorLocation, RigTarget.PositionTarget.);
+					EffectorTarget->Rotation = Goal.Rotation.Quaternion(); // FMath::Lerp(CurrentLinkRotation, EffectorRotation, CurEffector.RotationAlpha);
 					EffectorTarget->InitialPositionDistance = (EffectorTarget->Position - CurrentLinkLocation).Size();
 					EffectorTarget->InitialRotationDistance = (FBIKUtil::GetScaledRotationAxis(EffectorTarget->Rotation) - FBIKUtil::GetScaledRotationAxis(CurrentLinkRotation)).Size();
 
@@ -573,11 +574,11 @@ void UFullBodyIKSolver::UpdateEffectors()
 
 	// we have more nodes than goals
 	// which means we have something deleted
-	if (Effectors.Num() < EffectorToGoal.Num())
+	if (Effectors.Num() < EffectorToGoalName.Num())
 	{
 		TArray<FIKRigEffector> GoalEffectors;
 		// we have to remove things that don't belong
-		for (auto Iter = EffectorToGoal.CreateIterator(); Iter; ++Iter)
+		for (auto Iter = EffectorToGoalName.CreateIterator(); Iter; ++Iter)
 		{
 			GoalEffectors.Add(Iter.Key());
 		}
@@ -597,7 +598,7 @@ void UFullBodyIKSolver::UpdateEffectors()
 		for (TConstSetBitIterator<> Iter(RemoveFlags); Iter; ++Iter)
 		{
 			// remove things that don't belong
-			EffectorToGoal.Remove(GoalEffectors[Iter.GetIndex()]);
+			EffectorToGoalName.Remove(GoalEffectors[Iter.GetIndex()]);
 		}
 	}
 
diff --git a/Engine/Plugins/Experimental/FullBodyIK/Source/FullBodyIK/Private/FullBodyIKSolver.h b/Engine/Plugins/Experimental/FullBodyIK/Source/FullBodyIK/Private/FullBodyIKSolver.h
--- a/Engine/Plugins/Experimental/FullBodyIK/Source/FullBodyIK/Private/FullBodyIKSolver.h
+++ b/Engine/Plugins/Experimental/FullBodyIK/Source/FullBodyIK/Private/FullBodyIKSolver.h
@@ -99,8 +99,11 @@ private:
 	TArray<FJacobianDebugData> DebugData;
 
 protected:
-	virtual void InitInternal(const FIKRigTransforms& InGlobalTransforms) override;
-	virtual void SolveInternal(FIKRigTransforms& InOutGlobalTransforms, FControlRigDrawInterface* InOutDrawInterface) override;
+	virtual void Init(const FIKRigTransforms& InGlobalTransforms) override;
+	virtual void Solve(
+		FIKRigTransforms& InOutGlobalTransforms,
+		const FIKRigGoalContainer& Goals,
+		FControlRigDrawInterface* InOutDrawInterface) override;
 	virtual bool IsSolverActive() const override;
 
 public:
