Adding missing copyright boilerplate

diff --git a/Engine/Plugins/Animation/IKRig/Source/IKRig/Private/BoneSettings/RotationLimits.cpp b/Engine/Plugins/Animation/IKRig/Source/IKRig/Private/BoneSettings/RotationLimits.cpp
--- a/Engine/Plugins/Animation/IKRig/Source/IKRig/Private/BoneSettings/RotationLimits.cpp
+++ b/Engine/Plugins/Animation/IKRig/Source/IKRig/Private/BoneSettings/RotationLimits.cpp
@@ -0,0 +1,2 @@
+// Copyright Epic Games, Inc. All Rights Reserved.
+
