Remove dependency on IKRig plugin from FBIK plugin.

diff --git a/Engine/Plugins/Animation/IKRig/IKRig.uplugin b/Engine/Plugins/Animation/IKRig/IKRig.uplugin
--- a/Engine/Plugins/Animation/IKRig/IKRig.uplugin
+++ b/Engine/Plugins/Animation/IKRig/IKRig.uplugin
@@ -19,6 +19,10 @@
 		{
 			"Name": "ControlRig",
 			"Enabled": true
+		},
+		{
+			"Name": "FullBodyIK",
+			"Enabled": true
 		}
 	],
 	"Modules": [
diff --git a/Engine/Plugins/Animation/IKRig/Source/IKRig/IKRig.Build.cs b/Engine/Plugins/Animation/IKRig/Source/IKRig/IKRig.Build.cs
--- a/Engine/Plugins/Animation/IKRig/Source/IKRig/IKRig.Build.cs
+++ b/Engine/Plugins/Animation/IKRig/Source/IKRig/IKRig.Build.cs
@@ -17,7 +17,8 @@ public class IKRig : ModuleRules
 				"SlateCore",
 				"AnimationCore", 
 				"AnimGraphRuntime",
-				"ControlRig"
+				"ControlRig",
+				"FullBodyIK"
 			}
 			);
 
diff --git a/Engine/Plugins/Animation/IKRig/Source/IKRig/Private/IKRigDefinition.cpp b/Engine/Plugins/Animation/IKRig/Source/IKRig/Private/IKRigDefinition.cpp
--- a/Engine/Plugins/Animation/IKRig/Source/IKRig/Private/IKRigDefinition.cpp
+++ b/Engine/Plugins/Animation/IKRig/Source/IKRig/Private/IKRigDefinition.cpp
@@ -8,6 +8,11 @@ void UIKRigDefinition::GetGoalNamesFromSolvers(TArray<FName>& OutGoalNames) cons
 {
 	for (UIKRigSolver* Solver : Solvers)
 	{
+		if (!Solver)
+		{
+			continue;
+		}
+		
 		TSet<FName> GoalNames;
 		Solver->CollectGoalNames(GoalNames);
 
diff --git a/Engine/Plugins/Animation/IKRig/Source/IKRig/Private/Solvers/IKRigFBIKSolver.cpp b/Engine/Plugins/Animation/IKRig/Source/IKRig/Private/Solvers/IKRigFBIKSolver.cpp
new file mode 100644
--- /dev/null
+++ b/Engine/Plugins/Animation/IKRig/Source/IKRig/Private/Solvers/IKRigFBIKSolver.cpp
@@ -0,0 +1,566 @@
+// Copyright Epic Games, Inc. All Rights Reserved.
+
+/*=============================================================================
+	FullBodyIKSolver.cpp: Solver execution class for Transform
+=============================================================================*/
+
+#include "Solvers/IKRigFBIKSolver.h"
+#include "IKRigDataTypes.h"
+#include "IKRigHierarchy.h"
+#include "FBIKConstraint.h"
+#include "JacobianIK.h"
+#include "FBIKUtil.h"
+#include "FBIKShared.h"
+#include "Drawing/ControlRigDrawInterface.h"
+
+const FName UIKRigFBIKSolver::EffectorTargetPrefix = FName(TEXT("FullBodyIKTarget"));
+
+//////////////////////////////////////////////
+// utility functions
+////////////////////////////////////////////////
+static float EnsureToAddBoneToLinkData(const FIKRigTransforms& TransformModifier, const int32 CurrentItem, TArray<FFBIKLinkData>& LinkData,
+	TMap<int32, int32>& HierarchyToLinkDataMap, TMap<int32, int32>& LinkDataToHierarchyIndices)
+{
+	float ChainLength = 0;
+
+	// we insert from back to first, but only if the list doesn't have it
+	int32* FoundLinkIndex = HierarchyToLinkDataMap.Find(CurrentItem);
+	if (!FoundLinkIndex)
+	{
+		int32 NewLinkIndex = LinkData.AddDefaulted();
+		FFBIKLinkData& NewLink = LinkData[NewLinkIndex];
+
+		// find parent LinkIndex
+		const int32 ParentItem = TransformModifier.Hierarchy->GetParentIndex(CurrentItem);
+		FoundLinkIndex = HierarchyToLinkDataMap.Find(ParentItem);
+		NewLink.ParentLinkIndex = (FoundLinkIndex) ? *FoundLinkIndex : INDEX_NONE;
+		NewLink.SetTransform(TransformModifier.GetGlobalTransform(CurrentItem));
+
+		if (ParentItem != INDEX_NONE)
+		{
+			// @todo: currently we assume the input pose is init transform
+			// this could use reference pose getter from IKRigSolver or we could wrap it to TransformModifier
+			FVector DiffLocation = TransformModifier.GetGlobalTransform(CurrentItem).GetLocation() - TransformModifier.GetGlobalTransform(ParentItem).GetLocation();
+			// set Length
+			NewLink.Length = DiffLocation.Size();
+		}
+
+		// we create bidirectional look up table
+		HierarchyToLinkDataMap.Add(CurrentItem, NewLinkIndex);
+		LinkDataToHierarchyIndices.Add(NewLinkIndex, CurrentItem);
+
+		ChainLength = NewLink.Length;
+	}
+	else
+	{
+		ChainLength = LinkData[*FoundLinkIndex].Length;
+	}
+
+	return ChainLength;
+}
+
+static void AddToEffectorTarget(int32 EffectorIndex, const int32 Effector, TMap<int32, FFBIKEffectorTarget>& EffectorTargets, const TMap<int32, int32>& HierarchyToLinkDataMap,
+	TArray<int32>& EffectorLinkIndices, float ChainLength, const TArray<int32>& ChainIndices, int32 PositionDepth, int32 RotationDepth)
+{
+	const int32* EffectorLinkIndexID = HierarchyToLinkDataMap.Find(Effector);
+	check(EffectorLinkIndexID);
+	EffectorLinkIndices[EffectorIndex] = *EffectorLinkIndexID;
+	// add EffectorTarget for this link Index
+	FFBIKEffectorTarget& EffectorTarget = EffectorTargets.FindOrAdd(*EffectorLinkIndexID);
+	EffectorTarget.ChainLength = ChainLength;
+
+	// convert bone chain indices to link chain
+	const int32 MaxNum = FMath::Max(PositionDepth, RotationDepth);
+	if (ensure(MaxNum <= ChainIndices.Num()))
+	{
+		EffectorTarget.LinkChain.Reset(MaxNum);
+		for (int32 Index = 0; Index < MaxNum; ++Index)
+		{
+			const int32 Bone = ChainIndices[Index];
+			const int32* LinkIndexID = HierarchyToLinkDataMap.Find(Bone);
+			EffectorTarget.LinkChain.Add(*LinkIndexID);
+		}
+	}
+}
+
+// the output should be from current index -> to root parent 
+static bool GetBoneChain(const FIKRigTransforms& TransformModifier, const FName& Root, const FName& Current, TArray<int32>& ChainIndices)
+{
+	ChainIndices.Reset();
+
+	int32 RootIndex = TransformModifier.Hierarchy->GetIndex(Root);
+	int32 Iterator = TransformModifier.Hierarchy->GetIndex(Current);;
+
+	// iterates until key is valid
+	while (Iterator!=INDEX_NONE && Iterator != RootIndex)
+	{
+		ChainIndices.Insert(Iterator, 0);
+		Iterator = TransformModifier.Hierarchy->GetParentIndex(Iterator);
+	}
+
+	// add the last one if valid
+	if (Iterator!=INDEX_NONE)
+	{
+		ChainIndices.Insert(Iterator, 0);
+	}
+
+	// when you reached to something invalid, that means, we did not hit the 
+	// expected root, we iterated to the root and above, and we exhausted our option
+	// so if you hit or valid target by the hitting max depth, we should 
+	return !ChainIndices.IsEmpty();
+}
+
+static void AddEffectors(const FIKRigTransforms& TransformModifier, const FName Root, const TArray<FFBIKRigEffector>& Effectors,
+	TArray<FFBIKLinkData>& LinkData, TMap<int32, FFBIKEffectorTarget>& EffectorTargets, TArray<int32>& EffectorLinkIndices,
+	TMap<int32, int32>& LinkDataToHierarchyIndices, TMap<int32, int32>& HierarchyToLinkDataMap, const FSolverInput& SolverProperty)
+{
+	EffectorLinkIndices.SetNum(Effectors.Num());
+	// fill up all effector indices
+	for (int32 Index = 0; Index < Effectors.Num(); ++Index)
+	{
+		// clear link indices, so that we don't search
+		EffectorLinkIndices[Index] = INDEX_NONE;
+		// create LinkeData from root bone to all effectors 
+		const FName Bone = Effectors[Index].Target.Bone;
+		const int32 BoneIndex = TransformModifier.Hierarchy->GetIndex(Bone);
+
+		if (Bone == NAME_None || BoneIndex == INDEX_NONE)
+		{
+			continue;
+		}
+
+		const FFBIKRigEffector& CurrentEffector = Effectors[Index];
+
+		TArray<int32> ChainIndices;
+		// if we haven't got to root, this is not valid chain
+		if (GetBoneChain(TransformModifier, Root, Bone, ChainIndices))
+		{
+			auto CalculateStrength = [&](int32 InBoneChainDepth, const int32 MaxDepth, float CurrentStrength, float MinStrength) -> float
+			{
+				const float Range = FMath::Max(CurrentStrength - MinStrength, 0.f);
+				const float ApplicationStrength = (float)(1.f - (float)InBoneChainDepth / (float)MaxDepth) * Range;
+				return ApplicationStrength + MinStrength;
+				//return FMath::Clamp(ApplicationStrength + MinStrength, MinStrength, CurrentStrength);
+			};
+
+			auto UpdateMotionStrength = [&](int32 InBoneChainDepth,
+				const int32 MaxPositionDepth, const int32 MaxRotationDepth, FFBIKLinkData& InOutNewLink)
+			{
+				// add motion scales
+				float LinearMotionStrength;
+				float AngularMotionStrength;
+
+				if (CurrentEffector.PositionDepth <= InBoneChainDepth)
+				{
+					LinearMotionStrength = 0.f;
+				}
+				else
+				{
+					LinearMotionStrength = CalculateStrength(InBoneChainDepth, MaxPositionDepth, SolverProperty.LinearMotionStrength, SolverProperty.MinLinearMotionStrength);
+				}
+
+				if (CurrentEffector.RotationDepth <= InBoneChainDepth)
+				{
+					AngularMotionStrength = 0.f;
+				}
+				else
+				{
+					AngularMotionStrength = CalculateStrength(InBoneChainDepth, MaxRotationDepth, SolverProperty.AngularMotionStrength, SolverProperty.MinAngularMotionStrength);
+				}
+
+				InOutNewLink.AddMotionStrength(LinearMotionStrength, AngularMotionStrength);
+			};
+
+			// position depth and rotation depth can't go beyond of it
+			// for now we cull it. 
+			const int32 PositionDepth = FMath::Min(CurrentEffector.PositionDepth, ChainIndices.Num());
+			const int32 RotationDepth = FMath::Min(CurrentEffector.RotationDepth, ChainIndices.Num());
+
+			float ChainLength = 0.f;
+
+			// add to link data
+			for (int32 BoneChainIndex = 0; BoneChainIndex < ChainIndices.Num(); ++BoneChainIndex)
+			{
+				const int32 CurrentItem = ChainIndices[BoneChainIndex];
+				ChainLength += EnsureToAddBoneToLinkData(TransformModifier, CurrentItem, LinkData,
+					HierarchyToLinkDataMap, LinkDataToHierarchyIndices);
+
+				const int32 ChainDepth = ChainIndices.Num() - BoneChainIndex;
+				int32* FoundLinkIndex = HierarchyToLinkDataMap.Find(CurrentItem);
+
+				// now we should always have it
+				check(FoundLinkIndex);
+				UpdateMotionStrength(ChainDepth, PositionDepth, RotationDepth, LinkData[*FoundLinkIndex]);
+			}
+
+			// add to EffectorTargets
+			AddToEffectorTarget(Index, BoneIndex, EffectorTargets, HierarchyToLinkDataMap, EffectorLinkIndices, ChainLength, ChainIndices, PositionDepth, RotationDepth);
+		}
+	}
+}
+/////////////////////////////////////////////////
+
+UIKRigFBIKSolver::UIKRigFBIKSolver()
+{
+}
+
+void UIKRigFBIKSolver::Init(const FIKRigTransforms& InGlobalTransform)
+{
+
+	LinkData.Reset();
+	EffectorTargets.Reset();
+	EffectorLinkIndices.Reset();
+	LinkDataToHierarchyIndices.Reset();
+	HierarchyToLinkDataMap.Reset();
+
+	// verify the chain
+	AddEffectors(InGlobalTransform, Root, Effectors, LinkData, EffectorTargets, EffectorLinkIndices, LinkDataToHierarchyIndices, HierarchyToLinkDataMap, SolverProperty);
+	
+}
+
+void UIKRigFBIKSolver::Solve(
+	FIKRigTransforms& InOutGlobalTransform,
+	const FIKRigGoalContainer& Goals,
+	FControlRigDrawInterface* InOutDrawInterface)
+{
+	if (Effectors.IsEmpty())
+	{
+		return; // nothing to do
+	}
+	
+	if (!LinkDataToHierarchyIndices.IsEmpty())
+	{
+#if 0
+		// we do this every frame for now
+		if (Constraints.Num() > 0)
+		{
+			DECLARE_SCOPE_HIERARCHICAL_COUNTER(TEXT("Build Constraint"))
+			//Build constraints
+			FBIKConstraintLib::BuildConstraints(Constraints, InternalConstraints, InOutGlobalTransform, LinkData, LinkDataToHierarchyIndices, HierarchyToLinkDataMap);
+		}
+
+		// during only editor and update
+		// we expect solver type changes, it will reinit
+		// InternalConstraints can't be changed during runtime
+		if (!InternalConstraints.IsEmpty())
+		{
+			WorkData.IKSolver.SetPostProcessDelegateForIteration(FPostProcessDelegateForIteration::CreateStatic(&FBIKConstraintLib::ApplyConstraint, &InternalConstraints));
+		}
+		else
+		{
+			WorkData.IKSolver.ClearPostProcessDelegateForIteration();
+		}
+	//disable constraint for now
+#endif // 0 
+
+		// before update we finalize motion scale
+		// this code may go away once we have constraint
+		// update link data and end effectors
+		for (int32 LinkIndex = 0; LinkIndex < LinkData.Num(); ++LinkIndex)
+		{
+			const int32 Item = *LinkDataToHierarchyIndices.Find(LinkIndex);
+			LinkData[LinkIndex].SetTransform(InOutGlobalTransform.GetGlobalTransform(Item));
+			// @todo: fix this somewhere else - we can add this to prepare step
+			// @todo: we update motion scale here, then?
+			LinkData[LinkIndex].FinalizeForSolver();
+		}
+
+		// update mid effector info
+		const float LinearMotionStrength = FMath::Max(SolverProperty.LinearMotionStrength, SolverProperty.MinLinearMotionStrength);
+		const float AngularMotionStrength = FMath::Max(SolverProperty.AngularMotionStrength, SolverProperty.MinAngularMotionStrength);
+		const float LinearRange = LinearMotionStrength - SolverProperty.MinLinearMotionStrength;
+		const float AngularRange = AngularMotionStrength - SolverProperty.MinAngularMotionStrength;
+
+		//const TArray<FFBIKRigEffector>& Effectors = Effectors;
+
+		// update end effector info
+		for (int32 EffectorIndex = 0; EffectorIndex < Effectors.Num(); ++EffectorIndex)
+		{
+			int32 EffectorLinkIndex = EffectorLinkIndices[EffectorIndex];
+			if (EffectorLinkIndex != INDEX_NONE)
+			{
+				FFBIKEffectorTarget* EffectorTarget = EffectorTargets.Find(EffectorLinkIndex);
+				if (EffectorTarget)
+				{
+					const FFBIKRigEffector& CurEffector = Effectors[EffectorIndex];
+					const FVector CurrentLinkLocation = LinkData[EffectorLinkIndex].GetTransform().GetLocation();
+					const FQuat CurrentLinkRotation = LinkData[EffectorLinkIndex].GetTransform().GetRotation();
+					FIKRigGoal Goal;
+					ensure(GetGoalForEffector(CurEffector.Target, Goals, Goal));
+					EffectorTarget->Position = Goal.Position; // FMath::Lerp(CurrentLinkLocation, EffectorLocation, RigTarget.PositionTarget.);
+					EffectorTarget->Rotation = Goal.Rotation; // FMath::Lerp(CurrentLinkRotation, EffectorRotation, CurEffector.RotationAlpha);
+					EffectorTarget->InitialPositionDistance = (EffectorTarget->Position - CurrentLinkLocation).Size();
+					EffectorTarget->InitialRotationDistance = (FBIKUtil::GetScaledRotationAxis(EffectorTarget->Rotation) - FBIKUtil::GetScaledRotationAxis(CurrentLinkRotation)).Size();
+
+					const float Pull = FMath::Clamp(CurEffector.Pull, 0.f, 1.f);
+					// we want some impact of Pull, in order for Pull to have some impact, we clamp to some number
+					const float TargetClamp = FMath::Clamp(SolverProperty.DefaultTargetClamp, 0.f, 0.7f);
+					const float Scale = TargetClamp + Pull * (1.f - TargetClamp);
+					// Pull set up
+					EffectorTarget->LinearMotionStrength = LinearRange * Scale + SolverProperty.MinLinearMotionStrength;
+					EffectorTarget->AngularMotionStrength = AngularRange * Scale + SolverProperty.MinAngularMotionStrength;
+					EffectorTarget->ConvergeScale = Scale;
+					EffectorTarget->TargetClampScale = Scale;
+
+					EffectorTarget->bPositionEnabled = true;
+					EffectorTarget->bRotationEnabled = true;
+				}
+			}
+		}
+
+		DebugData.Reset();
+
+		const bool bDebugEnabled = DebugOption.bDrawDebugHierarchy || DebugOption.bDrawDebugEffector || DebugOption.bDrawDebugConstraints;
+
+		// we can't reuse memory until we fix the memory issue on RigVM
+		{
+			IKSolver.SolveJacobianIK(LinkData, EffectorTargets,
+				JacobianIK::FSolverParameter(SolverProperty.Damping, true, false, (SolverProperty.bUseJacobianTranspose) ? EJacobianSolver::JacobianTranspose : EJacobianSolver::JacobianPIDLS),
+				SolverProperty.MaxIterations, SolverProperty.Precision, &DebugData);
+
+			if (MotionProperty.bForceEffectorRotationTarget)
+			{
+				// if position is reached, we force rotation target
+				for (int32 EffectorIndex = 0; EffectorIndex < Effectors.Num(); ++EffectorIndex)
+				{
+					int32 EffectorLinkIndex = EffectorLinkIndices[EffectorIndex];
+					if (EffectorLinkIndex != INDEX_NONE)
+					{
+						FFBIKEffectorTarget* EffectorTarget = EffectorTargets.Find(EffectorLinkIndex);
+						if (EffectorTarget && EffectorTarget->bRotationEnabled)
+						{
+							bool bApplyRotation = true;
+
+							if (MotionProperty.bOnlyApplyWhenReachedToTarget)
+							{
+								// only do this when position is reached? This will conflict with converge scale
+								const FVector& BonePosition = LinkData[EffectorLinkIndex].GetTransform().GetLocation();
+								const FVector& TargetPosition = EffectorTarget->Position;
+
+								bApplyRotation = (FVector(BonePosition - TargetPosition).SizeSquared() <= SolverProperty.Precision * SolverProperty.Precision);
+							}
+
+							if (bApplyRotation)
+							{
+								FQuat NewRotation = EffectorTarget->Rotation;
+								FTransform NewTransform = LinkData[EffectorLinkIndex].GetTransform();
+								NewTransform.SetRotation(NewRotation);
+								LinkData[EffectorLinkIndex].SetTransform(NewTransform);
+							}
+						}
+					}
+				}
+			}
+		}
+		///////////////////////////////////////////////////////////////////////////
+		// debug draw start
+		///////////////////////////////////////////////////////////////////////////
+		if (bDebugEnabled && InOutDrawInterface != nullptr)
+		{
+			const int32 DebugDataNum = DebugData.Num();
+			if (!DebugData.IsEmpty())
+			{
+				for (int32 DebugIndex = DebugDataNum - 1; DebugIndex >= 0; --DebugIndex)
+				{
+					const TArray<FFBIKLinkData>& LocalLink = DebugData[DebugIndex].LinkData;
+
+					FTransform Offset = DebugOption.DrawWorldOffset;
+					Offset.SetLocation(Offset.GetLocation() * (DebugDataNum - DebugIndex));
+
+					if (DebugOption.bDrawDebugHierarchy)
+					{
+						for (int32 LinkIndex = 0; LinkIndex < LocalLink.Num(); ++LinkIndex)
+						{
+							const FFBIKLinkData& Data = LocalLink[LinkIndex];
+
+							FLinearColor DrawColor = FLinearColor::White;
+
+							float LineThickness = 0.f;
+							if (DebugOption.bColorAngularMotionStrength || DebugOption.bColorLinearMotionStrength)
+							{
+								DrawColor = FLinearColor::Black;
+								if (DebugOption.bColorAngularMotionStrength)
+								{
+									const float Range = FMath::Max(SolverProperty.AngularMotionStrength - SolverProperty.MinAngularMotionStrength, 0.f);
+									if (Range > 0.f)
+									{
+										float CurrentStrength = Data.GetAngularMotionStrength() - SolverProperty.MinAngularMotionStrength;
+										float Alpha = FMath::Clamp(CurrentStrength / Range, 0.f, 1.f);
+										DrawColor.R = LineThickness = Alpha;
+									}
+								}
+								else if (DebugOption.bColorLinearMotionStrength)
+								{
+									const float Range = FMath::Max(SolverProperty.LinearMotionStrength - SolverProperty.MinLinearMotionStrength, 0.f);
+									if (Range > 0.f)
+									{
+										float CurrentStrength = Data.GetLinearMotionStrength() - SolverProperty.MinLinearMotionStrength;
+										float Alpha = FMath::Clamp(CurrentStrength / Range, 0.f, 1.f);
+										DrawColor.B = LineThickness = Alpha;
+									}
+								}
+							}
+
+							if (Data.ParentLinkIndex != INDEX_NONE)
+							{
+								const FFBIKLinkData& ParentData = LocalLink[Data.ParentLinkIndex];
+								InOutDrawInterface->DrawLine(Offset, Data.GetPreviousTransform().GetLocation(), ParentData.GetPreviousTransform().GetLocation(), DrawColor, LineThickness);
+							}
+
+							if (DebugOption.bDrawDebugAxes)
+							{
+								InOutDrawInterface->DrawAxes(Offset, Data.GetPreviousTransform(), DebugOption.DrawSize);
+							}
+						}
+					}
+
+					if (DebugOption.bDrawDebugEffector)
+					{
+						for (auto Iter = EffectorTargets.CreateConstIterator(); Iter; ++Iter)
+						{
+							const FFBIKEffectorTarget& EffectorTarget = Iter.Value();
+							if (EffectorTarget.bPositionEnabled)
+							{
+								// draw effector target locations
+								InOutDrawInterface->DrawBox(Offset, FTransform(EffectorTarget.Position), FLinearColor::Yellow, DebugOption.DrawSize);
+							}
+
+							// draw effector link location
+							InOutDrawInterface->DrawBox(Offset, LocalLink[Iter.Key()].GetPreviousTransform(), FLinearColor::Green, DebugOption.DrawSize);
+						}
+
+						for (int32 Index = 0; Index < DebugData[DebugIndex].TargetVectorSources.Num(); ++Index)
+						{
+							// draw arrow to the target
+							InOutDrawInterface->DrawLine(Offset, DebugData[DebugIndex].TargetVectorSources[Index].GetLocation(),
+								DebugData[DebugIndex].TargetVectorSources[Index].GetLocation() + DebugData[DebugIndex].TargetVectors[Index], FLinearColor::Red);
+						}
+					}
+				}
+			}
+#if 0 
+			if (SolverDef->DebugOption.bDrawDebugConstraints && InternalConstraints.Num())
+			{
+				FTransform Offset = FTransform::Identity;
+
+				// draw frame if active
+				for (int32 Index = 0; Index < Constraints.Num(); ++Index)
+				{
+					if (Constraints[Index].bEnabled)
+					{
+						if (Constraints[Index].Item.IsValid())
+						{
+							const int32* Found = HierarchyToLinkDataMap.Find(Constraints[Index].Item);
+							if (Found)
+							{
+								FTransform ConstraintFrame = LinkData[*Found].GetTransform();
+								ConstraintFrame.ConcatenateRotation(FQuat(Constraints[Index].OffsetRotation));
+								InOutDrawInterface->DrawAxes(Offset, ConstraintFrame, 2.f);
+							}
+						}
+					}
+				}
+
+				for (int32 Index = 0; Index < InternalConstraints.Num(); ++Index)
+				{
+					// for now we have rotation limit only
+					if (InternalConstraints[Index].IsType< FRotationLimitConstraint>())
+					{
+						FRotationLimitConstraint& LimitConstraint = InternalConstraints[Index].Get<FRotationLimitConstraint>();
+						const FQuat LocalRefRotation = LimitConstraint.RelativelRefPose.GetRotation();
+						FTransform RotationTransform = LinkData[LimitConstraint.ConstrainedIndex].GetTransform();
+						// base is parent transform but in their space, we can get there by inversing local ref rotation
+						FTransform BaseTransform = FTransform(LocalRefRotation).GetRelativeTransformReverse(RotationTransform);
+						BaseTransform.ConcatenateRotation(LimitConstraint.BaseFrameOffset);
+						InOutDrawInterface->DrawAxes(Offset, BaseTransform, 5.f, 1.f);
+
+						// current transform
+						const FQuat LocalRotation = BaseTransform.GetRotation().Inverse() * RotationTransform.GetRotation();
+						const FQuat DeltaTransform = LocalRefRotation.Inverse() * LocalRotation;
+						RotationTransform.SetRotation(BaseTransform.GetRotation() * DeltaTransform);
+						RotationTransform.NormalizeRotation();
+						RotationTransform.SetLocation(BaseTransform.GetLocation());
+
+						// draw ref pose on their current transform
+						InOutDrawInterface->DrawAxes(Offset, RotationTransform, 10.f, 1.f);
+
+						FVector XAxis = BaseTransform.GetUnitAxis(EAxis::X);
+						FVector YAxis = BaseTransform.GetUnitAxis(EAxis::Y);
+						FVector ZAxis = BaseTransform.GetUnitAxis(EAxis::Z);
+
+						if (LimitConstraint.bXLimitSet)
+						{
+							FTransform XAxisConeTM(YAxis, XAxis ^ YAxis, XAxis, BaseTransform.GetTranslation());
+							XAxisConeTM.SetRotation(FQuat(XAxis, 0.f) * XAxisConeTM.GetRotation());
+							XAxisConeTM.SetScale3D(FVector(30.f));
+							InOutDrawInterface->DrawCone(Offset, XAxisConeTM, LimitConstraint.Limit.X, 0.0f, 24, false, FLinearColor::Red, GEngine->ConstraintLimitMaterialX->GetRenderProxy());
+						}
+
+						if (LimitConstraint.bYLimitSet)
+						{
+							FTransform YAxisConeTM(ZAxis, YAxis ^ ZAxis, YAxis, BaseTransform.GetTranslation());
+							YAxisConeTM.SetRotation(FQuat(YAxis, 0.f) * YAxisConeTM.GetRotation());
+							YAxisConeTM.SetScale3D(FVector(30.f));
+							InOutDrawInterface->DrawCone(Offset, YAxisConeTM, LimitConstraint.Limit.Y, 0.0f, 24, false, FLinearColor::Green, GEngine->ConstraintLimitMaterialY->GetRenderProxy());
+						}
+
+						if (LimitConstraint.bZLimitSet)
+						{
+							FTransform ZAxisConeTM(XAxis, ZAxis ^ XAxis, ZAxis, BaseTransform.GetTranslation());
+							ZAxisConeTM.SetRotation(FQuat(ZAxis, 0.f) * ZAxisConeTM.GetRotation());
+							ZAxisConeTM.SetScale3D(FVector(30.f));
+							InOutDrawInterface->DrawCone(Offset, ZAxisConeTM, LimitConstraint.Limit.Z, 0.0f, 24, false, FLinearColor::Blue, GEngine->ConstraintLimitMaterialZ->GetRenderProxy());
+						}
+					}
+
+					if (InternalConstraints[Index].IsType<FPoleVectorConstraint>())
+					{
+						// darw pole vector location
+						// draw 3 joints line and a plane to pole vector
+						FPoleVectorConstraint& Constraint = InternalConstraints[Index].Get<FPoleVectorConstraint>();
+						FTransform RootTransform = LinkData[Constraint.ParentBoneIndex].GetTransform();
+						FTransform JointTransform = LinkData[Constraint.BoneIndex].GetTransform();
+						FTransform ChildTransform = LinkData[Constraint.ChildBoneIndex].GetTransform();
+
+						FVector JointTarget = (Constraint.bUseLocalDir) ? Constraint.CalculateCurrentPoleVectorDir(RootTransform, JointTransform, ChildTransform, LinkData[Constraint.BoneIndex].LocalFrame) : Constraint.PoleVector;
+
+						// draw the plane, 
+						TArray<FVector> Positions;
+						Positions.Add(RootTransform.GetLocation());
+						Positions.Add(ChildTransform.GetLocation());
+						Positions.Add(ChildTransform.GetLocation());
+						Positions.Add(JointTarget);
+						Positions.Add(JointTarget);
+						Positions.Add(RootTransform.GetLocation());
+
+						InOutDrawInterface->DrawLines(Offset, Positions, FLinearColor::Gray, 1.2f);
+						InOutDrawInterface->DrawLine(Offset, JointTransform.GetLocation(), JointTarget, FLinearColor::Red, 1.2f);
+					}
+				}
+			}
+
+			#endif // 0 - disable all constraint
+		}
+		///////////////////////////////////////////////////////////////////////////
+		// debug draw end
+		///////////////////////////////////////////////////////////////////////////
+	}
+
+	// we update back to hierarchy
+	for (int32 LinkIndex = 0; LinkIndex < LinkData.Num(); ++LinkIndex)
+	{
+		// only propagate, if you are leaf joints here
+		// this means, only the last joint in the test
+		const int32 CurrentItem = *LinkDataToHierarchyIndices.Find(LinkIndex);
+		const FTransform& LinkTransform = LinkData[LinkIndex].GetTransform();
+		InOutGlobalTransform.SetGlobalTransform(CurrentItem, LinkTransform, true);
+	}
+}
+
+void UIKRigFBIKSolver::CollectGoalNames(TSet<FName>& OutGoals) const
+{
+	for (const FFBIKRigEffector& Effector: Effectors)
+	{
+		OutGoals.Add(Effector.Target.Goal);
+	}
+}
diff --git a/Engine/Plugins/Animation/IKRig/Source/IKRig/Public/Solvers/IKRigFBIKSolver.h b/Engine/Plugins/Animation/IKRig/Source/IKRig/Public/Solvers/IKRigFBIKSolver.h
new file mode 100644
--- /dev/null
+++ b/Engine/Plugins/Animation/IKRig/Source/IKRig/Public/Solvers/IKRigFBIKSolver.h
@@ -0,0 +1,101 @@
+// Copyright Epic Games, Inc. All Rights Reserved.
+
+/**
+ * Contains Transform Solver Execution
+ *
+ */
+
+#pragma once
+
+#include "IKRigSolver.h"
+#include "FBIKShared.h"
+#include "FBIKConstraintOption.h"
+#include "FBIKDebugOption.h"
+#include "IKRigFBIKSolver.generated.h"
+
+struct FFBIKLinkData;
+
+USTRUCT()
+struct FFBIKRigEffector
+{
+	GENERATED_BODY()
+
+	/**
+	* The last item in the chain to solve - the effector
+	*/
+	UPROPERTY(EditAnywhere, Category = FFBIKRigEffector)
+	FIKRigEffectorGoal Target;
+
+	UPROPERTY(EditAnywhere, Category = FFBIKRigEffector)
+	int32  PositionDepth = 1000;
+
+	UPROPERTY(EditAnywhere, Category = FFBIKRigEffector)
+	int32  RotationDepth = 1000;
+
+	/*
+	 * Clamps the total length to target by this scale for each iteration
+	 * This helps to stabilize solver to reduce singularity by avoiding to try to reach target too far.
+	 */
+	UPROPERTY(EditAnywhere, Category = FFBIKRigEffector)
+	float Pull = 0.f;
+
+	FFBIKRigEffector()
+	{
+	}
+};
+
+ // run time for UFullBodyIKSolverDefinition
+UCLASS(EditInlineNew)
+class IKRIG_API UIKRigFBIKSolver : public UIKRigSolver
+{
+	GENERATED_BODY()
+
+public:
+	
+	UPROPERTY(EditAnywhere, Category = Effectors)
+	FName Root;
+
+	UPROPERTY(EditAnywhere, Category = Effectors)
+	TArray<FFBIKRigEffector> Effectors;
+
+	UPROPERTY(EditAnywhere, Category = Solver)
+	FSolverInput SolverProperty;
+
+	UPROPERTY(EditAnywhere, Category = Motion)
+	FMotionProcessInput MotionProperty;
+
+	UPROPERTY(EditAnywhere, Category = Solver)
+	FFBIKDebugOption DebugOption;
+
+	UIKRigFBIKSolver();
+
+private:
+
+	/** list of Link Data for solvers - joints */
+	TArray<FFBIKLinkData> LinkData;
+	/** Effector Targets - search key is LinkData Index */
+	TMap<int32, FFBIKEffectorTarget> EffectorTargets; 
+	/** End Effector Link Indices - EndEffector index to LinkData index*/
+	TArray<int32> EffectorLinkIndices;
+	/** Map from LinkData index to Hierarchy Index*/
+	TMap<int32, int32> LinkDataToHierarchyIndices;
+	/** Map from Rig Hierarchy Index to LinkData index*/
+	TMap<int32, int32> HierarchyToLinkDataMap;
+	/** Constraints data */
+	TArray<ConstraintType> InternalConstraints;
+	/* Current Solver */
+	FJacobianSolver_FullbodyIK IKSolver;
+	/** Debug Data */
+	TArray<FJacobianDebugData> DebugData;
+
+protected:
+	virtual void Init(const FIKRigTransforms& InGlobalTransforms) override;
+	virtual void Solve(
+		FIKRigTransforms& InOutGlobalTransforms,
+		const FIKRigGoalContainer& Goals,
+		FControlRigDrawInterface* InOutDrawInterface) override;
+	virtual void CollectGoalNames(TSet<FName>& OutGoals) const override;
+
+public:
+	static const FName EffectorTargetPrefix;
+};
\ No newline at end of file
diff --git a/Engine/Plugins/Experimental/FullBodyIK/FullBodyIK.uplugin b/Engine/Plugins/Experimental/FullBodyIK/FullBodyIK.uplugin
--- a/Engine/Plugins/Experimental/FullBodyIK/FullBodyIK.uplugin
+++ b/Engine/Plugins/Experimental/FullBodyIK/FullBodyIK.uplugin
@@ -18,10 +18,6 @@
 		{
 			"Name": "ControlRig",
 			"Enabled": true
-		},
-		{
-			"Name": "IKRig",
-			"Enabled": true
 		}
 	],
 	"Modules": [
diff --git a/Engine/Plugins/Experimental/FullBodyIK/Source/FullBodyIK/FullBodyIK.Build.cs b/Engine/Plugins/Experimental/FullBodyIK/Source/FullBodyIK/FullBodyIK.Build.cs
--- a/Engine/Plugins/Experimental/FullBodyIK/Source/FullBodyIK/FullBodyIK.Build.cs
+++ b/Engine/Plugins/Experimental/FullBodyIK/Source/FullBodyIK/FullBodyIK.Build.cs
@@ -15,9 +15,8 @@ public class FullBodyIK : ModuleRules
 				// ... add other public dependencies that you statically link with here ...
 				"Eigen",
 				"RigVM",
-				"ControlRig",
-				"IKRig"
-            }
+				"ControlRig"
+			}
 			);
 			
 		
diff --git a/Engine/Plugins/Experimental/FullBodyIK/Source/FullBodyIK/Private/FBIKShared.h b/Engine/Plugins/Experimental/FullBodyIK/Source/FullBodyIK/Private/FBIKShared.h
deleted file mode 100644
diff --git a/Engine/Plugins/Experimental/FullBodyIK/Source/FullBodyIK/Private/IKRigFBIKSolver.cpp b/Engine/Plugins/Experimental/FullBodyIK/Source/FullBodyIK/Private/IKRigFBIKSolver.cpp
deleted file mode 100644
diff --git a/Engine/Plugins/Experimental/FullBodyIK/Source/FullBodyIK/Private/IKRigFBIKSolver.h b/Engine/Plugins/Experimental/FullBodyIK/Source/FullBodyIK/Private/IKRigFBIKSolver.h
deleted file mode 100644
diff --git a/Engine/Plugins/Experimental/FullBodyIK/Source/FullBodyIK/Public/FBIKShared.h b/Engine/Plugins/Experimental/FullBodyIK/Source/FullBodyIK/Public/FBIKShared.h
new file mode 100644
--- /dev/null
+++ b/Engine/Plugins/Experimental/FullBodyIK/Source/FullBodyIK/Public/FBIKShared.h
@@ -0,0 +1,76 @@
+// Copyright Epic Games, Inc. All Rights Reserved.
+
+#pragma once
+
+#include "CoreMinimal.h"
+#include "JacobianSolver.h"
+#include "Rigs/RigHierarchyDefines.h"
+#include "FBIKShared.generated.h"
+
+class FULLBODYIK_API FJacobianSolver_FullbodyIK : public FJacobianSolverBase
+{
+public:
+	FJacobianSolver_FullbodyIK() : FJacobianSolverBase() {}
+
+private:
+	virtual void InitializeSolver(TArray<FFBIKLinkData>& InOutLinkData, TMap<int32, FFBIKEffectorTarget>& InOutEndEffectors) const override;
+	virtual void PreSolve(TArray<FFBIKLinkData>& InOutLinkData, const TMap<int32, FFBIKEffectorTarget>& InEndEffectors) const override;
+};
+
+USTRUCT()
+struct FSolverInput
+{
+	GENERATED_BODY()
+
+	/*
+	 * This value is applied to the target information for effectors, which influence back to 
+	 * Joint's motion that are affected by the end effector
+	 * The reason min/max is used when we apply the depth through the chain that are affected
+
+	 */
+	UPROPERTY(EditAnywhere, Category = FSolverInput)
+	float	LinearMotionStrength = 3.f;
+
+	UPROPERTY(EditAnywhere, Category = FSolverInput)
+	float	MinLinearMotionStrength = 2.f;
+
+	/*
+	 * This value is applied to the target information for effectors, which influence back to 
+	 * Joint's motion that are affected by the end effector
+	 * The reason min/max is used when we apply the depth through the chain that are affected
+	 */
+	UPROPERTY(EditAnywhere, Category = FSolverInput)
+	float	AngularMotionStrength = 3.f;
+
+	UPROPERTY(EditAnywhere, Category = FSolverInput)
+	float	MinAngularMotionStrength = 2.f;
+
+	/* This is a scale value (range from 0-0.7) that is used to stablize the target vector. If less, it's more stable, but it can reduce speed of converge. */
+	UPROPERTY(EditAnywhere, Category = FSolverInput)
+	float	DefaultTargetClamp = 0.2f;
+
+	/**
+	 * The precision to use for the solver
+	 */
+	UPROPERTY(EditAnywhere, Category = FSolverInput)
+	float Precision = 0.1f;
+
+	/**
+	* The precision to use for the fabrik solver
+	*/
+	UPROPERTY(EditAnywhere, Category = FSolverInput)
+	float Damping = 30.f;
+
+	/**
+	 * The maximum number of iterations. Values between 4 and 16 are common.
+	 */
+	UPROPERTY(EditAnywhere, Category = FSolverInput)
+	int32 MaxIterations = 30;
+
+	/**
+	 * Cheaper solution than default Jacobian Pseudo Inverse Damped Least Square
+	 */
+	UPROPERTY(EditAnywhere, Category = FSolverInput)
+	bool bUseJacobianTranspose = false;
+};
+
