Fixed bug, solvers not loaded before goal list created in Anim Node.
Minor cleanup, removed pointless getters()

diff --git a/Engine/Plugins/Animation/IKRig/Source/IKRig/Private/IKRigProcessor.cpp b/Engine/Plugins/Animation/IKRig/Source/IKRig/Private/IKRigProcessor.cpp
--- a/Engine/Plugins/Animation/IKRig/Source/IKRig/Private/IKRigProcessor.cpp
+++ b/Engine/Plugins/Animation/IKRig/Source/IKRig/Private/IKRigProcessor.cpp
@@ -16,19 +16,24 @@ void UIKRigProcessor::Initialize(UIKRigDefinition* InRigDefinition)
 		return;
 	}
 	
-	RefPoseGlobalTransforms = RigDefinition->GetReferencePose();
-	GlobalBoneTransforms = FIKRigTransforms(&RigDefinition->GetHierarchy());
+	RefPoseGlobalTransforms = RigDefinition->RefPoseTransforms;
+	GlobalBoneTransforms = FIKRigTransforms(&RigDefinition->Hierarchy);
 	GlobalBoneTransforms.SetAllGlobalTransforms(RefPoseGlobalTransforms);
 
 	TArray<FName> GoalNames;
 	RigDefinition->GetGoalNamesFromSolvers(GoalNames);
 	Goals.InitializeGoalsFromNames(GoalNames);
 
-	TArray<UIKRigSolver*> RigSolvers = RigDefinition->GetSolvers();
+	TArray<UIKRigSolver*> RigSolvers = RigDefinition->Solvers;
 	const int32 NumSolvers = RigSolvers.Num();
 	Solvers.Reset(NumSolvers);
 	for (int32 Index = 0; Index < NumSolvers; ++Index)
 	{
+		if (!RigSolvers[Index])
+		{
+			continue;
+		}
+		
 		UIKRigSolver* Solver = DuplicateObject(RigSolvers[Index], this);
 		Solver->Init(GlobalBoneTransforms);
 		Solvers.Add(Solver);
@@ -83,7 +88,7 @@ int UIKRigProcessor::GetNumGoals() const
 const FIKRigHierarchy* UIKRigProcessor::GetHierarchy() const
 {
 	check(bInitialized);
-	return ensure(RigDefinition) ? &RigDefinition->GetHierarchy(): nullptr;
+	return ensure(RigDefinition) ? &RigDefinition->Hierarchy : nullptr;
 }
 
 void UIKRigProcessor::ResetToRefPose()
diff --git a/Engine/Plugins/Animation/IKRig/Source/IKRig/Public/IKRigDefinition.h b/Engine/Plugins/Animation/IKRig/Source/IKRig/Public/IKRigDefinition.h
--- a/Engine/Plugins/Animation/IKRig/Source/IKRig/Public/IKRigDefinition.h
+++ b/Engine/Plugins/Animation/IKRig/Source/IKRig/Public/IKRigDefinition.h
@@ -23,18 +23,6 @@ public:
 	UPROPERTY()
 	TSoftObjectPtr<UObject> SourceAsset;
 
-	UPROPERTY(EditAnywhere, instanced, Category = BoneSettings)
-	TArray<UIKRigBoneSetting*> BoneSettings;
-
-	UPROPERTY(EditAnywhere, Category = BoneSettings)
-	TArray<FName> ExcludedBones;
-
-private:
-	/** stack of solvers, of varying types, executed in serial fashion where
-	output of prior solve is input to the next */
-	UPROPERTY(EditAnywhere, instanced, Category = "Solvers")
-	TArray<UIKRigSolver*> Solvers;
-
 	/** parent/child relationships stored separately from skeleton-specific transforms */
 	UPROPERTY(VisibleAnywhere, Category = "Skeleton")
 	FIKRigHierarchy Hierarchy;
@@ -43,22 +31,16 @@ private:
 	UPROPERTY(VisibleAnywhere, Category = "Skeleton")
 	TArray<FTransform> RefPoseTransforms;
 
-public:
-
-	const TArray<FTransform>& GetReferencePose() const
-	{
-		return RefPoseTransforms;
-	}
+	UPROPERTY(EditAnywhere, instanced, Category = BoneSettings)
+	TArray<UIKRigBoneSetting*> BoneSettings;
 
-	const TArray<UIKRigSolver*>& GetSolvers() const 
-	{
-		return Solvers;
-	}
+	UPROPERTY(EditAnywhere, Category = BoneSettings)
+	TArray<FName> ExcludedBones;
 
-	const FIKRigHierarchy& GetHierarchy() const
-	{
-		return Hierarchy;
-	}
+	/** stack of solvers, of varying types, executed in serial fashion where
+	output of prior solve is input to the next */
+	UPROPERTY(EditAnywhere, instanced, Category = "Solvers")
+	TArray<UIKRigSolver*> Solvers;
 
 	void GetGoalNamesFromSolvers(TArray<FName>& OutGoalNames) const;
 
diff --git a/Engine/Plugins/Animation/IKRig/Source/IKRigDeveloper/Private/AnimGraphNode_IKRig.cpp b/Engine/Plugins/Animation/IKRig/Source/IKRigDeveloper/Private/AnimGraphNode_IKRig.cpp
--- a/Engine/Plugins/Animation/IKRig/Source/IKRigDeveloper/Private/AnimGraphNode_IKRig.cpp
+++ b/Engine/Plugins/Animation/IKRig/Source/IKRigDeveloper/Private/AnimGraphNode_IKRig.cpp
@@ -3,6 +3,7 @@
 #include "AnimGraphNode_IKRig.h"
 #include "Animation/AnimInstance.h"
 #include "IKRigDefinition.h"
+#include "IKRigSolver.h"
 
 /////////////////////////////////////////////////////
 // UAnimGraphNode_IKRig 
@@ -83,6 +84,10 @@ void UAnimGraphNode_IKRig::PreloadRequiredAssets()
 	if (Node.RigDefinitionAsset)
 	{
 		PreloadObject(Node.RigDefinitionAsset);
+		for (UIKRigSolver* Solver : Node.RigDefinitionAsset->Solvers)
+		{
+			PreloadObject(Solver);
+		}
 	}
 
 	Super::PreloadRequiredAssets();
diff --git a/Engine/Plugins/Animation/IKRig/Source/IKRigEditor/Private/IKRigController.cpp b/Engine/Plugins/Animation/IKRig/Source/IKRigEditor/Private/IKRigController.cpp
--- a/Engine/Plugins/Animation/IKRig/Source/IKRigEditor/Private/IKRigController.cpp
+++ b/Engine/Plugins/Animation/IKRig/Source/IKRigEditor/Private/IKRigController.cpp
@@ -67,14 +67,14 @@ void UIKRigController::SetIKRigDefinition(UIKRigDefinition* InIKRigDefinition)
 
 const FIKRigHierarchy* UIKRigController::GetHierarchy() const
 {
-	return (IKRigDefinition) ? &IKRigDefinition->GetHierarchy() : nullptr;
+	return IKRigDefinition ? &IKRigDefinition->Hierarchy : nullptr;
 }
 
 const TArray<FTransform>& UIKRigController::GetRefPoseTransforms() const
 {
 	if (IKRigDefinition)
 	{
-		return IKRigDefinition->GetReferencePose();
+		return IKRigDefinition->RefPoseTransforms;
 	}
 
 	static TArray<FTransform> Dummy;
