Pretzel: Adding retargeting ability to IKRig plugin.

- Added IKRetargeter anim node
- Added IKRetargeter.h/cpp API for copying animation between different hierarchies using an IK Rig
- Removed IKRigHierarchy and IKRigTransforms classes
- Added IKRigSkeleton class for all the hierarchy and transform manipulation
- Modified IK Rig to use new IKRigSkeleton format
- Modified FBIK in IKRig to use IKRigSkeleton
- Modified IKRigController to load data into an IKRigSkeleton

diff --git a/Engine/Plugins/Animation/IKRig/Source/IKRig/Private/AnimNodes/AnimNode_IKRetargeter.cpp b/Engine/Plugins/Animation/IKRig/Source/IKRig/Private/AnimNodes/AnimNode_IKRetargeter.cpp
new file mode 100644
--- /dev/null
+++ b/Engine/Plugins/Animation/IKRig/Source/IKRig/Private/AnimNodes/AnimNode_IKRetargeter.cpp
@@ -0,0 +1,192 @@
+// Copyright Epic Games, Inc. All Rights Reserved.
+
+#include "AnimNodes/AnimNode_IKRetargeter.h"
+#include "Animation/AnimInstanceProxy.h"
+
+//#pragma optimize("", off)
+
+FAnimNode_IKRetargeter::FAnimNode_IKRetargeter()
+    : SourceMeshComponent(nullptr),
+	bEnableIK(true),
+	bIsInitialized (false)
+{
+}
+
+void FAnimNode_IKRetargeter::Initialize_AnyThread(const FAnimationInitializeContext& Context)
+{
+	DECLARE_SCOPE_HIERARCHICAL_COUNTER_ANIMNODE(Initialize_AnyThread)
+    FAnimNode_Base::Initialize_AnyThread(Context);
+
+	// Initial update of the node, so we dont have a frame-delay on setup
+	GetEvaluateGraphExposedInputs().Execute(Context);
+}
+
+void FAnimNode_IKRetargeter::CacheBones_AnyThread(const FAnimationCacheBonesContext& Context)
+{
+	DECLARE_SCOPE_HIERARCHICAL_COUNTER_ANIMNODE(CacheBones_AnyThread)
+}
+
+void FAnimNode_IKRetargeter::Update_AnyThread(const FAnimationUpdateContext& Context)
+{
+	DECLARE_SCOPE_HIERARCHICAL_COUNTER_ANIMNODE(Update_AnyThread)
+	FAnimNode_Base::Update_AnyThread(Context);
+    // this introduces a frame of latency in setting the pin-driven source component,
+    // but we cannot do the work to extract transforms on a worker thread as it is not thread safe.
+    GetEvaluateGraphExposedInputs().Execute(Context);
+}
+void FAnimNode_IKRetargeter::Evaluate_AnyThread(FPoseContext& Output)
+{
+	DECLARE_SCOPE_HIERARCHICAL_COUNTER_ANIMNODE(Evaluate_AnyThread)
+
+	if (!bIsInitialized)
+	{
+		return;
+	}
+
+	// run the retargeter
+	CurrentlyUsedRetargeter->RunRetargeter(SourceMeshComponentSpaceBoneTransforms, bEnableIK);
+
+	// copy pose back
+	FCSPose<FCompactPose> ComponentPose;
+	ComponentPose.InitPose(Output.Pose);
+	for (FCompactPoseBoneIndex CompactBoneIndex : Output.Pose.ForEachBoneIndex())
+	{
+		ComponentPose.SetComponentSpaceTransform(CompactBoneIndex, CurrentlyUsedRetargeter->TargetSkeleton.OutputGlobalPose[CompactBoneIndex.GetInt()]);
+	}
+
+	// convert to local space
+	FCSPose<FCompactPose>::ConvertComponentPosesToLocalPoses(ComponentPose, Output.Pose);
+}
+
+void FAnimNode_IKRetargeter::PreUpdate(const UAnimInstance* InAnimInstance)
+{
+	DECLARE_SCOPE_HIERARCHICAL_COUNTER_ANIMNODE(PreUpdate)
+	EnsureInitialized(InAnimInstance);
+	if (bIsInitialized)
+	{
+		CopyBoneTransformsFromSource(InAnimInstance->GetSkelMeshComponent());
+	}
+}
+
+void FAnimNode_IKRetargeter::EnsureInitialized(const UAnimInstance* InAnimInstance)
+{
+	// has a source mesh been plugged in?
+	if (!(SourceMeshComponent.IsValid() && IKRetargeterAsset))
+	{
+		bIsInitialized = false;
+		return; // can't do anything if we don't have a source mesh
+	}
+	
+	// are all the parts already loaded?
+	const bool bPartsLoaded =
+	   CurrentlyUsedSourceMeshComponent.IsValid()
+    && CurrentlyUsedSourceMesh.IsValid()
+    && CurrentlyUsedTargetMesh.IsValid()
+	&& CurrentlyUsedRetargeter.IsValid()
+	&& CurrentlyUsedSourceIKRig.IsValid()
+	&& CurrentlyUsedTargetIKRig.IsValid();
+	if (!bPartsLoaded)
+	{
+		InitializeRetargetData(InAnimInstance); // nothing loaded yet, initialize
+		return;
+	}
+
+	// so parts have been loaded, but have any of the parts changed since we last initialized?
+	USkeletalMeshComponent* SourceMeshComp = SourceMeshComponent.Get();
+	USkeletalMesh* TargetMesh = InAnimInstance->GetSkelMeshComponent()->SkeletalMesh;
+	USkeletalMesh* SourceMesh = SourceMeshComp->SkeletalMesh;
+	const bool bSameParts =
+	   CurrentlyUsedSourceMeshComponent == SourceMeshComp 
+    && CurrentlyUsedTargetMesh == TargetMesh 
+    && CurrentlyUsedSourceMesh == SourceMesh 
+	&& CurrentlyUsedSourceIKRig == IKRetargeterAsset->SourceIKRigAsset
+	&& CurrentlyUsedTargetIKRig == IKRetargeterAsset->TargetIKRigAsset;
+	if (!bSameParts)
+	{
+		InitializeRetargetData(InAnimInstance); // parts have changed, re-initialize
+		return;
+	}
+}
+
+void FAnimNode_IKRetargeter::InitializeRetargetData(const UAnimInstance* InAnimInstance)
+{
+	// assume we fail until we don't
+	bIsInitialized = false;
+
+	// store all the components that were used to initialize
+	// if in future updates, any of this are mismatched, we have to re-initialize
+	CurrentlyUsedSourceMeshComponent = SourceMeshComponent;
+	CurrentlyUsedSourceMesh = SourceMeshComponent->SkeletalMesh;
+	CurrentlyUsedTargetMesh = InAnimInstance->GetSkelMeshComponent()->SkeletalMesh;
+	CurrentlyUsedRetargeter = DuplicateObject(IKRetargeterAsset, const_cast<UAnimInstance*>(InAnimInstance));
+	CurrentlyUsedSourceIKRig = CurrentlyUsedRetargeter->SourceIKRigAsset;
+	CurrentlyUsedTargetIKRig = CurrentlyUsedRetargeter->TargetIKRigAsset;
+
+	const bool bMeshesAreValid = CurrentlyUsedSourceMesh.IsValid() && CurrentlyUsedTargetMesh.IsValid();
+	if (!bMeshesAreValid)
+	{
+		return; // cannot initialize if components are missing skeletal mesh references
+	}
+
+	const bool bRetargeterIsValid =
+	   CurrentlyUsedRetargeter.IsValid()
+	&& CurrentlyUsedSourceIKRig.IsValid()
+	&& CurrentlyUsedTargetIKRig.IsValid();
+	if (!bRetargeterIsValid)
+	{
+		return; // cannot initialize unless we have a retargeter with BOTH source AND target IK Rigs
+	}
+
+	// initialize retargeter with source and target skeletal meshes
+	// (anim instance is passed in as outer UObject for new UIKRigProcessor) 
+	CurrentlyUsedRetargeter->Initialize(
+		CurrentlyUsedSourceMesh.Get(),
+		CurrentlyUsedTargetMesh.Get(),
+		const_cast<UAnimInstance*>(InAnimInstance));
+	
+	// made it!
+	bIsInitialized = true;
+}
+
+void FAnimNode_IKRetargeter::CopyBoneTransformsFromSource(USkeletalMeshComponent* TargetMeshComponent)
+{
+	USkeletalMeshComponent* SourceMeshComp = CurrentlyUsedSourceMeshComponent.IsValid() ? CurrentlyUsedSourceMeshComponent.Get() : nullptr;
+
+	// is the source mesh ticking?
+	if (!SourceMeshComp->IsRegistered())
+	{
+		CurrentlyUsedSourceMesh.Reset(); // forces reinitialization when re-registered
+		return; // skip copying pose when component is no longer ticking
+	}
+	
+	// if our source is running under master-pose, then get bone data from there
+	if(USkeletalMeshComponent* MasterPoseComponent = Cast<USkeletalMeshComponent>(SourceMeshComponent->MasterPoseComponent.Get()))
+	{
+		SourceMeshComp = MasterPoseComponent;
+	}
+
+	// re-check mesh component validity as it may have changed to master
+	if (!(SourceMeshComp->SkeletalMesh && SourceMeshComp->IsRegistered()))
+	{
+		return; // master pose either missing skeletal mesh reference or not ticking, either way, we aren't copying from it
+	}
+	
+	const bool bUROInSync = SourceMeshComp->ShouldUseUpdateRateOptimizations() && SourceMeshComp->AnimUpdateRateParams != nullptr && SourceMeshComponent->AnimUpdateRateParams == TargetMeshComponent->AnimUpdateRateParams;
+	const bool bUsingExternalInterpolation = SourceMeshComp->IsUsingExternalInterpolation();
+	const TArray<FTransform>& CachedComponentSpaceTransforms = SourceMeshComp->GetCachedComponentSpaceTransforms();
+	const bool bArraySizesMatch = CachedComponentSpaceTransforms.Num() == SourceMeshComp->GetComponentSpaceTransforms().Num();
+
+	// copy source array from the appropriate location
+	SourceMeshComponentSpaceBoneTransforms.Reset();
+	if ((bUROInSync || bUsingExternalInterpolation) && bArraySizesMatch)
+	{
+		SourceMeshComponentSpaceBoneTransforms.Append(CachedComponentSpaceTransforms); // copy from source's cache
+	}
+	else
+	{
+		SourceMeshComponentSpaceBoneTransforms.Append(SourceMeshComp->GetComponentSpaceTransforms()); // copy directly
+	}
+
+	// ref skeleton is need for parent index lookups later, so store it now
+	CurrentlyUsedSourceMesh = SourceMeshComp->SkeletalMesh;
+}
diff --git a/Engine/Plugins/Animation/IKRig/Source/IKRig/Private/AnimNodes/AnimNode_IKRig.cpp b/Engine/Plugins/Animation/IKRig/Source/IKRig/Private/AnimNodes/AnimNode_IKRig.cpp
--- a/Engine/Plugins/Animation/IKRig/Source/IKRig/Private/AnimNodes/AnimNode_IKRig.cpp
+++ b/Engine/Plugins/Animation/IKRig/Source/IKRig/Private/AnimNodes/AnimNode_IKRig.cpp
@@ -3,9 +3,9 @@
 #include "AnimNodes/AnimNode_IKRig.h"
 #include "IKRigProcessor.h"
 #include "IKRigDefinition.h"
+#include "Drawing/ControlRigDrawInterface.h"
 #include "Animation/AnimInstanceProxy.h"
-/////////////////////////////////////////////////////
-// AnimNode_IKRig
+
 
 FAnimNode_IKRig::FAnimNode_IKRig()
 	: RigDefinitionAsset (nullptr)
@@ -29,32 +29,34 @@ void FAnimNode_IKRig::Evaluate_AnyThread(FPoseContext& Output)
 		SourcePose.ResetToRefPose();
 	}
 
-	if (!RigProcessor)
+	if (!IKRigProcessor && IKRigProcessor->IsInitialized())
 	{
 		Output = SourcePose;
 		return;
 	}
 
+	FCompactPose& OutPose = SourcePose.Pose;
+	FIKRigSkeleton& IKRigSkeleton = IKRigProcessor->GetSkeleton();
+	
 	if (bStartFromRefPose)
 	{
-		RigProcessor->ResetToRefPose();
-	}
-		
-	FCompactPose& OutPose = SourcePose.Pose;
-	FIKRigTransforms& Transforms = RigProcessor->GetCurrentGlobalTransforms();
-
-	// copy input pose
-	for (FCompactPoseBoneIndex CPIndex : OutPose.ForEachBoneIndex())
+		// start Solve() from REFERENCE pose
+		IKRigProcessor->SetInputPoseToRefPose();
+	}else
 	{
-		int32* Index = CompactPoseToRigIndices.Find(CPIndex);
-		if (Index)
+		// start Solve() from INPUT pose
+
+		// copy local bone transforms into IKRigProcessor skeleton
+		for (FCompactPoseBoneIndex CPIndex : OutPose.ForEachBoneIndex())
 		{
-			// Todo: this is again slow. 
-			// the tricky part is that if we go start from child to parent
-			// we may still missing joints that don't exists
-			// so this is slow and for safety but needs revisit
-			Transforms.SetLocalTransform(*Index, OutPose[CPIndex], true);
+			int32* Index = CompactPoseToRigIndices.Find(CPIndex);
+			if (Index)
+			{
+				IKRigSkeleton.CurrentPoseLocal[*Index] = OutPose[CPIndex];
+			}
 		}
+		// update global pose in IK Rig
+		IKRigSkeleton.UpdateAllGlobalTransformFromLocal();
 	}
 
 	// update goal transforms before solve
@@ -62,23 +64,26 @@ void FAnimNode_IKRig::Evaluate_AnyThread(FPoseContext& Output)
 	{
 		if (ensure(GoalTransforms.IsValidIndex(GoalIndex)))
 		{
-			RigProcessor->SetGoalTransform(
+			IKRigProcessor->SetGoalTransform(
 				GoalNames[GoalIndex], 
 				GoalTransforms[GoalIndex].GetLocation(),
 				GoalTransforms[GoalIndex].GetRotation());
 		}
 	}
 
-	// run stack of solvers, generate new pose
-	RigProcessor->Solve();
+	// run stack of solvers, updates global transforms with new pose
+	IKRigProcessor->Solve();
+
+	// update local transforms of current IKRig pose
+	IKRigSkeleton.UpdateAllLocalTransformFromGlobal();
 
-	// copy output pose
+	// copy local transforms to output pose
 	for (FCompactPoseBoneIndex CPIndex : OutPose.ForEachBoneIndex())
 	{
 		int32* Index = CompactPoseToRigIndices.Find(CPIndex);
 		if (Index)
 		{
-			Output.Pose[CPIndex] = Transforms.GetLocalTransform(*Index);
+			Output.Pose[CPIndex] = IKRigSkeleton.CurrentPoseLocal[*Index];
 		}
 	}
 
@@ -98,42 +103,43 @@ void FAnimNode_IKRig::GatherDebugData(FNodeDebugData& DebugData)
 void FAnimNode_IKRig::Initialize_AnyThread(const FAnimationInitializeContext& Context)
 {
 	DECLARE_SCOPE_HIERARCHICAL_COUNTER_FUNC()
-
 	FAnimNode_Base::Initialize_AnyThread(Context);
 	Source.Initialize(Context);
-
-	if (RigProcessor)
+	if (!IKRigProcessor)
 	{
-		RigProcessor->Initialize(RigDefinitionAsset);
-		GoalNames.Reset();
-		RigProcessor->GetGoalNames(GoalNames);
+		return;
 	}
+	
+	IKRigProcessor->Initialize(RigDefinitionAsset);
+	GoalNames.Reset();
+	IKRigProcessor->GetGoalNames(GoalNames);
 }
 
 void FAnimNode_IKRig::OnInitializeAnimInstance(const FAnimInstanceProxy* InProxy, const UAnimInstance* InAnimInstance)
 {
 	DECLARE_SCOPE_HIERARCHICAL_COUNTER_FUNC()
-
-	if (RigDefinitionAsset && !RigProcessor)
+	FAnimNode_Base::OnInitializeAnimInstance(InProxy, InAnimInstance);
+	
+	if (RigDefinitionAsset && !IKRigProcessor)
 	{
-		RigProcessor = NewObject<UIKRigProcessor>(InAnimInstance->GetOwningComponent());
+		IKRigProcessor = UIKRigProcessor::MakeNewIKRigProcessor(InAnimInstance->GetOwningComponent());
 	}
-
-	FAnimNode_Base::OnInitializeAnimInstance(InProxy, InAnimInstance);
 }
 
 bool FAnimNode_IKRig::RebuildGoalList()
 {
-	if (RigDefinitionAsset)
+	if (!RigDefinitionAsset)
 	{
-		TArray<FName> SolverGoals;
-		RigDefinitionAsset->GetGoalNamesFromSolvers(SolverGoals);
-		const int32 GoalNum = SolverGoals.Num();
-		if (GoalTransforms.Num() != GoalNum)
-		{
-			GoalTransforms.SetNum(GoalNum);
-			return true;
-		}
+		return false;
+	}
+	
+	TArray<FName> SolverGoals;
+	RigDefinitionAsset->GetGoalNamesFromSolvers(SolverGoals);
+	const int32 GoalNum = SolverGoals.Num();
+	if (GoalTransforms.Num() != GoalNum)
+	{
+		GoalTransforms.SetNum(GoalNum);
+		return true;
 	}
 
 	return false;
@@ -157,7 +163,6 @@ FName FAnimNode_IKRig::GetGoalName(int32 Index) const
 void FAnimNode_IKRig::Update_AnyThread(const FAnimationUpdateContext& Context)
 {
 	GetEvaluateGraphExposedInputs().Execute(Context);
-
 	FAnimNode_Base::Update_AnyThread(Context);
 	Source.Update(Context);
 }
@@ -166,48 +171,38 @@ void FAnimNode_IKRig::CacheBones_AnyThread(const FAnimationCacheBonesContext& Co
 {
 	FAnimNode_Base::CacheBones_AnyThread(Context);
 	Source.CacheBones(Context);
-
-	if (RigProcessor)
+	
+	const FBoneContainer& RequiredBones = Context.AnimInstanceProxy->GetRequiredBones();
+	if (!RequiredBones.IsValid())
 	{
-		// fill up node names
-		const FBoneContainer& RequiredBones = Context.AnimInstanceProxy->GetRequiredBones();
-
-		CompactPoseToRigIndices.Reset();
+		return;
+	}
 
-		if (RequiredBones.IsValid())
+	// fill up node names
+	CompactPoseToRigIndices.Reset();
+	const TArray<FBoneIndexType>& RequiredBonesArray = RequiredBones.GetBoneIndicesArray();
+	const FReferenceSkeleton& RefSkeleton = RequiredBones.GetReferenceSkeleton();
+	const int32 NumBones = RequiredBonesArray.Num();
+	for (uint16 Index = 0; Index < NumBones; ++Index)
+	{
+		const int32 MeshBone = RequiredBonesArray[Index];
+		if (!ensure(MeshBone != INDEX_NONE))
 		{
-			const TArray<FBoneIndexType>& RequiredBonesArray = RequiredBones.GetBoneIndicesArray();
-			const int32 NumBones = RequiredBonesArray.Num();
-
-			const FReferenceSkeleton& RefSkeleton = RequiredBones.GetReferenceSkeleton();
-			const FIKRigHierarchy* Hierarchy = RigProcessor->GetHierarchy();
-
-			if (Hierarchy)
-			{
-
-				// even if not mapped, we map only node that exists in the controlrig
-				for (uint16 Index = 0; Index < NumBones; ++Index)
-				{
-					// get the name
-					int32 MeshBone = RequiredBonesArray[Index];
-					if (ensure(MeshBone != INDEX_NONE))
-					{
-						FCompactPoseBoneIndex CPIndex = RequiredBones.MakeCompactPoseIndex(FMeshPoseBoneIndex(MeshBone));
-						FName Name = RequiredBones.GetReferenceSkeleton().GetBoneName(MeshBone);
-						CompactPoseToRigIndices.Add(CPIndex) = Hierarchy->GetIndex(Name);
-					}
-				}
-			}
+			continue;
 		}
+		
+		FCompactPoseBoneIndex CPIndex = RequiredBones.MakeCompactPoseIndex(FMeshPoseBoneIndex(MeshBone));
+		const FName Name = RefSkeleton.GetBoneName(MeshBone);
+		CompactPoseToRigIndices.Add(CPIndex) = RigDefinitionAsset->Skeleton.GetBoneIndexFromName(Name);
 	}
 }
 
 void FAnimNode_IKRig::QueueDrawInterface(FAnimInstanceProxy* AnimProxy, const FTransform& ComponentToWorld)
 {
-	check (RigProcessor);
+	check (IKRigProcessor);
 	check (AnimProxy);
 
-	for (const FControlRigDrawInstruction& Instruction : RigProcessor->GetDrawInterface())
+	for (const FControlRigDrawInstruction& Instruction : IKRigProcessor->GetDrawInterface())
 	{
 		if (!Instruction.IsValid())
 		{
diff --git a/Engine/Plugins/Animation/IKRig/Source/IKRig/Private/IKRigDataTypes.cpp b/Engine/Plugins/Animation/IKRig/Source/IKRig/Private/IKRigDataTypes.cpp
--- a/Engine/Plugins/Animation/IKRig/Source/IKRig/Private/IKRigDataTypes.cpp
+++ b/Engine/Plugins/Animation/IKRig/Source/IKRig/Private/IKRigDataTypes.cpp
@@ -1,171 +1,3 @@
 // Copyright Epic Games, Inc. All Rights Reserved.
 
-#include "IKRigDataTypes.h"
-#include "IKRigHierarchy.h"
-
-FIKRigTransforms::FIKRigTransforms(const FIKRigHierarchy* InHierarchy)
-	: Hierarchy(InHierarchy)
-{
-	ensure(Hierarchy != nullptr);
-}
-
-// we don't want to allocate local transforms until required
-// so we don't do this in the constructor but on demand when to be used
-void FIKRigTransforms::EnsureLocalTransformsExist()
-{
-	// first check validation
-	if (Hierarchy&& Hierarchy->GetNum() == GlobalTransforms.Num())
-	{
-		if (LocalTransforms.Num() != GlobalTransforms.Num())
-		{
-			LocalTransforms.SetNumUninitialized(GlobalTransforms.Num());
-			LocalTransformDirtyFlags.Init(true, GlobalTransforms.Num());
-		}
-	}
-}
-
-void FIKRigTransforms::RecalculateLocalTransform()
-{
-	// make sure to call EnsureLocalTransformsExist before getting here
-	check (Hierarchy && LocalTransforms.Num() == GlobalTransforms.Num());
-	
-	for (TConstSetBitIterator<> It(LocalTransformDirtyFlags); It; ++It)
-	{
-		const int32 BoneIndex = It.GetIndex();
-		const int32 ParentIndex = Hierarchy->GetParentIndex(BoneIndex);
-		LocalTransforms[BoneIndex] = GetRelativeTransform(BoneIndex, ParentIndex);
-	}
-
-	// clear the bits
-	LocalTransformDirtyFlags.Init(false, GlobalTransforms.Num());
-}
-
-void FIKRigTransforms::UpdateLocalTransform(int32 Index)
-{
-	EnsureLocalTransformsExist();
-
-	if (Hierarchy && Hierarchy->IsValidIndex(Index) && LocalTransformDirtyFlags[Index])
-	{
-		int32 ParentIndex = Hierarchy->GetParentIndex(Index);
-		LocalTransforms[Index] = GetRelativeTransform(Index, ParentIndex);
-		LocalTransformDirtyFlags[Index] = false;
-	}
-}
-
-void FIKRigTransforms::SetGlobalTransform(int32 Index, const FTransform& InTransform, bool bPropagate)
-{
-	if (GlobalTransforms.IsValidIndex(Index))
-	{
-		if (bPropagate)
-		{
-			// update all children
-			SetGlobalTransform_Internal(Index, InTransform);
-		}
-		else
-		{
-			GlobalTransforms[Index] = InTransform;
-
-			// if we have local transform, we also update
-			if (LocalTransforms.IsValidIndex(Index))
-			{
-				LocalTransformDirtyFlags[Index] = true;
-				UpdateLocalTransform(Index);
-			}
-		}
-	}
-}
-
-void FIKRigTransforms::SetLocalTransform(int32 Index, const FTransform& InTransform, bool bPropagate)
-{
-	// we go through global 
-	// because we can't just update local only
-	if (Hierarchy && Hierarchy->IsValidIndex(Index))
-	{
-		const int32 ParentIndex = Hierarchy->GetParentIndex(Index);
-		// if parent is invalid, it will get identity
-		FTransform NewTransform = InTransform * GetGlobalTransform(ParentIndex);
-		SetGlobalTransform(Index, NewTransform, true);
-	}
-}
-
-const FTransform& FIKRigTransforms::GetLocalTransform(int32 Index) const
-{
-	if (Hierarchy && Hierarchy->IsValidIndex(Index))
-	{
-		// this is because this requires to modify LocalTransforms
-		const_cast<FIKRigTransforms*>(this)->UpdateLocalTransform(Index);
-		return LocalTransforms[Index];
-	}
-
-	return FTransform::Identity;
-}
-
-const FTransform& FIKRigTransforms::GetGlobalTransform(int32 Index) const
-{
-	if (GlobalTransforms.IsValidIndex(Index))
-	{
-		return GlobalTransforms[Index];
-	}
-
-	return FTransform::Identity;
-}
-/* This function does propagate through children
-	* of the current index and keeps last LocalTransform to up to date
-	*/
-void FIKRigTransforms::SetGlobalTransform_Internal(int32 Index, const FTransform& InTransform)
-{
-	// first calculate local transform if not found
-	EnsureLocalTransformsExist();
-	RecalculateLocalTransform();
-
-	// set global transform recursively without modifying Local Transform
-	SetGlobalTransform_Recursive(Index, InTransform);
-
-	// now update my local transform
-	LocalTransformDirtyFlags[Index] = true;
-	UpdateLocalTransform(Index);
-}
-
-/** Set Global Transform Recursive
-	*
-	*/
-void FIKRigTransforms::SetGlobalTransform_Recursive(int32 Index, const FTransform& InTransform)
-{
-	GlobalTransforms[Index] = InTransform;
-	TArray<int32> Children = Hierarchy->FindChildren(Index);
-	for (int32 Child : Children)
-	{
-		// calculate new global based on local 
-		FTransform NewTransform = LocalTransforms[Child] * InTransform;
-		SetGlobalTransform_Recursive(Child, NewTransform);
-	}
-}
-
-void FIKRigTransforms::SetAllGlobalTransforms(const TArray<FTransform>& InTransforms)
-{
-	if (Hierarchy && Hierarchy->GetNum() == InTransforms.Num())
-	{
-		GlobalTransforms = InTransforms;
-
-		// we don't want previous transform data
-		LocalTransforms.Reset();
-		LocalTransformDirtyFlags.Reset();
-	}
-}
-
-FTransform FIKRigTransforms::GetRelativeTransform(int32 ChildIndex, int32 ParentIndex) const
-{
-	if (GlobalTransforms.IsValidIndex(ChildIndex))
-	{
-		if (GlobalTransforms.IsValidIndex(ParentIndex))
-		{
-			return (GlobalTransforms[ChildIndex].GetRelativeTransform(GlobalTransforms[ParentIndex]));
-		}
-		else
-		{
-			return GlobalTransforms[ChildIndex];
-		}
-	}
-
-	return FTransform::Identity;
-}
\ No newline at end of file
+#include "IKRigDataTypes.h"
\ No newline at end of file
diff --git a/Engine/Plugins/Animation/IKRig/Source/IKRig/Private/IKRigDefinition.cpp b/Engine/Plugins/Animation/IKRig/Source/IKRig/Private/IKRigDefinition.cpp
--- a/Engine/Plugins/Animation/IKRig/Source/IKRig/Private/IKRigDefinition.cpp
+++ b/Engine/Plugins/Animation/IKRig/Source/IKRig/Private/IKRigDefinition.cpp
@@ -4,6 +4,18 @@
 #include "IKRigSolver.h"
 
 
+FBoneChain* FRetargetDefinition::GetBoneChainByName(FName ChainName)
+{
+	for (FBoneChain& Chain : BoneChains)
+	{
+		if (Chain.ChainName == ChainName)
+		{
+			return &Chain;
+		}
+	}
+	return nullptr;
+}
+
 void UIKRigDefinition::GetGoalNamesFromSolvers(TArray<FName>& OutGoalNames) const
 {
 	for (UIKRigSolver* Solver : Solvers)
@@ -25,170 +37,4 @@ void UIKRigDefinition::GetGoalNamesFromSolvers(TArray<FName>& OutGoalNames) cons
 			}
 		}
 	}
-}
-
-
-#if WITH_EDITOR
-
-// add a bone. Return false if it already has conflict name or parent is not found
-bool UIKRigDefinition::AddBone(const FName& InName, const FName& InParent, const FTransform& InGlobalTransform)
-{
-	// only add when it doesn't exist
-	if (Hierarchy.FindIndexFromBoneArray(InName) == INDEX_NONE)
-	{
-		// if no parent or we want to have parent BEFORE adding this joint
-		if (InParent == NAME_None || Hierarchy.FindIndexFromBoneArray(InParent) != INDEX_NONE)
-		{
-			Hierarchy.Bones.Add(FIKRigBone(InName, InParent));
-			RefPoseTransforms.Add(InGlobalTransform);
-
-			check (Hierarchy.Bones.Num() == RefPoseTransforms.Num());
-
-			Hierarchy.RebuildCacheData();
-			return true;
-		}
-	}
-
-	return false;
-}
-
-// ensure it's sorted from parent to children
-void UIKRigDefinition::EnsureSortedCorrectly(bool bReSortIfNeeded)
-{
-	// we want to ensure it's sorted from parent to child
-	// we do this twice, we use this data to sort, so after sort, we still have to do this again
-	Hierarchy.RebuildCacheData();
-
-	// @TODO: FIX THE CODE, it's confusing
-	bool bNeedReSort = false;
-	const TArray<FIKRigBone>& Bones = Hierarchy.Bones;
-	const TArray<int32>& ParentIndices = Hierarchy.ParentIndices;
-	// first check to see if we have violation
-	for (int32 Index = 0; Index < ParentIndices.Num(); ++Index)
-	{
-		// if my parent is ahead, I need to resort
-		if (ParentIndices[Index] >= Index)
-		{
-			bNeedReSort = true;
-		}
-	}
-
-	// resort if needed is true, then we're good
-	// but if resort is false, we want to ensure it's sorted already, and not have to resort
-	ensureAlways(bReSortIfNeeded || !bNeedReSort);
-
-	if (bNeedReSort)
-	{
-		// if so we create a tree and recreate Bones array
-		struct FBoneTreeNode
-		{
-			FIKRigBone Bone;
-			FBoneTreeNode* Children;
-			FBoneTreeNode* Sibling;
-
-			FBoneTreeNode()
-				: Children(nullptr)
-				, Sibling(nullptr)
-			{}
-
-			FBoneTreeNode(FIKRigBone InBone)
-				: Bone(InBone)
-				, Children(nullptr)
-				, Sibling(nullptr)
-			{}
-
-			void AddToBones(TArray<FIKRigBone>& OutBones)
-			{
-				OutBones.Add(Bone);
-
-				FBoneTreeNode* Iter = Children;
-				while (Iter)
-				{
-					Iter->AddToBones(OutBones);
-					Iter = Iter->Sibling;
-				}
-			}
-
-			void AddChild(FBoneTreeNode* InChild)
-			{
-				if (Children == nullptr)
-				{
-					Children = InChild;
-				}
-				else
-				{
-					FBoneTreeNode* Iter = Children;
-					do
-					{
-						if (Iter->Sibling == nullptr)
-						{
-							Iter->Sibling = InChild;
-							break;
-						}
-
-						Iter = Iter->Sibling;
-					} while (true);
-				}
-			}
-		};
-
-		// the assumption is that first one is the root
-		TArray<FBoneTreeNode> Tree;
-		Tree.AddDefaulted(Bones.Num());
-
-		for (int32 Index = 0; Index < Bones.Num(); ++Index)
-		{
-			Tree[Index].Bone = Bones[Index];
-			if (ParentIndices[Index] != INDEX_NONE)
-			{
-				Tree[ParentIndices[Index]].AddChild(&Tree[Index]);
-			}
-		}
-
-		TArray<FIKRigBone> NewBones;
-		// BFS for tree
-		Tree[0].AddToBones(NewBones);
-
-		// create new transform 
-		check(NewBones.Num() == Bones.Num());
-
-		// we need to move transform of old Bones to new Bones, and their indices have changed
-		// move transform, very slow move
-		TArray<FTransform> NewTransforms;
-		NewTransforms.AddUninitialized(Bones.Num());
-
-		for (int32 Index = 0; Index < Bones.Num(); ++Index)
-		{
-			const FName& Name = Bones[Index].Name;
-			for (int32 Jndex = 0; Jndex < NewBones.Num(); ++Jndex)
-			{
-				if (Name == NewBones[Jndex].Name)
-				{
-					NewTransforms[Jndex] = RefPoseTransforms[Index];
-				}
-			}
-		}
-
-		// override back to it
-		Hierarchy.Bones = NewBones;
-		RefPoseTransforms = NewTransforms;
-		Hierarchy.RebuildCacheData();
-	}
-}
-
-void UIKRigDefinition::ResetHierarchy()
-{
-	Hierarchy.Bones.Reset();
-	Hierarchy.ParentIndices.Reset();
-	Hierarchy.RuntimeNameLookupTable.Reset();
-	RefPoseTransforms.Reset();
-}
-
-#endif // WITH_EDITOR
-
-void UIKRigDefinition::PostLoad()
-{
-	Super::PostLoad();
-
-	Hierarchy.RebuildCacheData();
 }
\ No newline at end of file
diff --git a/Engine/Plugins/Animation/IKRig/Source/IKRig/Private/IKRigHierarchy.cpp b/Engine/Plugins/Animation/IKRig/Source/IKRig/Private/IKRigHierarchy.cpp
deleted file mode 100644
diff --git a/Engine/Plugins/Animation/IKRig/Source/IKRig/Private/IKRigProcessor.cpp b/Engine/Plugins/Animation/IKRig/Source/IKRig/Private/IKRigProcessor.cpp
--- a/Engine/Plugins/Animation/IKRig/Source/IKRig/Private/IKRigProcessor.cpp
+++ b/Engine/Plugins/Animation/IKRig/Source/IKRig/Private/IKRigProcessor.cpp
@@ -3,77 +3,89 @@
 #include "IKRigProcessor.h"
 #include "IKRigDefinition.h"
 #include "IKRigSolver.h"
-#include "Solvers/ConstraintSolver.h"
 
 
+UIKRigProcessor* UIKRigProcessor::MakeNewIKRigProcessor(UObject* Outer)
+{
+	return NewObject<UIKRigProcessor>(Outer);
+}
+
 void UIKRigProcessor::Initialize(UIKRigDefinition* InRigDefinition)
 {
 	bInitialized = false;
 
-	RigDefinition = InRigDefinition;
-	if (!RigDefinition)
+	if (!ensureMsgf(InRigDefinition, TEXT("Trying to initialize IKRigProcessor with a null IKRigDefinition asset.")))
 	{
 		return;
 	}
-	
-	RefPoseGlobalTransforms = RigDefinition->RefPoseTransforms;
-	GlobalBoneTransforms = FIKRigTransforms(&RigDefinition->Hierarchy);
-	GlobalBoneTransforms.SetAllGlobalTransforms(RefPoseGlobalTransforms);
 
+	// copy skeleton data from IKRigDefinition
+	Skeleton = InRigDefinition->Skeleton; // trivial copy assignment operator for POD
+
+	// initialize goal names based on solvers
 	TArray<FName> GoalNames;
-	RigDefinition->GetGoalNamesFromSolvers(GoalNames);
+	InRigDefinition->GetGoalNamesFromSolvers(GoalNames);
 	Goals.InitializeGoalsFromNames(GoalNames);
 
-	TArray<UIKRigSolver*> RigSolvers = RigDefinition->Solvers;
-	const int32 NumSolvers = RigSolvers.Num();
-	Solvers.Reset(NumSolvers);
-	for (int32 Index = 0; Index < NumSolvers; ++Index)
+	// create copies of all the solvers in the IK rig
+	Solvers.Reset(InRigDefinition->Solvers.Num());
+	for (UIKRigSolver* IKRigSolver : InRigDefinition->Solvers)
 	{
-		if (!RigSolvers[Index])
+		if (!IKRigSolver)
 		{
+			// this can happen if asset references deleted IK Solver type
+			// which should only happen during development (if at all)
 			continue;
 		}
 		
-		UIKRigSolver* Solver = DuplicateObject(RigSolvers[Index], this);
-		Solver->Init(GlobalBoneTransforms);
+		UIKRigSolver* Solver = DuplicateObject(IKRigSolver, this);
+		Solver->Initialize(Skeleton);
 		Solvers.Add(Solver);
 	}
 
 	bInitialized = true;
 }
 
-void UIKRigProcessor::Solve()
+void UIKRigProcessor::SetInputPoseGlobal(const TArray<FTransform>& InGlobalBoneTransforms) 
 {
-	if (!ensure(bInitialized))
-	{
-		return;
-	}
-
-	check (RigDefinition);
-
-	DrawInterface.Reset();
-
-	for (int32 Index = 0; Index < Solvers.Num(); ++Index)
-	{
-		Solvers[Index]->Solve(GlobalBoneTransforms, Goals, &DrawInterface);
-	}
+	check(bInitialized);
+	check(InGlobalBoneTransforms.Num() == Skeleton.CurrentPoseGlobal.Num());
+	Skeleton.CurrentPoseGlobal = InGlobalBoneTransforms;
+	Skeleton.UpdateAllLocalTransformFromGlobal();
 }
 
-FIKRigTransforms& UIKRigProcessor::GetCurrentGlobalTransforms() 
+void UIKRigProcessor::SetInputPoseToRefPose()
 {
 	check(bInitialized);
-	return GlobalBoneTransforms;
+	Skeleton.CurrentPoseGlobal = Skeleton.RefPoseGlobal;
+	Skeleton.UpdateAllLocalTransformFromGlobal();
 }
 
 void UIKRigProcessor::SetGoalTransform(
-	const FName& GoalName,
-	const FVector& Position,
-	const FQuat& Rotation)
+    const FName& GoalName,
+    const FVector& Position,
+    const FQuat& Rotation)
 {
 	check(bInitialized);
 	Goals.SetGoalTransform(GoalName, Position, Rotation);
 }
 
+void UIKRigProcessor::Solve()
+{
+	check(bInitialized);
+	DrawInterface.Reset();
+	for (UIKRigSolver* Solver : Solvers)
+	{
+		Solver->Solve(Skeleton, Goals, &DrawInterface);
+	}
+}
+
+void UIKRigProcessor::CopyOutputGlobalPoseToArray(TArray<FTransform>& OutputPoseGlobal) const
+{
+	check(bInitialized);
+	OutputPoseGlobal = Skeleton.CurrentPoseGlobal;
+}
+
 void UIKRigProcessor::GetGoalNames(TArray<FName>& OutGoalNames) const
 {
 	check(bInitialized);
@@ -82,17 +94,12 @@ void UIKRigProcessor::GetGoalNames(TArray<FName>& OutGoalNames) const
 
 int UIKRigProcessor::GetNumGoals() const
 {
+	check(bInitialized);
 	return Goals.GetNumGoals();
 }
 
-const FIKRigHierarchy* UIKRigProcessor::GetHierarchy() const
+FIKRigSkeleton& UIKRigProcessor::GetSkeleton()
 {
 	check(bInitialized);
-	return ensure(RigDefinition) ? &RigDefinition->Hierarchy : nullptr;
+	return Skeleton;
 }
-
-void UIKRigProcessor::ResetToRefPose()
-{
-	check(bInitialized);
-	GlobalBoneTransforms.SetAllGlobalTransforms(RefPoseGlobalTransforms);
-}
\ No newline at end of file
diff --git a/Engine/Plugins/Animation/IKRig/Source/IKRig/Private/IKRigSkeleton.cpp b/Engine/Plugins/Animation/IKRig/Source/IKRig/Private/IKRigSkeleton.cpp
new file mode 100644
--- /dev/null
+++ b/Engine/Plugins/Animation/IKRig/Source/IKRig/Private/IKRigSkeleton.cpp
@@ -0,0 +1,115 @@
+// Copyright Epic Games, Inc. All Rights Reserved.
+
+#include "IKRigSkeleton.h"
+
+void FIKRigSkeleton::Initialize(const FReferenceSkeleton& RefSkeleton)
+{
+	// copy names and parent indices into local storage
+	for (int32 BoneIndex=0; BoneIndex<RefSkeleton.GetNum(); ++BoneIndex)
+	{
+		BoneNames.Add(RefSkeleton.GetBoneName(BoneIndex));
+		ParentIndices.Add(RefSkeleton.GetParentIndex(BoneIndex));	
+	}
+
+	// store copy of reference pose (converted from local to global)
+	const TArray<FTransform>& RefPoseLocal = RefSkeleton.GetRefBonePose();
+	for (int32 BoneIndex=0; BoneIndex<RefPoseLocal.Num(); ++BoneIndex)
+	{
+		const int32 ParentIndex = ParentIndices[BoneIndex];
+		if (ParentIndex == INDEX_NONE)
+		{
+			RefPoseGlobal.Add(RefPoseLocal[BoneIndex]);
+			continue; // root bone is always in local space 
+		}
+
+		const FTransform& ChildLocalTransform  = RefPoseLocal[BoneIndex];
+		const FTransform& ParentGlobalTransform  = RefPoseGlobal[ParentIndex];
+		RefPoseGlobal.Add(ChildLocalTransform * ParentGlobalTransform);
+	}
+
+	// initialize current GLOBAL pose to reference pose
+	CurrentPoseGlobal = RefPoseGlobal;
+
+	// initialize current LOCAL pose
+	UpdateAllLocalTransformFromGlobal();
+}
+
+int32 FIKRigSkeleton::GetBoneIndexFromName(const FName InName) const
+{
+	for (int32 Index=0; Index<BoneNames.Num(); ++Index)
+	{
+		if (InName == BoneNames[Index])
+		{
+			return Index;
+		}
+	}
+
+	return INDEX_NONE;
+}
+
+int32 FIKRigSkeleton::GetParentIndex(const int32 BoneIndex) const
+{
+	if (BoneIndex < 0 || BoneIndex>ParentIndices.Num() || BoneIndex == INDEX_NONE)
+	{
+		return INDEX_NONE;
+	}
+
+	return ParentIndices[BoneIndex];
+}
+
+void FIKRigSkeleton::UpdateAllGlobalTransformFromLocal()
+{
+	// generate GLOBAL transforms based on LOCAL transforms
+	CurrentPoseGlobal = CurrentPoseLocal;
+	for (int32 BoneIndex=0; BoneIndex<CurrentPoseLocal.Num(); ++BoneIndex)
+	{
+		UpdateGlobalTransformFromLocal(BoneIndex);
+	}
+}
+
+void FIKRigSkeleton::UpdateGlobalTransformFromLocal(const int32 BoneIndex)
+{
+	const int32 ParentIndex = ParentIndices[BoneIndex];
+	if (ParentIndex == INDEX_NONE)
+	{
+		CurrentPoseGlobal[BoneIndex] = CurrentPoseLocal[BoneIndex]; // root bone is always in local space
+		return; 
+	}
+
+	const FTransform& ChildLocalTransform  = CurrentPoseLocal[BoneIndex];
+	const FTransform& ParentGlobalTransform  = CurrentPoseGlobal[ParentIndex];
+	CurrentPoseGlobal[BoneIndex] = ChildLocalTransform * ParentGlobalTransform;
+}
+
+void FIKRigSkeleton::UpdateAllLocalTransformFromGlobal()
+{
+	// generate LOCAL transforms based on GLOBAL transforms
+	CurrentPoseLocal = CurrentPoseGlobal;
+	for (int32 BoneIndex=0; BoneIndex<CurrentPoseGlobal.Num(); ++BoneIndex)
+	{
+		UpdateLocalTransformFromGlobal(BoneIndex);
+	}
+}
+
+void FIKRigSkeleton::UpdateLocalTransformFromGlobal(const int32 BoneIndex)
+{
+	const int32 ParentIndex = ParentIndices[BoneIndex];
+	if (ParentIndex == INDEX_NONE)
+	{
+		CurrentPoseLocal[BoneIndex] = CurrentPoseGlobal[BoneIndex]; // root bone is always in local space
+		return; 
+	}
+
+	const FTransform& ChildGlobal  = CurrentPoseGlobal[BoneIndex];
+	const FTransform& ParentGlobal  = CurrentPoseGlobal[ParentIndex];
+	CurrentPoseLocal[BoneIndex] = ChildGlobal.GetRelativeTransform(ParentGlobal);
+}
+
+void FIKRigSkeleton::PropagateGlobalPoseBelowBone(const int32 StartBoneIndex)
+{
+	for (int32 BoneIndex=StartBoneIndex+1; BoneIndex<CurrentPoseGlobal.Num(); ++BoneIndex)
+	{
+		UpdateGlobalTransformFromLocal(BoneIndex);
+		UpdateLocalTransformFromGlobal(BoneIndex);
+	}
+}
diff --git a/Engine/Plugins/Animation/IKRig/Source/IKRig/Private/IKRigSolver.cpp b/Engine/Plugins/Animation/IKRig/Source/IKRig/Private/IKRigSolver.cpp
--- a/Engine/Plugins/Animation/IKRig/Source/IKRig/Private/IKRigSolver.cpp
+++ b/Engine/Plugins/Animation/IKRig/Source/IKRig/Private/IKRigSolver.cpp
@@ -1,23 +1,3 @@
 // Copyright Epic Games, Inc. All Rights Reserved.
 
-#include "IKRigSolver.h"
-#include "IKRigDataTypes.h"
-
-
-/* TODO KIARAN: do we even want to auto rename goals?
- * I think the risk is high of being frustrating for user
-#if WITH_EDITOR
-
-void UIKRigSolver::RenameGoal(const FName& OldName, const FName& NewName)
-{
-	for (auto Iter = EffectorToGoalName.CreateIterator(); Iter; ++Iter)
-	{
-		if (Iter.Value() == OldName)
-		{
-			Iter.Value() = NewName;
-		}
-	}
-}
-
-#endif // WITH_EDITOR
-*/
\ No newline at end of file
+#include "IKRigSolver.h"
\ No newline at end of file
diff --git a/Engine/Plugins/Animation/IKRig/Source/IKRig/Private/Retargeter/IKRetargeter.cpp b/Engine/Plugins/Animation/IKRig/Source/IKRig/Private/Retargeter/IKRetargeter.cpp
new file mode 100644
--- /dev/null
+++ b/Engine/Plugins/Animation/IKRig/Source/IKRig/Private/Retargeter/IKRetargeter.cpp
@@ -0,0 +1,888 @@
+// Copyright Epic Games, Inc. All Rights Reserved.
+
+#include "Retargeter/IKRetargeter.h"
+
+/*
+TODO:
+
+Asset Testing:
+	- add chain for "upperarm_twist_01_l"
+			- will test single-bone chains AND fix twisted upperarm
+	- add finger chains to test fk AFTER ik
+	- export starfish in mannequin defailt pose
+
+IK Rig:
+	- add "alpha" value to goals, set alpha to zero by default
+		- internally LERP effectors from input pose to goal transform by alpha
+	- port PBIK to IK Rig
+
+IK Retargeter:
+	- we need option to override initial poses to be different than reference pose
+	- expose root stride scale factor as parameter
+
+FK Retarget:
+	- sort FK chains root to tip (no guarantee user sets them up this way)
+	
+IK Retarget:
+	
+	- set weight of unused goals to zero (so redundant goals don't screw up solve)
+	- set pole vectors goals on IKRig
+
+DONE:
+- create IKRigSkeleton to merge hierarchy and IKRigTransforms
+- IKRigTransforms to use only global space and provide conversion from local space
+- initialize all the things
+- encode (remove decode)
+
+
+NOTES:
+
+// WE PROBABLY DON'T NEED TO DO SEPARATE PASS OF BEFORE/AFTER IK RETARGETS BECAUSE THE
+// FK RETARGETING SHOULD ONLY AFFECT ROTATION WHICH IS INVARIANT ??
+//
+// ALWAYS RUN FK RETARGET?
+// WE MAY WANT TO ALWAYS RUN FK RETARGET BEFORE IK
+// THAT WAY THE CHAINS WILL BE CLOSE TO GOAL TARGETS!!
+// THIS IS DIFFERENT THAN MAYA BECAUSE WE HAVE FULL BODY IK
+// !!!! ACTUALLY IF WE DO FK UPDATE OF GLOBAL SPACE POSE BEFORE IK WE SHOULD BE FINE AND MUCH FASTER!!!
+	
+*/
+
+void FRetargetSkeleton::InitializeHierarchy(
+	const USkeletalMesh* SkeletalMesh,
+	const TArray<FTransform>& InLocalRetargetPose)
+{
+	// copy names and parent indices into local storage
+	const FReferenceSkeleton& RefSkeleton = SkeletalMesh->GetRefSkeleton();
+	for (int32 BoneIndex=0; BoneIndex<RefSkeleton.GetNum(); ++BoneIndex)
+	{
+		BoneNames.Add(RefSkeleton.GetBoneName(BoneIndex));
+		ParentIndices.Add(RefSkeleton.GetParentIndex(BoneIndex));	
+	}
+
+	// store copy of retarget pose
+	RetargetGlobalPose.Reset(InLocalRetargetPose.Num());
+	
+	// convert local transforms to global
+	for (int32 BoneIndex=0; BoneIndex<InLocalRetargetPose.Num(); ++BoneIndex)
+	{
+		const int32 ParentIndex = ParentIndices[BoneIndex];
+		if (ParentIndex == INDEX_NONE)
+		{
+			RetargetGlobalPose.Add(InLocalRetargetPose[BoneIndex]); // root bone is always in local space
+			continue; 
+		}
+
+		const FTransform& ChildLocalTransform  = InLocalRetargetPose[BoneIndex];
+		const FTransform& ParentGlobalTransform  = RetargetGlobalPose[ParentIndex];
+		RetargetGlobalPose.Add(ChildLocalTransform * ParentGlobalTransform);
+	}
+}
+
+int32 FRetargetSkeleton::FindBoneIndexByName(const FName InName) const
+{
+	for (int32 Index=0; Index<BoneNames.Num(); ++Index)
+	{
+		if (InName == BoneNames[Index])
+		{
+			return Index;
+		}
+	}
+
+	return INDEX_NONE;
+}
+
+int32 FRetargetSkeleton::GetParentIndex(const int32 BoneIndex) const
+{
+	if (BoneIndex < 0 || BoneIndex>ParentIndices.Num() || BoneIndex == INDEX_NONE)
+	{
+		return INDEX_NONE;
+	}
+
+	return ParentIndices[BoneIndex];
+}
+
+void FTargetSkeleton::Initialize(
+	const USkeletalMesh* SkeletalMesh,
+	const TArray<FTransform>& InTargetRetargetLocalPose)
+{
+	InitializeHierarchy(SkeletalMesh, InTargetRetargetLocalPose);
+
+	// initialize storage for output pose (the result of the retargeting)
+	OutputGlobalPose = RetargetGlobalPose;
+	// store local retarget pose so we can update partial skeleton during retarget process
+	RetargetLocalPose = InTargetRetargetLocalPose;
+
+	// make storage for per-bone "IsFK" flag (used for hierarchy updates)
+	IsBoneFK.Init(false, InTargetRetargetLocalPose.Num());
+}
+
+void FTargetSkeleton::SetBoneControlledByFKChain(const int32 BoneIndex, const bool IsFK)
+{
+	check(IsBoneFK.IsValidIndex(BoneIndex));
+	IsBoneFK[BoneIndex] = IsFK;
+}
+
+void FTargetSkeleton::UpdateGlobalTransformsBelowBone(const int32 StartBoneIndex)
+{
+	check(OutputGlobalPose.IsValidIndex(StartBoneIndex));
+	
+	for (int32 BoneIndex=StartBoneIndex+1; BoneIndex<OutputGlobalPose.Num(); ++BoneIndex)
+	{
+		UpdateGlobalTransformOfSingleBone(BoneIndex);
+	}
+}
+
+void FTargetSkeleton::UpdateGlobalTransformsAllNonFKBones()
+{
+	check(IsBoneFK.Num() == OutputGlobalPose.Num());
+	
+	for (int32 BoneIndex=0; BoneIndex<OutputGlobalPose.Num(); ++BoneIndex)
+	{
+		if (!IsBoneFK[BoneIndex])
+		{
+			UpdateGlobalTransformOfSingleBone(BoneIndex);
+		}
+	}
+}
+
+void FTargetSkeleton::UpdateGlobalTransformOfSingleBone(const int32 BoneIndex)
+{
+	const int32 ParentIndex = ParentIndices[BoneIndex];
+	if (ParentIndex == INDEX_NONE)
+	{
+		return; // root always in global space
+	}
+	const FTransform ChildLocalTransform = RetargetLocalPose[BoneIndex];
+	const FTransform ParentGlobalTransform = OutputGlobalPose[ParentIndex];
+	OutputGlobalPose[BoneIndex] = ChildLocalTransform * ParentGlobalTransform;
+}
+
+bool FChainEncoderFK::Initialize(const TArray<int32>& BoneIndices, const TArray<FTransform>& InitialGlobalPose)
+{
+	check(!BoneIndices.IsEmpty());
+
+	// store all the initial bone transforms in the bone chain
+	for (int32 Index=0; Index < BoneIndices.Num(); ++Index)
+	{
+		const int32 BoneIndex = BoneIndices[Index];
+		check(BoneIndex < InitialGlobalPose.Num());
+		InitialBoneTransforms.Emplace(InitialGlobalPose[BoneIndex]);
+	}
+
+	// initialize storage for current bones
+	CurrentBoneTransforms = InitialBoneTransforms;
+
+	// calculate parameter of each bone, normalized by the length of the bone chain
+	return CalculateBoneParameters();
+}
+
+bool FChainEncoderFK::CalculateBoneParameters()
+{
+	// special case, a single-bone chain
+	if (InitialBoneTransforms.Num() == 1)
+	{
+		Params.Add(1.0f);
+		return true;
+	}
+
+	// calculate bone lengths in chain and accumulate total length
+	TArray<float> BoneDistances;
+	float TotalChainLength = 0.0f;
+	BoneDistances.Add(0.0f);
+	for (int32 i=1; i<InitialBoneTransforms.Num(); ++i)
+	{
+		TotalChainLength += (InitialBoneTransforms[i].GetTranslation() - InitialBoneTransforms[i-1].GetTranslation()).Size();
+		BoneDistances.Add(TotalChainLength);
+	}
+
+	// cannot retarget chain if all the bones are sitting directly on each other
+	if (TotalChainLength <= KINDA_SMALL_NUMBER)
+	{
+		UE_LOG(LogTemp, Warning, TEXT("IK Retargeter bone chain length is too small to reliably retarget."));
+		return false;
+	}
+
+	// calc each bone's param along length
+	for (int32 i=0; i<InitialBoneTransforms.Num(); ++i)
+	{
+		Params.Add(BoneDistances[i] / TotalChainLength); 
+	}
+
+	return true;
+}
+
+void FChainEncoderFK::EncodePose(
+	const TArray<int32>& SourceBoneIndices,
+    const TArray<FTransform> &InputGlobalPose)
+{
+	check(SourceBoneIndices.Num() == CurrentBoneTransforms.Num());
+	
+	// copy the input pose for the chain
+	for (int32 ChainIndex=0; ChainIndex<SourceBoneIndices.Num(); ++ChainIndex)
+	{
+		const int32 BoneIndex = SourceBoneIndices[ChainIndex];
+		CurrentBoneTransforms[ChainIndex] = InputGlobalPose[BoneIndex];
+	}
+}
+
+void FChainEncoderFK::DecodePose(
+	const TArray<int32>& TargetBoneIndices,
+    const FChainEncoderFK& SourceChain,
+    TArray<FTransform> &InOutGlobalPose)
+{
+	check(TargetBoneIndices.Num() == CurrentBoneTransforms.Num());
+	check(TargetBoneIndices.Num() == Params.Num());
+	
+	for (int32 ChainIndex=0; ChainIndex<TargetBoneIndices.Num(); ++ChainIndex)
+	{
+		// get the initial and current transform of source chain at param
+		// this is the interpolated transform along the chain
+		const float Param = Params[ChainIndex];
+		FTransform SourceCurrentTransform = GetTransformAtParam(
+			SourceChain.CurrentBoneTransforms,
+			SourceChain.Params,
+			Param);
+		FTransform SourceInitialTransform = GetTransformAtParam(
+			SourceChain.InitialBoneTransforms,
+			SourceChain.Params,
+			Param);
+
+		// the initial transform on the target chain
+		const FTransform& TargetInitialTransform = InitialBoneTransforms[ChainIndex];
+
+		// calculate output ROTATION
+		const FQuat& SourceCurrentRotation = SourceCurrentTransform.GetRotation();
+		const FQuat& SourceInitialRotation = SourceInitialTransform.GetRotation();
+		const FQuat& TargetInitialRotation = TargetInitialTransform.GetRotation();
+		const FQuat RotationDelta = TargetInitialRotation * SourceInitialRotation.Inverse();
+		const FQuat OutRotation = RotationDelta * SourceCurrentRotation;
+
+		// calculate output POSITION
+		const FVector& SourceCurrentPosition = SourceCurrentTransform.GetTranslation();
+		const FVector& SourceInitialPosition = SourceInitialTransform.GetTranslation();
+		const FVector& TargetInitialPosition = TargetInitialTransform.GetTranslation();
+		const FVector OutPosition = SourceCurrentPosition + (TargetInitialPosition - SourceInitialPosition);
+
+		// calculate output SCALE
+		const FVector& SourceCurrentScale = SourceCurrentTransform.GetScale3D();
+		const FVector& SourceInitialScale = SourceInitialTransform.GetScale3D();
+		const FVector& TargetInitialScale = TargetInitialTransform.GetScale3D();
+		const FVector OutScale = SourceCurrentScale + (TargetInitialScale - SourceInitialScale);
+
+		InOutGlobalPose[TargetBoneIndices[ChainIndex]] = FTransform(OutRotation, OutPosition, OutScale);
+	}
+}
+
+FTransform FChainEncoderFK::GetTransformAtParam(
+	const TArray<FTransform>& Transforms,
+	const TArray<float>& InParams,
+	const float Param) const
+{
+	if (Param < KINDA_SMALL_NUMBER)
+	{
+		return Transforms[0];
+	}
+
+	if (Param > 1.0f - KINDA_SMALL_NUMBER)
+	{
+		return Transforms.Last();
+	}
+
+	for (int32 ChainIndex=1; ChainIndex<InParams.Num(); ++ChainIndex)
+	{
+		const float CurrentParam = InParams[ChainIndex];
+		if (CurrentParam <= Param)
+		{
+			continue;
+		}
+		
+		const float PrevParam = InParams[ChainIndex-1];
+		const float PercentBetweenParams = (Param - PrevParam) / (CurrentParam - PrevParam);
+		const FTransform& Prev = Transforms[ChainIndex-1];
+		const FTransform& Next = Transforms[ChainIndex];
+		const FVector Position = FMath::Lerp(Prev.GetTranslation(), Next.GetTranslation(), PercentBetweenParams);
+		const FQuat Rotation = FQuat::FastLerp(Prev.GetRotation(), Next.GetRotation(), PercentBetweenParams);
+		const FVector Scale = FMath::Lerp(Prev.GetScale3D(), Next.GetScale3D(), PercentBetweenParams);
+		
+		return FTransform(Rotation,Position, Scale);
+	}
+
+	checkNoEntry();
+	return FTransform::Identity;
+}
+
+bool FChainEncoderIK::Initialize(
+	const TArray<int32>& BoneIndices,
+	const TArray<FTransform>& InitialGlobalPose)
+{
+	if (BoneIndices.Num() < 3)
+	{
+		UE_LOG(LogTemp, Warning, TEXT("IK Retargeter trying to retarget bone chain with IK but it has less than 3 joints."));
+		return false;
+	}
+	
+	BoneIndexA = BoneIndices[0];
+	BoneIndexB = BoneIndices[1];
+	BoneIndexC = BoneIndices.Last();
+	
+	const FTransform& End = InitialGlobalPose[BoneIndexC];
+	EndPositionOrig = End.GetTranslation();
+	EndRotationOrig = End.GetRotation();
+
+	const FTransform& Start = InitialGlobalPose[BoneIndexA];
+	const float Length = (Start.GetTranslation() - EndPositionOrig).Size();
+
+	if (Length <= KINDA_SMALL_NUMBER)
+	{
+		UE_LOG(LogTemp, Warning, TEXT("IK Retargeter trying to retarget bone chain with IK, but it is zero length!"));
+    	return false;
+	}
+
+	InverseLength = 1.0f / Length;
+	
+	return true;
+}
+
+void FChainEncoderIK::EncodePose(const TArray<FTransform>& InputGlobalPose)
+{
+	FVector A = InputGlobalPose[BoneIndexA].GetTranslation();
+	FVector B = InputGlobalPose[BoneIndexB].GetTranslation();
+	FVector C = InputGlobalPose[BoneIndexC].GetTranslation();
+
+    // get the normalized direction / length of the IK limb (how extended it is as percentage of original length)
+    const FVector AC = C - A;
+	float ACLength;
+	FVector ACDirection;
+	AC.ToDirectionAndLength(ACDirection, ACLength);
+	const float NormalizedLimbLength = ACLength * InverseLength;
+
+	EncodedResult.EndDirectionNormalized = ACDirection * NormalizedLimbLength;
+	EncodedResult.EndRotation = InputGlobalPose[BoneIndexC].GetRotation();
+	EncodedResult.EndRotationOrig = EndRotationOrig;
+	EncodedResult.HeightFromGroundNormalized = (C.Z - EndPositionOrig.Z)  * InverseLength;
+	EncodedResult.PoleVectorDirection = FVector::OneVector; // TBD
+}
+
+bool FChainDecoderIK::Initialize(
+	const TArray<int32>& BoneIndices,
+	const TArray<FTransform> &InitialGlobalPose)
+{
+	if (BoneIndices.Num() < 3)
+	{
+		UE_LOG(LogTemp, Warning, TEXT("IK Retargeter trying to retarget bone chain with IK but it has less than 3 joints."));
+		return false;
+	}
+	
+	BoneIndexA = BoneIndices[0];
+	const FTransform& Last = InitialGlobalPose[BoneIndices.Last()];
+	EndPositionOrig = Last.GetTranslation();
+	EndRotationOrig = Last.GetRotation();
+	Length = (InitialGlobalPose[BoneIndexA].GetTranslation() - Last.GetTranslation()).Size();
+
+	if (Length <= KINDA_SMALL_NUMBER)
+	{
+		UE_LOG(LogTemp, Warning, TEXT("IK Retargeter trying to retarget bone chain with IK, but it is zero length!"));
+		return false;
+	}
+	
+	return true;
+}
+	
+void FChainDecoderIK::DecodePose(
+	const FEncodedIKChain& EncodedChain,
+    const TArray<FTransform>& OutGlobalPose,
+    FDecodedIKChain& OutResults) const
+{
+	// calculate end effector position
+	const FVector Start = OutGlobalPose[BoneIndexA].GetTranslation();
+	OutResults.EndEffectorPosition = Start + EncodedChain.EndDirectionNormalized * Length; // * LimbScale (TBD)
+	// optionally factor in ground height
+	const bool IsGrounded = false;
+	if (IsGrounded) // TBD expose this
+	{
+		OutResults.EndEffectorPosition.Z = (EncodedChain.HeightFromGroundNormalized * Length) + EndPositionOrig.Z;
+	}
+
+	// calculate end effector rotation
+	const FQuat Delta = EndRotationOrig * EncodedChain.EndRotationOrig.Inverse();
+	OutResults.EndEffectorRotation = Delta * EncodedChain.EndRotation;
+
+	// TBD calc pole vector position
+	OutResults.PoleVectorPosition = FVector::OneVector;
+}
+
+bool FRetargetChainPair::Initialize(
+    const FBoneChain& SourceBoneChain,
+    const FBoneChain& TargetBoneChain,
+    const FRetargetSkeleton& SourceSkeleton,
+    const FTargetSkeleton& TargetSkeleton)
+{
+	// validate source bone chain is compatible with source skeletal mesh
+	const bool bIsSourceValid = ValidateBoneChainWithSkeletalMesh(true, SourceBoneChain, SourceSkeleton);
+	if (!bIsSourceValid)
+	{
+		UE_LOG(LogTemp, Warning, TEXT("IK Retargeter source bone chain, '%s', is not compatible with Skeletal Mesh: '%s'"),
+			*SourceBoneChain.ChainName.ToString(), *SourceSkeleton.SkeletalMeshName);
+		return false;
+	}
+
+	// validate target bone chain is compatible with target skeletal mesh
+	const bool bIsTargetValid = ValidateBoneChainWithSkeletalMesh(false, TargetBoneChain, TargetSkeleton);
+	if (!bIsTargetValid)
+    {
+		UE_LOG(LogTemp, Warning, TEXT("IK Retargeter target bone chain, '%s', is not compatible with Skeletal Mesh: '%s'"),
+            *TargetBoneChain.ChainName.ToString(), *TargetSkeleton.SkeletalMeshName);
+		return false;
+    }
+
+	// store attributes of chain
+	SourceBoneChainName = SourceBoneChain.ChainName;
+	TargetBoneChainName = TargetBoneChain.ChainName;
+	
+	return true;
+}
+
+bool FRetargetChainPair::ValidateBoneChainWithSkeletalMesh(
+    const bool IsSource,
+    const FBoneChain& BoneChain,
+    const FRetargetSkeleton& RetargetSkeleton)
+{	
+	// validate start bone exists
+	const int32 StartIndex = RetargetSkeleton.FindBoneIndexByName(BoneChain.StartBone);
+	if (StartIndex == INDEX_NONE)
+	{
+		UE_LOG(LogTemp, Warning, TEXT("IK Retargeter bone chain, %s, could not find start bone, %s in mesh %s"),
+            *BoneChain.ChainName.ToString(), *BoneChain.StartBone.ToString(), *RetargetSkeleton.SkeletalMeshName);
+		return false;
+	}
+	
+	// validate end bone exists
+	const int32 EndIndex = RetargetSkeleton.FindBoneIndexByName(BoneChain.EndBone);
+	if (EndIndex == INDEX_NONE)
+	{
+		UE_LOG(LogTemp, Warning, TEXT("IK Retargeter bone chain, %s, could not find end bone, %s in mesh %s"),
+            *BoneChain.ChainName.ToString(), *BoneChain.EndBone.ToString(), *RetargetSkeleton.SkeletalMeshName);
+		return false;
+	}
+	
+	// validate end bone is child of start bone
+	// record all bones in chain while walking up the hierarchy (tip to root of chain)
+	TArray<int>& BoneIndices = IsSource ? SourceBoneIndices : TargetBoneIndices;
+	BoneIndices.Reset();
+	int32 NextBoneIndex = EndIndex;
+	while (true)
+	{
+		BoneIndices.Add(NextBoneIndex);
+		
+		if (NextBoneIndex == StartIndex)
+		{
+			break;
+		}
+
+		NextBoneIndex = RetargetSkeleton.GetParentIndex(NextBoneIndex);
+		
+		if (NextBoneIndex == 0)
+		{
+			UE_LOG(LogTemp, Warning, TEXT("IK Retargeter bone chain, '%s' has end bone, '%s' that is not a child of start bone, '%s' in mesh '%s'"),
+            *BoneChain.ChainName.ToString(), *BoneChain.EndBone.ToString(), *BoneChain.StartBone.ToString(), *RetargetSkeleton.SkeletalMeshName);
+			return false;
+		}
+	}
+
+	// reverse the indices (we want root to tip order)
+	Algo::Reverse(BoneIndices);
+	
+	return true;
+}
+
+bool FRetargetChainPairFK::Initialize(
+	const FBoneChain& SourceBoneChain,
+	const FBoneChain& TargetBoneChain,
+	const FRetargetSkeleton& SourceSkeleton,
+	const FTargetSkeleton& TargetSkeleton)
+{
+	const bool bChainInitialized = FRetargetChainPair::Initialize(SourceBoneChain, TargetBoneChain, SourceSkeleton, TargetSkeleton);
+	if (!bChainInitialized)
+	{
+		return false;
+	}
+
+	// initialize SOURCE FK chain encoder with retarget pose
+	const bool bFKEncoderInitialized = FKEncoder.Initialize(SourceBoneIndices, SourceSkeleton.RetargetGlobalPose);
+	if (!bFKEncoderInitialized)
+	{
+		UE_LOG(LogTemp, Warning, TEXT("IK Retargeter failed to initialize FK encoder, '%s', on Skeletal Mesh: '%s'"),
+            *SourceBoneChainName.ToString(), *SourceSkeleton.SkeletalMeshName);
+		return false;
+	}
+
+	// initialize TARGET FK chain decoder with retarget pose
+	const bool bFKDecoderInitialized = FKDecoder.Initialize(TargetBoneIndices, TargetSkeleton.RetargetGlobalPose);
+	if (!bFKDecoderInitialized)
+	{
+		UE_LOG(LogTemp, Warning, TEXT("IK Retargeter failed to initialize FK decoder, '%s', on Skeletal Mesh: '%s'"),
+            *TargetBoneChainName.ToString(), *TargetSkeleton.SkeletalMeshName);
+		return false;
+	}
+
+	return true;
+}
+
+bool FRetargetChainPairIK::Initialize(
+	const FBoneChain& SourceBoneChain,
+	const FBoneChain& TargetBoneChain,
+	const FRetargetSkeleton& SourceSkeleton,
+	const FTargetSkeleton& TargetSkeleton)
+{
+	const bool bChainInitialized = FRetargetChainPair::Initialize(SourceBoneChain, TargetBoneChain, SourceSkeleton, TargetSkeleton);
+	if (!bChainInitialized)
+	{
+		return false;
+	}
+	
+	IKGoalName = TargetBoneChain.IKGoalName;
+
+	// initialize SOURCE IK chain encoder with retarget pose
+	const bool bIKEncoderInitialized = IKEncoder.Initialize(SourceBoneIndices, SourceSkeleton.RetargetGlobalPose);
+	if (!bIKEncoderInitialized)
+	{
+		UE_LOG(LogTemp, Warning, TEXT("IK Retargeter failed to initialize IK encoder, '%s', on Skeletal Mesh: '%s'"),
+        *SourceBoneChainName.ToString(), *SourceSkeleton.SkeletalMeshName);
+		return false;
+	}
+
+	// initialize TARGET IK chain decoder with retarget pose
+	const bool bIKDecoderInitialized = IKDecoder.Initialize(TargetBoneIndices, TargetSkeleton.RetargetGlobalPose);
+	if (!bIKDecoderInitialized)
+	{
+		UE_LOG(LogTemp, Warning, TEXT("IK Retargeter failed to initialize IK decoder, '%s', on Skeletal Mesh: '%s'"),
+        *TargetBoneChainName.ToString(), *TargetSkeleton.SkeletalMeshName);
+		return false;
+	}
+
+	return true;
+}
+
+bool FRootEncoder::Initialize(const FName RootBoneName, const FRetargetSkeleton& SourceSkeleton)
+{
+	// validate root bone exists
+	BoneIndex = SourceSkeleton.FindBoneIndexByName(RootBoneName);
+	if (BoneIndex == INDEX_NONE)
+	{
+		UE_LOG(LogTemp, Warning, TEXT("IK Retargeter could not find source root bone, %s in mesh %s"),
+			*RootBoneName.ToString(), *SourceSkeleton.SkeletalMeshName);
+		return false;
+	}
+	
+	// record initial root data
+	const FTransform InitialTransform = SourceSkeleton.RetargetGlobalPose[BoneIndex]; 
+	float InitialHeight = InitialTransform.GetTranslation().Z;
+	EncodedResult.InitialRotation = InitialTransform.GetRotation();
+
+	// ensure root height is not at origin, this happens if user sets root to ACTUAL skeleton root and not pelvis
+	if (InitialHeight < KINDA_SMALL_NUMBER)
+	{
+		// warn user and push it up slightly to avoid divide by zero
+		UE_LOG(LogTemp, Warning, TEXT("IK Retargeter root bone is very near the ground plane. This is probably not correct."));
+		InitialHeight = 1.0f;
+	}
+
+	// invert height
+	InvInitialHeight = 1.0f / InitialHeight;
+
+	return true;
+}
+
+void FRootEncoder::EncodePose(const TArray<FTransform>& InputComponentSpaceBoneTransforms)
+{
+	const FTransform& Transform = InputComponentSpaceBoneTransforms[BoneIndex];
+	EncodedResult.NormalizedPosition = Transform.GetTranslation() * InvInitialHeight;
+	EncodedResult.Rotation = Transform.GetRotation();	
+}
+
+bool FRootDecoder::Initialize(const FName RootBoneName, const FTargetSkeleton& TargetSkeleton)
+{
+	// validate root bone exists
+	BoneIndex = TargetSkeleton.FindBoneIndexByName(RootBoneName);
+	if (BoneIndex == INDEX_NONE)
+	{
+		UE_LOG(LogTemp, Warning, TEXT("IK Retargeter could not find target root bone, %s in mesh %s"),
+            *RootBoneName.ToString(), *TargetSkeleton.SkeletalMeshName);
+		return false;
+	}
+
+	const FTransform InitialTransform = TargetSkeleton.RetargetGlobalPose[BoneIndex];
+	InitialHeight = InitialTransform.GetTranslation().Z;
+	InitialRotation = InitialTransform.GetRotation();
+
+	return true;
+}
+
+void FRootDecoder::DecodePose(
+	const FEncodedRoot& EncodedResult,
+	TArray<FTransform>& InOutTargetBoneTransforms,
+	const float StrideScale) const
+{
+	// scale normalized position by root height
+	FVector Position = EncodedResult.NormalizedPosition * InitialHeight;
+	// scale horizontal displacement by stride scale factor
+	Position.X *= StrideScale;
+	Position.Y *= StrideScale;
+	// apply position to target
+	InOutTargetBoneTransforms[BoneIndex].SetTranslation(Position);
+
+	// calc offset between initial source/target root rotations
+	const FQuat RotationOffset = InitialRotation * EncodedResult.InitialRotation.Inverse();
+	// add offset to the current source rotation
+	const FQuat Rotation = RotationOffset * EncodedResult.Rotation;
+	// apply rotation to target
+	InOutTargetBoneTransforms[BoneIndex].SetRotation(Rotation);	
+}
+
+UIKRetargeter::UIKRetargeter()
+	: SourceIKRigAsset(nullptr),
+    TargetIKRigAsset(nullptr),
+    bIsLoadedAndValid(false)
+{
+}
+
+void UIKRetargeter::Initialize(USkeletalMesh* SourceSkeletalMesh, USkeletalMesh* TargetSkeletalMesh, UObject* Outer)
+{
+	bIsLoadedAndValid = false;
+
+	check(SourceSkeletalMesh && TargetSkeletalMesh);
+	
+	// initialize skeleton data for source and target
+	const TArray<FTransform>& SourceRetargetPoseLocal = SourceSkeletalMesh->GetRefSkeleton().GetRefBonePose();
+	const TArray<FTransform>& TargetRetargetPoseLocal = TargetSkeletalMesh->GetRefSkeleton().GetRefBonePose();
+	SourceSkeleton.InitializeHierarchy(SourceSkeletalMesh, SourceRetargetPoseLocal);
+	TargetSkeleton.Initialize(TargetSkeletalMesh, TargetRetargetPoseLocal);
+
+	// initialize roots
+	const bool bRootsInitialized = InitializeRoots();
+	if (!bRootsInitialized)
+	{
+		// couldn't match up any BoneChain pairs, no retargeting possible
+		UE_LOG(LogTemp, Warning, TEXT("IK Retargeter unable to initialize one or more root bones on source, %s and target, %s"),
+            *SourceSkeleton.SkeletalMeshName, *TargetSkeleton.SkeletalMeshName);
+		return;
+	}
+
+	// initialize pairs of bone chains
+	const bool bAtLeastOneValidBoneChainPair = InitializeBoneChainPairs();
+	if (!bAtLeastOneValidBoneChainPair)
+	{
+		// couldn't match up any BoneChain pairs, no retargeting possible
+		UE_LOG(LogTemp, Warning, TEXT("IK Retargeter unable to create any Bone Chain pairs between source, %s and target, %s"),
+            *SourceSkeleton.SkeletalMeshName, *TargetSkeleton.SkeletalMeshName);
+		return;
+	}
+
+	// initialize the IKRigProcessor for doing IK decoding
+	const bool bIKRigInitialized = InitializeIKRig(Outer);
+	if (!bIKRigInitialized)
+	{
+		// couldn't initialize the IK Rig
+		UE_LOG(LogTemp, Warning, TEXT("IK Retargeter unable to initialize IK Rig for %s. See output for details."),
+			*TargetSkeleton.SkeletalMeshName);
+		return;
+	}
+
+	bIsLoadedAndValid = true;
+}
+
+bool UIKRetargeter::InitializeRoots()
+{
+	check(SourceIKRigAsset && TargetIKRigAsset);
+
+	// initialize root encoder
+	const FName SourceRootBoneName = SourceIKRigAsset->RetargetDefinition.RootBone;
+	const bool bRootEncoderInit = RootEncoder.Initialize(SourceRootBoneName, SourceSkeleton);
+	if (!bRootEncoderInit)
+	{
+		UE_LOG(LogTemp, Warning, TEXT("IK Retargeter unable to initialize source root, '%s' on skeletal mesh: '%s'"),
+            *SourceRootBoneName.ToString(), *SourceSkeleton.SkeletalMeshName);
+		return false;
+	}
+
+	// initialize root decoder
+	const FName TargetRootBoneName = SourceIKRigAsset->RetargetDefinition.RootBone;
+	const bool bRootDecoderInit = RootDecoder.Initialize(TargetRootBoneName, TargetSkeleton);
+	if (!bRootDecoderInit)
+	{
+		UE_LOG(LogTemp, Warning, TEXT("IK Retargeter unable to initialize target root, '%s' on skeletal mesh: '%s'"),
+            *TargetRootBoneName.ToString(), *TargetSkeleton.SkeletalMeshName);
+		return false;
+	}
+
+	return true;
+}
+
+bool UIKRetargeter::InitializeBoneChainPairs()
+{
+	check(SourceIKRigAsset && TargetIKRigAsset);
+	
+	FRetargetDefinition& SourceRetargetDef = SourceIKRigAsset->RetargetDefinition;
+	FRetargetDefinition& TargetRetargetDef = TargetIKRigAsset->RetargetDefinition;
+
+	for (const TPair<FName, FName>& NamePair : ChainMapping)
+	{
+		// get source bone chain
+		FBoneChain* SourceBoneChain = SourceRetargetDef.GetBoneChainByName(NamePair.Key);
+		if (!SourceBoneChain)
+		{
+			UE_LOG(LogTemp, Warning, TEXT("IK Retargeter missing source bone chain: %s"), *NamePair.Key.ToString());
+			continue;
+		}
+
+		// get target bone chain
+		FBoneChain* TargetBoneChain = TargetRetargetDef.GetBoneChainByName(NamePair.Value);
+		if (!TargetBoneChain)
+		{
+			UE_LOG(LogTemp, Warning, TEXT("IK Retargeter missing target bone chain: %s"), *NamePair.Value.ToString());
+			continue;
+		}
+
+		// all chains are loaded as FK (giving IK better starting pose)
+		FRetargetChainPairFK ChainPair;
+		const bool bChainPairValid = ChainPair.Initialize(*SourceBoneChain, *TargetBoneChain, SourceSkeleton, TargetSkeleton);
+		if (bChainPairValid)
+		{
+			ChainPairsFK.Add(ChainPair);
+		}
+		
+		// load IK chain
+		if (TargetBoneChain->bUseIK)
+		{
+			FRetargetChainPairIK ChainPairIK;
+			const bool bChainPairIKValid = ChainPairIK.Initialize(*SourceBoneChain, *TargetBoneChain, SourceSkeleton, TargetSkeleton);
+			if (bChainPairIKValid)
+			{
+				ChainPairsIK.Add(ChainPairIK);
+			}
+		}
+	}
+
+	// record which bones in the target skeleton are controlled by FK
+	for (const FRetargetChainPairFK& FKChainPair : ChainPairsFK)
+	{
+		for (const int32 BoneIndex : FKChainPair.TargetBoneIndices)
+		{
+			TargetSkeleton.SetBoneControlledByFKChain(BoneIndex, true);
+		}
+	}
+
+	// root is updated before IK as well
+	TargetSkeleton.SetBoneControlledByFKChain(RootDecoder.BoneIndex, true);
+
+	// return true if at least 1 pair of bone chains were initialized
+	return !(ChainPairsIK.IsEmpty() && ChainPairsFK.IsEmpty());
+}
+
+bool UIKRetargeter::InitializeIKRig(UObject* Outer)
+{
+	// make a new IKRigProcessor if we haven't already
+	if (!IKRig)
+	{
+		IKRig = UIKRigProcessor::MakeNewIKRigProcessor(Outer);
+	}
+	
+	// initialize IK Rig runtime processor
+	IKRig->Initialize(TargetIKRigAsset);
+	if (!IKRig->IsInitialized())
+	{
+		return false;
+	}
+
+	// validate that all IK bone chains have an associated Goal
+	TArray<FName> GoalNames;
+	IKRig->GetGoalNames(GoalNames);
+	for (FRetargetChainPairIK& ChainPair : ChainPairsIK)
+	{
+		// does the IK rig have the IK goal this bone chain requires?
+		if (!GoalNames.Contains(ChainPair.IKGoalName))
+		{
+			UE_LOG(LogTemp, Warning, TEXT("IK Retargeter has target bone chain, %s that references an IK Goal, %s that is not present in any of the solvers in the IK Rig asset."),
+            *ChainPair.TargetBoneChainName.ToString(), *ChainPair.IKGoalName.ToString());
+			return false;
+		}
+	}
+	
+	return true;
+}
+
+void UIKRetargeter::RunRetargeter(const TArray<FTransform>& InSourceGlobalPose, bool bEnableIK)
+{
+	check(bIsLoadedAndValid);
+
+	// ROOT retargeting
+	RunRootRetarget(InSourceGlobalPose, TargetSkeleton.OutputGlobalPose);
+
+	// update global transforms below root
+	TargetSkeleton.UpdateGlobalTransformsBelowBone(RootDecoder.BoneIndex);
+
+	// FK CHAIN retargeting
+	RunFKRetarget(InSourceGlobalPose,TargetSkeleton.OutputGlobalPose, true);
+
+	// update all the bones that are not controlled by FK chains or root
+	TargetSkeleton.UpdateGlobalTransformsAllNonFKBones();
+	
+	// IK CHAIN retargeting
+	if (bEnableIK)
+	{
+		RunIKRetarget(InSourceGlobalPose, TargetSkeleton.OutputGlobalPose);
+	}	
+}
+
+
+void UIKRetargeter::RunRootRetarget(
+	const TArray<FTransform>& InGlobalTransforms,
+    TArray<FTransform>& OutGlobalTransforms)
+{
+	RootEncoder.EncodePose(InGlobalTransforms);
+	const float StrideScale = 1.0f;
+	RootDecoder.DecodePose(RootEncoder.EncodedResult, OutGlobalTransforms, StrideScale);
+}
+
+void UIKRetargeter::RunFKRetarget(
+	const TArray<FTransform>& InGlobalTransforms,
+    TArray<FTransform>& OutGlobalTransforms,
+    const bool bPreIK)
+{
+	// spin through chains and encode/decode them all using the input pose
+    for (FRetargetChainPairFK& ChainPair : ChainPairsFK)
+    {
+    	ChainPair.FKEncoder.EncodePose(ChainPair.SourceBoneIndices, InGlobalTransforms);
+    	ChainPair.FKDecoder.DecodePose(ChainPair.TargetBoneIndices, ChainPair.FKEncoder,OutGlobalTransforms);
+    }
+}
+
+void UIKRetargeter::RunIKRetarget(
+	const TArray<FTransform>& InGlobalPose,
+    TArray<FTransform>& OutGlobalPose)
+{
+	if (ChainPairsIK.IsEmpty())
+	{
+		return; // skip IK
+	}
+	
+	// spin through IK chains
+	for (FRetargetChainPairIK& ChainPair : ChainPairsIK)
+	{
+		// encode them all using the input pose
+		ChainPair.IKEncoder.EncodePose(InGlobalPose);
+		// decode the IK goal and apply to IKRig
+		FDecodedIKChain OutIKGoal;
+		ChainPair.IKDecoder.DecodePose(ChainPair.IKEncoder.EncodedResult, OutGlobalPose, OutIKGoal);
+		// set the goal transform on the IK Rig
+		IKRig->SetGoalTransform(ChainPair.IKGoalName, OutIKGoal.EndEffectorPosition, OutIKGoal.EndEffectorRotation);
+	}
+
+	// copy input pose to start IK solve from
+	IKRig->SetInputPoseGlobal(OutGlobalPose);
+	// run IK solve
+	IKRig->Solve();
+	// copy results of solve
+	IKRig->CopyOutputGlobalPoseToArray(OutGlobalPose);
+}
diff --git a/Engine/Plugins/Animation/IKRig/Source/IKRig/Private/Solvers/ConstraintSolver.cpp b/Engine/Plugins/Animation/IKRig/Source/IKRig/Private/Solvers/ConstraintSolver.cpp
--- a/Engine/Plugins/Animation/IKRig/Source/IKRig/Private/Solvers/ConstraintSolver.cpp
+++ b/Engine/Plugins/Animation/IKRig/Source/IKRig/Private/Solvers/ConstraintSolver.cpp
@@ -3,13 +3,13 @@
 #include "Solvers/ConstraintSolver.h"
 #include "IKRigBoneSetting.h"
 
-void UIKRigConstraintSolver::Init(const FIKRigTransforms& InGlobalTransform)
+void UIKRigConstraintSolver::Initialize(const FIKRigSkeleton& InSkeleton)
 {
 	// TODO get constraints that this solver can deal with
 }
 
 void UIKRigConstraintSolver::Solve(
-	FIKRigTransforms& InOutGlobalTransform,
+	FIKRigSkeleton& IKRigSkeleton,
 	const FIKRigGoalContainer& Goals,
 	FControlRigDrawInterface* InOutDrawInterface)
 {
diff --git a/Engine/Plugins/Animation/IKRig/Source/IKRig/Private/Solvers/IKRigFBIKSolver.cpp b/Engine/Plugins/Animation/IKRig/Source/IKRig/Private/Solvers/IKRigFBIKSolver.cpp
--- a/Engine/Plugins/Animation/IKRig/Source/IKRig/Private/Solvers/IKRigFBIKSolver.cpp
+++ b/Engine/Plugins/Animation/IKRig/Source/IKRig/Private/Solvers/IKRigFBIKSolver.cpp
@@ -1,12 +1,7 @@
 // Copyright Epic Games, Inc. All Rights Reserved.
 
-/*=============================================================================
-	FullBodyIKSolver.cpp: Solver execution class for Transform
-=============================================================================*/
-
 #include "Solvers/IKRigFBIKSolver.h"
 #include "IKRigDataTypes.h"
-#include "IKRigHierarchy.h"
 #include "FBIKConstraint.h"
 #include "JacobianIK.h"
 #include "FBIKUtil.h"
@@ -18,7 +13,7 @@ const FName UIKRigFBIKSolver::EffectorTargetPrefix = FName(TEXT("FullBodyIKTarge
 //////////////////////////////////////////////
 // utility functions
 ////////////////////////////////////////////////
-static float EnsureToAddBoneToLinkData(const FIKRigTransforms& TransformModifier, const int32 CurrentItem, TArray<FFBIKLinkData>& LinkData,
+static float EnsureToAddBoneToLinkData(const FIKRigSkeleton& InSkeleton, const int32 CurrentItem, TArray<FFBIKLinkData>& LinkData,
 	TMap<int32, int32>& HierarchyToLinkDataMap, TMap<int32, int32>& LinkDataToHierarchyIndices)
 {
 	float ChainLength = 0;
@@ -31,16 +26,16 @@ static float EnsureToAddBoneToLinkData(const FIKRigTransforms& TransformModifier
 		FFBIKLinkData& NewLink = LinkData[NewLinkIndex];
 
 		// find parent LinkIndex
-		const int32 ParentItem = TransformModifier.Hierarchy->GetParentIndex(CurrentItem);
+		const int32 ParentItem = InSkeleton.GetParentIndex(CurrentItem);
 		FoundLinkIndex = HierarchyToLinkDataMap.Find(ParentItem);
 		NewLink.ParentLinkIndex = (FoundLinkIndex) ? *FoundLinkIndex : INDEX_NONE;
-		NewLink.SetTransform(TransformModifier.GetGlobalTransform(CurrentItem));
+		NewLink.SetTransform(InSkeleton.RefPoseGlobal[CurrentItem]);
 
 		if (ParentItem != INDEX_NONE)
 		{
 			// @todo: currently we assume the input pose is init transform
 			// this could use reference pose getter from IKRigSolver or we could wrap it to TransformModifier
-			FVector DiffLocation = TransformModifier.GetGlobalTransform(CurrentItem).GetLocation() - TransformModifier.GetGlobalTransform(ParentItem).GetLocation();
+			FVector DiffLocation = InSkeleton.RefPoseGlobal[CurrentItem].GetLocation() - InSkeleton.RefPoseGlobal[ParentItem].GetLocation();
 			// set Length
 			NewLink.Length = DiffLocation.Size();
 		}
@@ -84,18 +79,18 @@ static void AddToEffectorTarget(int32 EffectorIndex, const int32 Effector, TMap<
 }
 
 // the output should be from current index -> to root parent 
-static bool GetBoneChain(const FIKRigTransforms& TransformModifier, const FName& Root, const FName& Current, TArray<int32>& ChainIndices)
+static bool GetBoneChain(const FIKRigSkeleton& InSkeleton, const FName& Root, const FName& Current, TArray<int32>& ChainIndices)
 {
 	ChainIndices.Reset();
 
-	int32 RootIndex = TransformModifier.Hierarchy->GetIndex(Root);
-	int32 Iterator = TransformModifier.Hierarchy->GetIndex(Current);;
+	int32 RootIndex = InSkeleton.GetBoneIndexFromName(Root);
+	int32 Iterator = InSkeleton.GetBoneIndexFromName(Current);;
 
 	// iterates until key is valid
 	while (Iterator!=INDEX_NONE && Iterator != RootIndex)
 	{
 		ChainIndices.Insert(Iterator, 0);
-		Iterator = TransformModifier.Hierarchy->GetParentIndex(Iterator);
+		Iterator = InSkeleton.GetParentIndex(Iterator);
 	}
 
 	// add the last one if valid
@@ -110,7 +105,7 @@ static bool GetBoneChain(const FIKRigTransforms& TransformModifier, const FName&
 	return !ChainIndices.IsEmpty();
 }
 
-static void AddEffectors(const FIKRigTransforms& TransformModifier, const FName Root, const TArray<FFBIKRigEffector>& Effectors,
+static void AddEffectors(const FIKRigSkeleton& IKRigSkeleton, const FName Root, const TArray<FFBIKRigEffector>& Effectors,
 	TArray<FFBIKLinkData>& LinkData, TMap<int32, FFBIKEffectorTarget>& EffectorTargets, TArray<int32>& EffectorLinkIndices,
 	TMap<int32, int32>& LinkDataToHierarchyIndices, TMap<int32, int32>& HierarchyToLinkDataMap, const FSolverInput& SolverProperty)
 {
@@ -122,7 +117,7 @@ static void AddEffectors(const FIKRigTransforms& TransformModifier, const FName
 		EffectorLinkIndices[Index] = INDEX_NONE;
 		// create LinkeData from root bone to all effectors 
 		const FName Bone = Effectors[Index].Target.Bone;
-		const int32 BoneIndex = TransformModifier.Hierarchy->GetIndex(Bone);
+		const int32 BoneIndex = IKRigSkeleton.GetBoneIndexFromName(Bone);
 
 		if (Bone == NAME_None || BoneIndex == INDEX_NONE)
 		{
@@ -133,7 +128,7 @@ static void AddEffectors(const FIKRigTransforms& TransformModifier, const FName
 
 		TArray<int32> ChainIndices;
 		// if we haven't got to root, this is not valid chain
-		if (GetBoneChain(TransformModifier, Root, Bone, ChainIndices))
+		if (GetBoneChain(IKRigSkeleton, Root, Bone, ChainIndices))
 		{
 			auto CalculateStrength = [&](int32 InBoneChainDepth, const int32 MaxDepth, float CurrentStrength, float MinStrength) -> float
 			{
@@ -182,7 +177,7 @@ static void AddEffectors(const FIKRigTransforms& TransformModifier, const FName
 			for (int32 BoneChainIndex = 0; BoneChainIndex < ChainIndices.Num(); ++BoneChainIndex)
 			{
 				const int32 CurrentItem = ChainIndices[BoneChainIndex];
-				ChainLength += EnsureToAddBoneToLinkData(TransformModifier, CurrentItem, LinkData,
+				ChainLength += EnsureToAddBoneToLinkData(IKRigSkeleton, CurrentItem, LinkData,
 					HierarchyToLinkDataMap, LinkDataToHierarchyIndices);
 
 				const int32 ChainDepth = ChainIndices.Num() - BoneChainIndex;
@@ -204,7 +199,7 @@ UIKRigFBIKSolver::UIKRigFBIKSolver()
 {
 }
 
-void UIKRigFBIKSolver::Init(const FIKRigTransforms& InGlobalTransform)
+void UIKRigFBIKSolver::Initialize(const FIKRigSkeleton& IKRigSkeleton)
 {
 
 	LinkData.Reset();
@@ -214,12 +209,12 @@ void UIKRigFBIKSolver::Init(const FIKRigTransforms& InGlobalTransform)
 	HierarchyToLinkDataMap.Reset();
 
 	// verify the chain
-	AddEffectors(InGlobalTransform, Root, Effectors, LinkData, EffectorTargets, EffectorLinkIndices, LinkDataToHierarchyIndices, HierarchyToLinkDataMap, SolverProperty);
+	AddEffectors(IKRigSkeleton, Root, Effectors, LinkData, EffectorTargets, EffectorLinkIndices, LinkDataToHierarchyIndices, HierarchyToLinkDataMap, SolverProperty);
 	
 }
 
 void UIKRigFBIKSolver::Solve(
-	FIKRigTransforms& InOutGlobalTransform,
+	FIKRigSkeleton& IKRigSkeleton,
 	const FIKRigGoalContainer& Goals,
 	FControlRigDrawInterface* InOutDrawInterface)
 {
@@ -259,7 +254,7 @@ void UIKRigFBIKSolver::Solve(
 		for (int32 LinkIndex = 0; LinkIndex < LinkData.Num(); ++LinkIndex)
 		{
 			const int32 Item = *LinkDataToHierarchyIndices.Find(LinkIndex);
-			LinkData[LinkIndex].SetTransform(InOutGlobalTransform.GetGlobalTransform(Item));
+			LinkData[LinkIndex].SetTransform(IKRigSkeleton.CurrentPoseGlobal[Item]);
 			// @todo: fix this somewhere else - we can add this to prepare step
 			// @todo: we update motion scale here, then?
 			LinkData[LinkIndex].FinalizeForSolver();
@@ -546,6 +541,10 @@ void UIKRigFBIKSolver::Solve(
 		///////////////////////////////////////////////////////////////////////////
 	}
 
+	// record bones that are controlled by solver
+	TArray<bool> SolvedBones;
+	SolvedBones.Init(false, IKRigSkeleton.CurrentPoseGlobal.Num());
+	
 	// we update back to hierarchy
 	for (int32 LinkIndex = 0; LinkIndex < LinkData.Num(); ++LinkIndex)
 	{
@@ -553,7 +552,17 @@ void UIKRigFBIKSolver::Solve(
 		// this means, only the last joint in the test
 		const int32 CurrentItem = *LinkDataToHierarchyIndices.Find(LinkIndex);
 		const FTransform& LinkTransform = LinkData[LinkIndex].GetTransform();
-		InOutGlobalTransform.SetGlobalTransform(CurrentItem, LinkTransform, true);
+		IKRigSkeleton.CurrentPoseGlobal[CurrentItem] = LinkTransform;
+		SolvedBones[CurrentItem] = true;
+	}
+
+	// update non-solved global pose
+	for (int32 BoneIndex=0; BoneIndex<IKRigSkeleton.CurrentPoseGlobal.Num(); ++BoneIndex)
+	{
+		if (!SolvedBones[BoneIndex])
+		{
+			IKRigSkeleton.UpdateGlobalTransformFromLocal(BoneIndex);
+		}
 	}
 }
 
diff --git a/Engine/Plugins/Animation/IKRig/Source/IKRig/Private/Solvers/TransformSolver.cpp b/Engine/Plugins/Animation/IKRig/Source/IKRig/Private/Solvers/TransformSolver.cpp
--- a/Engine/Plugins/Animation/IKRig/Source/IKRig/Private/Solvers/TransformSolver.cpp
+++ b/Engine/Plugins/Animation/IKRig/Source/IKRig/Private/Solvers/TransformSolver.cpp
@@ -3,7 +3,6 @@
 
 #include "Solvers/TransformSolver.h"
 #include "IKRigDataTypes.h"
-#include "IKRigHierarchy.h"
 
 
 UTransformSolver::UTransformSolver()
@@ -11,13 +10,13 @@ UTransformSolver::UTransformSolver()
 	Effector.Goal = "DefaultGoal";
 }
 
-void UTransformSolver::Init(const FIKRigTransforms& InGlobalTransform)
+void UTransformSolver::Initialize(const FIKRigSkeleton& IKRigSkeleton)
 {
-	
+	BoneIndex = IKRigSkeleton.GetBoneIndexFromName(Effector.Bone);
 }
 
 void UTransformSolver::Solve(
-	FIKRigTransforms& InOutGlobalTransform,
+	FIKRigSkeleton& IKRigSkeleton,
 	const FIKRigGoalContainer& Goals,
 	FControlRigDrawInterface* InOutDrawInterface)
 {
@@ -26,25 +25,20 @@ void UTransformSolver::Solve(
 	{
 		return;
 	}
-	
-	int32 Index = InOutGlobalTransform.Hierarchy->GetIndex(Effector.Bone);
-	if (Index == INDEX_NONE)
-	{
-		return;
-	}
 
-	FTransform CurrentTransform = InOutGlobalTransform.GetGlobalTransform(Index);
+	FTransform& CurrentTransform = IKRigSkeleton.CurrentPoseGlobal[BoneIndex];
 
 	if (bEnablePosition)
 	{
 		CurrentTransform.SetLocation(Goal.Position);
 	}
+	
 	if (bEnableRotation)
 	{
 		CurrentTransform.SetRotation(Goal.Rotation);
 	}
 
-	InOutGlobalTransform.SetGlobalTransform(Index, CurrentTransform, true);
+	IKRigSkeleton.PropagateGlobalPoseBelowBone(BoneIndex);
 }
 
 void UTransformSolver::CollectGoalNames(TSet<FName>& OutGoals) const
diff --git a/Engine/Plugins/Animation/IKRig/Source/IKRig/Public/AnimNodes/AnimNode_IKRetargeter.h b/Engine/Plugins/Animation/IKRig/Source/IKRig/Public/AnimNodes/AnimNode_IKRetargeter.h
new file mode 100644
--- /dev/null
+++ b/Engine/Plugins/Animation/IKRig/Source/IKRig/Public/AnimNodes/AnimNode_IKRetargeter.h
@@ -0,0 +1,58 @@
+// Copyright Epic Games, Inc. All Rights Reserved.
+
+#pragma once
+
+#include "CoreMinimal.h"
+#include "BoneControllers/AnimNode_SkeletalControlBase.h"
+#include "Retargeter/IKRetargeter.h"
+#include "AnimNode_IKRetargeter.generated.h"
+
+USTRUCT(BlueprintInternalUseOnly)
+struct IKRIG_API FAnimNode_IKRetargeter : public FAnimNode_Base
+{
+	GENERATED_BODY()
+
+	/** The Skeletal Mesh Component to retarget animation from. Assumed to be animated and tick BEFORE this anim instance.*/
+	UPROPERTY(BlueprintReadWrite, transient, Category=Settings, meta=(PinShownByDefault))
+	TWeakObjectPtr<USkeletalMeshComponent> SourceMeshComponent;
+	
+	/** Retarget asset to use. Must define a Source and Target IK Rig compatible with the SourceMeshComponent and current anim instance.*/
+	UPROPERTY(EditAnywhere, Category = Settings)
+	UIKRetargeter* IKRetargeterAsset;
+
+	/** Retarget asset to use. Must define a Source and Target IK Rig compatible with the SourceMeshComponent and current anim instance.*/
+	UPROPERTY(EditAnywhere, Category = Settings)
+	bool bEnableIK;
+
+	FAnimNode_IKRetargeter();
+
+	// FAnimNode_Base interface
+	virtual void Initialize_AnyThread(const FAnimationInitializeContext& Context) override;
+	virtual void CacheBones_AnyThread(const FAnimationCacheBonesContext& Context) override;
+	virtual void Update_AnyThread(const FAnimationUpdateContext& Context) override;
+	virtual void Evaluate_AnyThread(FPoseContext& Output) override;
+	//virtual void GatherDebugData(FNodeDebugData& DebugData) override;
+	virtual bool HasPreUpdate() const override { return true; }
+	virtual void PreUpdate(const UAnimInstance* InAnimInstance) override;
+	// End of FAnimNode_Base interface
+
+private:
+	void EnsureInitialized(const UAnimInstance* InAnimInstance);
+	void InitializeRetargetData(const UAnimInstance* InAnimInstance);
+	void CopyBoneTransformsFromSource(USkeletalMeshComponent* TargetMeshComponent);
+
+	// indicates that all prerequisites are met and node is ready to operate
+	bool bIsInitialized;
+	
+	// source mesh references, cached during init so that we can compare and see if it has changed
+	TWeakObjectPtr<USkeletalMeshComponent>	CurrentlyUsedSourceMeshComponent;
+	TWeakObjectPtr<USkeletalMesh>			CurrentlyUsedSourceMesh;
+	TWeakObjectPtr<USkeletalMesh>			CurrentlyUsedTargetMesh;
+	TWeakObjectPtr<UIKRetargeter>			CurrentlyUsedRetargeter;
+	TWeakObjectPtr<UIKRigDefinition>		CurrentlyUsedSourceIKRig;
+	TWeakObjectPtr<UIKRigDefinition>		CurrentlyUsedTargetIKRig;
+
+	// cached transforms, copied on the game thread
+	TArray<FTransform> SourceMeshComponentSpaceBoneTransforms;
+	TArray<FTransform> TargetMeshComponentSpaceBoneTransforms;
+};
diff --git a/Engine/Plugins/Animation/IKRig/Source/IKRig/Public/AnimNodes/AnimNode_IKRig.h b/Engine/Plugins/Animation/IKRig/Source/IKRig/Public/AnimNodes/AnimNode_IKRig.h
--- a/Engine/Plugins/Animation/IKRig/Source/IKRig/Public/AnimNodes/AnimNode_IKRig.h
+++ b/Engine/Plugins/Animation/IKRig/Source/IKRig/Public/AnimNodes/AnimNode_IKRig.h
@@ -10,10 +10,6 @@
 class UIKRigDefinition;
 class UIKRigProcessor;
 
-/**
-*	Controller which implements the FABRIK IK approximation algorithm -  see http://www.academia.edu/9165835/FABRIK_A_fast_iterative_solver_for_the_Inverse_Kinematics_problem for details
-*/
-
 USTRUCT(BlueprintInternalUseOnly)
 struct IKRIG_API FAnimNode_IKRig : public FAnimNode_Base
 {
@@ -29,7 +25,7 @@ struct IKRIG_API FAnimNode_IKRig : public FAnimNode_Base
 	UPROPERTY(EditAnywhere, BlueprintReadWrite, EditFixedSize, Category = Goal, meta = (PinShownByDefault))
 	TArray<FTransform> GoalTransforms;
 
-	/** For every evaluation, we'll start with refpose, this is safer if you have joints that don't exists in the target */
+	/** optionally ignore the input pose and start from the reference pose each solve */
 	UPROPERTY(EditAnywhere, Category = Solver)
 	bool bStartFromRefPose = false;
 
@@ -41,7 +37,7 @@ struct IKRIG_API FAnimNode_IKRig : public FAnimNode_Base
 
 private: 
 	UPROPERTY(Transient)
-	UIKRigProcessor* RigProcessor = nullptr;
+	UIKRigProcessor* IKRigProcessor = nullptr;
 
 	TArray<FName> GoalNames;
 	TMap<FCompactPoseBoneIndex, int32, FDefaultSetAllocator, TCompactPoseBoneIndexMapKeyFuncs<int32>> CompactPoseToRigIndices;
diff --git a/Engine/Plugins/Animation/IKRig/Source/IKRig/Public/IKRigDataTypes.h b/Engine/Plugins/Animation/IKRig/Source/IKRig/Public/IKRigDataTypes.h
--- a/Engine/Plugins/Animation/IKRig/Source/IKRig/Public/IKRigDataTypes.h
+++ b/Engine/Plugins/Animation/IKRig/Source/IKRig/Public/IKRigDataTypes.h
@@ -16,27 +16,28 @@ struct IKRIG_API FIKRigGoal
 {
 	GENERATED_BODY()
 
-		FIKRigGoal()
-		: Position(ForceInitToZero),
-		Rotation(ForceInitToZero)
+	FIKRigGoal()
+	: Position(ForceInitToZero),
+	Rotation(ForceInitToZero)
 	{
 		Name = NAME_None;
 	}
+	
 	FIKRigGoal(const FName& GoalName)
-		: Position(ForceInitToZero),
-		Rotation(ForceInitToZero)
+	: Position(ForceInitToZero),
+	Rotation(ForceInitToZero)
 	{
 		Name = GoalName;
 	}
 
 	UPROPERTY(VisibleAnywhere, Category = FIKRigGoal)
-		FName Name;
+	FName Name;
 
 	UPROPERTY(EditAnywhere, Category = FIKRigGoal)
-		FVector Position;
+	FVector Position;
 
 	UPROPERTY(EditAnywhere, Category = FIKRigGoal)
-		FQuat Rotation;
+	FQuat Rotation;
 };
 
 inline uint32 GetTypeHash(FIKRigGoal ObjectRef) { return GetTypeHash(ObjectRef.Name); }
@@ -46,7 +47,7 @@ struct IKRIG_API FIKRigGoalContainer
 {
 	GENERATED_BODY()
 
-		void InitializeGoalsFromNames(const TArray<FName>& InGoalNames)
+	void InitializeGoalsFromNames(const TArray<FName>& InGoalNames)
 	{
 		Goals.Reserve(InGoalNames.Num());
 		for (const FName& Name : InGoalNames)
@@ -94,43 +95,3 @@ private:
 
 	TMap<FName, FIKRigGoal> Goals;
 };
-
-USTRUCT()
-struct IKRIG_API FIKRigTransforms
-{
-	GENERATED_BODY()
-	
-	FIKRigTransforms()
-		: Hierarchy(nullptr)
-	{
-	}
-
-	FIKRigTransforms(const FIKRigHierarchy* InHierarchy);
-
-	void SetGlobalTransform(int32 Index, const FTransform& InTransform, bool bPropagate);
-	void SetLocalTransform(int32 Index, const FTransform& InTransform, bool bUpdate);
-	
-	const FTransform& GetLocalTransform(int32 Index) const; // this is mutating, but const_cast 
-	const FTransform& GetGlobalTransform(int32 Index) const;
-
-	void SetAllGlobalTransforms(const TArray<FTransform>& InTransforms);
-
-	const FIKRigHierarchy* Hierarchy;
-
-private:
-
-	TArray<FTransform> GlobalTransforms;
-	TArray<FTransform> LocalTransforms;
-	TBitArray<>	LocalTransformDirtyFlags;
-
-	void EnsureLocalTransformsExist();
-	void RecalculateLocalTransform();
-	void UpdateLocalTransform(int32 Index);
-
-	/** this function propagates to children and updates LocalTransforms */
-	void SetGlobalTransform_Internal(int32 Index, const FTransform& InTransform);
-	void SetGlobalTransform_Recursive(int32 Index, const FTransform& InTransform);
-
-	FTransform GetRelativeTransform(int32 Index, int32 BaseIndex) const;
-};
-
diff --git a/Engine/Plugins/Animation/IKRig/Source/IKRig/Public/IKRigDefinition.h b/Engine/Plugins/Animation/IKRig/Source/IKRig/Public/IKRigDefinition.h
--- a/Engine/Plugins/Animation/IKRig/Source/IKRig/Public/IKRigDefinition.h
+++ b/Engine/Plugins/Animation/IKRig/Source/IKRig/Public/IKRigDefinition.h
@@ -3,16 +3,65 @@
 #pragma once
 
 #include "CoreMinimal.h"
-#include "UObject/ObjectMacros.h"
 #include "UObject/Object.h"
 #include "IKRigBoneSetting.h"
-#include "IKRigHierarchy.h"
-#include "IKRigDataTypes.h"
+#include "IKRigSkeleton.h"
 
 #include "IKRigDefinition.generated.h"
 
 class UIKRigSolver;
 
+USTRUCT(Blueprintable)
+struct IKRIG_API FBoneChain
+{
+	FBoneChain()
+		: bUseIK(false)
+	{
+	}
+
+	GENERATED_BODY()
+
+	UPROPERTY(EditAnywhere, Category = BoneChain)
+	FName ChainName;
+
+	UPROPERTY(EditAnywhere, Category = BoneChain)
+	FName StartBone;
+
+	UPROPERTY(EditAnywhere, Category = BoneChain)
+	FName EndBone;
+
+	UPROPERTY(EditAnywhere, Category = IK)
+	bool bUseIK;
+	
+	UPROPERTY(EditAnywhere, Category = IK)
+	FName IKGoalName;
+
+	/*
+	UPROPERTY(EditAnywhere, Category = IK)
+	FName PoleVectorGoalName;
+
+	UPROPERTY(EditAnywhere, Category = IK)
+	FName PoleVectorBoneName;
+
+	UPROPERTY(EditAnywhere, Category = IK)
+	FVector PoleVectorDirection;
+	*/
+};
+
+USTRUCT(Blueprintable)
+struct IKRIG_API FRetargetDefinition
+{
+	GENERATED_BODY()
+
+	UPROPERTY(EditAnywhere, Category = BodyParts)
+	FName RootBone;
+
+	UPROPERTY(EditAnywhere, Category = BodyParts)
+	TArray<FBoneChain> BoneChains;
+
+	FBoneChain* GetBoneChainByName(FName ChainName);
+};
+
 UCLASS(Blueprintable)
 class IKRIG_API UIKRigDefinition : public UObject
 {
@@ -20,46 +69,31 @@ class IKRIG_API UIKRigDefinition : public UObject
 
 public:
 
-	UPROPERTY()
-	TSoftObjectPtr<UObject> SourceAsset;
-
-	/** parent/child relationships stored separately from skeleton-specific transforms */
+	/** load a source skeleton to import a hierarchy from */
 	UPROPERTY(VisibleAnywhere, Category = "Skeleton")
-	FIKRigHierarchy Hierarchy;
+	TSoftObjectPtr<UObject> SourceAsset;
 
-	/** default transforms of the reference pose skeleton */
+	/** hierarchy and bone-pose transforms */
 	UPROPERTY(VisibleAnywhere, Category = "Skeleton")
-	TArray<FTransform> RefPoseTransforms;
-
-	UPROPERTY(EditAnywhere, instanced, Category = BoneSettings)
-	TArray<UIKRigBoneSetting*> BoneSettings;
-
-	UPROPERTY(EditAnywhere, Category = BoneSettings)
-	TArray<FName> ExcludedBones;
+	FIKRigSkeleton Skeleton;
 
 	/** stack of solvers, of varying types, executed in serial fashion where
 	output of prior solve is input to the next */
-	UPROPERTY(EditAnywhere, instanced, Category = "Solvers")
+	UPROPERTY(EditAnywhere, instanced, Category = "IK")
 	TArray<UIKRigSolver*> Solvers;
 
-	void GetGoalNamesFromSolvers(TArray<FName>& OutGoalNames) const;
-
-private:
-
-	// BEGIN UObject functions 
-	virtual void PostLoad() override;
-	// END UObject functions
+	/** per-bone settings, used by solvers that support them */
+	UPROPERTY(EditAnywhere, instanced, Category = "IK")
+	TArray<UIKRigBoneSetting*> BoneSettings;
 
-#if WITH_EDITOR
+	/** bone chains for animation retargeting */
+	UPROPERTY(EditAnywhere, Category = "Retargeting")
+	FRetargetDefinition RetargetDefinition;
 
-	//** add a bone. Return false if it already has conflict name or parent is not found */
-	bool AddBone(const FName& InName, const FName& InParent, const FTransform& InGlobalTransform);
-	//** reset skeleton data */
-	void ResetHierarchy();
-	// ensure it's sorted from parent to children
-	void EnsureSortedCorrectly(bool bReSortIfNeeded=false);
+	/** get a list of all goal names used in the solvers */
+	void GetGoalNamesFromSolvers(TArray<FName>& OutGoalNames) const;
 
-#endif // WITH_EDITOR
+private:
 
 	friend class UIKRigController;
 };
diff --git a/Engine/Plugins/Animation/IKRig/Source/IKRig/Public/IKRigHierarchy.h b/Engine/Plugins/Animation/IKRig/Source/IKRig/Public/IKRigHierarchy.h
deleted file mode 100644
diff --git a/Engine/Plugins/Animation/IKRig/Source/IKRig/Public/IKRigProcessor.h b/Engine/Plugins/Animation/IKRig/Source/IKRig/Public/IKRigProcessor.h
--- a/Engine/Plugins/Animation/IKRig/Source/IKRig/Public/IKRigProcessor.h
+++ b/Engine/Plugins/Animation/IKRig/Source/IKRig/Public/IKRigProcessor.h
@@ -5,13 +5,14 @@
 #include "CoreMinimal.h"
 #include "UObject/ObjectMacros.h"
 #include "UObject/Object.h"
-#include "IKRigDataTypes.h"
 #include "Drawing/ControlRigDrawInterface.h"
+#include "IKRigDataTypes.h"
+#include "IKRigSkeleton.h"
+
 #include "IKRigProcessor.generated.h"
 
 class UIKRigDefinition;
 class UIKRigSolver;
-struct FControlRigDrawInterface;
 
 UCLASS(BlueprintType)
 class IKRIG_API UIKRigProcessor : public UObject
@@ -20,60 +21,77 @@ class IKRIG_API UIKRigProcessor : public UObject
 
 public:
 
-	// the runtime for an IKRig, general steps are:
-	// 1. Initialize() with an IKRigDefinition asset
-	// 2. each tick, call SetGoalTransform() and update GlobalBoneTransforms
-	// 3. Call Solve()
-	// 4. Copy back output transforms...
-	// 5. Profit.
+	/** the runtime for an IKRig to convert an input pose into
+	*   a solved output pose given the goal transforms:
+	* 1. Create a new IKRigProcessor once using MakeNewIKRigProcessor()
+	* 2. Initialize() with an IKRigDefinition asset
+	* 3. each tick, call SetGoalTransform() and update GlobalBoneTransforms
+	* 4. Call Solve()
+	* 5. Copy back output transforms...
+	*/
+
+	static UIKRigProcessor* MakeNewIKRigProcessor(UObject* Outer);
 	 
-	// setup
+	/** setup a new processor to run the given IKRig asset  */
 	void Initialize(UIKRigDefinition* InRigDefinition);
 
-	// run entire stack of solvers
-	void Solve();
-
-	// optionally can be called before Solve() to do deterministic solve
-	void ResetToRefPose();
+	//
+	// BEGIN UPDATE SEQUENCE FUNCTIONS
+	//
+	// This is the general sequence of function calls to run a typical IK solve:
+	// 
+	
+	/** set all transforms in global space */
+	void SetInputPoseGlobal(const TArray<FTransform>& InGlobalBoneTransforms);
 
-	// get all global bone transforms
-	FIKRigTransforms& GetCurrentGlobalTransforms();
+	/** optionally can be called before Solve() to use the reference pose as start pose */
+	void SetInputPoseToRefPose();
 
-	// set goal transform by name
+	/** set goal transform by name */
 	void SetGoalTransform(const FName& GoalName, const FVector& Position, const FQuat& Rotation);
 
-	// get all goal names in the solver stack
+	/** run entire stack of solvers */
+	void Solve();
+
+	/** get the results after calling Solve() */
+	void CopyOutputGlobalPoseToArray(TArray<FTransform>& OutputPoseGlobal) const;
+
+	//
+	// END UPDATE SEQUENCE FUNCTIONS
+	//
+
+	/** get all goal names in the solver stack */
 	void GetGoalNames(TArray<FName>& OutGoalNames) const;
 
-	// total number of goals in solver stack
+	/** total number of goals in solver stack */
 	int GetNumGoals() const;
 
-	// get access to the hierarchy the solver stack expects
-	const FIKRigHierarchy* GetHierarchy() const;
+	bool IsInitialized() const { return bInitialized; };
+
+	/** interface for debug drawing */
+	const FControlRigDrawInterface& GetDrawInterface() const { return DrawInterface; }
 
-	// interface for debug drawing
-	const FControlRigDrawInterface& GetDrawInterface() { return DrawInterface; }
+	/** get access to the internal skeleton data */
+	FIKRigSkeleton& GetSkeleton();
 	
 private:
 
+	/** solving disabled until this flag is true */
 	bool bInitialized = false;
 
-	UPROPERTY(EditAnywhere, Category = "RigDefinition")
-	UIKRigDefinition* RigDefinition = nullptr;
-
-	UPROPERTY(transient)
-	TArray<FTransform> RefPoseGlobalTransforms;
-
-	UPROPERTY(transient)
-	FIKRigTransforms GlobalBoneTransforms;
-
+	/** the stack of solvers to run in order */
 	UPROPERTY(transient)
 	TArray<UIKRigSolver*> Solvers;
 
+	/** storage for solver debug drawing data (lines, boxes etc) */
 	UPROPERTY(transient)
 	FControlRigDrawInterface DrawInterface;
 
+	/** the named transforms that solvers use as end effectors */
 	UPROPERTY(transient)
 	FIKRigGoalContainer Goals;
-};
 
+	/** storage for hierarchy and bone transforms */
+	UPROPERTY(transient)
+	FIKRigSkeleton Skeleton;
+};
\ No newline at end of file
diff --git a/Engine/Plugins/Animation/IKRig/Source/IKRig/Public/IKRigSkeleton.h b/Engine/Plugins/Animation/IKRig/Source/IKRig/Public/IKRigSkeleton.h
new file mode 100644
--- /dev/null
+++ b/Engine/Plugins/Animation/IKRig/Source/IKRig/Public/IKRigSkeleton.h
@@ -0,0 +1,34 @@
+// Copyright Epic Games, Inc. All Rights Reserved.
+
+#pragma once
+
+#include "CoreMinimal.h"
+#include "ReferenceSkeleton.h"
+
+#include "IKRigSkeleton.generated.h"
+
+USTRUCT()
+struct IKRIG_API FIKRigSkeleton
+{
+	GENERATED_BODY()
+
+	UPROPERTY(VisibleAnywhere, Category = Skeleton)
+	TArray<FName> BoneNames;
+	UPROPERTY(VisibleAnywhere, Category = Skeleton)
+	TArray<int32> ParentIndices;
+	UPROPERTY(VisibleAnywhere, Category = Skeleton)
+	TArray<FTransform> CurrentPoseGlobal;
+	UPROPERTY(VisibleAnywhere, Category = Skeleton)
+	TArray<FTransform> CurrentPoseLocal;
+	UPROPERTY(VisibleAnywhere, Category = Skeleton)
+	TArray<FTransform> RefPoseGlobal;
+
+	void Initialize(const FReferenceSkeleton& RefSkeleton);
+	int32 GetBoneIndexFromName(const FName InName) const;
+	int32 GetParentIndex(const int32 BoneIndex) const;
+	void UpdateAllGlobalTransformFromLocal();
+	void UpdateAllLocalTransformFromGlobal();
+	void UpdateGlobalTransformFromLocal(const int32 BoneIndex);
+	void UpdateLocalTransformFromGlobal(const int32 BoneIndex);
+	void PropagateGlobalPoseBelowBone(const int32 BoneIndex);
+};
diff --git a/Engine/Plugins/Animation/IKRig/Source/IKRig/Public/IKRigSolver.h b/Engine/Plugins/Animation/IKRig/Source/IKRig/Public/IKRigSolver.h
--- a/Engine/Plugins/Animation/IKRig/Source/IKRig/Public/IKRigSolver.h
+++ b/Engine/Plugins/Animation/IKRig/Source/IKRig/Public/IKRigSolver.h
@@ -5,6 +5,7 @@
 #include "UObject/ObjectMacros.h"
 #include "UObject/Object.h"
 #include "IKRigDataTypes.h"
+#include "IKRigSkeleton.h"
 #include "IKRigSolver.generated.h"
 
 struct FIKRigTarget;
@@ -35,12 +36,12 @@ public:
 	UPROPERTY(EditAnywhere, Category = "Settings")
 	bool bEnabled = true;
 
-	/** override Init() to setup internal data based on ref pose */
-	virtual void Init(const FIKRigTransforms& InGlobalTransform) PURE_VIRTUAL("Init");
+	/** override to setup internal data based on ref pose */
+	virtual void Initialize(const FIKRigSkeleton& IKRigSkeleton) PURE_VIRTUAL("Init");
 
 	/** override Solve() to evaluate new output pose (InOutGlobalTransform) */
 	virtual void Solve(
-		FIKRigTransforms& InOutGlobalTransform,
+		FIKRigSkeleton& IKRigSkeleton,
 		const FIKRigGoalContainer& Goals,
 		FControlRigDrawInterface* InOutDrawInterface) PURE_VIRTUAL("Solve");
 
diff --git a/Engine/Plugins/Animation/IKRig/Source/IKRig/Public/Retargeter/IKRetargeter.h b/Engine/Plugins/Animation/IKRig/Source/IKRig/Public/Retargeter/IKRetargeter.h
new file mode 100644
--- /dev/null
+++ b/Engine/Plugins/Animation/IKRig/Source/IKRig/Public/Retargeter/IKRetargeter.h
@@ -0,0 +1,283 @@
+// Copyright Epic Games, Inc. All Rights Reserved.
+
+#pragma once
+
+#include "Engine/SkeletalMesh.h"
+#include "Containers/UnrealString.h"
+#include "Math/TransformNonVectorized.h"
+
+#include "IKRigDefinition.h"
+#include "IKRigProcessor.h"
+
+#include "IKRetargeter.generated.h"
+
+struct FRetargetSkeleton
+{
+	TArray<FName> BoneNames;
+	TArray<int32> ParentIndices;
+	TArray<FTransform> RetargetGlobalPose;
+	FString SkeletalMeshName;
+
+	void InitializeHierarchy(const USkeletalMesh* SkeletalMesh, const TArray<FTransform>& InLocalRetargetPose);
+	int32 FindBoneIndexByName(const FName InName) const;
+	int32 GetParentIndex(const int32 BoneIndex) const;
+};
+
+struct FTargetSkeleton : public FRetargetSkeleton
+{
+	TArray<FTransform> OutputGlobalPose;
+	TArray<FTransform> RetargetLocalPose;
+	FString SkeletalMeshName;
+	TArray<bool> IsBoneFK;
+
+	void Initialize(const USkeletalMesh* SkeletalMesh, const TArray<FTransform>& InTargetRetargetLocalPose);
+	void SetBoneControlledByFKChain(const int32 BoneIndex, const bool IsFK);
+	void UpdateGlobalTransformsBelowBone(const int32 StartBoneIndex);
+	void UpdateGlobalTransformsAllNonFKBones();
+
+private:
+	void UpdateGlobalTransformOfSingleBone(const int32 BoneIndex);
+};
+
+struct FEncodedRoot
+{
+	FVector NormalizedPosition;
+	
+	FQuat Rotation;
+
+	FQuat InitialRotation;
+};
+
+struct FRootEncoder
+{
+	int32 BoneIndex;
+
+	float InvInitialHeight;
+
+	FEncodedRoot EncodedResult;
+
+	bool Initialize(const FName RootBoneName, const FRetargetSkeleton& SourceSkeleton);
+
+	void EncodePose(const TArray<FTransform> &InputGlobalPose);
+};
+
+struct FRootDecoder
+{
+	int32 BoneIndex;
+	
+	float InitialHeight;
+
+	FQuat InitialRotation;
+
+	bool Initialize(const FName RootBoneName, const FTargetSkeleton& TargetSkeleton);
+
+	void DecodePose(
+        const FEncodedRoot& EncodedResult,
+        TArray<FTransform> &InOutTargetBoneTransforms,
+        const float StrideScale) const;
+};
+
+struct FChainEncoderFK
+{
+	TArray<FTransform> InitialBoneTransforms;
+
+	TArray<FTransform> CurrentBoneTransforms;
+
+	TArray<float> Params;
+
+	bool Initialize(
+		const TArray<int32>& BoneIndices,
+		const TArray<FTransform> &InitialGlobalPose);
+	
+	bool CalculateBoneParameters();
+
+	void EncodePose(
+		const TArray<int32>& SourceBoneIndices,
+		const TArray<FTransform> &InputGlobalPose);
+
+	void DecodePose(
+		const TArray<int32>& TargetBoneIndices,
+		const FChainEncoderFK& SourceChain,
+		TArray<FTransform> &InOutGlobalPose);
+
+	FTransform GetTransformAtParam(
+		const TArray<FTransform>& Transforms,
+		const TArray<float>& InParams,
+		const float Param) const;
+};
+
+struct FEncodedIKChain
+{
+	FVector EndDirectionNormalized;
+	FQuat EndRotation;
+	FQuat EndRotationOrig;
+	float HeightFromGroundNormalized;
+	FVector PoleVectorDirection;
+};
+
+struct FChainEncoderIK
+{
+	int32 BoneIndexA;
+	
+	int32 BoneIndexB;
+	
+	int32 BoneIndexC;
+
+	FVector EndPositionOrig;
+	
+	FQuat EndRotationOrig;
+	
+	float InverseLength;
+
+	FEncodedIKChain EncodedResult;
+
+	bool Initialize(const TArray<int32>& BoneIndices, const TArray<FTransform> &InitialGlobalPose);
+
+	void EncodePose(const TArray<FTransform> &InputGlobalPose);
+};
+
+struct FDecodedIKChain
+{
+	FVector EndEffectorPosition;
+	
+	FQuat EndEffectorRotation;
+	
+	FVector PoleVectorPosition;
+};
+
+struct FChainDecoderIK
+{
+	int32 BoneIndexA;
+
+	float Length;
+
+	FVector EndPositionOrig;
+	
+	FQuat EndRotationOrig;
+
+	FDecodedIKChain DecodedResults;
+
+	bool Initialize(const TArray<int32>& BoneIndices, const TArray<FTransform> &InitialGlobalPose);
+	
+	void DecodePose(
+		const FEncodedIKChain& EncodedChain,
+		const TArray<FTransform>& OutGlobalPose,
+		FDecodedIKChain& OutResults) const;
+};
+
+struct FRetargetChainPair
+{
+	TArray<int32> SourceBoneIndices;
+	
+	TArray<int32> TargetBoneIndices;
+	
+	FName SourceBoneChainName;
+	
+	FName TargetBoneChainName;
+
+	virtual ~FRetargetChainPair(){};
+	
+	virtual bool Initialize(
+		const FBoneChain& SourceBoneChain,
+		const FBoneChain& TargetBoneChain,
+		const FRetargetSkeleton& SourceSkeleton,
+		const FTargetSkeleton& TargetSkeleton);
+
+private:
+
+	bool ValidateBoneChainWithSkeletalMesh(
+		const bool IsSource,
+		const FBoneChain& BoneChain,
+		const FRetargetSkeleton& RetargetSkeleton);
+};
+
+struct FRetargetChainPairFK : FRetargetChainPair
+{
+	FChainEncoderFK FKEncoder;
+
+	FChainEncoderFK FKDecoder;
+	
+	virtual bool Initialize(
+        const FBoneChain& SourceBoneChain,
+        const FBoneChain& TargetBoneChain,
+        const FRetargetSkeleton& SourceSkeleton,
+        const FTargetSkeleton& TargetSkeleton) override;
+};
+
+struct FRetargetChainPairIK : FRetargetChainPair
+{
+	FChainEncoderIK IKEncoder;
+	
+	FChainDecoderIK IKDecoder;
+	
+	FName IKGoalName;
+	
+	FName PoleVectorGoalName;
+
+	virtual bool Initialize(
+        const FBoneChain& SourceBoneChain,
+        const FBoneChain& TargetBoneChain,
+        const FRetargetSkeleton& SourceSkeleton,
+        const FTargetSkeleton& TargetSkeleton) override;
+};
+
+UCLASS(Blueprintable)
+class IKRIG_API UIKRetargeter : public UObject
+{
+	GENERATED_BODY()
+	
+public:
+
+	UPROPERTY(EditAnywhere, Category = Rigs)
+	UIKRigDefinition *SourceIKRigAsset;
+
+	UPROPERTY(EditAnywhere, Category = Rigs)
+	UIKRigDefinition *TargetIKRigAsset;
+
+	UPROPERTY(EditAnywhere, Category = Mapping)
+	TMap<FName, FName> ChainMapping;
+
+	FRetargetSkeleton SourceSkeleton;
+	
+	FTargetSkeleton TargetSkeleton;
+	
+	bool bIsLoadedAndValid;
+
+	UIKRetargeter();
+	
+	void Initialize(USkeletalMesh *SourceSkeleton, USkeletalMesh *TargetSkeleton, UObject* Outer);
+	
+	void RunRetargeter(const TArray<FTransform>& InSourceGlobalPose, bool bEnableIK);
+
+private:
+
+	UPROPERTY(Transient)
+	UIKRigProcessor* IKRig = nullptr;
+	
+	TArray<FRetargetChainPairFK> ChainPairsFK;
+	
+	TArray<FRetargetChainPairIK> ChainPairsIK;
+
+	FRootEncoder RootEncoder;
+	
+	FRootDecoder RootDecoder;
+
+	bool InitializeRoots();
+		
+	bool InitializeBoneChainPairs();
+	
+	bool InitializeIKRig(UObject* Outer);
+
+	void RunRootRetarget(
+		const TArray<FTransform>& InGlobalTransforms,
+		TArray<FTransform>& OutGlobalTransforms);
+	
+	void RunFKRetarget(
+		const TArray<FTransform>& InGlobalTransforms,
+		TArray<FTransform>& OutGlobalTransforms,
+		const bool bPreIK);
+
+	void RunIKRetarget(
+		const TArray<FTransform>& InGlobalPose,
+		TArray<FTransform>& OutGlobalPose);
+};
diff --git a/Engine/Plugins/Animation/IKRig/Source/IKRig/Public/Solvers/ConstraintSolver.h b/Engine/Plugins/Animation/IKRig/Source/IKRig/Public/Solvers/ConstraintSolver.h
--- a/Engine/Plugins/Animation/IKRig/Source/IKRig/Public/Solvers/ConstraintSolver.h
+++ b/Engine/Plugins/Animation/IKRig/Source/IKRig/Public/Solvers/ConstraintSolver.h
@@ -13,9 +13,9 @@ class IKRIG_API UIKRigConstraintSolver : public UIKRigSolver
 
 public: 
 
-	virtual void Init(const FIKRigTransforms& InGlobalTransform) override;
+	virtual void Initialize(const FIKRigSkeleton& IKRigSkeleton) override;
 	virtual void Solve(
-		FIKRigTransforms& InOutGlobalTransform,
+		FIKRigSkeleton& IKRigSkeleton,
 		const FIKRigGoalContainer& Goals,
 		FControlRigDrawInterface* InOutDrawInterface) override;
 	virtual void CollectGoalNames(TSet<FName>& OutGoals) const override;
diff --git a/Engine/Plugins/Animation/IKRig/Source/IKRig/Public/Solvers/IKRigFBIKSolver.h b/Engine/Plugins/Animation/IKRig/Source/IKRig/Public/Solvers/IKRigFBIKSolver.h
--- a/Engine/Plugins/Animation/IKRig/Source/IKRig/Public/Solvers/IKRigFBIKSolver.h
+++ b/Engine/Plugins/Animation/IKRig/Source/IKRig/Public/Solvers/IKRigFBIKSolver.h
@@ -89,9 +89,9 @@ private:
 	TArray<FJacobianDebugData> DebugData;
 
 protected:
-	virtual void Init(const FIKRigTransforms& InGlobalTransforms) override;
+	virtual void Initialize(const FIKRigSkeleton& IKRigSkeleton) override;
 	virtual void Solve(
-		FIKRigTransforms& InOutGlobalTransforms,
+		FIKRigSkeleton& IKRigSkeleton,
 		const FIKRigGoalContainer& Goals,
 		FControlRigDrawInterface* InOutDrawInterface) override;
 	virtual void CollectGoalNames(TSet<FName>& OutGoals) const override;
diff --git a/Engine/Plugins/Animation/IKRig/Source/IKRig/Public/Solvers/TransformSolver.h b/Engine/Plugins/Animation/IKRig/Source/IKRig/Public/Solvers/TransformSolver.h
--- a/Engine/Plugins/Animation/IKRig/Source/IKRig/Public/Solvers/TransformSolver.h
+++ b/Engine/Plugins/Animation/IKRig/Source/IKRig/Public/Solvers/TransformSolver.h
@@ -26,11 +26,15 @@ public:
 
 protected:
 	
-	virtual void Init(const FIKRigTransforms& InGlobalTransform) override;
+	virtual void Initialize(const FIKRigSkeleton& IKRigSkeleton) override;
 	virtual void Solve(
-		FIKRigTransforms& InOutGlobalTransform, 
+		FIKRigSkeleton& IKRigSkeleton, 
 		const FIKRigGoalContainer& Goals,
 		FControlRigDrawInterface* InOutDrawInterface) override;
 	virtual void CollectGoalNames(TSet<FName>& OutGoals) const override;
+
+private:
+	
+	int32 BoneIndex;
 };
 
diff --git a/Engine/Plugins/Animation/IKRig/Source/IKRigDeveloper/IKRigDeveloper.Build.cs b/Engine/Plugins/Animation/IKRig/Source/IKRigDeveloper/IKRigDeveloper.Build.cs
--- a/Engine/Plugins/Animation/IKRig/Source/IKRigDeveloper/IKRigDeveloper.Build.cs
+++ b/Engine/Plugins/Animation/IKRig/Source/IKRigDeveloper/IKRigDeveloper.Build.cs
@@ -6,16 +6,14 @@ namespace UnrealBuildTool.Rules
     {
         public IKRigDeveloper(ReadOnlyTargetRules Target) : base(Target)
         {
-            // Copying some these from ControlRig.build.cs, our deps are likely leaner
-            // and therefore these could be pruned if needed:
-            PrivateDependencyModuleNames.AddRange(
+	        PrivateDependencyModuleNames.AddRange(
                 new string[]
                 {
 					"Core",
 					"CoreUObject",
 					"Engine",
 					"IKRig",
-				}
+                }
             );
 
             PublicDependencyModuleNames.AddRange(
@@ -29,7 +27,10 @@ namespace UnrealBuildTool.Rules
                 PrivateDependencyModuleNames.AddRange(
                     new string[]
                     {
-						"AnimGraph",
+                        "EditorFramework",
+                        "UnrealEd",
+                        "Kismet",
+                        "AnimGraph",
                         "BlueprintGraph",
                     }
                 );
diff --git a/Engine/Plugins/Animation/IKRig/Source/IKRigDeveloper/Private/AnimGraphNode_IKRetargeter.cpp b/Engine/Plugins/Animation/IKRig/Source/IKRigDeveloper/Private/AnimGraphNode_IKRetargeter.cpp
new file mode 100644
--- /dev/null
+++ b/Engine/Plugins/Animation/IKRig/Source/IKRigDeveloper/Private/AnimGraphNode_IKRetargeter.cpp
@@ -0,0 +1,97 @@
+// Copyright Epic Games, Inc. All Rights Reserved.
+
+#include "AnimGraphNode_IKRetargeter.h"
+#include "Animation/AnimInstance.h"
+#include "Kismet2/CompilerResultsLog.h"
+
+
+#define LOCTEXT_NAMESPACE "AnimGraphNode_IKRig"
+const FName UAnimGraphNode_IKRetargeter::AnimModeName(TEXT("IKRig.IKRigEditor.IKRigEditMode"));
+
+void UAnimGraphNode_IKRetargeter::Draw(FPrimitiveDrawInterface* PDI, USkeletalMeshComponent* PreviewSkelMeshComp) const
+{
+}
+
+FText UAnimGraphNode_IKRetargeter::GetNodeTitle(ENodeTitleType::Type TitleType) const
+{
+	return LOCTEXT("AnimGraphNode_IKRetargeter_Title", "IK Retargeter");
+}
+
+void UAnimGraphNode_IKRetargeter::CopyNodeDataToPreviewNode(FAnimNode_Base* InPreviewNode)
+{
+	FAnimNode_IKRetargeter* IKRetargeterNode = static_cast<FAnimNode_IKRetargeter*>(InPreviewNode);
+}
+
+FEditorModeID UAnimGraphNode_IKRetargeter::GetEditorMode() const
+{
+	return AnimModeName;
+}
+
+void UAnimGraphNode_IKRetargeter::PostEditChangeProperty(struct FPropertyChangedEvent& PropertyChangedEvent)
+{
+}
+
+void UAnimGraphNode_IKRetargeter::PostLoad()
+{
+	Super::PostLoad();
+}
+
+void UAnimGraphNode_IKRetargeter::ValidateAnimNodeDuringCompilation(USkeleton* ForSkeleton,	FCompilerResultsLog& MessageLog)
+{
+	Super::ValidateAnimNodeDuringCompilation(ForSkeleton, MessageLog);
+
+	// validate source mesh component is not null
+	//if (!Node.SourceMeshComponent.())
+	if (!IsPinExposedAndLinked(GET_MEMBER_NAME_STRING_CHECKED(FAnimNode_IKRetargeter, SourceMeshComponent)))
+	{
+		MessageLog.Warning(TEXT("@@ is missing an Source Skeletal Mesh Component reference."), this);
+		return;
+	}
+
+	// validate IK Rig asset has been assigned
+	if (!Node.IKRetargeterAsset)
+	{
+		MessageLog.Warning(TEXT("@@ is missing an IKRetargeter asset."), this);
+		return;
+	}
+
+	// validate SOURCE IK Rig asset has been assigned
+	if (!Node.IKRetargeterAsset->SourceIKRigAsset)
+	{
+		MessageLog.Warning(TEXT("@@ has IK Retargeter that is missing a source IK Rig asset."), this);
+	}
+
+	// validate TARGET IK Rig asset has been assigned
+	if (!Node.IKRetargeterAsset->TargetIKRigAsset)
+	{
+		MessageLog.Warning(TEXT("@@ has IK Retargeter that is missing a target IK Rig asset."), this);
+	}
+
+	if (!(Node.IKRetargeterAsset->SourceIKRigAsset && Node.IKRetargeterAsset->TargetIKRigAsset))
+	{
+		return;
+	}
+	
+	// validate that target bone chains exist on this skeleton
+	const FReferenceSkeleton &RefSkel = ForSkeleton->GetReferenceSkeleton();
+	const TArray<FBoneChain> &TargetBoneChains = Node.IKRetargeterAsset->TargetIKRigAsset->RetargetDefinition.BoneChains;
+    for (const FBoneChain &Chain : TargetBoneChains)
+    {
+        if (RefSkel.FindBoneIndex(Chain.StartBone) == INDEX_NONE)
+        {
+        	MessageLog.Warning(*LOCTEXT("StartBoneNotFound", "@@ - Start Bone in target IK Rig Bone Chain not found.").ToString(), this);
+        }
+
+    	if (RefSkel.FindBoneIndex(Chain.EndBone) == INDEX_NONE)
+    	{
+    		MessageLog.Warning(*LOCTEXT("EndBoneNotFound", "@@ - End Bone in target IK Rig Bone Chain not found.").ToString(), this);
+    	}
+    }
+}
+
+void UAnimGraphNode_IKRetargeter::PreloadRequiredAssets()
+{
+	Super::PreloadRequiredAssets();
+}
+
+#undef LOCTEXT_NAMESPACE
diff --git a/Engine/Plugins/Animation/IKRig/Source/IKRigDeveloper/Public/AnimGraphNode_IKRetargeter.h b/Engine/Plugins/Animation/IKRig/Source/IKRigDeveloper/Public/AnimGraphNode_IKRetargeter.h
new file mode 100644
--- /dev/null
+++ b/Engine/Plugins/Animation/IKRig/Source/IKRigDeveloper/Public/AnimGraphNode_IKRetargeter.h
@@ -0,0 +1,39 @@
+// Copyright Epic Games, Inc. All Rights Reserved.
+
+#pragma once
+
+#include "CoreMinimal.h"
+#include "UObject/ObjectMacros.h"
+#include "AnimGraphNode_Base.h"
+#include "AnimNodes/AnimNode_IKRetargeter.h"
+#include "AnimGraphNode_IKRetargeter.generated.h"
+
+class FPrimitiveDrawInterface;
+class USkeletalMeshComponent;
+
+// Editor node for IKRig 
+UCLASS()
+class IKRIGDEVELOPER_API UAnimGraphNode_IKRetargeter : public UAnimGraphNode_Base
+{
+	GENERATED_BODY()
+
+	UPROPERTY(EditAnywhere, Category = Settings)
+	FAnimNode_IKRetargeter Node;
+
+public:
+	// UEdGraphNode interface
+	virtual FText GetNodeTitle(ENodeTitleType::Type TitleType) const override;
+	virtual void PostEditChangeProperty(struct FPropertyChangedEvent& PropertyChangedEvent) override;
+	virtual void PostLoad() override;
+	virtual void ValidateAnimNodeDuringCompilation(USkeleton* ForSkeleton, FCompilerResultsLog& MessageLog);
+	// End of UEdGraphNode interface
+
+	// UAnimGraphNode_Base interface
+	virtual void PreloadRequiredAssets() override;
+	virtual void CopyNodeDataToPreviewNode(FAnimNode_Base* AnimNode) override;
+	virtual FEditorModeID GetEditorMode() const override;
+	virtual void Draw(FPrimitiveDrawInterface* PDI, USkeletalMeshComponent* PreviewSkelMeshComp) const override;
+	// End of UAnimGraphNode_Base interface
+
+	static const FName AnimModeName;
+};
diff --git a/Engine/Plugins/Animation/IKRig/Source/IKRigEditor/Private/AnimMode/IKRigEditMode.cpp b/Engine/Plugins/Animation/IKRig/Source/IKRigEditor/Private/AnimMode/IKRigEditMode.cpp
--- a/Engine/Plugins/Animation/IKRig/Source/IKRigEditor/Private/AnimMode/IKRigEditMode.cpp
+++ b/Engine/Plugins/Animation/IKRig/Source/IKRigEditor/Private/AnimMode/IKRigEditMode.cpp
@@ -5,18 +5,20 @@
 #include "IPersonaPreviewScene.h"
 #include "Animation/DebugSkelMeshComponent.h"
 
+// NOTE: This mode is disabled until we can revisit the idea of pulling effectors in ABP viewport
+
 void FIKRigEditMode::EnterMode(class UAnimGraphNode_Base* InEditorNode, struct FAnimNode_Base* InRuntimeNode)
 {
-	RuntimeNode = static_cast<FAnimNode_IKRig*>(InRuntimeNode);
-	GraphNode = CastChecked<UAnimGraphNode_IKRig>(InEditorNode);
+	//RuntimeNode = static_cast<FAnimNode_IKRig*>(InRuntimeNode);
+	//GraphNode = CastChecked<UAnimGraphNode_IKRig>(InEditorNode);
 
 	FAnimNodeEditMode::EnterMode(InEditorNode, InRuntimeNode);
 }
 
 void FIKRigEditMode::ExitMode()
 {
-	RuntimeNode = nullptr;
-	GraphNode = nullptr;
+	//RuntimeNode = nullptr;
+	//GraphNode = nullptr;
 
 	FAnimNodeEditMode::ExitMode();
 }
diff --git a/Engine/Plugins/Animation/IKRig/Source/IKRigEditor/Private/AssetTypeActions_IKRetargeter.cpp b/Engine/Plugins/Animation/IKRig/Source/IKRigEditor/Private/AssetTypeActions_IKRetargeter.cpp
new file mode 100644
--- /dev/null
+++ b/Engine/Plugins/Animation/IKRig/Source/IKRigEditor/Private/AssetTypeActions_IKRetargeter.cpp
@@ -0,0 +1,31 @@
+// Copyright Epic Games, Inc. All Rights Reserved.
+
+#include "AssetTypeActions_IKRetargeter.h"
+#include "Retargeter/IKRetargeter.h"
+
+//#include "ThumbnailRendering/SceneThumbnailInfo.h"
+//#include "ToolMenus.h"
+
+#define LOCTEXT_NAMESPACE "AssetTypeActions"
+
+UClass* FAssetTypeActions_IKRetargeter::GetSupportedClass() const
+{
+	return UIKRetargeter::StaticClass();
+}
+
+// UThumbnailInfo* FAssetTypeActions_IKRigDefinition::GetThumbnailInfo(UObject* Asset) const
+// {
+// 	UChaosSolver* ChaosSolver = CastChecked<UChaosSolver>(Asset);
+// 	return NewObject<USceneThumbnailInfo>(ChaosSolver, NAME_None, RF_Transactional);
+// }
+
+void FAssetTypeActions_IKRetargeter::GetActions(const TArray<UObject*>& InObjects, FToolMenuSection& Section)
+{
+	FAssetTypeActions_Base::GetActions(InObjects, Section);
+}
+
+// void FAssetTypeActions_IKRetargeter::OpenAssetEditor(const TArray<UObject*>& InObjects, TSharedPtr<IToolkitHost> EditWithinLevelEditor)
+// {
+// }
+
+#undef LOCTEXT_NAMESPACE
diff --git a/Engine/Plugins/Animation/IKRig/Source/IKRigEditor/Private/AssetTypeActions_IKRetargeter.h b/Engine/Plugins/Animation/IKRig/Source/IKRigEditor/Private/AssetTypeActions_IKRetargeter.h
new file mode 100644
--- /dev/null
+++ b/Engine/Plugins/Animation/IKRig/Source/IKRigEditor/Private/AssetTypeActions_IKRetargeter.h
@@ -0,0 +1,20 @@
+// Copyright Epic Games, Inc. All Rights Reserved.
+
+#pragma once
+
+#include "CoreMinimal.h"
+#include "Toolkits/IToolkitHost.h"
+#include "AssetTypeActions_Base.h"
+
+class FAssetTypeActions_IKRetargeter : public FAssetTypeActions_Base
+{
+public:
+	// IAssetTypeActions Implementation
+	virtual FText GetName() const override { return NSLOCTEXT("AssetTypeActions", "AssetTypeActions_IKRetargeter", "IK Retargeter"); }
+	virtual void GetActions(const TArray<UObject*>& InObjects, struct FToolMenuSection& Section) override;
+	virtual FColor GetTypeColor() const override { return FColor(255, 192, 128); }
+	virtual UClass* GetSupportedClass() const override;
+	//virtual void OpenAssetEditor(const TArray<UObject*>& InObjects, TSharedPtr<class IToolkitHost> EditWithinLevelEditor = TSharedPtr<IToolkitHost>()) override;
+	virtual uint32 GetCategories() override { return EAssetTypeCategories::Animation; }
+	//virtual UThumbnailInfo* GetThumbnailInfo(UObject* Asset) const override;
+};
diff --git a/Engine/Plugins/Animation/IKRig/Source/IKRigEditor/Private/IKRetargeterFactory.cpp b/Engine/Plugins/Animation/IKRig/Source/IKRigEditor/Private/IKRetargeterFactory.cpp
new file mode 100644
--- /dev/null
+++ b/Engine/Plugins/Animation/IKRig/Source/IKRigEditor/Private/IKRetargeterFactory.cpp
@@ -0,0 +1,51 @@
+// Copyright Epic Games, Inc. All Rights Reserved.
+
+#include "IKRetargeterFactory.h"
+#include "IKRig/Public/Retargeter/IKRetargeter.h"
+#include "AssetTypeCategories.h"
+
+#define LOCTEXT_NAMESPACE "IKRetargeterFactory"
+
+
+UIKRetargeterFactory::UIKRetargeterFactory()
+{
+	bCreateNew = true;
+	bEditAfterNew = true;
+	SupportedClass = UIKRetargeter::StaticClass();
+}
+
+UObject* UIKRetargeterFactory::FactoryCreateNew(UClass* Class, UObject* InParent, FName Name, EObjectFlags Flags, UObject* Context, FFeedbackContext* Warn)
+{
+	return NewObject<UIKRetargeter>(InParent, Name, Flags | RF_Transactional);;
+}
+
+bool UIKRetargeterFactory::ShouldShowInNewMenu() const
+{
+	return true;
+}
+
+bool UIKRetargeterFactory::ConfigureProperties()
+{
+	return true;
+}
+
+FText UIKRetargeterFactory::GetDisplayName() const
+{
+	return LOCTEXT("IKRetargeter_DisplayName", "IK Retargeter");
+}
+
+uint32 UIKRetargeterFactory::GetMenuCategories() const
+{
+	return EAssetTypeCategories::Animation;
+}
+
+FText UIKRetargeterFactory::GetToolTip() const
+{
+	return LOCTEXT("IKRetargeter_Tooltip", "Defines a pair of Source/Target Retarget Rigs and the mapping between them.");
+}
+
+FString UIKRetargeterFactory::GetDefaultNewAssetName() const
+{
+	return FString(TEXT("NewIKRetargeter"));
+}
+#undef LOCTEXT_NAMESPACE
diff --git a/Engine/Plugins/Animation/IKRig/Source/IKRigEditor/Private/IKRetargeterFactory.h b/Engine/Plugins/Animation/IKRig/Source/IKRigEditor/Private/IKRetargeterFactory.h
new file mode 100644
--- /dev/null
+++ b/Engine/Plugins/Animation/IKRig/Source/IKRigEditor/Private/IKRetargeterFactory.h
@@ -0,0 +1,28 @@
+// Copyright Epic Games, Inc. All Rights Reserved.
+
+#pragma once
+
+#include "CoreMinimal.h"
+#include "UObject/ObjectMacros.h"
+#include "Factories/Factory.h"
+#include "IKRetargeterFactory.generated.h"
+
+
+UCLASS(BlueprintType, hidecategories=Object)
+class UIKRetargeterFactory : public UFactory
+{
+	GENERATED_BODY()
+
+public:
+
+	UIKRetargeterFactory();
+
+	// UFactory Interface
+	virtual FText GetDisplayName() const override;
+	virtual uint32 GetMenuCategories() const override;
+	virtual FText GetToolTip() const override;
+	virtual FString GetDefaultNewAssetName() const override;
+	virtual UObject* FactoryCreateNew(UClass* Class, UObject* InParent, FName Name, EObjectFlags Flags, UObject* Context, FFeedbackContext* Warn) override;
+	virtual bool ConfigureProperties() override;
+	virtual bool ShouldShowInNewMenu() const override;
+};
diff --git a/Engine/Plugins/Animation/IKRig/Source/IKRigEditor/Private/IKRigController.cpp b/Engine/Plugins/Animation/IKRig/Source/IKRigEditor/Private/IKRigController.cpp
--- a/Engine/Plugins/Animation/IKRig/Source/IKRigEditor/Private/IKRigController.cpp
+++ b/Engine/Plugins/Animation/IKRig/Source/IKRigEditor/Private/IKRigController.cpp
@@ -65,23 +65,7 @@ void UIKRigController::SetIKRigDefinition(UIKRigDefinition* InIKRigDefinition)
 // SKELETON
 //
 
-const FIKRigHierarchy* UIKRigController::GetHierarchy() const
-{
-	return IKRigDefinition ? &IKRigDefinition->Hierarchy : nullptr;
-}
-
-const TArray<FTransform>& UIKRigController::GetRefPoseTransforms() const
-{
-	if (IKRigDefinition)
-	{
-		return IKRigDefinition->RefPoseTransforms;
-	}
-
-	static TArray<FTransform> Dummy;
-	return Dummy;
-}
-
-void UIKRigController::SetSkeleton(const FReferenceSkeleton& InSkeleton)
+void UIKRigController::SetSkeleton(const FReferenceSkeleton& InSkeleton) const
 {
 	if (!IKRigDefinition)
 	{
@@ -91,48 +75,15 @@ void UIKRigController::SetSkeleton(const FReferenceSkeleton& InSkeleton)
 	FScopedTransaction Transaction(LOCTEXT("SetSkeleton_Label", "Set Skeleton"));
 	IKRigDefinition->Modify();
 
-	IKRigDefinition->ResetHierarchy();
-
-	const TArray<FMeshBoneInfo>& RefBoneInfo = InSkeleton.GetRefBoneInfo();
-	TArray<FTransform> RefPoseInCS;
-	FAnimationRuntime::FillUpComponentSpaceTransforms(InSkeleton, InSkeleton.GetRefBonePose(), RefPoseInCS);
-	ensure(RefPoseInCS.Num() == RefBoneInfo.Num());
-	for (int32 Index = 0; Index < RefBoneInfo.Num(); ++Index)
-	{
-		int32 ParentIndex = RefBoneInfo[Index].ParentIndex;
-		ensure(IKRigDefinition->AddBone(RefBoneInfo[Index].Name, (ParentIndex != INDEX_NONE) ? RefBoneInfo[ParentIndex].Name : NAME_None, RefPoseInCS[Index]));
-	}
-
-	ensure(IKRigDefinition->Hierarchy.GetNum() == IKRigDefinition->RefPoseTransforms.Num());
+	IKRigDefinition->Skeleton.Initialize(InSkeleton);
 }
 
-bool UIKRigController::AddBone(const FName& InName, const FName& InParent, const FTransform& InGlobalTransform)
+FIKRigSkeleton& UIKRigController::GetSkeleton() const
 {
-	if (!IKRigDefinition)
-	{
-		return false;
-	}
-
-	FScopedTransaction Transaction(LOCTEXT("AddBone_Label", "Add Bone"));
-	IKRigDefinition->Modify();
-
-	return IKRigDefinition->AddBone(InName, InParent, InGlobalTransform);
+	return IKRigDefinition->Skeleton;
 }
 
 
-void UIKRigController::ResetHierarchy()
-{
-	if (!IKRigDefinition)
-	{
-		return;
-	}
-
-	FScopedTransaction Transaction(LOCTEXT("ResetHierarchy_Label", "Reset Hierarchy"));
-	IKRigDefinition->Modify();
-
-	IKRigDefinition->ResetHierarchy();
-}
-
 // -------------------------------------------------------
 // SOLVERS
 //
diff --git a/Engine/Plugins/Animation/IKRig/Source/IKRigEditor/Private/IKRigDefinitionFactory.cpp b/Engine/Plugins/Animation/IKRig/Source/IKRigEditor/Private/IKRigDefinitionFactory.cpp
--- a/Engine/Plugins/Animation/IKRig/Source/IKRigEditor/Private/IKRigDefinitionFactory.cpp
+++ b/Engine/Plugins/Animation/IKRig/Source/IKRigEditor/Private/IKRigDefinitionFactory.cpp
@@ -7,9 +7,6 @@
 #define LOCTEXT_NAMESPACE "IKRigDefinitionFactory"
 
 
-/* UFactory interface
- *****************************************************************************/
-
 UIKRigDefinitionFactory::UIKRigDefinitionFactory()
 {
 	bCreateNew = true;
@@ -44,7 +41,7 @@ uint32 UIKRigDefinitionFactory::GetMenuCategories() const
 
 FText UIKRigDefinitionFactory::GetToolTip() const
 {
-	return LOCTEXT("IKRigDefinition_Tooltip", "Create a customized IK Rig to be embedded in an Anim Graph or Control Rig.");
+	return LOCTEXT("IKRigDefinition_Tooltip", "Defines a set of IK Solvers and Effectors to pose a skeleton with Goals.");
 }
 
 FString UIKRigDefinitionFactory::GetDefaultNewAssetName() const
diff --git a/Engine/Plugins/Animation/IKRig/Source/IKRigEditor/Private/IKRigDefinitionFactory.h b/Engine/Plugins/Animation/IKRig/Source/IKRigEditor/Private/IKRigDefinitionFactory.h
--- a/Engine/Plugins/Animation/IKRig/Source/IKRigEditor/Private/IKRigDefinitionFactory.h
+++ b/Engine/Plugins/Animation/IKRig/Source/IKRigEditor/Private/IKRigDefinitionFactory.h
@@ -7,9 +7,7 @@
 #include "Factories/Factory.h"
 #include "IKRigDefinitionFactory.generated.h"
 
-/**
- * Implements a factory for UIKRigDefinitionFactory objects.
- */
+
 UCLASS(BlueprintType, hidecategories=Object)
 class UIKRigDefinitionFactory : public UFactory
 {
diff --git a/Engine/Plugins/Animation/IKRig/Source/IKRigEditor/Private/IKRigEditor.cpp b/Engine/Plugins/Animation/IKRig/Source/IKRigEditor/Private/IKRigEditor.cpp
--- a/Engine/Plugins/Animation/IKRig/Source/IKRigEditor/Private/IKRigEditor.cpp
+++ b/Engine/Plugins/Animation/IKRig/Source/IKRigEditor/Private/IKRigEditor.cpp
@@ -13,6 +13,7 @@
 #include "EditorModeRegistry.h"
 #include "AnimGraphNode_IKRig.h"
 #include "AssetTypeActions_IKRigDefinition.h"
+#include "AssetTypeActions_IKRetargeter.h"
 #include "IKRigDefinition.h"
 #include "AnimMode/IKRigEditMode.h"
 
@@ -24,13 +25,20 @@ DEFINE_LOG_CATEGORY_STATIC(LogIKRigEditor, Log, All);
 
 void FIKRigEditor::StartupModule()
 {
+	// register IKRigDefinition asset type
 	IKRigDefinitionAssetAction = MakeShareable(new FAssetTypeActions_IKRigDefinition);
 	FModuleManager::LoadModuleChecked<FAssetToolsModule>("AssetTools").Get().RegisterAssetTypeActions(IKRigDefinitionAssetAction.ToSharedRef());
 
+	// register IKRetargeter asset type
+	IKRetargeterAssetAction = MakeShareable(new FAssetTypeActions_IKRetargeter);
+	FModuleManager::LoadModuleChecked<FAssetToolsModule>("AssetTools").Get().RegisterAssetTypeActions(IKRetargeterAssetAction.ToSharedRef());
+
+	// register details panel customization
 	FPropertyEditorModule& PropertyModule = FModuleManager::GetModuleChecked<FPropertyEditorModule>("PropertyEditor");
 	PropertyModule.RegisterCustomClassLayout(UIKRigDefinition::StaticClass()->GetFName(), FOnGetDetailCustomizationInstance::CreateStatic(&FIKRigDefinitionDetails::MakeInstance));
 
-	FEditorModeRegistry::Get().RegisterMode<FIKRigEditMode>(UAnimGraphNode_IKRig::AnimModeName, LOCTEXT("IKRigEditMode", "IKRig"), FSlateIcon(), false);
+	// register custom editor mode (TBD)
+	//FEditorModeRegistry::Get().RegisterMode<FIKRigEditMode>(UAnimGraphNode_IKRig::AnimModeName, LOCTEXT("IKRigEditMode", "IKRig"), FSlateIcon(), false);
 }
 
 void FIKRigEditor::ShutdownModule()
diff --git a/Engine/Plugins/Animation/IKRig/Source/IKRigEditor/Public/IKRigController.h b/Engine/Plugins/Animation/IKRig/Source/IKRigEditor/Public/IKRigController.h
--- a/Engine/Plugins/Animation/IKRig/Source/IKRigEditor/Public/IKRigController.h
+++ b/Engine/Plugins/Animation/IKRig/Source/IKRigEditor/Public/IKRigController.h
@@ -26,11 +26,8 @@ public:
 	static UIKRigController* GetControllerByRigDefinition(UIKRigDefinition* InIKRigDefinition);
 
 	//** skeleton */
-	const FIKRigHierarchy* GetHierarchy() const;
-	const TArray<FTransform>& GetRefPoseTransforms() const;
-	void SetSkeleton(const FReferenceSkeleton& InSkeleton);
-	bool AddBone(const FName& InName, const FName& InParent, const FTransform& InGlobalTransform);
-	void ResetHierarchy();
+	void SetSkeleton(const FReferenceSkeleton& InSkeleton) const;
+	FIKRigSkeleton& GetSkeleton() const;
 
 	//** solvers */
 	UIKRigSolver* AddSolver(TSubclassOf<UIKRigSolver> InSolverClass);
diff --git a/Engine/Plugins/Animation/IKRig/Source/IKRigEditor/Public/IKRigEditor.h b/Engine/Plugins/Animation/IKRig/Source/IKRigEditor/Public/IKRigEditor.h
--- a/Engine/Plugins/Animation/IKRig/Source/IKRigEditor/Public/IKRigEditor.h
+++ b/Engine/Plugins/Animation/IKRig/Source/IKRigEditor/Public/IKRigEditor.h
@@ -13,4 +13,6 @@ public:
 
 private:
 	TSharedPtr<class FAssetTypeActions_IKRigDefinition> IKRigDefinitionAssetAction;
+	TSharedPtr<class FAssetTypeActions_IKRetargeter> IKRetargeterAssetAction;
+	TSharedPtr<class FAssetTypeActions_IKRetargetRig> IKRetargetRigAssetAction;
 };
