Fixing CIS warning, "Dereferencing of the null pointer 'IKRigProcessor'"
JIRA: UE-112082

diff --git a/Engine/Plugins/Animation/IKRig/Source/IKRig/Private/AnimNodes/AnimNode_IKRig.cpp b/Engine/Plugins/Animation/IKRig/Source/IKRig/Private/AnimNodes/AnimNode_IKRig.cpp
--- a/Engine/Plugins/Animation/IKRig/Source/IKRig/Private/AnimNodes/AnimNode_IKRig.cpp
+++ b/Engine/Plugins/Animation/IKRig/Source/IKRig/Private/AnimNodes/AnimNode_IKRig.cpp
@@ -29,7 +29,13 @@ void FAnimNode_IKRig::Evaluate_AnyThread(FPoseContext& Output)
 		SourcePose.ResetToRefPose();
 	}
 
-	if (!IKRigProcessor && IKRigProcessor->IsInitialized())
+	if (!IKRigProcessor)
+	{
+		Output = SourcePose;
+		return;
+	}
+	
+	if (!IKRigProcessor->IsInitialized())
 	{
 		Output = SourcePose;
 		return;
