Adding IK Rig actor component/interface.

fyi: fernando.castillocoello

diff --git a/Engine/Plugins/Animation/IKRig/Source/IKRig/Private/ActorComponents/IKRigComponent.cpp b/Engine/Plugins/Animation/IKRig/Source/IKRig/Private/ActorComponents/IKRigComponent.cpp
new file mode 100644
--- /dev/null
+++ b/Engine/Plugins/Animation/IKRig/Source/IKRig/Private/ActorComponents/IKRigComponent.cpp
@@ -0,0 +1 @@
+// Copyright Epic Games, Inc. All Rights Reserved.
\ No newline at end of file
diff --git a/Engine/Plugins/Animation/IKRig/Source/IKRig/Private/AnimNodes/AnimNode_IKRig.cpp b/Engine/Plugins/Animation/IKRig/Source/IKRig/Private/AnimNodes/AnimNode_IKRig.cpp
--- a/Engine/Plugins/Animation/IKRig/Source/IKRig/Private/AnimNodes/AnimNode_IKRig.cpp
+++ b/Engine/Plugins/Animation/IKRig/Source/IKRig/Private/AnimNodes/AnimNode_IKRig.cpp
@@ -3,6 +3,8 @@
 #include "AnimNodes/AnimNode_IKRig.h"
 #include "IKRigProcessor.h"
 #include "IKRigDefinition.h"
+#include "IKRigSolver.h"
+#include "ActorComponents/IKRigInterface.h"
 #include "Drawing/ControlRigDrawInterface.h"
 #include "Animation/AnimInstanceProxy.h"
 
@@ -20,7 +22,7 @@ void FAnimNode_IKRig::Evaluate_AnyThread(FPoseContext& Output)
 	DECLARE_SCOPE_HIERARCHICAL_COUNTER_FUNC()
 	FPoseContext SourcePose(Output);
 
-	if (Source.GetLinkNode())
+	if (Source.GetLinkNode() && !bStartFromRefPose)
 	{
 		Source.Evaluate(SourcePose);
 	}
@@ -66,15 +68,15 @@ void FAnimNode_IKRig::Evaluate_AnyThread(FPoseContext& Output)
 	}
 
 	// update goal transforms before solve
-	for (int32 GoalIndex = 0; GoalIndex < GoalNames.Num(); ++GoalIndex)
+	for (const FIKRigGoal& Goal : Goals)
 	{
-		if (ensure(GoalTransforms.IsValidIndex(GoalIndex)))
-		{
-			IKRigProcessor->SetGoalTransform(
-				GoalNames[GoalIndex], 
-				GoalTransforms[GoalIndex].GetLocation(),
-				GoalTransforms[GoalIndex].GetRotation());
-		}
+		IKRigProcessor->SetIKGoal(Goal);
+	}
+
+	// goals from goal creator components will override any goals that were manually set (they take precedence)
+	for (const TPair<FName, FIKRigGoal>& GoalPair : GoalsFromGoalCreators)
+	{
+		IKRigProcessor->SetIKGoal(GoalPair.Value);
 	}
 
 	// run stack of solvers, updates global transforms with new pose
@@ -93,8 +95,6 @@ void FAnimNode_IKRig::Evaluate_AnyThread(FPoseContext& Output)
 		}
 	}
 
-	Output.Pose.NormalizeRotations();
-
 	// draw the interface
 	QueueDrawInterface(Output.AnimInstanceProxy, Output.AnimInstanceProxy->GetComponentTransform());
 }
@@ -102,8 +102,23 @@ void FAnimNode_IKRig::Evaluate_AnyThread(FPoseContext& Output)
 void FAnimNode_IKRig::GatherDebugData(FNodeDebugData& DebugData)
 {
 	DECLARE_SCOPE_HIERARCHICAL_COUNTER_FUNC()
-	// alpha later?
-	Source.GatherDebugData(DebugData.BranchFlow(1.f));
+
+	DebugData.AddDebugItem(FString::Printf(TEXT("%s IK Rig evaluated with %d Goals."), *DebugData.GetNodeName(this), Goals.Num()));
+		
+	for (const TPair<FName, FIKRigGoal>& GoalPair : GoalsFromGoalCreators)
+	{
+		DebugData.AddDebugItem(FString::Printf(TEXT("Goal supplied by actor component: %s"), *GoalPair.Value.ToString()));
+	}
+
+	for (const FIKRigGoal& Goal : Goals)
+	{
+		if (GoalsFromGoalCreators.Contains(Goal.Name))
+		{
+			continue;
+		}
+		
+		DebugData.AddDebugItem(FString::Printf(TEXT("Goal supplied by node pin: %s"), *Goal.ToString()));
+	}
 }
 
 void FAnimNode_IKRig::Initialize_AnyThread(const FAnimationInitializeContext& Context)
@@ -115,10 +130,24 @@ void FAnimNode_IKRig::Initialize_AnyThread(const FAnimationInitializeContext& Co
 	{
 		return;
 	}
-	
+
+	// cache list of goal creator components on the actor
+	GoalCreators.Reset();
+	USkeletalMeshComponent* SkelMeshComponent = Context.AnimInstanceProxy->GetSkelMeshComponent();
+	AActor* OwningActor = SkelMeshComponent->GetOwner();
+	TArray<UActorComponent*> GoalCreatorComponents =  OwningActor->GetComponentsByInterface( UIKGoalCreatorInterface::StaticClass() );
+	for (UActorComponent* GoalCreatorComponent : GoalCreatorComponents)
+	{
+		IIKGoalCreatorInterface* GoalCreator = Cast<IIKGoalCreatorInterface>(GoalCreatorComponent);
+		if (!ensureMsgf(GoalCreator, TEXT("Goal creator component failed cast to IIKGoalCreatorInterface.")))
+		{
+			continue;
+		}
+		GoalCreators.Add(GoalCreator);
+	}
+
+	// initialize the IK Rig
 	IKRigProcessor->Initialize(RigDefinitionAsset);
-	GoalNames.Reset();
-	IKRigProcessor->GetGoalNames(GoalNames);
 }
 
 void FAnimNode_IKRig::OnInitializeAnimInstance(const FAnimInstanceProxy* InProxy, const UAnimInstance* InAnimInstance)
@@ -132,6 +161,22 @@ void FAnimNode_IKRig::OnInitializeAnimInstance(const FAnimInstanceProxy* InProxy
 	}
 }
 
+void FAnimNode_IKRig::PreUpdate(const UAnimInstance* InAnimInstance)
+{
+	// pull all the goals out of any goal creators on the owning actor
+	// this is done on the main thread because we're talking to actor components here
+	GoalsFromGoalCreators.Empty();
+	for (IIKGoalCreatorInterface* GoalCreator : GoalCreators)
+	{
+		TMap<FName, FIKRigGoal> OutGoals;
+		GoalCreator->GetIKGoals_Implementation(OutGoals);
+		for (const TPair<FName, FIKRigGoal>& GoalPair : OutGoals)
+		{
+			GoalsFromGoalCreators.Add(GoalPair);
+		}
+	}
+}
+
 bool FAnimNode_IKRig::RebuildGoalList()
 {
 	if (!RigDefinitionAsset)
@@ -139,12 +184,16 @@ bool FAnimNode_IKRig::RebuildGoalList()
 		return false;
 	}
 	
-	TArray<FName> SolverGoals;
-	RigDefinitionAsset->GetGoalNamesFromSolvers(SolverGoals);
-	const int32 GoalNum = SolverGoals.Num();
-	if (GoalTransforms.Num() != GoalNum)
+	TArray<FIKRigEffectorGoal> AllGoalNamesInIKRig;
+	RigDefinitionAsset->GetGoalNamesFromSolvers(AllGoalNamesInIKRig);
+	const int32 NumGoalsInRig = AllGoalNamesInIKRig.Num();
+	if (Goals.Num() != NumGoalsInRig)
 	{
-		GoalTransforms.SetNum(GoalNum);
+		Goals.SetNum(NumGoalsInRig);
+		for (int32 i=0; i<NumGoalsInRig; ++i)
+		{
+			Goals[i].Name = AllGoalNamesInIKRig[i].Goal;
+		}
 		return true;
 	}
 
@@ -155,11 +204,11 @@ FName FAnimNode_IKRig::GetGoalName(int32 Index) const
 {
 	if (RigDefinitionAsset)
 	{
-		TArray<FName> Names;
+		TArray<FIKRigEffectorGoal> Names;
 		RigDefinitionAsset->GetGoalNamesFromSolvers(Names);
 		if (Names.IsValidIndex(Index))
 		{
-			return Names[Index];
+			return Names[Index].Goal;
 		}
 	}
 
@@ -203,7 +252,7 @@ void FAnimNode_IKRig::CacheBones_AnyThread(const FAnimationCacheBonesContext& Co
 	}
 }
 
-void FAnimNode_IKRig::QueueDrawInterface(FAnimInstanceProxy* AnimProxy, const FTransform& ComponentToWorld)
+void FAnimNode_IKRig::QueueDrawInterface(FAnimInstanceProxy* AnimProxy, const FTransform& ComponentToWorld) const
 {
 	check (IKRigProcessor);
 	check (AnimProxy);
diff --git a/Engine/Plugins/Animation/IKRig/Source/IKRig/Private/IKRigDataTypes.cpp b/Engine/Plugins/Animation/IKRig/Source/IKRig/Private/IKRigDataTypes.cpp
--- a/Engine/Plugins/Animation/IKRig/Source/IKRig/Private/IKRigDataTypes.cpp
+++ b/Engine/Plugins/Animation/IKRig/Source/IKRig/Private/IKRigDataTypes.cpp
@@ -1,3 +1,29 @@
 // Copyright Epic Games, Inc. All Rights Reserved.
 
-#include "IKRigDataTypes.h"
\ No newline at end of file
+#include "IKRigDataTypes.h"
+#include "IKRigSolver.h"
+
+void FIKRigGoalContainer::InitializeGoalsFromNames(const TArray<FIKRigEffectorGoal>& InGoalNames)
+{
+	Goals.Empty(InGoalNames.Num());
+	for (const FIKRigEffectorGoal& GoalNames : InGoalNames)
+	{
+		Goals.Emplace(GoalNames.Goal, GoalNames.Goal);
+	}
+}
+
+void FIKRigGoalContainer::SetIKGoal(const FIKRigGoal& InGoal)
+{
+	Goals.Add(InGoal.Name, InGoal);
+}
+
+bool FIKRigGoalContainer::GetGoalByName(const FName& InGoalName, FIKRigGoal& OutGoal) const
+{
+	if (const FIKRigGoal* Goal = Goals.Find(InGoalName))
+	{
+		OutGoal = *Goal;
+		return true;
+	}
+
+	return false;
+}
\ No newline at end of file
diff --git a/Engine/Plugins/Animation/IKRig/Source/IKRig/Private/IKRigDefinition.cpp b/Engine/Plugins/Animation/IKRig/Source/IKRig/Private/IKRigDefinition.cpp
--- a/Engine/Plugins/Animation/IKRig/Source/IKRig/Private/IKRigDefinition.cpp
+++ b/Engine/Plugins/Animation/IKRig/Source/IKRig/Private/IKRigDefinition.cpp
@@ -13,28 +13,36 @@ FBoneChain* FRetargetDefinition::GetBoneChainByName(FName ChainName)
 			return &Chain;
 		}
 	}
+	
 	return nullptr;
 }
 
-void UIKRigDefinition::GetGoalNamesFromSolvers(TArray<FName>& OutGoalNames) const
+void UIKRigDefinition::GetGoalNamesFromSolvers(TArray<FIKRigEffectorGoal>& OutGoalNames) const
 {
+	TSet<FIKRigEffectorGoal> GoalNames;
 	for (UIKRigSolver* Solver : Solvers)
 	{
-		if (!Solver)
+		if (Solver)
 		{
-			continue;
+			Solver->CollectGoalNames(GoalNames);
 		}
-		
-		TSet<FName> GoalNames;
-		Solver->CollectGoalNames(GoalNames);
+	}
+
+	// user code needs to use indices, so we bake set into an array
+	OutGoalNames = GoalNames.Array();
+}
 
-		// not using a TSet here because user code relies on indices
-		for (FName Name : GoalNames)
+FName UIKRigDefinition::GetBoneNameForGoal(FName GoalName) const
+{
+	TArray<FIKRigEffectorGoal> AllGoalNames;
+	GetGoalNamesFromSolvers(AllGoalNames);
+	for (const FIKRigEffectorGoal& Names : AllGoalNames)
+	{
+		if (Names.Goal == GoalName)
 		{
-			if (!OutGoalNames.Contains(Name))
-			{
-				OutGoalNames.Add(Name);
-			}
+			return Names.Bone;
 		}
 	}
-}
\ No newline at end of file
+
+	return NAME_None;
+}
diff --git a/Engine/Plugins/Animation/IKRig/Source/IKRig/Private/IKRigProcessor.cpp b/Engine/Plugins/Animation/IKRig/Source/IKRig/Private/IKRigProcessor.cpp
--- a/Engine/Plugins/Animation/IKRig/Source/IKRig/Private/IKRigProcessor.cpp
+++ b/Engine/Plugins/Animation/IKRig/Source/IKRig/Private/IKRigProcessor.cpp
@@ -23,9 +23,37 @@ void UIKRigProcessor::Initialize(UIKRigDefinition* InRigDefinition)
 	Skeleton = InRigDefinition->Skeleton; // trivial copy assignment operator for POD
 
 	// initialize goal names based on solvers
-	TArray<FName> GoalNames;
+	TArray<FIKRigEffectorGoal> GoalNames;
 	InRigDefinition->GetGoalNamesFromSolvers(GoalNames);
-	Goals.InitializeGoalsFromNames(GoalNames);
+	GoalContainer.InitializeGoalsFromNames(GoalNames);
+	for (const FIKRigEffectorGoal& GoalName : GoalNames)
+	{
+		FGoalBone NewGoalBone;
+		NewGoalBone.BoneName = GoalName.Goal;
+		NewGoalBone.BoneIndex = Skeleton.GetBoneIndexFromName(GoalName.Bone);
+
+		// validate that the skeleton we are trying to solve this goal on contains the bone the goal expects
+		if (NewGoalBone.BoneIndex == INDEX_NONE)
+		{
+			UE_LOG(LogTemp, Warning, TEXT("IK Rig, %s has a Goal, '%s' that references an unknown bone, '%s'. Cannot evaluate."),
+				*GetName(), *GoalName.Goal.ToString(), *GoalName.Bone.ToString());
+			return;
+		}
+
+		// validate that there is not already a different goal, with the same name, that is using a different bone
+		// (all goals with the same name must reference the same bone within a single IK Rig)
+		if (const FGoalBone* Bone = GoalBones.Find(GoalName.Goal))
+		{
+			if (Bone->BoneName != NewGoalBone.BoneName)
+			{
+				UE_LOG(LogTemp, Warning, TEXT("IK Rig, %s has a Goal, '%s' that references different bones in different solvers, '%s' and '%s'. Cannot evaluate."),
+                *GetName(), *GoalName.Goal.ToString(), *Bone->BoneName.ToString(), *NewGoalBone.BoneName.ToString());
+				return;
+			}
+		}
+		
+		GoalBones.Add(GoalName.Goal, NewGoalBone);
+	}
 
 	// create copies of all the solvers in the IK rig
 	Solvers.Reset(InRigDefinition->Solvers.Num());
@@ -35,6 +63,7 @@ void UIKRigProcessor::Initialize(UIKRigDefinition* InRigDefinition)
 		{
 			// this can happen if asset references deleted IK Solver type
 			// which should only happen during development (if at all)
+			UE_LOG(LogTemp, Warning, TEXT("IK Rig, %s has null/unknown solver in it. Please remove it."), *GetName());
 			continue;
 		}
 		
@@ -61,45 +90,69 @@ void UIKRigProcessor::SetInputPoseToRefPose()
 	Skeleton.UpdateAllLocalTransformFromGlobal();
 }
 
-void UIKRigProcessor::SetGoalTransform(
-    const FName& GoalName,
-    const FVector& Position,
-    const FQuat& Rotation)
+void UIKRigProcessor::SetIKGoal(const FIKRigGoal& InGoal)
 {
 	check(bInitialized);
-	Goals.SetGoalTransform(GoalName, Position, Rotation);
+	GoalContainer.SetIKGoal(InGoal);
 }
 
 void UIKRigProcessor::Solve()
 {
 	check(bInitialized);
+	
 	DrawInterface.Reset();
+
+	// blend goals towards input pose by alpha
+	for (TPair<FName, FIKRigGoal>& GoalPair : GoalContainer.Goals)
+	{
+		FIKRigGoal& Goal = GoalPair.Value;
+		const FGoalBone& GoalBone = GoalBones[Goal.Name];
+		const FTransform& InputPoseBoneTransform = Skeleton.CurrentPoseGlobal[GoalBone.BoneIndex];
+		Goal.Position = FMath::Lerp(
+			InputPoseBoneTransform.GetTranslation(),
+			Goal.Position,
+			Goal.PositionAlpha);
+		Goal.Rotation = FQuat::FastLerp(
+			InputPoseBoneTransform.GetRotation(),
+			Goal.Rotation.Quaternion(),
+			Goal.RotationAlpha).Rotator();
+	}
+
+	// run all the solvers
 	for (UIKRigSolver* Solver : Solvers)
 	{
-		Solver->Solve(Skeleton, Goals, &DrawInterface);
+		Solver->Solve(Skeleton, GoalContainer, &DrawInterface);
 	}
+
+	// make sure rotations are normalized coming out
+	Skeleton.NormalizeRotations(Skeleton.CurrentPoseGlobal);
 }
 
 void UIKRigProcessor::CopyOutputGlobalPoseToArray(TArray<FTransform>& OutputPoseGlobal) const
 {
-	check(bInitialized);
+	check(OutputPoseGlobal.Num() == Skeleton.CurrentPoseGlobal.Num());
 	OutputPoseGlobal = Skeleton.CurrentPoseGlobal;
 }
 
-void UIKRigProcessor::GetGoalNames(TArray<FName>& OutGoalNames) const
+const FIKRigGoalContainer& UIKRigProcessor::GetGoalContainer() const
 {
 	check(bInitialized);
-	return Goals.GetNames(OutGoalNames);
+	return GoalContainer;
 }
 
-int UIKRigProcessor::GetNumGoals() const
+FIKRigSkeleton& UIKRigProcessor::GetSkeleton()
 {
 	check(bInitialized);
-	return Goals.GetNumGoals();
+	return Skeleton;
 }
 
-FIKRigSkeleton& UIKRigProcessor::GetSkeleton()
+bool UIKRigProcessor::GetBoneForGoal(FName GoalName, FGoalBone& OutBone) const
 {
-	check(bInitialized);
-	return Skeleton;
+	const FGoalBone* FoundBone = GoalBones.Find(GoalName);
+	if (FoundBone)
+	{
+		OutBone = *FoundBone;
+		return true;
+	}
+	return false;
 }
diff --git a/Engine/Plugins/Animation/IKRig/Source/IKRig/Private/IKRigSkeleton.cpp b/Engine/Plugins/Animation/IKRig/Source/IKRig/Private/IKRigSkeleton.cpp
--- a/Engine/Plugins/Animation/IKRig/Source/IKRig/Private/IKRigSkeleton.cpp
+++ b/Engine/Plugins/Animation/IKRig/Source/IKRig/Private/IKRigSkeleton.cpp
@@ -113,3 +113,11 @@ void FIKRigSkeleton::PropagateGlobalPoseBelowBone(const int32 StartBoneIndex)
 		UpdateLocalTransformFromGlobal(BoneIndex);
 	}
 }
+
+void FIKRigSkeleton::NormalizeRotations(TArray<FTransform>& Transforms)
+{
+	for (FTransform& Transform : Transforms)
+	{
+		Transform.NormalizeRotation();
+	}
+}
\ No newline at end of file
diff --git a/Engine/Plugins/Animation/IKRig/Source/IKRig/Private/Retargeter/IKRetargeter.cpp b/Engine/Plugins/Animation/IKRig/Source/IKRig/Private/Retargeter/IKRetargeter.cpp
--- a/Engine/Plugins/Animation/IKRig/Source/IKRig/Private/Retargeter/IKRetargeter.cpp
+++ b/Engine/Plugins/Animation/IKRig/Source/IKRig/Private/Retargeter/IKRetargeter.cpp
@@ -783,21 +783,21 @@ bool UIKRetargeter::InitializeBoneChainPairs()
 bool UIKRetargeter::InitializeIKRig(UObject* Outer)
 {
 	// make a new IKRigProcessor if we haven't already
-	if (!IKRig)
+	if (!IKRigProcessor)
 	{
-		IKRig = UIKRigProcessor::MakeNewIKRigProcessor(Outer);
+		IKRigProcessor = UIKRigProcessor::MakeNewIKRigProcessor(Outer);
 	}
 	
 	// initialize IK Rig runtime processor
-	IKRig->Initialize(TargetIKRigAsset);
-	if (!IKRig->IsInitialized())
+	IKRigProcessor->Initialize(TargetIKRigAsset);
+	if (!IKRigProcessor->IsInitialized())
 	{
 		return false;
 	}
 
 	// validate that all IK bone chains have an associated Goal
 	TArray<FName> GoalNames;
-	IKRig->GetGoalNames(GoalNames);
+	IKRigProcessor->GetGoalContainer().Goals.GenerateKeyArray(GoalNames);
 	for (FRetargetChainPairIK& ChainPair : ChainPairsIK)
 	{
 		// does the IK rig have the IK goal this bone chain requires?
@@ -876,13 +876,19 @@ void UIKRetargeter::RunIKRetarget(
 		FDecodedIKChain OutIKGoal;
 		ChainPair.IKDecoder.DecodePose(ChainPair.IKEncoder.EncodedResult, OutGlobalPose, OutIKGoal);
 		// set the goal transform on the IK Rig
-		IKRig->SetGoalTransform(ChainPair.IKGoalName, OutIKGoal.EndEffectorPosition, OutIKGoal.EndEffectorRotation);
+		FIKRigGoal Goal = FIKRigGoal(
+			ChainPair.IKGoalName,
+			OutIKGoal.EndEffectorPosition,
+			OutIKGoal.EndEffectorRotation,
+			1.0f,
+			1.0f);
+		IKRigProcessor->SetIKGoal(Goal);
 	}
 
 	// copy input pose to start IK solve from
-	IKRig->SetInputPoseGlobal(OutGlobalPose);
+	IKRigProcessor->SetInputPoseGlobal(OutGlobalPose);
 	// run IK solve
-	IKRig->Solve();
+	IKRigProcessor->Solve();
 	// copy results of solve
-	IKRig->CopyOutputGlobalPoseToArray(OutGlobalPose);
+	IKRigProcessor->CopyOutputGlobalPoseToArray(OutGlobalPose);
 }
diff --git a/Engine/Plugins/Animation/IKRig/Source/IKRig/Private/Solvers/ConstraintSolver.cpp b/Engine/Plugins/Animation/IKRig/Source/IKRig/Private/Solvers/ConstraintSolver.cpp
--- a/Engine/Plugins/Animation/IKRig/Source/IKRig/Private/Solvers/ConstraintSolver.cpp
+++ b/Engine/Plugins/Animation/IKRig/Source/IKRig/Private/Solvers/ConstraintSolver.cpp
@@ -30,7 +30,7 @@ void UIKRigConstraintSolver::Solve(
 	}*/
 }
 
-void UIKRigConstraintSolver::CollectGoalNames(TSet<FName>& OutGoals) const
+void UIKRigConstraintSolver::CollectGoalNames(TSet<FIKRigEffectorGoal>& OutGoals) const
 {
 
 }
\ No newline at end of file
diff --git a/Engine/Plugins/Animation/IKRig/Source/IKRig/Private/Solvers/IKRigFBIKSolver.cpp b/Engine/Plugins/Animation/IKRig/Source/IKRig/Private/Solvers/IKRigFBIKSolver.cpp
--- a/Engine/Plugins/Animation/IKRig/Source/IKRig/Private/Solvers/IKRigFBIKSolver.cpp
+++ b/Engine/Plugins/Animation/IKRig/Source/IKRig/Private/Solvers/IKRigFBIKSolver.cpp
@@ -2,6 +2,7 @@
 
 #include "Solvers/IKRigFBIKSolver.h"
 #include "IKRigDataTypes.h"
+#include "IKRigSkeleton.h"
 #include "FBIKConstraint.h"
 #include "JacobianIK.h"
 #include "FBIKUtil.h"
@@ -280,10 +281,10 @@ void UIKRigFBIKSolver::Solve(
 					const FFBIKRigEffector& CurEffector = Effectors[EffectorIndex];
 					const FVector CurrentLinkLocation = LinkData[EffectorLinkIndex].GetTransform().GetLocation();
 					const FQuat CurrentLinkRotation = LinkData[EffectorLinkIndex].GetTransform().GetRotation();
-					FIKRigGoal Goal;
-					ensure(GetGoalForEffector(CurEffector.Target, Goals, Goal));
-					EffectorTarget->Position = Goal.Position; // FMath::Lerp(CurrentLinkLocation, EffectorLocation, RigTarget.PositionTarget.);
-					EffectorTarget->Rotation = Goal.Rotation; // FMath::Lerp(CurrentLinkRotation, EffectorRotation, CurEffector.RotationAlpha);
+					FIKRigGoal OutGoal;
+					ensure(Goals.GetGoalByName(CurEffector.Target.Goal, OutGoal));
+					EffectorTarget->Position = OutGoal.Position; // FMath::Lerp(CurrentLinkLocation, EffectorLocation, RigTarget.PositionTarget.);
+					EffectorTarget->Rotation = OutGoal.Rotation.Quaternion(); // FMath::Lerp(CurrentLinkRotation, EffectorRotation, CurEffector.RotationAlpha);
 					EffectorTarget->InitialPositionDistance = (EffectorTarget->Position - CurrentLinkLocation).Size();
 					EffectorTarget->InitialRotationDistance = (FBIKUtil::GetScaledRotationAxis(EffectorTarget->Rotation) - FBIKUtil::GetScaledRotationAxis(CurrentLinkRotation)).Size();
 
@@ -566,10 +567,10 @@ void UIKRigFBIKSolver::Solve(
 	}
 }
 
-void UIKRigFBIKSolver::CollectGoalNames(TSet<FName>& OutGoals) const
+void UIKRigFBIKSolver::CollectGoalNames(TSet<FIKRigEffectorGoal>& OutGoals) const
 {
 	for (const FFBIKRigEffector& Effector: Effectors)
 	{
-		OutGoals.Add(Effector.Target.Goal);
+		OutGoals.Add(Effector.Target);
 	}
 }
diff --git a/Engine/Plugins/Animation/IKRig/Source/IKRig/Private/Solvers/TransformSolver.cpp b/Engine/Plugins/Animation/IKRig/Source/IKRig/Private/Solvers/TransformSolver.cpp
--- a/Engine/Plugins/Animation/IKRig/Source/IKRig/Private/Solvers/TransformSolver.cpp
+++ b/Engine/Plugins/Animation/IKRig/Source/IKRig/Private/Solvers/TransformSolver.cpp
@@ -3,7 +3,7 @@
 
 #include "Solvers/TransformSolver.h"
 #include "IKRigDataTypes.h"
-
+#include "IKRigSkeleton.h"
 
 UTransformSolver::UTransformSolver()
 {
@@ -20,8 +20,8 @@ void UTransformSolver::Solve(
 	const FIKRigGoalContainer& Goals,
 	FControlRigDrawInterface* InOutDrawInterface)
 {
-	FIKRigGoal Goal;
-	if (!GetGoalForEffector(Effector, Goals, Goal))
+	FIKRigGoal OutGoal;
+	if (!Goals.GetGoalByName(Effector.Goal, OutGoal))
 	{
 		return;
 	}
@@ -30,18 +30,18 @@ void UTransformSolver::Solve(
 
 	if (bEnablePosition)
 	{
-		CurrentTransform.SetLocation(Goal.Position);
+		CurrentTransform.SetLocation(OutGoal.Position);
 	}
 	
 	if (bEnableRotation)
 	{
-		CurrentTransform.SetRotation(Goal.Rotation);
+		CurrentTransform.SetRotation(OutGoal.Rotation.Quaternion());
 	}
 
 	IKRigSkeleton.PropagateGlobalPoseBelowBone(BoneIndex);
 }
 
-void UTransformSolver::CollectGoalNames(TSet<FName>& OutGoals) const
+void UTransformSolver::CollectGoalNames(TSet<FIKRigEffectorGoal>& OutGoals) const
 {
-	OutGoals.Add(Effector.Goal);
+	OutGoals.Add(Effector);
 }
diff --git a/Engine/Plugins/Animation/IKRig/Source/IKRig/Public/ActorComponents/IKRigComponent.h b/Engine/Plugins/Animation/IKRig/Source/IKRig/Public/ActorComponents/IKRigComponent.h
new file mode 100644
--- /dev/null
+++ b/Engine/Plugins/Animation/IKRig/Source/IKRig/Public/ActorComponents/IKRigComponent.h
@@ -0,0 +1,63 @@
+// Copyright Epic Games, Inc. All Rights Reserved.
+
+#pragma once
+
+#include "Components/ActorComponent.h"
+
+#include "../IKRigDataTypes.h"
+#include "IKRigInterface.h"
+
+#include "IKRigComponent.generated.h"
+
+
+UCLASS(ClassGroup = IKRig, BlueprintType, Blueprintable, EditInlineNew, meta = (BlueprintSpawnableComponent))
+class IKRIG_API UIKRigComponent : public UActorComponent, public IIKGoalCreatorInterface
+{
+	GENERATED_BODY()
+
+public:
+
+	/** Set an IK Rig Goal position and rotation (assumed in Component Space of Skeletal Mesh) with separate alpha values. */
+	UFUNCTION(BlueprintCallable, Category=IKRigGoals)
+    void SetIKRigGoalPositionAndRotation(
+        const FName GoalName,
+        const FVector Position,
+        const FQuat Rotation,
+        const float PositionAlpha,
+        const float RotationAlpha)
+	{
+		GoalContainer.SetIKGoal(FIKRigGoal(GoalName, Position, Rotation, PositionAlpha, RotationAlpha));
+	};
+
+	/** Set an IK Rig Goal transform (assumed in Component Space of Skeletal Mesh) with separate alpha values. */
+	UFUNCTION(BlueprintCallable, Category=IKRigGoals)
+    void SetIKRigGoalTransform(
+        const FName GoalName,
+        const FTransform Transform,
+        const float PositionAlpha,
+        const float RotationAlpha)
+	{
+		const FVector Position = Transform.GetTranslation();
+		const FQuat Rotation = Transform.GetRotation();
+		GoalContainer.SetIKGoal(FIKRigGoal(GoalName, Position, Rotation, PositionAlpha, RotationAlpha));
+	};
+	
+	/** Apply a IKRigGoal and store it on this rig. Goal transform assumed in Component Space of Skeletal Mesh. */
+	UFUNCTION(BlueprintCallable, Category=IKRigGoals)
+    void SetIKRigGoal(const FIKRigGoal& Goal)
+	{
+		GoalContainer.SetIKGoal(Goal);
+	};
+
+	/** Remove all stored goals in this component. */
+	UFUNCTION(BlueprintCallable, Category=IKRigGoals)
+    void ClearAllGoals(){ GoalContainer.Goals.Empty(); };
+
+	// BEGIN IIKRigGoalCreator interface
+	/** Called from the IK Rig Anim node to obtain all the goals that have been set on this actor.*/
+	void GetIKGoals_Implementation(TMap<FName, FIKRigGoal>& OutGoals) { OutGoals = GoalContainer.Goals; };
+	// END IIKRigGoalCreator interface
+
+private:
+	FIKRigGoalContainer GoalContainer;
+};
diff --git a/Engine/Plugins/Animation/IKRig/Source/IKRig/Public/ActorComponents/IKRigInterface.h b/Engine/Plugins/Animation/IKRig/Source/IKRig/Public/ActorComponents/IKRigInterface.h
new file mode 100644
--- /dev/null
+++ b/Engine/Plugins/Animation/IKRig/Source/IKRig/Public/ActorComponents/IKRigInterface.h
@@ -0,0 +1,25 @@
+// Copyright Epic Games, Inc. All Rights Reserved.
+
+#pragma once
+
+#include "UObject/Interface.h"
+#include "../IKRigDataTypes.h"
+
+#include "IKRigInterface.generated.h"
+
+UINTERFACE(MinimalAPI, Blueprintable)
+class UIKGoalCreatorInterface : public UInterface
+{
+	GENERATED_BODY()
+};
+
+class IIKGoalCreatorInterface
+{    
+	GENERATED_BODY()
+
+public:
+	
+	/** Provide a Goal container for callers to read from.*/
+	UFUNCTION(BlueprintCallable, BlueprintNativeEvent, Category=IKRigGoals)
+    void GetIKGoals(TMap<FName, FIKRigGoal>& OutGoals);
+};
\ No newline at end of file
diff --git a/Engine/Plugins/Animation/IKRig/Source/IKRig/Public/AnimNodes/AnimNode_IKRig.h b/Engine/Plugins/Animation/IKRig/Source/IKRig/Public/AnimNodes/AnimNode_IKRig.h
--- a/Engine/Plugins/Animation/IKRig/Source/IKRig/Public/AnimNodes/AnimNode_IKRig.h
+++ b/Engine/Plugins/Animation/IKRig/Source/IKRig/Public/AnimNodes/AnimNode_IKRig.h
@@ -3,10 +3,13 @@
 #pragma once
 
 #include "CoreMinimal.h"
+
+#include "IKRigDataTypes.h"
 #include "UObject/ObjectMacros.h"
 #include "Animation/AnimNodeBase.h"
 #include "AnimNode_IKRig.generated.h"
 
+class IIKGoalCreatorInterface;
 class UIKRigDefinition;
 class UIKRigProcessor;
 
@@ -15,15 +18,17 @@ struct IKRIG_API FAnimNode_IKRig : public FAnimNode_Base
 {
 	GENERATED_BODY()
 
+	/** The input pose to start the IK solve relative to. */
 	UPROPERTY(EditAnywhere, Category = Links)
 	FPoseLink Source;
 
+	/** The IK rig to use to modify the incoming source pose. */
 	UPROPERTY(EditAnywhere, BlueprintReadWrite, Category = RigDefinition, meta = (NeverAsPin))
 	UIKRigDefinition* RigDefinitionAsset;
 
-	/** Coordinates for target location of tip bone - if EffectorLocationSpace is bone, this is the offset from Target Bone to use as target location*/
+	/** The input goal transforms used by the IK Rig solvers.*/
 	UPROPERTY(EditAnywhere, BlueprintReadWrite, EditFixedSize, Category = Goal, meta = (PinShownByDefault))
-	TArray<FTransform> GoalTransforms;
+	TArray<FIKRigGoal> Goals;
 
 	/** optionally ignore the input pose and start from the reference pose each solve */
 	UPROPERTY(EditAnywhere, Category = Solver)
@@ -35,14 +40,19 @@ struct IKRIG_API FAnimNode_IKRig : public FAnimNode_Base
 	bool bEnableDebugDraw;
 #endif
 
-private: 
+private:
+	
 	UPROPERTY(Transient)
 	UIKRigProcessor* IKRigProcessor = nullptr;
 
-	TArray<FName> GoalNames;
+	/** a cached list of components on the owning actor that implement the goal creator interface */
+	TArray<IIKGoalCreatorInterface*> GoalCreators;
+	TMap<FName, FIKRigGoal> GoalsFromGoalCreators;
+	
 	TMap<FCompactPoseBoneIndex, int32, FDefaultSetAllocator, TCompactPoseBoneIndexMapKeyFuncs<int32>> CompactPoseToRigIndices;
 
 public:
+	
 	FAnimNode_IKRig();
 
 	// FAnimNode_Base interface
@@ -53,13 +63,15 @@ public:
 	virtual void Evaluate_AnyThread(FPoseContext& Output) override;
 	virtual void Update_AnyThread(const FAnimationUpdateContext& Context) override;
 	virtual bool NeedsOnInitializeAnimInstance() const override { return true; }
+	virtual bool HasPreUpdate() const override { return true; }
+	virtual void PreUpdate(const UAnimInstance* InAnimInstance) override;
 	// End of FAnimNode_Base interface
 
 private:
 	bool RebuildGoalList();
 	FName GetGoalName(int32 Index) const;
 
-	void QueueDrawInterface(FAnimInstanceProxy* AnimProxy, const FTransform& ComponentToWorld);
+	void QueueDrawInterface(FAnimInstanceProxy* AnimProxy, const FTransform& ComponentToWorld) const;
 
 	friend class UAnimGraphNode_IKRig;
 };
diff --git a/Engine/Plugins/Animation/IKRig/Source/IKRig/Public/IKRigDataTypes.h b/Engine/Plugins/Animation/IKRig/Source/IKRig/Public/IKRigDataTypes.h
--- a/Engine/Plugins/Animation/IKRig/Source/IKRig/Public/IKRigDataTypes.h
+++ b/Engine/Plugins/Animation/IKRig/Source/IKRig/Public/IKRigDataTypes.h
@@ -4,40 +4,72 @@
 #pragma once
 
 #include "CoreMinimal.h"
+
 #include "UObject/ObjectMacros.h"
 #include "UObject/Object.h"
 #include "Misc/Guid.h"
 #include "IKRigDataTypes.generated.h"
 
+struct FIKRigEffectorGoal;
 struct FIKRigHierarchy;
 
-USTRUCT()
+
+USTRUCT(Blueprintable)
 struct IKRIG_API FIKRigGoal
 {
 	GENERATED_BODY()
+	
+	/** Name of the IK Goal. Must correspond to the name of a Goal in the target IKRig asset. */
+	UPROPERTY(EditAnywhere, BlueprintReadWrite, Category = FIKRigGoal)
+	FName Name;
+
+	/** Position of the IK goal in Component Space of target actor component. */
+	UPROPERTY(EditAnywhere, BlueprintReadWrite, Category = FIKRigGoal)
+	FVector Position;
+
+	/** Rotation of the IK goal in Component Space of target actor component. */
+	UPROPERTY(EditAnywhere, BlueprintReadWrite, Category = FIKRigGoal)
+	FRotator Rotation;
+
+	/** Range 0-1. Smoothly blends the Goal position from the input pose (0.0) to the Goal position (1.0). */
+	UPROPERTY(EditAnywhere, BlueprintReadWrite, Category = FIKRigGoal, meta = (ClampMin = "0.0", ClampMax = "1.0", UIMin = "0.0", UIMax = "1.0"))
+	float PositionAlpha;
+
+	/** Range 0-1. Smoothly blends the Goal rotation from the input pose (0.0) to the Goal rotation (1.0). */
+	UPROPERTY(EditAnywhere, BlueprintReadWrite, Category = FIKRigGoal, meta = (ClampMin = "0.0", ClampMax = "1.0", UIMin = "0.0", UIMax = "1.0"))
+	float RotationAlpha;
 
 	FIKRigGoal()
-	: Position(ForceInitToZero),
-	Rotation(ForceInitToZero)
-	{
-		Name = NAME_None;
-	}
+    : Name(NAME_None),
+	Position(ForceInitToZero),
+    Rotation(FRotator::ZeroRotator),
+    PositionAlpha(1.f),
+    RotationAlpha(0.0f){}
 	
 	FIKRigGoal(const FName& GoalName)
-	: Position(ForceInitToZero),
-	Rotation(ForceInitToZero)
+    : Name(GoalName),
+	Position(ForceInitToZero),
+    Rotation(FRotator::ZeroRotator),
+    PositionAlpha(1.f),
+    RotationAlpha(0.0f){}
+
+	FIKRigGoal(
+        const FName& Name,
+        const FVector& Position,
+        const FQuat& Rotation,
+        const float PositionAlpha,
+        const float RotationAlpha)
+        : Name(Name),
+          Position(Position),
+          Rotation(Rotation.Rotator()),
+          PositionAlpha(PositionAlpha),
+          RotationAlpha(RotationAlpha){}
+
+	FString ToString() const
 	{
-		Name = GoalName;
+		return FString::Printf(TEXT("Name=%s, Pos=(%s, Alpha=%3.3f), Rot=(%s, Alpha=%3.3f)"),
+			*Name.ToString(), *Position.ToString(), PositionAlpha, *Rotation.ToString(), RotationAlpha);
 	}
-
-	UPROPERTY(VisibleAnywhere, Category = FIKRigGoal)
-	FName Name;
-
-	UPROPERTY(EditAnywhere, Category = FIKRigGoal)
-	FVector Position;
-
-	UPROPERTY(EditAnywhere, Category = FIKRigGoal)
-	FQuat Rotation;
 };
 
 inline uint32 GetTypeHash(FIKRigGoal ObjectRef) { return GetTypeHash(ObjectRef.Name); }
@@ -47,51 +79,19 @@ struct IKRIG_API FIKRigGoalContainer
 {
 	GENERATED_BODY()
 
-	void InitializeGoalsFromNames(const TArray<FName>& InGoalNames)
-	{
-		Goals.Reserve(InGoalNames.Num());
-		for (const FName& Name : InGoalNames)
-		{
-			Goals.Emplace(Name, Name);
-		}
-	}
-
-	void SetGoalTransform(
-		const FName& GoalName,
-		const FVector& InPosition,
-		const FQuat& InRotation)
-	{
-		FIKRigGoal* Goal = Goals.Find(GoalName);
-		if (Goal)
-		{
-			Goal->Position = InPosition;
-			Goal->Rotation = InRotation;
-		}
-	}
-
-	bool GetGoalByName(const FName& InGoalName, FIKRigGoal& OutGoal) const
-	{
-		const FIKRigGoal* Goal = Goals.Find(InGoalName);
-		if (Goal)
-		{
-			OutGoal = *Goal;
-			return true;
-		}
-
-		return false;
-	}
-
-	void GetNames(TArray<FName>& OutNames) const
-	{
-		Goals.GenerateKeyArray(OutNames);
-	}
+public:
+	
+	/** Pre-load all the names of goals (optional, you can just call SetIKGoal to add as needed) */
+	void InitializeGoalsFromNames(const TArray<FIKRigEffectorGoal>& InGoalNames);
 
-	FORCEINLINE int GetNumGoals() const
-	{
-		return Goals.Num();
-	}
+	/** Set an IK goal to go to a specific location and rotation (in component space) blended by alpha.
+	 * Will ADD the goal if none exist with the input name. */
+	void SetIKGoal(const FIKRigGoal& InGoal);
 
-private:
+	/** Get an IK goal with the given name. Returns false if no goal is found in the container with the name. */
+	bool GetGoalByName(const FName& InGoalName, FIKRigGoal& OutGoal) const;
 
+	/** Keys are IK Rig Goal names. Values are the Goal data structures.
+	 * These are consumed by an IKRig asset to drive effectors. */
 	TMap<FName, FIKRigGoal> Goals;
-};
+};
\ No newline at end of file
diff --git a/Engine/Plugins/Animation/IKRig/Source/IKRig/Public/IKRigDefinition.h b/Engine/Plugins/Animation/IKRig/Source/IKRig/Public/IKRigDefinition.h
--- a/Engine/Plugins/Animation/IKRig/Source/IKRig/Public/IKRigDefinition.h
+++ b/Engine/Plugins/Animation/IKRig/Source/IKRig/Public/IKRigDefinition.h
@@ -9,6 +9,7 @@
 
 #include "IKRigDefinition.generated.h"
 
+struct FIKRigEffectorGoal;
 class UIKRigSolver;
 
 USTRUCT(Blueprintable)
@@ -91,10 +92,13 @@ public:
 	FRetargetDefinition RetargetDefinition;
 
 	/** get a list of all goal names used in the solvers */
-	void GetGoalNamesFromSolvers(TArray<FName>& OutGoalNames) const;
+	void GetGoalNamesFromSolvers(TArray<FIKRigEffectorGoal>& OutGoalNames) const;
 
-private:
+	/** get bone associated with a given goal */
+	FName GetBoneNameForGoal(FName GoalName) const;
 
+private:
+	
 	friend class UIKRigController;
 };
 
diff --git a/Engine/Plugins/Animation/IKRig/Source/IKRig/Public/IKRigProcessor.h b/Engine/Plugins/Animation/IKRig/Source/IKRig/Public/IKRigProcessor.h
--- a/Engine/Plugins/Animation/IKRig/Source/IKRig/Public/IKRigProcessor.h
+++ b/Engine/Plugins/Animation/IKRig/Source/IKRig/Public/IKRigProcessor.h
@@ -14,6 +14,15 @@
 class UIKRigDefinition;
 class UIKRigSolver;
 
+USTRUCT()
+struct FGoalBone
+{
+	GENERATED_BODY()
+	
+	FName BoneName;
+	int32 BoneIndex;
+};
+
 UCLASS(BlueprintType)
 class IKRIG_API UIKRigProcessor : public UObject
 {
@@ -22,12 +31,14 @@ class IKRIG_API UIKRigProcessor : public UObject
 public:
 
 	/** the runtime for an IKRig to convert an input pose into
-	*   a solved output pose given the goal transforms:
+	*   a solved output pose given a set of IK Rig Goals:
+	*   
 	* 1. Create a new IKRigProcessor once using MakeNewIKRigProcessor()
 	* 2. Initialize() with an IKRigDefinition asset
-	* 3. each tick, call SetGoalTransform() and update GlobalBoneTransforms
+	* 3. each tick, call SetIKGoal() and SetInputPoseGlobal()
 	* 4. Call Solve()
-	* 5. Copy back output transforms...
+	* 5. Copy output transforms with CopyOutputGlobalPoseToArray()
+	* 
 	*/
 
 	static UIKRigProcessor* MakeNewIKRigProcessor(UObject* Outer);
@@ -41,38 +52,46 @@ public:
 	// This is the general sequence of function calls to run a typical IK solve:
 	// 
 	
-	/** set all transforms in global space */
+	/** Set all bone transforms in global space. This is the pose the IK solve will start from */
 	void SetInputPoseGlobal(const TArray<FTransform>& InGlobalBoneTransforms);
 
-	/** optionally can be called before Solve() to use the reference pose as start pose */
+	/** Optionally can be called before Solve() to use the reference pose as start pose */
 	void SetInputPoseToRefPose();
 
-	/** set goal transform by name */
-	void SetGoalTransform(const FName& GoalName, const FVector& Position, const FQuat& Rotation);
+	/** Set a named IK goal to go to a specific location and rotation (assumed in component space) blended by separate position/rotation alpha (0-1)*/
+	void SetIKGoal(
+		const FName& GoalName,
+		const FVector& Position,
+		const FQuat& Rotation,
+		const float PositionAlpha,
+		const float RotationAlpha);
+
+	/** Set a named IK goal to go to a specific location and rotation (assumed in component space) blended by separate position/rotation alpha (0-1)*/
+	void SetIKGoal(const FIKRigGoal& Goal);
 
-	/** run entire stack of solvers */
+	/** Run entire stack of solvers */
 	void Solve();
 
-	/** get the results after calling Solve() */
+	/** Get the results after calling Solve() */
 	void CopyOutputGlobalPoseToArray(TArray<FTransform>& OutputPoseGlobal) const;
 
 	//
 	// END UPDATE SEQUENCE FUNCTIONS
 	//
 
-	/** get all goal names in the solver stack */
-	void GetGoalNames(TArray<FName>& OutGoalNames) const;
-
-	/** total number of goals in solver stack */
-	int GetNumGoals() const;
-
 	bool IsInitialized() const { return bInitialized; };
 
-	/** interface for debug drawing */
+	/** Get interface for debug drawing. */
 	const FControlRigDrawInterface& GetDrawInterface() const { return DrawInterface; }
 
-	/** get access to the internal skeleton data */
+	/** Get access to the internal goal data (read only) */
+	const FIKRigGoalContainer& GetGoalContainer() const;
+	
+	/** Get access to the internal skeleton data */
 	FIKRigSkeleton& GetSkeleton();
+
+	/** Get bone for goal */
+	bool GetBoneForGoal(FName GoalName, FGoalBone& OutBone) const;
 	
 private:
 
@@ -89,7 +108,10 @@ private:
 
 	/** the named transforms that solvers use as end effectors */
 	UPROPERTY(transient)
-	FIKRigGoalContainer Goals;
+	FIKRigGoalContainer GoalContainer;
+
+	/** map of goal names to bone names/indices */
+	TMap<FName, FGoalBone> GoalBones;
 
 	/** storage for hierarchy and bone transforms */
 	UPROPERTY(transient)
diff --git a/Engine/Plugins/Animation/IKRig/Source/IKRig/Public/IKRigSkeleton.h b/Engine/Plugins/Animation/IKRig/Source/IKRig/Public/IKRigSkeleton.h
--- a/Engine/Plugins/Animation/IKRig/Source/IKRig/Public/IKRigSkeleton.h
+++ b/Engine/Plugins/Animation/IKRig/Source/IKRig/Public/IKRigSkeleton.h
@@ -14,21 +14,36 @@ struct IKRIG_API FIKRigSkeleton
 
 	UPROPERTY(VisibleAnywhere, Category = Skeleton)
 	TArray<FName> BoneNames;
+	
 	UPROPERTY(VisibleAnywhere, Category = Skeleton)
 	TArray<int32> ParentIndices;
+	
 	UPROPERTY(VisibleAnywhere, Category = Skeleton)
 	TArray<FTransform> CurrentPoseGlobal;
+	
 	UPROPERTY(VisibleAnywhere, Category = Skeleton)
 	TArray<FTransform> CurrentPoseLocal;
+	
 	UPROPERTY(VisibleAnywhere, Category = Skeleton)
 	TArray<FTransform> RefPoseGlobal;
 
 	void Initialize(const FReferenceSkeleton& RefSkeleton);
+	
 	int32 GetBoneIndexFromName(const FName InName) const;
+	
 	int32 GetParentIndex(const int32 BoneIndex) const;
+	
 	void UpdateAllGlobalTransformFromLocal();
+	
 	void UpdateAllLocalTransformFromGlobal();
+	
 	void UpdateGlobalTransformFromLocal(const int32 BoneIndex);
+	
 	void UpdateLocalTransformFromGlobal(const int32 BoneIndex);
+	
 	void PropagateGlobalPoseBelowBone(const int32 BoneIndex);
+
+	//static void NormalizeRotations(TArray<FTransform>& Transforms);
+	static void NormalizeRotations(TArray<FTransform>& Transforms);
 };
+
diff --git a/Engine/Plugins/Animation/IKRig/Source/IKRig/Public/IKRigSolver.h b/Engine/Plugins/Animation/IKRig/Source/IKRig/Public/IKRigSolver.h
--- a/Engine/Plugins/Animation/IKRig/Source/IKRig/Public/IKRigSolver.h
+++ b/Engine/Plugins/Animation/IKRig/Source/IKRig/Public/IKRigSolver.h
@@ -4,14 +4,12 @@
 
 #include "UObject/ObjectMacros.h"
 #include "UObject/Object.h"
-#include "IKRigDataTypes.h"
-#include "IKRigSkeleton.h"
 #include "IKRigSolver.generated.h"
 
-struct FIKRigTarget;
-struct FIKRigTransform;
-struct FIKRigTransforms;
+
 struct FControlRigDrawInterface;
+struct FIKRigGoalContainer;
+struct FIKRigSkeleton;
 
 USTRUCT()
 struct IKRIG_API FIKRigEffectorGoal
@@ -23,8 +21,12 @@ struct IKRIG_API FIKRigEffectorGoal
 
 	UPROPERTY(EditAnywhere, Category = FIKRigEffector)
 	FName Bone;
+
+	bool operator==(const FIKRigEffectorGoal& Other) const { return Goal == Other.Goal; }
 };
 
+inline uint32 GetTypeHash(FIKRigEffectorGoal ObjectRef) { return GetTypeHash(ObjectRef.Goal); }
+
 // run time processor 
 UCLASS(abstract, hidecategories = UObject)
 class IKRIG_API UIKRigSolver : public UObject
@@ -47,22 +49,12 @@ public:
 
 	/** override CollectGoalNames() to tell the processor which goals to collect for this solver
 	 * NOTE: only ADD to OutGoals, do not reset or remove */
-	virtual void CollectGoalNames(TSet<FName>& OutGoals) const PURE_VIRTUAL("CollectGoalNames");
+	virtual void CollectGoalNames(TSet<FIKRigEffectorGoal>& OutGoals) const PURE_VIRTUAL("CollectGoalNames");
 
 	#if WITH_EDITOR
 	/** override RenameGoal() when UI renames a goal, you can auto-update any effectors
 	 * that were previously mapped to the old goal name. */
 	virtual void RenameGoal(const FName& OldName, const FName& NewName) {};
 	#endif // WITH_EDITOR
-	
-protected:
-
-	static bool GetGoalForEffector(
-		const FIKRigEffectorGoal& InEffector,
-		const FIKRigGoalContainer &Goals,
-		FIKRigGoal& OutGoal)
-	{
-		return Goals.GetGoalByName(InEffector.Goal, OutGoal);
-	}
 };
 
diff --git a/Engine/Plugins/Animation/IKRig/Source/IKRig/Public/Retargeter/IKRetargeter.h b/Engine/Plugins/Animation/IKRig/Source/IKRig/Public/Retargeter/IKRetargeter.h
--- a/Engine/Plugins/Animation/IKRig/Source/IKRig/Public/Retargeter/IKRetargeter.h
+++ b/Engine/Plugins/Animation/IKRig/Source/IKRig/Public/Retargeter/IKRetargeter.h
@@ -252,7 +252,7 @@ public:
 private:
 
 	UPROPERTY(Transient)
-	UIKRigProcessor* IKRig = nullptr;
+	UIKRigProcessor* IKRigProcessor = nullptr;
 	
 	TArray<FRetargetChainPairFK> ChainPairsFK;
 	
diff --git a/Engine/Plugins/Animation/IKRig/Source/IKRig/Public/Solvers/ConstraintSolver.h b/Engine/Plugins/Animation/IKRig/Source/IKRig/Public/Solvers/ConstraintSolver.h
--- a/Engine/Plugins/Animation/IKRig/Source/IKRig/Public/Solvers/ConstraintSolver.h
+++ b/Engine/Plugins/Animation/IKRig/Source/IKRig/Public/Solvers/ConstraintSolver.h
@@ -18,6 +18,6 @@ public:
 		FIKRigSkeleton& IKRigSkeleton,
 		const FIKRigGoalContainer& Goals,
 		FControlRigDrawInterface* InOutDrawInterface) override;
-	virtual void CollectGoalNames(TSet<FName>& OutGoals) const override;
+	virtual void CollectGoalNames(TSet<FIKRigEffectorGoal>& OutGoals) const override;
 };
 
diff --git a/Engine/Plugins/Animation/IKRig/Source/IKRig/Public/Solvers/IKRigFBIKSolver.h b/Engine/Plugins/Animation/IKRig/Source/IKRig/Public/Solvers/IKRigFBIKSolver.h
--- a/Engine/Plugins/Animation/IKRig/Source/IKRig/Public/Solvers/IKRigFBIKSolver.h
+++ b/Engine/Plugins/Animation/IKRig/Source/IKRig/Public/Solvers/IKRigFBIKSolver.h
@@ -94,7 +94,7 @@ protected:
 		FIKRigSkeleton& IKRigSkeleton,
 		const FIKRigGoalContainer& Goals,
 		FControlRigDrawInterface* InOutDrawInterface) override;
-	virtual void CollectGoalNames(TSet<FName>& OutGoals) const override;
+	virtual void CollectGoalNames(TSet<FIKRigEffectorGoal>& OutGoals) const override;
 
 public:
 	static const FName EffectorTargetPrefix;
diff --git a/Engine/Plugins/Animation/IKRig/Source/IKRig/Public/Solvers/TransformSolver.h b/Engine/Plugins/Animation/IKRig/Source/IKRig/Public/Solvers/TransformSolver.h
--- a/Engine/Plugins/Animation/IKRig/Source/IKRig/Public/Solvers/TransformSolver.h
+++ b/Engine/Plugins/Animation/IKRig/Source/IKRig/Public/Solvers/TransformSolver.h
@@ -31,7 +31,7 @@ protected:
 		FIKRigSkeleton& IKRigSkeleton, 
 		const FIKRigGoalContainer& Goals,
 		FControlRigDrawInterface* InOutDrawInterface) override;
-	virtual void CollectGoalNames(TSet<FName>& OutGoals) const override;
+	virtual void CollectGoalNames(TSet<FIKRigEffectorGoal>& OutGoals) const override;
 
 private:
 	
diff --git a/Engine/Plugins/Animation/IKRig/Source/IKRigDeveloper/Private/AnimGraphNode_IKRig.cpp b/Engine/Plugins/Animation/IKRig/Source/IKRigDeveloper/Private/AnimGraphNode_IKRig.cpp
--- a/Engine/Plugins/Animation/IKRig/Source/IKRigDeveloper/Private/AnimGraphNode_IKRig.cpp
+++ b/Engine/Plugins/Animation/IKRig/Source/IKRigDeveloper/Private/AnimGraphNode_IKRig.cpp
@@ -30,9 +30,6 @@ FText UAnimGraphNode_IKRig::GetNodeTitle(ENodeTitleType::Type TitleType) const
 void UAnimGraphNode_IKRig::CopyNodeDataToPreviewNode(FAnimNode_Base* InPreviewNode)
 {
 	FAnimNode_IKRig* IKRigNode = static_cast<FAnimNode_IKRig*>(InPreviewNode);
-
-	// copies Pin values from the internal node to get data which are not compiled yet
-	//Fabrik->EffectorTransform = Node.EffectorTransform;
 }
 
 FEditorModeID UAnimGraphNode_IKRig::GetEditorMode() const
@@ -58,8 +55,9 @@ void UAnimGraphNode_IKRig::CustomizePinData(UEdGraphPin* Pin, FName SourceProper
 {
 	Super::CustomizePinData(Pin, SourcePropertyName, ArrayIndex);
 
-	FString PinString = Pin->PinName.ToString();
-	FString GoalPinString = GET_MEMBER_NAME_STRING_CHECKED(FAnimNode_IKRig, GoalTransforms);
+	// set names of GoalTransform pins to names of goals
+	const FString PinString = Pin->PinName.ToString();
+	const FString GoalPinString = GET_MEMBER_NAME_STRING_CHECKED(FAnimNode_IKRig, Goals);
 	const int32 GoalStringLen = GoalPinString.Len();
 	if (PinString.Len() > GoalStringLen && PinString.Left(GoalStringLen) == GoalPinString)
 	{
diff --git a/Engine/Plugins/Animation/IKRig/Source/IKRigEditor/Private/IKRigController.cpp b/Engine/Plugins/Animation/IKRig/Source/IKRigEditor/Private/IKRigController.cpp
--- a/Engine/Plugins/Animation/IKRig/Source/IKRigEditor/Private/IKRigController.cpp
+++ b/Engine/Plugins/Animation/IKRig/Source/IKRigEditor/Private/IKRigController.cpp
@@ -149,7 +149,7 @@ UIKRigBoneSetting* UIKRigController::AddBoneSetting(TSubclassOf<UIKRigBoneSettin
 // GOALS
 //
 
-void UIKRigController::GetGoalNames(TArray<FName>& OutGoals) const
+void UIKRigController::GetGoalNames(TArray<FIKRigEffectorGoal>& OutGoals) const
 {
 	if (!IKRigDefinition)
 	{
diff --git a/Engine/Plugins/Animation/IKRig/Source/IKRigEditor/Public/IKRigController.h b/Engine/Plugins/Animation/IKRig/Source/IKRigEditor/Public/IKRigController.h
--- a/Engine/Plugins/Animation/IKRig/Source/IKRigEditor/Public/IKRigController.h
+++ b/Engine/Plugins/Animation/IKRig/Source/IKRigEditor/Public/IKRigController.h
@@ -40,7 +40,7 @@ public:
 
 	//** goals */
 	void RenameGoal(const FName& OldName, const FName& NewName) const;
-	void GetGoalNames(TArray<FName>& OutGoals) const;
+	void GetGoalNames(TArray<FIKRigEffectorGoal>& OutGoals) const;
 	DECLARE_MULTICAST_DELEGATE(FGoalModified);
 	FGoalModified OnGoalModified;
 
