Fixing Engine Tests after breaking IKRig API change.

diff --git a/Engine/Plugins/Animation/IKRig/Source/IKRig/Private/IKRigProcessor.cpp b/Engine/Plugins/Animation/IKRig/Source/IKRig/Private/IKRigProcessor.cpp
--- a/Engine/Plugins/Animation/IKRig/Source/IKRig/Private/IKRigProcessor.cpp
+++ b/Engine/Plugins/Animation/IKRig/Source/IKRig/Private/IKRigProcessor.cpp
@@ -130,7 +130,6 @@ void UIKRigProcessor::Solve()
 
 void UIKRigProcessor::CopyOutputGlobalPoseToArray(TArray<FTransform>& OutputPoseGlobal) const
 {
-	check(OutputPoseGlobal.Num() == Skeleton.CurrentPoseGlobal.Num());
 	OutputPoseGlobal = Skeleton.CurrentPoseGlobal;
 }
 
diff --git a/Engine/Plugins/Animation/IKRig/Source/IKRig/Public/IKRigProcessor.h b/Engine/Plugins/Animation/IKRig/Source/IKRig/Public/IKRigProcessor.h
--- a/Engine/Plugins/Animation/IKRig/Source/IKRig/Public/IKRigProcessor.h
+++ b/Engine/Plugins/Animation/IKRig/Source/IKRig/Public/IKRigProcessor.h
@@ -58,14 +58,6 @@ public:
 	/** Optionally can be called before Solve() to use the reference pose as start pose */
 	void SetInputPoseToRefPose();
 
-	/** Set a named IK goal to go to a specific location and rotation (assumed in component space) blended by separate position/rotation alpha (0-1)*/
-	void SetIKGoal(
-		const FName& GoalName,
-		const FVector& Position,
-		const FQuat& Rotation,
-		const float PositionAlpha,
-		const float RotationAlpha);
-
 	/** Set a named IK goal to go to a specific location and rotation (assumed in component space) blended by separate position/rotation alpha (0-1)*/
 	void SetIKGoal(const FIKRigGoal& Goal);
 
