Fixing uninitialized struct members in AnimNode_IKRetargeter
JIRA: UE-112350

diff --git a/Engine/Plugins/Animation/IKRig/Source/IKRig/Private/AnimNodes/AnimNode_IKRetargeter.cpp b/Engine/Plugins/Animation/IKRig/Source/IKRig/Private/AnimNodes/AnimNode_IKRetargeter.cpp
--- a/Engine/Plugins/Animation/IKRig/Source/IKRig/Private/AnimNodes/AnimNode_IKRetargeter.cpp
+++ b/Engine/Plugins/Animation/IKRig/Source/IKRig/Private/AnimNodes/AnimNode_IKRetargeter.cpp
@@ -5,13 +5,6 @@
 
 //#pragma optimize("", off)
 
-FAnimNode_IKRetargeter::FAnimNode_IKRetargeter()
-    : SourceMeshComponent(nullptr),
-	bEnableIK(true),
-	bIsInitialized (false)
-{
-}
-
 void FAnimNode_IKRetargeter::Initialize_AnyThread(const FAnimationInitializeContext& Context)
 {
 	DECLARE_SCOPE_HIERARCHICAL_COUNTER_ANIMNODE(Initialize_AnyThread)
diff --git a/Engine/Plugins/Animation/IKRig/Source/IKRig/Public/AnimNodes/AnimNode_IKRetargeter.h b/Engine/Plugins/Animation/IKRig/Source/IKRig/Public/AnimNodes/AnimNode_IKRetargeter.h
--- a/Engine/Plugins/Animation/IKRig/Source/IKRig/Public/AnimNodes/AnimNode_IKRetargeter.h
+++ b/Engine/Plugins/Animation/IKRig/Source/IKRig/Public/AnimNodes/AnimNode_IKRetargeter.h
@@ -14,17 +14,15 @@ struct IKRIG_API FAnimNode_IKRetargeter : public FAnimNode_Base
 
 	/** The Skeletal Mesh Component to retarget animation from. Assumed to be animated and tick BEFORE this anim instance.*/
 	UPROPERTY(BlueprintReadWrite, transient, Category=Settings, meta=(PinShownByDefault))
-	TWeakObjectPtr<USkeletalMeshComponent> SourceMeshComponent;
+	TWeakObjectPtr<USkeletalMeshComponent> SourceMeshComponent = nullptr;
 	
 	/** Retarget asset to use. Must define a Source and Target IK Rig compatible with the SourceMeshComponent and current anim instance.*/
 	UPROPERTY(EditAnywhere, Category = Settings)
-	UIKRetargeter* IKRetargeterAsset;
+	UIKRetargeter* IKRetargeterAsset = nullptr;
 
 	/** Retarget asset to use. Must define a Source and Target IK Rig compatible with the SourceMeshComponent and current anim instance.*/
 	UPROPERTY(EditAnywhere, Category = Settings)
-	bool bEnableIK;
-
-	FAnimNode_IKRetargeter();
+	bool bEnableIK = true;
 
 	// FAnimNode_Base interface
 	virtual void Initialize_AnyThread(const FAnimationInitializeContext& Context) override;
@@ -42,7 +40,8 @@ private:
 	void CopyBoneTransformsFromSource(USkeletalMeshComponent* TargetMeshComponent);
 
 	// indicates that all prerequisites are met and node is ready to operate
-	bool bIsInitialized;
+	UPROPERTY(Transient)
+	bool bIsInitialized = false;
 	
 	// source mesh references, cached during init so that we can compare and see if it has changed
 	TWeakObjectPtr<USkeletalMeshComponent>	CurrentlyUsedSourceMeshComponent;
