Fixes bug, IKRig anim node stomped on incoming float curve values.

diff --git a/Engine/Plugins/Animation/IKRig/Source/IKRig/Private/AnimNodes/AnimNode_IKRig.cpp b/Engine/Plugins/Animation/IKRig/Source/IKRig/Private/AnimNodes/AnimNode_IKRig.cpp
--- a/Engine/Plugins/Animation/IKRig/Source/IKRig/Private/AnimNodes/AnimNode_IKRig.cpp
+++ b/Engine/Plugins/Animation/IKRig/Source/IKRig/Private/AnimNodes/AnimNode_IKRig.cpp
@@ -20,30 +20,27 @@ FAnimNode_IKRig::FAnimNode_IKRig()
 void FAnimNode_IKRig::Evaluate_AnyThread(FPoseContext& Output) 
 {
 	DECLARE_SCOPE_HIERARCHICAL_COUNTER_FUNC()
-	FPoseContext SourcePose(Output);
 
 	if (Source.GetLinkNode() && !bStartFromRefPose)
 	{
-		Source.Evaluate(SourcePose);
+		Source.Evaluate(Output);
 	}
 	else
 	{
-		SourcePose.ResetToRefPose();
+		Output.ResetToRefPose();
 	}
 
 	if (!IKRigProcessor)
 	{
-		Output = SourcePose;
 		return;
 	}
 	
 	if (!IKRigProcessor->IsInitialized())
 	{
-		Output = SourcePose;
 		return;
 	}
 
-	FCompactPose& OutPose = SourcePose.Pose;
+	FCompactPose& OutPose = Output.Pose;
 	FIKRigSkeleton& IKRigSkeleton = IKRigProcessor->GetSkeleton();
 	
 	if (bStartFromRefPose)
