Fixing CIS issue in IK Retargeter
JIRA: UE-112791

diff --git a/Engine/Plugins/Animation/IKRig/Source/IKRig/Private/AnimNodes/AnimNode_IKRetargeter.cpp b/Engine/Plugins/Animation/IKRig/Source/IKRig/Private/AnimNodes/AnimNode_IKRetargeter.cpp
--- a/Engine/Plugins/Animation/IKRig/Source/IKRig/Private/AnimNodes/AnimNode_IKRetargeter.cpp
+++ b/Engine/Plugins/Animation/IKRig/Source/IKRig/Private/AnimNodes/AnimNode_IKRetargeter.cpp
@@ -143,10 +143,16 @@ void FAnimNode_IKRetargeter::InitializeRetargetData(const UAnimInstance* InAnimI
 
 void FAnimNode_IKRetargeter::CopyBoneTransformsFromSource(USkeletalMeshComponent* TargetMeshComponent)
 {
-	USkeletalMeshComponent* SourceMeshComp = CurrentlyUsedSourceMeshComponent.IsValid() ? CurrentlyUsedSourceMeshComponent.Get() : nullptr;
+	if (!CurrentlyUsedSourceMeshComponent.IsValid())
+	{
+		return; 
+	}
 
+	USkeletalMeshComponent* SourceMeshComp =  CurrentlyUsedSourceMeshComponent.Get();
+	
 	// is the source mesh ticking?
 	if (!SourceMeshComp->IsRegistered())
+		
 	{
 		CurrentlyUsedSourceMesh.Reset(); // forces reinitialization when re-registered
 		return; // skip copying pose when component is no longer ticking
