IKRigEditor module, moving headers out Private folder into Public

diff --git a/Engine/Plugins/Animation/IKRig/Source/IKRigEditor/Private/AnimMode/IKRigEditMode.cpp b/Engine/Plugins/Animation/IKRig/Source/IKRigEditor/Private/AnimMode/IKRigEditMode.cpp
deleted file mode 100644
diff --git a/Engine/Plugins/Animation/IKRig/Source/IKRigEditor/Private/AnimMode/IKRigEditMode.h b/Engine/Plugins/Animation/IKRig/Source/IKRigEditor/Private/AnimMode/IKRigEditMode.h
deleted file mode 100644
diff --git a/Engine/Plugins/Animation/IKRig/Source/IKRigEditor/Private/AssetTypeActions_IKRetargeter.h b/Engine/Plugins/Animation/IKRig/Source/IKRigEditor/Private/AssetTypeActions_IKRetargeter.h
deleted file mode 100644
diff --git a/Engine/Plugins/Animation/IKRig/Source/IKRigEditor/Private/AssetTypeActions_IKRigDefinition.h b/Engine/Plugins/Animation/IKRig/Source/IKRigEditor/Private/AssetTypeActions_IKRigDefinition.h
deleted file mode 100644
diff --git a/Engine/Plugins/Animation/IKRig/Source/IKRigEditor/Private/IKRetargeterFactory.h b/Engine/Plugins/Animation/IKRig/Source/IKRigEditor/Private/IKRetargeterFactory.h
deleted file mode 100644
diff --git a/Engine/Plugins/Animation/IKRig/Source/IKRigEditor/Private/IKRigDefinitionDetails.h b/Engine/Plugins/Animation/IKRig/Source/IKRigEditor/Private/IKRigDefinitionDetails.h
deleted file mode 100644
diff --git a/Engine/Plugins/Animation/IKRig/Source/IKRigEditor/Private/IKRigDefinitionFactory.h b/Engine/Plugins/Animation/IKRig/Source/IKRigEditor/Private/IKRigDefinitionFactory.h
deleted file mode 100644
diff --git a/Engine/Plugins/Animation/IKRig/Source/IKRigEditor/Private/IKRigEditMode.cpp b/Engine/Plugins/Animation/IKRig/Source/IKRigEditor/Private/IKRigEditMode.cpp
new file mode 100644
--- /dev/null
+++ b/Engine/Plugins/Animation/IKRig/Source/IKRigEditor/Private/IKRigEditMode.cpp
@@ -0,0 +1,50 @@
+// Copyright Epic Games, Inc. All Rights Reserved.
+
+#include "IKRigEditMode.h"
+#include "AnimGraphNode_IKRig.h"
+#include "IPersonaPreviewScene.h"
+#include "Animation/DebugSkelMeshComponent.h"
+
+// NOTE: This mode is disabled until we can revisit the idea of pulling effectors in ABP viewport
+
+void FIKRigEditMode::EnterMode(class UAnimGraphNode_Base* InEditorNode, struct FAnimNode_Base* InRuntimeNode)
+{
+	//RuntimeNode = static_cast<FAnimNode_IKRig*>(InRuntimeNode);
+	//GraphNode = CastChecked<UAnimGraphNode_IKRig>(InEditorNode);
+
+	FAnimNodeEditMode::EnterMode(InEditorNode, InRuntimeNode);
+}
+
+void FIKRigEditMode::ExitMode()
+{
+	//RuntimeNode = nullptr;
+	//GraphNode = nullptr;
+
+	FAnimNodeEditMode::ExitMode();
+}
+
+FVector FIKRigEditMode::GetWidgetLocation() const
+{
+	USkeletalMeshComponent* SkelComp = GetAnimPreviewScene().GetPreviewMeshComponent();
+
+// 	FBoneSocketTarget& Target = RuntimeNode->EffectorTarget;
+// 	FVector Location = RuntimeNode->EffectorTransform.GetLocation();
+// 	EBoneControlSpace Space = RuntimeNode->EffectorTransformSpace;
+// 	FVector WidgetLoc = ConvertWidgetLocation(SkelComp, RuntimeNode->ForwardedPose, Target, Location, Space);
+	return SkelComp->GetComponentToWorld().GetLocation();
+}
+
+UE::Widget::EWidgetMode FIKRigEditMode::GetWidgetMode() const
+{
+	// allow translation all the time for effectot target
+	return UE::Widget::WM_Translate;
+}
+
+void FIKRigEditMode::DoTranslation(FVector& InTranslation)
+{
+// 	USkeletalMeshComponent* SkelComp = GetAnimPreviewScene().GetPreviewMeshComponent();
+// 	FVector Offset = ConvertCSVectorToBoneSpace(SkelComp, InTranslation, RuntimeNode->ForwardedPose, RuntimeNode->EffectorTarget, RuntimeNode->EffectorTransformSpace);
+// 
+// 	RuntimeNode->EffectorTransform.AddToTranslation(Offset);
+// 	GraphNode->Node.EffectorTransform.SetTranslation(RuntimeNode->EffectorTransform.GetTranslation());
+}
diff --git a/Engine/Plugins/Animation/IKRig/Source/IKRigEditor/Private/IKRigEditor.cpp b/Engine/Plugins/Animation/IKRig/Source/IKRigEditor/Private/IKRigEditor.cpp
--- a/Engine/Plugins/Animation/IKRig/Source/IKRigEditor/Private/IKRigEditor.cpp
+++ b/Engine/Plugins/Animation/IKRig/Source/IKRigEditor/Private/IKRigEditor.cpp
@@ -15,7 +15,7 @@
 #include "AssetTypeActions_IKRigDefinition.h"
 #include "AssetTypeActions_IKRetargeter.h"
 #include "IKRigDefinition.h"
-#include "AnimMode/IKRigEditMode.h"
+#include "IKRigEditMode.h"
 
 IMPLEMENT_MODULE(FIKRigEditor, IKRigEditor)
 
diff --git a/Engine/Plugins/Animation/IKRig/Source/IKRigEditor/Public/AssetTypeActions_IKRetargeter.h b/Engine/Plugins/Animation/IKRig/Source/IKRigEditor/Public/AssetTypeActions_IKRetargeter.h
new file mode 100644
--- /dev/null
+++ b/Engine/Plugins/Animation/IKRig/Source/IKRigEditor/Public/AssetTypeActions_IKRetargeter.h
@@ -0,0 +1,20 @@
+// Copyright Epic Games, Inc. All Rights Reserved.
+
+#pragma once
+
+#include "CoreMinimal.h"
+#include "Toolkits/IToolkitHost.h"
+#include "AssetTypeActions_Base.h"
+
+class FAssetTypeActions_IKRetargeter : public FAssetTypeActions_Base
+{
+public:
+	// IAssetTypeActions Implementation
+	virtual FText GetName() const override { return NSLOCTEXT("AssetTypeActions", "AssetTypeActions_IKRetargeter", "IK Retargeter"); }
+	virtual void GetActions(const TArray<UObject*>& InObjects, struct FToolMenuSection& Section) override;
+	virtual FColor GetTypeColor() const override { return FColor(255, 192, 128); }
+	virtual UClass* GetSupportedClass() const override;
+	//virtual void OpenAssetEditor(const TArray<UObject*>& InObjects, TSharedPtr<class IToolkitHost> EditWithinLevelEditor = TSharedPtr<IToolkitHost>()) override;
+	virtual uint32 GetCategories() override { return EAssetTypeCategories::Animation; }
+	//virtual UThumbnailInfo* GetThumbnailInfo(UObject* Asset) const override;
+};
diff --git a/Engine/Plugins/Animation/IKRig/Source/IKRigEditor/Public/AssetTypeActions_IKRigDefinition.h b/Engine/Plugins/Animation/IKRig/Source/IKRigEditor/Public/AssetTypeActions_IKRigDefinition.h
new file mode 100644
--- /dev/null
+++ b/Engine/Plugins/Animation/IKRig/Source/IKRigEditor/Public/AssetTypeActions_IKRigDefinition.h
@@ -0,0 +1,20 @@
+// Copyright Epic Games, Inc. All Rights Reserved.
+
+#pragma once
+
+#include "CoreMinimal.h"
+#include "Toolkits/IToolkitHost.h"
+#include "AssetTypeActions_Base.h"
+
+class FAssetTypeActions_IKRigDefinition : public FAssetTypeActions_Base
+{
+public:
+	// IAssetTypeActions Implementation
+	virtual FText GetName() const override { return NSLOCTEXT("AssetTypeActions", "AssetTypeActions_IKRigDefinition", "IK Rig"); }
+	virtual void GetActions(const TArray<UObject*>& InObjects, struct FToolMenuSection& Section) override;
+	virtual FColor GetTypeColor() const override { return FColor(255, 192, 128); }
+	virtual UClass* GetSupportedClass() const override;
+	//virtual void OpenAssetEditor(const TArray<UObject*>& InObjects, TSharedPtr<class IToolkitHost> EditWithinLevelEditor = TSharedPtr<IToolkitHost>()) override;
+	virtual uint32 GetCategories() override { return EAssetTypeCategories::Animation; }
+	//virtual UThumbnailInfo* GetThumbnailInfo(UObject* Asset) const override;
+};
diff --git a/Engine/Plugins/Animation/IKRig/Source/IKRigEditor/Public/IKRetargeterFactory.h b/Engine/Plugins/Animation/IKRig/Source/IKRigEditor/Public/IKRetargeterFactory.h
new file mode 100644
--- /dev/null
+++ b/Engine/Plugins/Animation/IKRig/Source/IKRigEditor/Public/IKRetargeterFactory.h
@@ -0,0 +1,28 @@
+// Copyright Epic Games, Inc. All Rights Reserved.
+
+#pragma once
+
+#include "CoreMinimal.h"
+#include "UObject/ObjectMacros.h"
+#include "Factories/Factory.h"
+#include "IKRetargeterFactory.generated.h"
+
+
+UCLASS(BlueprintType, hidecategories=Object)
+class UIKRetargeterFactory : public UFactory
+{
+	GENERATED_BODY()
+
+public:
+
+	UIKRetargeterFactory();
+
+	// UFactory Interface
+	virtual FText GetDisplayName() const override;
+	virtual uint32 GetMenuCategories() const override;
+	virtual FText GetToolTip() const override;
+	virtual FString GetDefaultNewAssetName() const override;
+	virtual UObject* FactoryCreateNew(UClass* Class, UObject* InParent, FName Name, EObjectFlags Flags, UObject* Context, FFeedbackContext* Warn) override;
+	virtual bool ConfigureProperties() override;
+	virtual bool ShouldShowInNewMenu() const override;
+};
diff --git a/Engine/Plugins/Animation/IKRig/Source/IKRigEditor/Public/IKRigDefinitionDetails.h b/Engine/Plugins/Animation/IKRig/Source/IKRigEditor/Public/IKRigDefinitionDetails.h
new file mode 100644
--- /dev/null
+++ b/Engine/Plugins/Animation/IKRig/Source/IKRigEditor/Public/IKRigDefinitionDetails.h
@@ -0,0 +1,51 @@
+// Copyright Epic Games, Inc. All Rights Reserved.
+
+#pragma once
+
+#include "UObject/WeakObjectPtrTemplates.h"
+#include "IDetailCustomization.h"
+#include "Input/Reply.h"
+
+class IDetailLayoutBuilder;
+class UIKRigDefinition;
+class UIKRigController;
+struct FAssetData;
+class FReply;
+class ITableRow;
+class STableViewBase;
+class IPropertyHandle;
+
+class FIKRigDefinitionDetails : public IDetailCustomization
+{
+
+public:
+
+	virtual ~FIKRigDefinitionDetails();
+
+	/** Makes a new instance of this detail layout class for a specific detail view requesting it */
+	static TSharedRef<IDetailCustomization> MakeInstance();
+
+	/** IDetailCustomization interface */
+	virtual void CustomizeDetails( IDetailLayoutBuilder& DetailBuilder ) override;
+	virtual void CustomizeDetails(const TSharedPtr<IDetailLayoutBuilder>& DetailBuilder) override;
+	/** IDetailCustomization interface */
+
+private:
+
+	/** pointers to "Model" and "Controller" in MVC pattern */
+	TWeakObjectPtr<UIKRigDefinition> IKRigDefinition;
+	UIKRigController* IKRigController;
+
+	/** source skeleton asset, used to initialize hierarchy */
+	FString GetCurrentSourceAsset() const;
+	bool ShouldFilterAsset(const FAssetData& AssetData);
+	void OnAssetSelected(const FAssetData& AssetData);
+	TWeakObjectPtr<UObject> SelectedAsset;
+
+	/** import new skeleton */
+	FReply OnImportHierarchy();
+	bool CanImport() const;
+
+	TWeakPtr< IDetailLayoutBuilder > DetailBuilderWeakPtr;
+};
+
diff --git a/Engine/Plugins/Animation/IKRig/Source/IKRigEditor/Public/IKRigDefinitionFactory.h b/Engine/Plugins/Animation/IKRig/Source/IKRigEditor/Public/IKRigDefinitionFactory.h
new file mode 100644
--- /dev/null
+++ b/Engine/Plugins/Animation/IKRig/Source/IKRigEditor/Public/IKRigDefinitionFactory.h
@@ -0,0 +1,28 @@
+// Copyright Epic Games, Inc. All Rights Reserved.
+
+#pragma once
+
+#include "CoreMinimal.h"
+#include "UObject/ObjectMacros.h"
+#include "Factories/Factory.h"
+#include "IKRigDefinitionFactory.generated.h"
+
+
+UCLASS(BlueprintType, hidecategories=Object)
+class UIKRigDefinitionFactory : public UFactory
+{
+	GENERATED_BODY()
+
+public:
+
+	UIKRigDefinitionFactory();
+
+	// UFactory Interface
+	virtual FText GetDisplayName() const override;
+	virtual uint32 GetMenuCategories() const override;
+	virtual FText GetToolTip() const override;
+	virtual FString GetDefaultNewAssetName() const override;
+	virtual UObject* FactoryCreateNew(UClass* Class, UObject* InParent, FName Name, EObjectFlags Flags, UObject* Context, FFeedbackContext* Warn) override;
+	virtual bool ConfigureProperties() override;
+	virtual bool ShouldShowInNewMenu() const override;
+};
diff --git a/Engine/Plugins/Animation/IKRig/Source/IKRigEditor/Public/IKRigEditMode.h b/Engine/Plugins/Animation/IKRig/Source/IKRigEditor/Public/IKRigEditMode.h
new file mode 100644
--- /dev/null
+++ b/Engine/Plugins/Animation/IKRig/Source/IKRigEditor/Public/IKRigEditMode.h
@@ -0,0 +1,22 @@
+// Copyright Epic Games, Inc. All Rights Reserved.
+
+#pragma once
+
+#include "CoreMinimal.h"
+#include "UnrealWidgetFwd.h"
+#include "AnimNodeEditMode.h"
+
+class FIKRigEditMode : public FAnimNodeEditMode
+{
+public:
+	/** IAnimNodeEditMode interface */
+	virtual void EnterMode(class UAnimGraphNode_Base* InEditorNode, struct FAnimNode_Base* InRuntimeNode) override;
+	virtual void ExitMode() override;
+	virtual FVector GetWidgetLocation() const override;
+	virtual UE::Widget::EWidgetMode GetWidgetMode() const override;
+	virtual void DoTranslation(FVector& InTranslation) override;
+
+private:
+	struct FAnimNode_IKRig* RuntimeNode;
+	class UAnimGraphNode_IKRig* GraphNode;
+};
