Fixed bug, IKRig accumulates bones when reimporting skeleton instead of reseting.
Fixed bug, IKRig default goal position value is now 0.

diff --git a/Engine/Plugins/Animation/IKRig/Source/IKRig/Private/IKRigSkeleton.cpp b/Engine/Plugins/Animation/IKRig/Source/IKRig/Private/IKRigSkeleton.cpp
--- a/Engine/Plugins/Animation/IKRig/Source/IKRig/Private/IKRigSkeleton.cpp
+++ b/Engine/Plugins/Animation/IKRig/Source/IKRig/Private/IKRigSkeleton.cpp
@@ -4,6 +4,12 @@
 
 void FIKRigSkeleton::Initialize(const FReferenceSkeleton& RefSkeleton)
 {
+	BoneNames.Reset();
+	ParentIndices.Reset();
+	CurrentPoseGlobal.Reset();
+	CurrentPoseLocal.Reset();
+	RefPoseGlobal.Reset();
+	
 	// copy names and parent indices into local storage
 	for (int32 BoneIndex=0; BoneIndex<RefSkeleton.GetNum(); ++BoneIndex)
 	{
diff --git a/Engine/Plugins/Animation/IKRig/Source/IKRig/Public/IKRigDataTypes.h b/Engine/Plugins/Animation/IKRig/Source/IKRig/Public/IKRigDataTypes.h
--- a/Engine/Plugins/Animation/IKRig/Source/IKRig/Public/IKRigDataTypes.h
+++ b/Engine/Plugins/Animation/IKRig/Source/IKRig/Public/IKRigDataTypes.h
@@ -68,7 +68,7 @@ struct IKRIG_API FIKRigGoal
     : Name(NAME_None),
 	Position(ForceInitToZero),
     Rotation(FRotator::ZeroRotator),
-    PositionAlpha(1.f),
+    PositionAlpha(0.0f),
     RotationAlpha(0.0f),
 	FinalBlendedPosition(ForceInitToZero),
 	FinalBlendedRotation(FQuat::Identity){}
@@ -77,7 +77,7 @@ struct IKRIG_API FIKRigGoal
     : Name(GoalName),
 	Position(ForceInitToZero),
     Rotation(FRotator::ZeroRotator),
-    PositionAlpha(1.f),
+    PositionAlpha(0.0f),
     RotationAlpha(0.0f),
 	FinalBlendedPosition(ForceInitToZero),
     FinalBlendedRotation(FQuat::Identity){}
