Added ability for IK Retargeter to find it's own source mesh from it's attached parent (same as CopyPoseFromMesh)

diff --git a/Engine/Plugins/Animation/IKRig/Source/IKRig/Private/AnimNodes/AnimNode_IKRetargeter.cpp b/Engine/Plugins/Animation/IKRig/Source/IKRig/Private/AnimNodes/AnimNode_IKRetargeter.cpp
--- a/Engine/Plugins/Animation/IKRig/Source/IKRig/Private/AnimNodes/AnimNode_IKRetargeter.cpp
+++ b/Engine/Plugins/Animation/IKRig/Source/IKRig/Private/AnimNodes/AnimNode_IKRetargeter.cpp
@@ -64,8 +64,26 @@ void FAnimNode_IKRetargeter::PreUpdate(const UAnimInstance* InAnimInstance)
 
 void FAnimNode_IKRetargeter::EnsureInitialized(const UAnimInstance* InAnimInstance)
 {
-	// has a source mesh been plugged in?
-	if (!(SourceMeshComponent.IsValid() && IKRetargeterAsset))
+	// has user supplied a retargeter asset?
+	if (!IKRetargeterAsset)
+	{
+		bIsInitialized = false;
+		return;
+	}
+
+	// if user hasn't explicitly connected a source mesh, optionally use the parent mesh component (if there is one) 
+	if (!SourceMeshComponent.IsValid() && bUseAttachedParent)
+	{
+		USkeletalMeshComponent* TargetMesh = InAnimInstance->GetSkelMeshComponent();
+		USkeletalMeshComponent* ParentComponent = Cast<USkeletalMeshComponent>(TargetMesh->GetAttachParent());
+		if (ParentComponent)
+		{
+			SourceMeshComponent = ParentComponent;
+		}
+	}
+	
+	// has a source mesh been plugged in or found?
+	if (!SourceMeshComponent.IsValid())
 	{
 		bIsInitialized = false;
 		return; // can't do anything if we don't have a source mesh
diff --git a/Engine/Plugins/Animation/IKRig/Source/IKRig/Public/AnimNodes/AnimNode_IKRetargeter.h b/Engine/Plugins/Animation/IKRig/Source/IKRig/Public/AnimNodes/AnimNode_IKRetargeter.h
--- a/Engine/Plugins/Animation/IKRig/Source/IKRig/Public/AnimNodes/AnimNode_IKRetargeter.h
+++ b/Engine/Plugins/Animation/IKRig/Source/IKRig/Public/AnimNodes/AnimNode_IKRetargeter.h
@@ -16,6 +16,10 @@ struct IKRIG_API FAnimNode_IKRetargeter : public FAnimNode_Base
 	UPROPERTY(BlueprintReadWrite, transient, Category=Settings, meta=(PinShownByDefault))
 	TWeakObjectPtr<USkeletalMeshComponent> SourceMeshComponent = nullptr;
 
+	/* If SourceMeshComponent is not valid, and if this is true, it will look for attached parent as a source */
+	UPROPERTY(EditAnywhere, BlueprintReadWrite, Category=Settings, meta = (NeverAsPin))
+	bool bUseAttachedParent = true;
+
 	/** Map of chain names to per-chain retarget settings (can be modified at runtime).*/
 	UPROPERTY(BlueprintReadWrite, transient, Category=Settings, meta=(PinShownByDefault))
 	TMap<FName, FIKRetargetChainSettings> ChainSettings;
@@ -24,17 +28,9 @@ struct IKRIG_API FAnimNode_IKRetargeter : public FAnimNode_Base
 	UPROPERTY(EditAnywhere, Category = Settings)
 	UIKRetargeter* IKRetargeterAsset = nullptr;
 
-	/** Retarget asset to use. Must define a Source and Target IK Rig compatible with the SourceMeshComponent and current anim instance.*/
+	/** When false, IK is not applied as part of retargeter. Useful for debugging limb issues suspected to be caused by IK.*/
 	UPROPERTY(EditAnywhere, Category = Settings)
 	bool bEnableIK = true;
-
-	/** Retarget asset to use. Must define a Source and Target IK Rig compatible with the SourceMeshComponent and current anim instance.*/
-	UPROPERTY(EditAnywhere, Category = Settings)
-	bool bAutoFindSourceMeshByTag = false;
-
-	/** Retarget asset to use. Must define a Source and Target IK Rig compatible with the SourceMeshComponent and current anim instance.*/
-	UPROPERTY(EditAnywhere, Category = Settings, meta = (EditCondition="bAutoFindSourceMeshByTag"))
-	FName SourceMeshComponentTag = "RetargetSource";
 	
 	// FAnimNode_Base interface
 	virtual void Initialize_AnyThread(const FAnimationInitializeContext& Context) override;
diff --git a/Engine/Plugins/Animation/IKRig/Source/IKRigDeveloper/Private/AnimGraphNode_IKRetargeter.cpp b/Engine/Plugins/Animation/IKRig/Source/IKRigDeveloper/Private/AnimGraphNode_IKRetargeter.cpp
--- a/Engine/Plugins/Animation/IKRig/Source/IKRigDeveloper/Private/AnimGraphNode_IKRetargeter.cpp
+++ b/Engine/Plugins/Animation/IKRig/Source/IKRigDeveloper/Private/AnimGraphNode_IKRetargeter.cpp
@@ -4,6 +4,7 @@
 #include "Animation/AnimInstance.h"
 #include "Kismet2/CompilerResultsLog.h"
 
+//#pragma optimize("", off)
 
 #define LOCTEXT_NAMESPACE "AnimGraphNode_IKRig"
 const FName UAnimGraphNode_IKRetargeter::AnimModeName(TEXT("IKRig.IKRigEditor.IKRigEditMode"));
@@ -27,8 +28,26 @@ FEditorModeID UAnimGraphNode_IKRetargeter::GetEditorMode() const
 	return AnimModeName;
 }
 
+void UAnimGraphNode_IKRetargeter::CustomizePinData(UEdGraphPin* Pin, FName SourcePropertyName, int32 ArrayIndex) const
+{
+	Super::CustomizePinData(Pin, SourcePropertyName, ArrayIndex);
+
+	// hide the Source Mesh Component input pin when bAutoFindSourceMeshByTag is true
+	const FString SourceMeshPropertyName = GET_MEMBER_NAME_STRING_CHECKED(FAnimNode_IKRetargeter, SourceMeshComponent);
+	const FString SourcePropertyString = SourcePropertyName.ToString();
+	if (SourcePropertyString == SourceMeshPropertyName)
+	{
+		Pin->bHidden = Node.bUseAttachedParent;	
+	}
+}
+
 void UAnimGraphNode_IKRetargeter::PostEditChangeProperty(struct FPropertyChangedEvent& PropertyChangedEvent)
 {
+	const FName PropertyName = (PropertyChangedEvent.Property ? PropertyChangedEvent.Property->GetFName() : NAME_None);
+	if ((PropertyName == GET_MEMBER_NAME_STRING_CHECKED(FAnimNode_IKRetargeter, bUseAttachedParent)))
+	{
+		ReconstructNode();
+	}
 }
 
 void UAnimGraphNode_IKRetargeter::PostLoad()
@@ -41,11 +60,13 @@ void UAnimGraphNode_IKRetargeter::ValidateAnimNodeDuringCompilation(USkeleton* F
 	Super::ValidateAnimNodeDuringCompilation(ForSkeleton, MessageLog);
 
 	// validate source mesh component is not null
-	//if (!Node.SourceMeshComponent.())
-	if (!IsPinExposedAndLinked(GET_MEMBER_NAME_STRING_CHECKED(FAnimNode_IKRetargeter, SourceMeshComponent)))
+	if (!Node.bUseAttachedParent)
 	{
-		MessageLog.Warning(TEXT("@@ is missing an Source Skeletal Mesh Component reference."), this);
-		return;
+		if (!IsPinExposedAndLinked(GET_MEMBER_NAME_STRING_CHECKED(FAnimNode_IKRetargeter, SourceMeshComponent)))
+		{
+			MessageLog.Warning(TEXT("@@ is missing a Source Skeletal Mesh Component reference."), this);
+			return;
+		}
 	}
 
 	// validate IK Rig asset has been assigned
@@ -92,6 +113,13 @@ void UAnimGraphNode_IKRetargeter::ValidateAnimNodeDuringCompilation(USkeleton* F
 void UAnimGraphNode_IKRetargeter::PreloadRequiredAssets()
 {
 	Super::PreloadRequiredAssets();
+	
+	if (Node.IKRetargeterAsset)
+	{
+		PreloadObject(Node.IKRetargeterAsset);
+		PreloadObject(Node.IKRetargeterAsset->SourceIKRigAsset);
+		PreloadObject(Node.IKRetargeterAsset->TargetIKRigAsset);
+	}
 }
 
 #undef LOCTEXT_NAMESPACE
diff --git a/Engine/Plugins/Animation/IKRig/Source/IKRigDeveloper/Public/AnimGraphNode_IKRetargeter.h b/Engine/Plugins/Animation/IKRig/Source/IKRigDeveloper/Public/AnimGraphNode_IKRetargeter.h
--- a/Engine/Plugins/Animation/IKRig/Source/IKRigDeveloper/Public/AnimGraphNode_IKRetargeter.h
+++ b/Engine/Plugins/Animation/IKRig/Source/IKRigDeveloper/Public/AnimGraphNode_IKRetargeter.h
@@ -33,6 +33,7 @@ public:
 	virtual void CopyNodeDataToPreviewNode(FAnimNode_Base* AnimNode) override;
 	virtual FEditorModeID GetEditorMode() const override;
 	virtual void Draw(FPrimitiveDrawInterface* PDI, USkeletalMeshComponent* PreviewSkelMeshComp) const override;
+	virtual void CustomizePinData(UEdGraphPin* Pin, FName SourcePropertyName, int32 ArrayIndex) const override;
 	// End of UAnimGraphNode_Base interface
 
 	static const FName AnimModeName;
