Fix, crash bug in IKRig anim node when connected but not referencing an IK Rig definition asset.

diff --git a/Engine/Plugins/Animation/IKRig/Source/IKRig/Private/AnimNodes/AnimNode_IKRig.cpp b/Engine/Plugins/Animation/IKRig/Source/IKRig/Private/AnimNodes/AnimNode_IKRig.cpp
--- a/Engine/Plugins/Animation/IKRig/Source/IKRig/Private/AnimNodes/AnimNode_IKRig.cpp
+++ b/Engine/Plugins/Animation/IKRig/Source/IKRig/Private/AnimNodes/AnimNode_IKRig.cpp
@@ -218,6 +218,11 @@ void FAnimNode_IKRig::CacheBones_AnyThread(const FAnimationCacheBonesContext& Co
 	{
 		return;
 	}
+	
+	if (!RigDefinitionAsset)
+	{
+		return;
+	}
 
 	// fill up node names
 	CompactPoseToRigIndices.Reset();
