Fixed bug, IKRig running on stored ref pose instead of ref pose from actual SkeletalMesh.

diff --git a/Engine/Plugins/Animation/IKRig/Source/IKRig/Private/AnimNodes/AnimNode_IKRig.cpp b/Engine/Plugins/Animation/IKRig/Source/IKRig/Private/AnimNodes/AnimNode_IKRig.cpp
--- a/Engine/Plugins/Animation/IKRig/Source/IKRig/Private/AnimNodes/AnimNode_IKRig.cpp
+++ b/Engine/Plugins/Animation/IKRig/Source/IKRig/Private/AnimNodes/AnimNode_IKRig.cpp
@@ -123,14 +123,32 @@ void FAnimNode_IKRig::Initialize_AnyThread(const FAnimationInitializeContext& Co
 	DECLARE_SCOPE_HIERARCHICAL_COUNTER_FUNC()
 	FAnimNode_Base::Initialize_AnyThread(Context);
 	Source.Initialize(Context);
-	if (!IKRigProcessor)
+	
+	// get the retargeted local ref pose to initialize the IK with
+	const FReferenceSkeleton& RefSkeleton = Context.AnimInstanceProxy->GetSkelMeshComponent()->SkeletalMesh->GetRefSkeleton();
+	// initialize the IK Rig
+	IKRigProcessor->Initialize(RigDefinitionAsset, RefSkeleton);
+}
+
+void FAnimNode_IKRig::OnInitializeAnimInstance(const FAnimInstanceProxy* InProxy, const UAnimInstance* InAnimInstance)
+{
+	DECLARE_SCOPE_HIERARCHICAL_COUNTER_FUNC()
+	FAnimNode_Base::OnInitializeAnimInstance(InProxy, InAnimInstance);
+	
+	if (RigDefinitionAsset && !IKRigProcessor)
 	{
-		return;
+		IKRigProcessor = UIKRigProcessor::MakeNewIKRigProcessor(InAnimInstance->GetOwningComponent());
 	}
+}
 
+void FAnimNode_IKRig::PreUpdate(const UAnimInstance* InAnimInstance)
+{
 	// cache list of goal creator components on the actor
+	// TODO tried doing this in Initialize_AnyThread but it would miss some GoalCreator components
+	// so it was moved here to be more robust, but we need to profile this and make sure it's not hurting perf
+	// (it may be enough to run this once and then never again...needs testing)
 	GoalCreators.Reset();
-	USkeletalMeshComponent* SkelMeshComponent = Context.AnimInstanceProxy->GetSkelMeshComponent();
+	USkeletalMeshComponent* SkelMeshComponent = InAnimInstance->GetSkelMeshComponent();
 	AActor* OwningActor = SkelMeshComponent->GetOwner();
 	TArray<UActorComponent*> GoalCreatorComponents =  OwningActor->GetComponentsByInterface( UIKGoalCreatorInterface::StaticClass() );
 	for (UActorComponent* GoalCreatorComponent : GoalCreatorComponents)
@@ -142,24 +160,7 @@ void FAnimNode_IKRig::Initialize_AnyThread(const FAnimationInitializeContext& Co
 		}
 		GoalCreators.Add(GoalCreator);
 	}
-
-	// initialize the IK Rig
-	IKRigProcessor->Initialize(RigDefinitionAsset);
-}
-
-void FAnimNode_IKRig::OnInitializeAnimInstance(const FAnimInstanceProxy* InProxy, const UAnimInstance* InAnimInstance)
-{
-	DECLARE_SCOPE_HIERARCHICAL_COUNTER_FUNC()
-	FAnimNode_Base::OnInitializeAnimInstance(InProxy, InAnimInstance);
 	
-	if (RigDefinitionAsset && !IKRigProcessor)
-	{
-		IKRigProcessor = UIKRigProcessor::MakeNewIKRigProcessor(InAnimInstance->GetOwningComponent());
-	}
-}
-
-void FAnimNode_IKRig::PreUpdate(const UAnimInstance* InAnimInstance)
-{
 	// pull all the goals out of any goal creators on the owning actor
 	// this is done on the main thread because we're talking to actor components here
 	GoalsFromGoalCreators.Reset();
diff --git a/Engine/Plugins/Animation/IKRig/Source/IKRig/Private/IKRigProcessor.cpp b/Engine/Plugins/Animation/IKRig/Source/IKRig/Private/IKRigProcessor.cpp
--- a/Engine/Plugins/Animation/IKRig/Source/IKRig/Private/IKRigProcessor.cpp
+++ b/Engine/Plugins/Animation/IKRig/Source/IKRig/Private/IKRigProcessor.cpp
@@ -10,18 +10,28 @@ UIKRigProcessor* UIKRigProcessor::MakeNewIKRigProcessor(UObject* Outer)
 	return NewObject<UIKRigProcessor>(Outer);
 }
 
-void UIKRigProcessor::Initialize(UIKRigDefinition* InRigDefinition)
+void UIKRigProcessor::Initialize(UIKRigDefinition* InRigDefinition, const FReferenceSkeleton& RefSkeleton)
 {
 	bInitialized = false;
 
-	if (!ensureMsgf(InRigDefinition, TEXT("Trying to initialize IKRigProcessor with a null IKRigDefinition asset.")))
+	if (!InRigDefinition)
 	{
+		UE_LOG(LogTemp, Warning, TEXT("Trying to initialize IKRigProcessor with a null IKRigDefinition asset."));
 		return;
 	}
 
 	// copy skeleton data from IKRigDefinition
-	Skeleton = InRigDefinition->Skeleton; // trivial copy assignment operator for POD
-
+	// we use the serialized bone names and parent indices (from when asset was initialized)
+	// but we use the CURRENT ref pose from the currently running skeletal mesh (RefSkeleton)
+	Skeleton.BoneNames = InRigDefinition->Skeleton.BoneNames;
+	Skeleton.ParentIndices = InRigDefinition->Skeleton.ParentIndices;
+	const bool bSkeletonIsCompatible = Skeleton.CopyPosesFromRefSkeleton(RefSkeleton);
+	if (!bSkeletonIsCompatible)
+	{
+		UE_LOG(LogTemp, Warning, TEXT("IK Rig, %s trying to run on a skeleton that does not have the required bones."), *GetName());
+		return;
+	}
+	
 	// initialize goal names based on solvers
 	TArray<FIKRigEffectorGoal>& EffectorGoals = InRigDefinition->GetEffectorGoals();
 	GoalContainer.InitializeGoalsFromNames(EffectorGoals);
diff --git a/Engine/Plugins/Animation/IKRig/Source/IKRig/Private/IKRigSkeleton.cpp b/Engine/Plugins/Animation/IKRig/Source/IKRig/Private/IKRigSkeleton.cpp
--- a/Engine/Plugins/Animation/IKRig/Source/IKRig/Private/IKRigSkeleton.cpp
+++ b/Engine/Plugins/Animation/IKRig/Source/IKRig/Private/IKRigSkeleton.cpp
@@ -11,33 +11,15 @@ void FIKRigSkeleton::Initialize(const FReferenceSkeleton& RefSkeleton)
 	RefPoseGlobal.Reset();
 	
 	// copy names and parent indices into local storage
-	for (int32 BoneIndex=0; BoneIndex<RefSkeleton.GetNum(); ++BoneIndex)
+	TArray<FMeshBoneInfo> RawBoneInfo = RefSkeleton.GetRawRefBoneInfo();
+	for (int32 BoneIndex=0; BoneIndex<RefSkeleton.GetRawBoneNum(); ++BoneIndex)
 	{
-		BoneNames.Add(RefSkeleton.GetBoneName(BoneIndex));
-		ParentIndices.Add(RefSkeleton.GetParentIndex(BoneIndex));	
+		BoneNames.Add(RawBoneInfo[BoneIndex].Name);
+		ParentIndices.Add(RawBoneInfo[BoneIndex].ParentIndex);	
 	}
 
-	// store copy of reference pose (converted from local to global)
-	const TArray<FTransform>& RefPoseLocal = RefSkeleton.GetRefBonePose();
-	for (int32 BoneIndex=0; BoneIndex<RefPoseLocal.Num(); ++BoneIndex)
-	{
-		const int32 ParentIndex = ParentIndices[BoneIndex];
-		if (ParentIndex == INDEX_NONE)
-		{
-			RefPoseGlobal.Add(RefPoseLocal[BoneIndex]);
-			continue; // root bone is always in local space 
-		}
-
-		const FTransform& ChildLocalTransform  = RefPoseLocal[BoneIndex];
-		const FTransform& ParentGlobalTransform  = RefPoseGlobal[ParentIndex];
-		RefPoseGlobal.Add(ChildLocalTransform * ParentGlobalTransform);
-	}
-
-	// initialize current GLOBAL pose to reference pose
-	CurrentPoseGlobal = RefPoseGlobal;
-
-	// initialize current LOCAL pose
-	UpdateAllLocalTransformFromGlobal();
+	// copy all the poses out of the ref skeleton
+	CopyPosesFromRefSkeleton(RefSkeleton);
 }
 
 int32 FIKRigSkeleton::GetBoneIndexFromName(const FName InName) const
@@ -63,6 +45,56 @@ int32 FIKRigSkeleton::GetParentIndex(const int32 BoneIndex) const
 	return ParentIndices[BoneIndex];
 }
 
+bool FIKRigSkeleton::CopyPosesFromRefSkeleton(const FReferenceSkeleton& RefSkeleton)
+{
+	// get a compacted local ref pose based on the stored names
+	TArray<FTransform> CompactRefPoseLocal;
+	for (const FName& BoneName : BoneNames)
+	{
+		const int32 BoneIndex = RefSkeleton.FindBoneIndex(BoneName);
+		if (BoneIndex == INDEX_NONE)
+		{
+			UE_LOG(LogTemp, Warning, TEXT("IK Rig is running on a skeleton that is missing bone named, '%s'."), *BoneName.ToString());
+			return false;
+		}
+		CompactRefPoseLocal.Add(RefSkeleton.GetRefBonePose()[BoneIndex]);
+	}
+	
+	// copy local ref pose to global
+	ConvertLocalPoseToGlobal(ParentIndices, CompactRefPoseLocal, RefPoseGlobal);
+
+	// initialize current GLOBAL pose to reference pose
+	CurrentPoseGlobal = RefPoseGlobal;
+
+	// initialize current LOCAL pose
+	UpdateAllLocalTransformFromGlobal();
+
+	return true; // supplied ref skeleton had all the bones we are looking for
+}
+
+void FIKRigSkeleton::ConvertLocalPoseToGlobal(
+	const TArray<int32>& InParentIndices,
+	const TArray<FTransform>& InLocalPose,
+	TArray<FTransform>& OutGlobalPose)
+{
+	check(InLocalPose.Num() == InParentIndices.Num());
+	
+	OutGlobalPose.Reset();
+	for (int32 BoneIndex=0; BoneIndex<InLocalPose.Num(); ++BoneIndex)
+	{
+		const int32 ParentIndex = InParentIndices[BoneIndex];
+		if (ParentIndex == INDEX_NONE)
+		{
+			OutGlobalPose.Add(InLocalPose[BoneIndex]);
+			continue; // root bone is always in local space 
+		}
+
+		const FTransform& ChildLocalTransform  = InLocalPose[BoneIndex];
+		const FTransform& ParentGlobalTransform  = OutGlobalPose[ParentIndex];
+		OutGlobalPose.Add(ChildLocalTransform * ParentGlobalTransform);
+	}
+}
+
 void FIKRigSkeleton::UpdateAllGlobalTransformFromLocal()
 {
 	// generate GLOBAL transforms based on LOCAL transforms
diff --git a/Engine/Plugins/Animation/IKRig/Source/IKRig/Private/Retargeter/IKRetargeter.cpp b/Engine/Plugins/Animation/IKRig/Source/IKRig/Private/Retargeter/IKRetargeter.cpp
--- a/Engine/Plugins/Animation/IKRig/Source/IKRig/Private/Retargeter/IKRetargeter.cpp
+++ b/Engine/Plugins/Animation/IKRig/Source/IKRig/Private/Retargeter/IKRetargeter.cpp
@@ -680,7 +680,7 @@ void UIKRetargeter::Initialize(USkeletalMesh* SourceSkeletalMesh, USkeletalMesh*
 	}
 
 	// initialize the IKRigProcessor for doing IK decoding
-	const bool bIKRigInitialized = InitializeIKRig(Outer);
+	const bool bIKRigInitialized = InitializeIKRig(Outer, TargetSkeletalMesh->GetRefSkeleton());
 	if (!bIKRigInitialized)
 	{
 		// couldn't initialize the IK Rig
@@ -780,7 +780,7 @@ bool UIKRetargeter::InitializeBoneChainPairs()
 	return !(ChainPairsIK.IsEmpty() && ChainPairsFK.IsEmpty());
 }
 
-bool UIKRetargeter::InitializeIKRig(UObject* Outer)
+bool UIKRetargeter::InitializeIKRig(UObject* Outer, const FReferenceSkeleton& InRefSkeleton)
 {
 	// make a new IKRigProcessor if we haven't already
 	if (!IKRigProcessor)
@@ -789,7 +789,7 @@ bool UIKRetargeter::InitializeIKRig(UObject* Outer)
 	}
 	
 	// initialize IK Rig runtime processor
-	IKRigProcessor->Initialize(TargetIKRigAsset);
+	IKRigProcessor->Initialize(TargetIKRigAsset, InRefSkeleton);
 	if (!IKRigProcessor->IsInitialized())
 	{
 		return false;
diff --git a/Engine/Plugins/Animation/IKRig/Source/IKRig/Public/IKRigProcessor.h b/Engine/Plugins/Animation/IKRig/Source/IKRig/Public/IKRigProcessor.h
--- a/Engine/Plugins/Animation/IKRig/Source/IKRig/Public/IKRigProcessor.h
+++ b/Engine/Plugins/Animation/IKRig/Source/IKRig/Public/IKRigProcessor.h
@@ -44,7 +44,7 @@ public:
 	static UIKRigProcessor* MakeNewIKRigProcessor(UObject* Outer);
 	 
 	/** setup a new processor to run the given IKRig asset  */
-	void Initialize(UIKRigDefinition* InRigDefinition);
+	void Initialize(UIKRigDefinition* InRigDefinition, const FReferenceSkeleton& RefSkeleton);
 
 	//
 	// BEGIN UPDATE SEQUENCE FUNCTIONS
diff --git a/Engine/Plugins/Animation/IKRig/Source/IKRig/Public/IKRigSkeleton.h b/Engine/Plugins/Animation/IKRig/Source/IKRig/Public/IKRigSkeleton.h
--- a/Engine/Plugins/Animation/IKRig/Source/IKRig/Public/IKRigSkeleton.h
+++ b/Engine/Plugins/Animation/IKRig/Source/IKRig/Public/IKRigSkeleton.h
@@ -32,6 +32,13 @@ struct IKRIG_API FIKRigSkeleton
 	int32 GetBoneIndexFromName(const FName InName) const;
 	
 	int32 GetParentIndex(const int32 BoneIndex) const;
+
+	bool CopyPosesFromRefSkeleton(const FReferenceSkeleton& RefSkeleton);
+
+	static void ConvertLocalPoseToGlobal(
+		const TArray<int32>& InParentIndices,
+		const TArray<FTransform>& InLocalPose,
+		TArray<FTransform>& OutGlobalPose);
 	
 	void UpdateAllGlobalTransformFromLocal();
 	
@@ -43,7 +50,6 @@ struct IKRIG_API FIKRigSkeleton
 	
 	void PropagateGlobalPoseBelowBone(const int32 BoneIndex);
 
-	//static void NormalizeRotations(TArray<FTransform>& Transforms);
 	static void NormalizeRotations(TArray<FTransform>& Transforms);
 };
 
diff --git a/Engine/Plugins/Animation/IKRig/Source/IKRig/Public/Retargeter/IKRetargeter.h b/Engine/Plugins/Animation/IKRig/Source/IKRig/Public/Retargeter/IKRetargeter.h
--- a/Engine/Plugins/Animation/IKRig/Source/IKRig/Public/Retargeter/IKRetargeter.h
+++ b/Engine/Plugins/Animation/IKRig/Source/IKRig/Public/Retargeter/IKRetargeter.h
@@ -313,7 +313,7 @@ private:
 		
 	bool InitializeBoneChainPairs();
 	
-	bool InitializeIKRig(UObject* Outer);
+	bool InitializeIKRig(UObject* Outer, const FReferenceSkeleton& InRefSkeleton);
 
 	void RunRootRetarget(
 		const TArray<FTransform>& InGlobalTransforms,
