IK Rig Asset Editor: new Persona editor for IK rig assets.

diff --git a/Engine/Plugins/Animation/IKRig/Source/IKRig/IKRig.Build.cs b/Engine/Plugins/Animation/IKRig/Source/IKRig/IKRig.Build.cs
--- a/Engine/Plugins/Animation/IKRig/Source/IKRig/IKRig.Build.cs
+++ b/Engine/Plugins/Animation/IKRig/Source/IKRig/IKRig.Build.cs
@@ -15,13 +15,16 @@ public class IKRig : ModuleRules
 				"Engine",
 				"Slate",
 				"SlateCore",
-				"AnimationCore", 
-				"AnimGraphRuntime",
-				
+				"AnimationCore",
 				"PBIK"
 			}
 			);
 
+		if (Target.bBuildEditor)
+		{
+			PrivateDependencyModuleNames.Add("UnrealEd");
+		}
+
 		PublicDependencyModuleNames.AddRange(
 			new string[]
 			{
diff --git a/Engine/Plugins/Animation/IKRig/Source/IKRig/Private/AnimNodes/AnimNode_IKRig.cpp b/Engine/Plugins/Animation/IKRig/Source/IKRig/Private/AnimNodes/AnimNode_IKRig.cpp
--- a/Engine/Plugins/Animation/IKRig/Source/IKRig/Private/AnimNodes/AnimNode_IKRig.cpp
+++ b/Engine/Plugins/Animation/IKRig/Source/IKRig/Private/AnimNodes/AnimNode_IKRig.cpp
@@ -30,70 +30,95 @@ void FAnimNode_IKRig::Evaluate_AnyThread(FPoseContext& Output)
 		Output.ResetToRefPose();
 	}
 
-	if (!IKRigProcessor)
+	if (!RigDefinitionAsset)
 	{
 		return;
 	}
-	
-	if (!IKRigProcessor->IsInitialized())
+
+	if (IKRigProcessor.NeedsInitialized(RigDefinitionAsset))
 	{
 		return;
 	}
 
-	FCompactPose& OutPose = Output.Pose;
-	FIKRigSkeleton& IKRigSkeleton = IKRigProcessor->GetSkeleton();
+	// copy input pose to solver stack
+	CopyInputPoseToSolver(Output.Pose);
+	// update target goal transforms
+	AssignGoalTargets();
+	// run stack of solvers, 
+	IKRigProcessor.Solve();
+	// updates transforms with new pose
+	CopyOutputPoseToAnimGraph(Output.Pose);
 	
+	// debug drawing
+	QueueDrawInterface(Output.AnimInstanceProxy, Output.AnimInstanceProxy->GetComponentTransform());
+}
+
+void FAnimNode_IKRig::CopyInputPoseToSolver(FCompactPose& InputPose)
+{
+	// start Solve() from REFERENCE pose
 	if (bStartFromRefPose)
 	{
-		// start Solve() from REFERENCE pose
-		IKRigProcessor->SetInputPoseToRefPose();
-	}else
+		IKRigProcessor.SetInputPoseToRefPose();
+		return;
+	}
+	
+	// start Solve() from INPUT pose
+	// copy local bone transforms into IKRigProcessor skeleton
+	FIKRigSkeleton& IKRigSkeleton = IKRigProcessor.GetSkeleton();
+	for (FCompactPoseBoneIndex CPIndex : InputPose.ForEachBoneIndex())
 	{
-		// start Solve() from INPUT pose
-
-		// copy local bone transforms into IKRigProcessor skeleton
-		for (FCompactPoseBoneIndex CPIndex : OutPose.ForEachBoneIndex())
+		int32* Index = CompactPoseToRigIndices.Find(CPIndex);
+		if (Index)
 		{
-			int32* Index = CompactPoseToRigIndices.Find(CPIndex);
-			if (Index)
-			{
-				IKRigSkeleton.CurrentPoseLocal[*Index] = OutPose[CPIndex];
-			}
+			IKRigSkeleton.CurrentPoseLocal[*Index] = InputPose[CPIndex];
 		}
-		// update global pose in IK Rig
-		IKRigSkeleton.UpdateAllGlobalTransformFromLocal();
 	}
+	// update global pose in IK Rig
+	IKRigSkeleton.UpdateAllGlobalTransformFromLocal();
+}
 
+void FAnimNode_IKRig::AssignGoalTargets()
+{
 	// update goal transforms before solve
+	// these transforms can come from a few different sources, handled here...
+
+	// use the goal transforms from the source asset itself
+	// this is used to live preview results from the IK Rig editor
+	if (bDriveWithSourceAsset)
+	{
+		IKRigProcessor.CopyAllInputsFromSourceAssetAtRuntime(RigDefinitionAsset);
+		return;
+	}
+	
+	// copy transforms from this anim node's goal pins from blueprint
 	for (const FIKRigGoal& Goal : Goals)
 	{
-		IKRigProcessor->SetIKGoal(Goal);
+		IKRigProcessor.SetIKGoal(Goal);
 	}
 
-	// goals from goal creator components will override any goals that were manually set (they take precedence)
+	// override any goals that were manually set with goals from goal creator components (they take precedence)
 	for (const TPair<FName, FIKRigGoal>& GoalPair : GoalsFromGoalCreators)
 	{
-		IKRigProcessor->SetIKGoal(GoalPair.Value);
+		IKRigProcessor.SetIKGoal(GoalPair.Value);
 	}
+}
 
-	// run stack of solvers, updates global transforms with new pose
-	IKRigProcessor->Solve();
-
+void FAnimNode_IKRig::CopyOutputPoseToAnimGraph(FCompactPose& OutputPose)
+{
+	FIKRigSkeleton& IKRigSkeleton = IKRigProcessor.GetSkeleton();
+	
 	// update local transforms of current IKRig pose
 	IKRigSkeleton.UpdateAllLocalTransformFromGlobal();
 
 	// copy local transforms to output pose
-	for (FCompactPoseBoneIndex CPIndex : OutPose.ForEachBoneIndex())
+	for (FCompactPoseBoneIndex CPIndex : OutputPose.ForEachBoneIndex())
 	{
 		int32* Index = CompactPoseToRigIndices.Find(CPIndex);
 		if (Index)
 		{
-			Output.Pose[CPIndex] = IKRigSkeleton.CurrentPoseLocal[*Index];
+			OutputPose[CPIndex] = IKRigSkeleton.CurrentPoseLocal[*Index];
 		}
 	}
-
-	// draw the interface
-	QueueDrawInterface(Output.AnimInstanceProxy, Output.AnimInstanceProxy->GetComponentTransform());
 }
 
 void FAnimNode_IKRig::GatherDebugData(FNodeDebugData& DebugData)
@@ -123,26 +148,25 @@ void FAnimNode_IKRig::Initialize_AnyThread(const FAnimationInitializeContext& Co
 	DECLARE_SCOPE_HIERARCHICAL_COUNTER_FUNC()
 	FAnimNode_Base::Initialize_AnyThread(Context);
 	Source.Initialize(Context);
-	
-	// get the retargeted local ref pose to initialize the IK with
-	const FReferenceSkeleton& RefSkeleton = Context.AnimInstanceProxy->GetSkelMeshComponent()->SkeletalMesh->GetRefSkeleton();
-	// initialize the IK Rig
-	IKRigProcessor->Initialize(RigDefinitionAsset, RefSkeleton);
 }
 
-void FAnimNode_IKRig::OnInitializeAnimInstance(const FAnimInstanceProxy* InProxy, const UAnimInstance* InAnimInstance)
+void FAnimNode_IKRig::Update_AnyThread(const FAnimationUpdateContext& Context)
 {
-	DECLARE_SCOPE_HIERARCHICAL_COUNTER_FUNC()
-	FAnimNode_Base::OnInitializeAnimInstance(InProxy, InAnimInstance);
-	
-	if (RigDefinitionAsset && !IKRigProcessor)
-	{
-		IKRigProcessor = UIKRigProcessor::MakeNewIKRigProcessor(InAnimInstance->GetOwningComponent());
-	}
+	GetEvaluateGraphExposedInputs().Execute(Context);
+	FAnimNode_Base::Update_AnyThread(Context);
+	Source.Update(Context);
 }
 
 void FAnimNode_IKRig::PreUpdate(const UAnimInstance* InAnimInstance)
 {
+	if (IKRigProcessor.NeedsInitialized(RigDefinitionAsset))
+	{
+		// get the retargeted local ref pose to initialize the IK with
+		const FReferenceSkeleton& RefSkeleton = InAnimInstance->CurrentSkeleton->GetReferenceSkeleton();
+		// initialize the IK Rig (will only try once on the current version of the rig asset)
+		IKRigProcessor.Initialize(RigDefinitionAsset, RefSkeleton, InAnimInstance->GetSkelMeshComponent());
+	}
+	
 	// cache list of goal creator components on the actor
 	// TODO tried doing this in Initialize_AnyThread but it would miss some GoalCreator components
 	// so it was moved here to be more robust, but we need to profile this and make sure it's not hurting perf
@@ -176,39 +200,43 @@ bool FAnimNode_IKRig::RebuildGoalList()
 	{
 		return false;
 	}
-	
-	const TArray<FIKRigEffectorGoal>& AllGoalNamesInIKRig = RigDefinitionAsset->GetEffectorGoals();
-	const int32 NumGoalsInRig = AllGoalNamesInIKRig.Num();
+
+	// number of goals changed
+	const int32 NumGoalsInRig = RigDefinitionAsset->Goals.Num();
 	if (Goals.Num() != NumGoalsInRig)
 	{
 		Goals.SetNum(NumGoalsInRig);
 		for (int32 i=0; i<NumGoalsInRig; ++i)
 		{
-			Goals[i].Name = AllGoalNamesInIKRig[i].Goal;
+			Goals[i].Name = RigDefinitionAsset->Goals[i]->GoalName;
 		}
 		return true;
 	}
 
-	return false;
+	// potentially number of goals remains identical, but only names changed
+	bool bNameUpdated = false;
+	for (int32 i=0; i<NumGoalsInRig; ++i)
+	{
+		if (Goals[i].Name != RigDefinitionAsset->Goals[i]->GoalName)
+		{
+			Goals[i].Name = RigDefinitionAsset->Goals[i]->GoalName;
+			bNameUpdated = true;
+		}
+	}
+
+	return bNameUpdated;
 }
 
 FName FAnimNode_IKRig::GetGoalName(int32 Index) const
 {
-	if (RigDefinitionAsset)
+	if (RigDefinitionAsset && RigDefinitionAsset->Goals.IsValidIndex(Index))
 	{
-		return RigDefinitionAsset->GetGoalName(Index);
+		return RigDefinitionAsset->Goals[Index]->GoalName;	
 	}
 
 	return NAME_None;
 }
 
-void FAnimNode_IKRig::Update_AnyThread(const FAnimationUpdateContext& Context)
-{
-	GetEvaluateGraphExposedInputs().Execute(Context);
-	FAnimNode_Base::Update_AnyThread(Context);
-	Source.Update(Context);
-}
-
 void FAnimNode_IKRig::CacheBones_AnyThread(const FAnimationCacheBonesContext& Context)
 {
 	FAnimNode_Base::CacheBones_AnyThread(Context);
@@ -246,9 +274,10 @@ void FAnimNode_IKRig::CacheBones_AnyThread(const FAnimationCacheBonesContext& Co
 
 void FAnimNode_IKRig::QueueDrawInterface(FAnimInstanceProxy* AnimProxy, const FTransform& ComponentToWorld) const
 {
-	check (IKRigProcessor);
 	check (AnimProxy);
 
+	/*
+	 *TODO implement basic drawing interface for IKRig solvers
 	for (const FControlRigDrawInstruction& Instruction : IKRigProcessor->GetDrawInterface())
 	{
 		if (!Instruction.IsValid())
@@ -295,10 +324,10 @@ void FAnimNode_IKRig::QueueDrawInterface(FAnimInstanceProxy* AnimProxy, const FT
 // 			FDynamicMeshBuilder MeshBuilder(PDI->View->GetFeatureLevel());
 // 			MeshBuilder.AddVertices(Instruction.MeshVerts);
 // 			MeshBuilder.AddTriangles(Instruction.MeshIndices);
-// 			MeshBuilder.Draw(PDI, InstructionTransform.ToMatrixWithScale(), Instruction.MaterialRenderProxy, SDPG_World/*SDPG_Foreground*/);
+// 			MeshBuilder.Draw(PDI, InstructionTransform.ToMatrixWithScale(), Instruction.MaterialRenderProxy, SDPG_World);
 			break;
 		}
 
 		}
-	}
+	}*/
 }
\ No newline at end of file
diff --git a/Engine/Plugins/Animation/IKRig/Source/IKRig/Private/BoneSettings/RotationLimits.cpp b/Engine/Plugins/Animation/IKRig/Source/IKRig/Private/BoneSettings/RotationLimits.cpp
deleted file mode 100644
diff --git a/Engine/Plugins/Animation/IKRig/Source/IKRig/Private/IKRigBoneSetting.cpp b/Engine/Plugins/Animation/IKRig/Source/IKRig/Private/IKRigBoneSetting.cpp
deleted file mode 100644
diff --git a/Engine/Plugins/Animation/IKRig/Source/IKRig/Private/IKRigDataTypes.cpp b/Engine/Plugins/Animation/IKRig/Source/IKRig/Private/IKRigDataTypes.cpp
--- a/Engine/Plugins/Animation/IKRig/Source/IKRig/Private/IKRigDataTypes.cpp
+++ b/Engine/Plugins/Animation/IKRig/Source/IKRig/Private/IKRigDataTypes.cpp
@@ -3,12 +3,12 @@
 #include "IKRigDataTypes.h"
 #include "IKRigSolver.h"
 
-void FIKRigGoalContainer::InitializeGoalsFromNames(const TArray<FIKRigEffectorGoal>& InGoalNames)
+void FIKRigGoalContainer::InitializeFromGoals(const TArray<UIKRigEffectorGoal*>& InGoals)
 {
-	Goals.Empty(InGoalNames.Num());
-	for (const FIKRigEffectorGoal& GoalNames : InGoalNames)
+	Goals.Empty(InGoals.Num());
+	for (const UIKRigEffectorGoal* Goal : InGoals)
 	{
-		Goals.Emplace(GoalNames.Goal, GoalNames.Goal);
+		Goals.Emplace(Goal->GoalName, Goal);
 	}
 }
 
@@ -17,6 +17,20 @@ void FIKRigGoalContainer::SetIKGoal(const FIKRigGoal& InGoal)
 	Goals.Add(InGoal.Name, InGoal);
 }
 
+void FIKRigGoalContainer::SetIKGoal(const UIKRigEffectorGoal* InEffectorGoal)
+{
+	if (FIKRigGoal* Goal = Goals.Find(InEffectorGoal->GoalName))
+	{
+		Goal->Position = InEffectorGoal->CurrentTransform.GetTranslation();
+        Goal->Rotation = InEffectorGoal->CurrentTransform.Rotator();
+        Goal->PositionAlpha = InEffectorGoal->PositionAlpha;
+        Goal->RotationAlpha = InEffectorGoal->RotationAlpha;
+	}else
+	{
+		Goals.Emplace(InEffectorGoal->GoalName, InEffectorGoal);
+	}
+}
+
 bool FIKRigGoalContainer::GetGoalByName(const FName& InGoalName, FIKRigGoal& OutGoal) const
 {
 	if (const FIKRigGoal* Goal = Goals.Find(InGoalName))
diff --git a/Engine/Plugins/Animation/IKRig/Source/IKRig/Private/IKRigDefinition.cpp b/Engine/Plugins/Animation/IKRig/Source/IKRig/Private/IKRigDefinition.cpp
--- a/Engine/Plugins/Animation/IKRig/Source/IKRig/Private/IKRigDefinition.cpp
+++ b/Engine/Plugins/Animation/IKRig/Source/IKRig/Private/IKRigDefinition.cpp
@@ -1,8 +1,13 @@
 // Copyright Epic Games, Inc. All Rights Reserved.
 
 #include "IKRigDefinition.h"
-#include "IKRigSolver.h"
+#include "Engine/SkeletalMesh.h"
 
+#if WITH_EDITOR
+#include "ScopedTransaction.h"
+#endif
+
+#define LOCTEXT_NAMESPACE	"IKRigDefinition"
 
 FBoneChain* FRetargetDefinition::GetBoneChainByName(FName ChainName)
 {
@@ -17,69 +22,68 @@ FBoneChain* FRetargetDefinition::GetBoneChainByName(FName ChainName)
 	return nullptr;
 }
 
-TArray<FIKRigEffectorGoal>& UIKRigDefinition::GetEffectorGoals()
+void UIKRigDefinition::SetPreviewMesh(USkeletalMesh* PreviewMesh, bool bMarkAsDirty)
+{
+	#if WITH_EDITOR
+	if (bMarkAsDirty)
+	{
+		FScopedTransaction Transaction(LOCTEXT("SetPreviewMesh_Label", "Set Preview Mesh"));
+		Modify();
+		PreviewSkeletalMesh = PreviewMesh;
+		return;
+	}
+	#endif
+	
+	PreviewSkeletalMesh = PreviewMesh;
+}
+
+USkeletalMesh* UIKRigDefinition::GetPreviewMesh() const
 {
-	UpdateGoalNameArray();
-	return EffectorGoals;
+	return PreviewSkeletalMesh.Get();
 }
 
 #if WITH_EDITOR
-void UIKRigDefinition::PostEditChangeProperty(FPropertyChangedEvent& PropertyChangedEvent)
+bool UIKRigDefinition::Modify(bool bAlwaysMarkDirty /*=true*/)
 {
-	// update list of goal names whenever a solver is modified
-	const FName PropertyName = (PropertyChangedEvent.Property ? PropertyChangedEvent.Property->GetFName() : NAME_None);
-	if ((PropertyName == GET_MEMBER_NAME_STRING_CHECKED(UIKRigDefinition, Solvers)))
+	const bool bSavedToTransactionBuffer = Super::Modify(bAlwaysMarkDirty);
+	if (bSavedToTransactionBuffer)
 	{
-		bEffectorGoalsDirty = true;
+		AssetVersion++; // inform any runtime/editor systems they should copy-up modifications	
 	}
-
-	Super::PostEditChangeProperty(PropertyChangedEvent);
+	return bSavedToTransactionBuffer;
 }
 #endif
 
-FName UIKRigDefinition::GetBoneNameForGoal(const FName& GoalName)
+void UIKRigDefinition::PostLoad()
 {
-	UpdateGoalNameArray();
-	
-	for (const FIKRigEffectorGoal& EffectorGoal : EffectorGoals)
-	{
-		if (EffectorGoal.Goal == GoalName)
-		{
-			return EffectorGoal.Bone;
-		}
-	}
-	
-	return NAME_None;
+	Super::PostLoad();
+	ResetGoalTransforms();
 }
 
-FName UIKRigDefinition::GetGoalName(int32 GoalIndex)
+void UIKRigDefinition::ResetGoalTransforms() const
 {
-	UpdateGoalNameArray();
-	
-	if (!EffectorGoals.IsValidIndex(GoalIndex))
+	for (UIKRigEffectorGoal* Goal : Goals)
 	{
-		return NAME_None;
-	}
-
-	return EffectorGoals[GoalIndex].Goal;
+		const FTransform InitialTransform =  GetGoalInitialTransform(Goal);
+		Goal->InitialTransform = InitialTransform;
+		Goal->CurrentTransform = InitialTransform;
+	}	
 }
 
-void UIKRigDefinition::UpdateGoalNameArray()
+FTransform UIKRigDefinition::GetGoalInitialTransform(UIKRigEffectorGoal* Goal) const
 {
-	if (!bEffectorGoalsDirty)
+	if (!Goal)
 	{
-		return;	
+		return FTransform::Identity; // null
 	}
-	
-	EffectorGoals.Reset();
-	for (UIKRigSolver* Solver : Solvers)
+
+	const int32 BoneIndex = Skeleton.GetBoneIndexFromName(Goal->BoneName);
+	if (BoneIndex == INDEX_NONE)
 	{
-		if (Solver)
-		{
-			Solver->AddGoalsInSolver(EffectorGoals);
-		}
+		return FTransform::Identity; // goal references unknown bone
 	}
 
-	bEffectorGoalsDirty = false;
+	return Skeleton.RefPoseGlobal[BoneIndex];
 }
 
+#undef LOCTEXT_NAMESPACE
diff --git a/Engine/Plugins/Animation/IKRig/Source/IKRig/Private/IKRigProcessor.cpp b/Engine/Plugins/Animation/IKRig/Source/IKRig/Private/IKRigProcessor.cpp
--- a/Engine/Plugins/Animation/IKRig/Source/IKRig/Private/IKRigProcessor.cpp
+++ b/Engine/Plugins/Animation/IKRig/Source/IKRig/Private/IKRigProcessor.cpp
@@ -5,86 +5,97 @@
 #include "IKRigSolver.h"
 
 
-UIKRigProcessor* UIKRigProcessor::MakeNewIKRigProcessor(UObject* Outer)
-{
-	return NewObject<UIKRigProcessor>(Outer);
-}
-
-void UIKRigProcessor::Initialize(UIKRigDefinition* InRigDefinition, const FReferenceSkeleton& RefSkeleton)
+void FIKRigProcessor::Initialize(UIKRigDefinition* InRigAsset, const FReferenceSkeleton& RefSkeleton, UObject* Outer)
 {
+	// we instantiate UObjects here which MUST be done on game thread...
+	check(IsInGameThread());
+	
 	bInitialized = false;
+	InitializedWithIKRigAssetVersion = -1;
 
-	if (!InRigDefinition)
+	if (!InRigAsset)
 	{
 		UE_LOG(LogTemp, Warning, TEXT("Trying to initialize IKRigProcessor with a null IKRigDefinition asset."));
 		return;
 	}
 
+	// bail out if we've already tried initializing with this exact version of the rig asset
+	if (LastVersionTried == InRigAsset->GetAssetVersion())
+	{
+		return; // don't keep spamming
+	}
+	LastVersionTried = InRigAsset->GetAssetVersion();
+
 	// copy skeleton data from IKRigDefinition
 	// we use the serialized bone names and parent indices (from when asset was initialized)
 	// but we use the CURRENT ref pose from the currently running skeletal mesh (RefSkeleton)
-	Skeleton.BoneNames = InRigDefinition->Skeleton.BoneNames;
-	Skeleton.ParentIndices = InRigDefinition->Skeleton.ParentIndices;
+	Skeleton.BoneNames = InRigAsset->Skeleton.BoneNames;
+	Skeleton.ParentIndices = InRigAsset->Skeleton.ParentIndices;
 	const bool bSkeletonIsCompatible = Skeleton.CopyPosesFromRefSkeleton(RefSkeleton);
 	if (!bSkeletonIsCompatible)
 	{
-		UE_LOG(LogTemp, Warning, TEXT("IK Rig, %s trying to run on a skeleton that does not have the required bones."), *GetName());
+		UE_LOG(LogTemp, Warning, TEXT("IK Rig, %s trying to run on a skeleton that does not have the required bones."), *InRigAsset->GetName());
 		return;
 	}
 	
 	// initialize goal names based on solvers
-	TArray<FIKRigEffectorGoal>& EffectorGoals = InRigDefinition->GetEffectorGoals();
-	GoalContainer.InitializeGoalsFromNames(EffectorGoals);
-	for (const FIKRigEffectorGoal& EffectorGoal : EffectorGoals)
+	GoalContainer.InitializeFromGoals(InRigAsset->Goals);
+	for (const UIKRigEffectorGoal* EffectorGoal : InRigAsset->Goals)
 	{
 		FGoalBone NewGoalBone;
-		NewGoalBone.BoneName = EffectorGoal.Bone;
-		NewGoalBone.BoneIndex = Skeleton.GetBoneIndexFromName(EffectorGoal.Bone);
+		NewGoalBone.BoneName = EffectorGoal->BoneName;
+		NewGoalBone.BoneIndex = Skeleton.GetBoneIndexFromName(EffectorGoal->BoneName);
 
 		// validate that the skeleton we are trying to solve this goal on contains the bone the goal expects
 		if (NewGoalBone.BoneIndex == INDEX_NONE)
 		{
 			UE_LOG(LogTemp, Warning, TEXT("IK Rig, %s has a Goal, '%s' that references an unknown bone, '%s'. Cannot evaluate."),
-				*GetName(), *EffectorGoal.Goal.ToString(), *EffectorGoal.Bone.ToString());
+				*InRigAsset->GetName(), *EffectorGoal->GoalName.ToString(), *EffectorGoal->BoneName.ToString());
 			return;
 		}
 
 		// validate that there is not already a different goal, with the same name, that is using a different bone
 		// (all goals with the same name must reference the same bone within a single IK Rig)
-		if (const FGoalBone* Bone = GoalBones.Find(EffectorGoal.Goal))
+		if (const FGoalBone* Bone = GoalBones.Find(EffectorGoal->GoalName))
 		{
 			if (Bone->BoneName != NewGoalBone.BoneName)
 			{
 				UE_LOG(LogTemp, Warning, TEXT("IK Rig, %s has a Goal, '%s' that references different bones in different solvers, '%s' and '%s'. Cannot evaluate."),
-                *GetName(), *EffectorGoal.Goal.ToString(), *Bone->BoneName.ToString(), *NewGoalBone.BoneName.ToString());
+                *InRigAsset->GetName(), *EffectorGoal->GoalName.ToString(), *Bone->BoneName.ToString(), *NewGoalBone.BoneName.ToString());
 				return;
 			}
 		}
 		
-		GoalBones.Add(EffectorGoal.Goal, NewGoalBone);
+		GoalBones.Add(EffectorGoal->GoalName, NewGoalBone);
 	}
 
 	// create copies of all the solvers in the IK rig
-	Solvers.Reset(InRigDefinition->Solvers.Num());
-	for (UIKRigSolver* IKRigSolver : InRigDefinition->Solvers)
+	Solvers.Reset(InRigAsset->Solvers.Num());
+	int32 SolverIndex = 0;
+	for (UIKRigSolver* IKRigSolver : InRigAsset->Solvers)
 	{
 		if (!IKRigSolver)
 		{
 			// this can happen if asset references deleted IK Solver type
 			// which should only happen during development (if at all)
-			UE_LOG(LogTemp, Warning, TEXT("IK Rig, %s has null/unknown solver in it. Please remove it."), *GetName());
+			UE_LOG(LogTemp, Warning, TEXT("IK Rig, %s has null/unknown solver in it. Please remove it."), *InRigAsset->GetName());
 			continue;
 		}
-		
-		UIKRigSolver* Solver = DuplicateObject(IKRigSolver, this);
+
+		// new solver name
+		FString Name = IKRigSolver->GetName() + "_SolverInstance_";
+		Name.AppendInt(SolverIndex++);
+		UIKRigSolver* Solver = DuplicateObject(IKRigSolver, Outer, FName(*Name));
+		Solver->AddToRoot();
 		Solver->Initialize(Skeleton);
 		Solvers.Add(Solver);
 	}
 
+	InitializedWithIKRigAssetVersion = InRigAsset->GetAssetVersion();
 	bInitialized = true;
 }
 
-void UIKRigProcessor::SetInputPoseGlobal(const TArray<FTransform>& InGlobalBoneTransforms) 
+void FIKRigProcessor::SetInputPoseGlobal(const TArray<FTransform>& InGlobalBoneTransforms) 
 {
 	check(bInitialized);
 	check(InGlobalBoneTransforms.Num() == Skeleton.CurrentPoseGlobal.Num());
@@ -92,20 +103,26 @@ void UIKRigProcessor::SetInputPoseGlobal(const TArray<FTransform>& InGlobalBoneT
 	Skeleton.UpdateAllLocalTransformFromGlobal();
 }
 
-void UIKRigProcessor::SetInputPoseToRefPose()
+void FIKRigProcessor::SetInputPoseToRefPose()
 {
 	check(bInitialized);
 	Skeleton.CurrentPoseGlobal = Skeleton.RefPoseGlobal;
 	Skeleton.UpdateAllLocalTransformFromGlobal();
 }
 
-void UIKRigProcessor::SetIKGoal(const FIKRigGoal& InGoal)
+void FIKRigProcessor::SetIKGoal(const FIKRigGoal& InGoal)
+{
+	check(bInitialized);
+	GoalContainer.SetIKGoal(InGoal);
+}
+
+void FIKRigProcessor::SetIKGoal(const UIKRigEffectorGoal* InGoal)
 {
 	check(bInitialized);
 	GoalContainer.SetIKGoal(InGoal);
 }
 
-void UIKRigProcessor::Solve()
+void FIKRigProcessor::Solve()
 {
 	check(bInitialized);
 
@@ -113,12 +130,11 @@ void UIKRigProcessor::Solve()
 	BlendGoalsByAlpha();
 
 	// run all the solvers
-	DrawInterface.Reset();
 	for (UIKRigSolver* Solver : Solvers)
 	{
 		if (Solver->bEnabled)
 		{
-			Solver->Solve(Skeleton, GoalContainer, &DrawInterface);
+			Solver->Solve(Skeleton, GoalContainer);
 		}
 	}
 
@@ -126,35 +142,62 @@ void UIKRigProcessor::Solve()
 	Skeleton.NormalizeRotations(Skeleton.CurrentPoseGlobal);
 }
 
-void UIKRigProcessor::CopyOutputGlobalPoseToArray(TArray<FTransform>& OutputPoseGlobal) const
+void FIKRigProcessor::CopyOutputGlobalPoseToArray(TArray<FTransform>& OutputPoseGlobal) const
 {
 	OutputPoseGlobal = Skeleton.CurrentPoseGlobal;
 }
 
-const FIKRigGoalContainer& UIKRigProcessor::GetGoalContainer() const
+void FIKRigProcessor::CopyAllInputsFromSourceAssetAtRuntime(UIKRigDefinition* IKRigAsset)
+{
+	// copy goal settings
+	for (const UIKRigEffectorGoal* AssetGoal : IKRigAsset->Goals)
+	{
+		SetIKGoal(AssetGoal);
+	}
+
+	// copy solver settings
+	check(Solvers.Num() == IKRigAsset->Solvers.Num()); // if number of solvers has been changed, processor should have been reinitialized
+	for (int32 SolverIndex=0; SolverIndex<Solvers.Num(); ++SolverIndex)
+	{
+		Solvers[SolverIndex]->bEnabled = IKRigAsset->Solvers[SolverIndex]->bEnabled;
+		Solvers[SolverIndex]->UpdateSolverSettings(IKRigAsset->Solvers[SolverIndex]);
+	}
+}
+
+bool FIKRigProcessor::NeedsInitialized(UIKRigDefinition* IKRigAsset) const
+{
+	if (!bInitialized)
+	{
+		return true; // not initialized yet at all
+	}
+
+	if (!IKRigAsset)
+	{
+		return true; // lost connection to asset
+	}
+
+	if (InitializedWithIKRigAssetVersion != IKRigAsset->GetAssetVersion())
+	{
+		return true; // IKRig asset has been modified since last initialization
+	}
+
+	return false;
+}
+
+const FIKRigGoalContainer& FIKRigProcessor::GetGoalContainer() const
 {
 	check(bInitialized);
 	return GoalContainer;
 }
 
-FIKRigSkeleton& UIKRigProcessor::GetSkeleton()
+FIKRigSkeleton& FIKRigProcessor::GetSkeleton()
 {
 	check(bInitialized);
 	return Skeleton;
 }
 
-bool UIKRigProcessor::GetBoneForGoal(FName GoalName, FGoalBone& OutBone) const
-{
-	const FGoalBone* FoundBone = GoalBones.Find(GoalName);
-	if (FoundBone)
-	{
-		OutBone = *FoundBone;
-		return true;
-	}
-	return false;
-}
 
-void UIKRigProcessor::BlendGoalsByAlpha()
+void FIKRigProcessor::BlendGoalsByAlpha()
 {
 	for (TPair<FName, FIKRigGoal>& GoalPair : GoalContainer.Goals)
 	{
diff --git a/Engine/Plugins/Animation/IKRig/Source/IKRig/Private/IKRigSkeleton.cpp b/Engine/Plugins/Animation/IKRig/Source/IKRig/Private/IKRigSkeleton.cpp
--- a/Engine/Plugins/Animation/IKRig/Source/IKRig/Private/IKRigSkeleton.cpp
+++ b/Engine/Plugins/Animation/IKRig/Source/IKRig/Private/IKRigSkeleton.cpp
@@ -4,11 +4,8 @@
 
 void FIKRigSkeleton::Initialize(const FReferenceSkeleton& RefSkeleton)
 {
-	BoneNames.Reset();
-	ParentIndices.Reset();
-	CurrentPoseGlobal.Reset();
-	CurrentPoseLocal.Reset();
-	RefPoseGlobal.Reset();
+	// reset all containers
+	Reset();
 	
 	// copy names and parent indices into local storage
 	TArray<FMeshBoneInfo> RawBoneInfo = RefSkeleton.GetRawRefBoneInfo();
@@ -22,6 +19,15 @@ void FIKRigSkeleton::Initialize(const FReferenceSkeleton& RefSkeleton)
 	CopyPosesFromRefSkeleton(RefSkeleton);
 }
 
+void FIKRigSkeleton::Reset()
+{
+	BoneNames.Reset();
+	ParentIndices.Reset();
+	CurrentPoseGlobal.Reset();
+	CurrentPoseLocal.Reset();
+	RefPoseGlobal.Reset();
+}
+
 int32 FIKRigSkeleton::GetBoneIndexFromName(const FName InName) const
 {
 	for (int32 Index=0; Index<BoneNames.Num(); ++Index)
@@ -152,10 +158,37 @@ void FIKRigSkeleton::PropagateGlobalPoseBelowBone(const int32 StartBoneIndex)
 	}
 }
 
+bool FIKRigSkeleton::IsBoneInDirectLineage(const FName& Child, const FName& PotentialParent) const
+{
+	const int32 ChildIndex = GetBoneIndexFromName(Child);
+	if (ChildIndex == INDEX_NONE)
+	{
+		return false;
+	}
+
+	const int32 PotentialParentIndex = GetBoneIndexFromName(PotentialParent);
+	if (PotentialParentIndex == INDEX_NONE)
+	{
+		return false;
+	}
+
+	int32 NextParentIndex = ChildIndex;
+	while(NextParentIndex != INDEX_NONE)
+	{
+		if (NextParentIndex == PotentialParentIndex)
+		{
+			return true;
+		}
+		NextParentIndex = ParentIndices[NextParentIndex];
+	}
+
+	return false;
+}
+
 void FIKRigSkeleton::NormalizeRotations(TArray<FTransform>& Transforms)
 {
 	for (FTransform& Transform : Transforms)
 	{
 		Transform.NormalizeRotation();
 	}
-}
\ No newline at end of file
+}
diff --git a/Engine/Plugins/Animation/IKRig/Source/IKRig/Private/IKRigSolver.cpp b/Engine/Plugins/Animation/IKRig/Source/IKRig/Private/IKRigSolver.cpp
deleted file mode 100644
diff --git a/Engine/Plugins/Animation/IKRig/Source/IKRig/Private/Retargeter/IKRetargeter.cpp b/Engine/Plugins/Animation/IKRig/Source/IKRig/Private/Retargeter/IKRetargeter.cpp
--- a/Engine/Plugins/Animation/IKRig/Source/IKRig/Private/Retargeter/IKRetargeter.cpp
+++ b/Engine/Plugins/Animation/IKRig/Source/IKRig/Private/Retargeter/IKRetargeter.cpp
@@ -781,23 +781,17 @@ bool UIKRetargeter::InitializeBoneChainPairs()
 }
 
 bool UIKRetargeter::InitializeIKRig(UObject* Outer, const FReferenceSkeleton& InRefSkeleton)
-{
-	// make a new IKRigProcessor if we haven't already
-	if (!IKRigProcessor)
-	{
-		IKRigProcessor = UIKRigProcessor::MakeNewIKRigProcessor(Outer);
-	}
-	
+{	
 	// initialize IK Rig runtime processor
-	IKRigProcessor->Initialize(TargetIKRigAsset, InRefSkeleton);
-	if (!IKRigProcessor->IsInitialized())
+	IKRigProcessor.Initialize(TargetIKRigAsset, InRefSkeleton, Outer);
+	if (IKRigProcessor.NeedsInitialized(TargetIKRigAsset))
 	{
 		return false;
 	}
 
 	// validate that all IK bone chains have an associated Goal
 	TArray<FName> GoalNames;
-	IKRigProcessor->GetGoalContainer().Goals.GenerateKeyArray(GoalNames);
+	IKRigProcessor.GetGoalContainer().Goals.GenerateKeyArray(GoalNames);
 	for (FRetargetChainPairIK& ChainPair : ChainPairsIK)
 	{
 		// does the IK rig have the IK goal this bone chain requires?
@@ -882,13 +876,13 @@ void UIKRetargeter::RunIKRetarget(
 			OutIKGoal.EndEffectorRotation,
 			1.0f,
 			1.0f);
-		IKRigProcessor->SetIKGoal(Goal);
+		IKRigProcessor.SetIKGoal(Goal);
 	}
 
 	// copy input pose to start IK solve from
-	IKRigProcessor->SetInputPoseGlobal(OutGlobalPose);
+	IKRigProcessor.SetInputPoseGlobal(OutGlobalPose);
 	// run IK solve
-	IKRigProcessor->Solve();
+	IKRigProcessor.Solve();
 	// copy results of solve
-	IKRigProcessor->CopyOutputGlobalPoseToArray(OutGlobalPose);
+	IKRigProcessor.CopyOutputGlobalPoseToArray(OutGlobalPose);
 }
diff --git a/Engine/Plugins/Animation/IKRig/Source/IKRig/Private/Solvers/IKRig_PBIKSolver.cpp b/Engine/Plugins/Animation/IKRig/Source/IKRig/Private/Solvers/IKRig_PBIKSolver.cpp
--- a/Engine/Plugins/Animation/IKRig/Source/IKRig/Private/Solvers/IKRig_PBIKSolver.cpp
+++ b/Engine/Plugins/Animation/IKRig/Source/IKRig/Private/Solvers/IKRig_PBIKSolver.cpp
@@ -9,9 +9,9 @@ void UIKRigPBIKSolver::Initialize(const FIKRigSkeleton& InSkeleton)
 {	
 	// check how many effectors are assigned to a bone
 	int NumEffectors = 0;
-	for (const FIKRig_PBIKEffector& Effector : Effectors)
+	for (const UIKRig_FBIKEffector* Effector : Effectors)
 	{
-		if (InSkeleton.GetBoneIndexFromName(Effector.Goal.Bone) != INDEX_NONE)
+		if (InSkeleton.GetBoneIndexFromName(Effector->BoneName) != INDEX_NONE)
 		{
 			++NumEffectors; // bone is set and exists!
 		}
@@ -43,19 +43,16 @@ void UIKRigPBIKSolver::Initialize(const FIKRigSkeleton& InSkeleton)
 	}
 
 	// create effectors
-	for (FIKRig_PBIKEffector& Effector : Effectors)
+	for (UIKRig_FBIKEffector* Effector : Effectors)
 	{
-		Effector.IndexInSolver = Solver.AddEffector(Effector.Goal.Bone);
+		Effector->IndexInSolver = Solver.AddEffector(Effector->BoneName);
 	}
 		
 	// initialize solver
 	Solver.Initialize();
 }
 
-void UIKRigPBIKSolver::Solve(
-    FIKRigSkeleton& IKRigSkeleton,
-    const FIKRigGoalContainer& Goals,
-    FControlRigDrawInterface* InOutDrawInterface)
+void UIKRigPBIKSolver::Solve(FIKRigSkeleton& IKRigSkeleton, const FIKRigGoalContainer& Goals)
 {
 	if (!Solver.IsReadyToSimulate())
 	{
@@ -75,33 +72,41 @@ void UIKRigPBIKSolver::Solve(
 	}
 
 	// update bone settings
-	for (const FPBIKBoneSetting& BoneSetting : BoneSettings)
+	for (const UIKRig_PBIKBoneSettings* BoneSetting : BoneSettings)
 	{
-		const int32 BoneIndex = Solver.GetBoneIndex(BoneSetting.Bone);
+		const int32 BoneIndex = Solver.GetBoneIndex(BoneSetting->Bone);
 		if (PBIK::FBoneSettings* InternalSettings = Solver.GetBoneSettings(BoneIndex))
 		{
-			BoneSetting.CopyToCoreStruct(*InternalSettings);
+			BoneSetting->CopyToCoreStruct(*InternalSettings);
 		}
 	}
 
 	// update effectors
 	const float OffsetAlpha = 1.0f; // this is constant because IKRig manages offset alphas itself
-	for (const FIKRig_PBIKEffector& Effector : Effectors)
+	for (const UIKRig_FBIKEffector* Effector : Effectors)
 	{
-		if (Effector.IndexInSolver < 0)
+		if (Effector->IndexInSolver < 0)
 		{
 			continue;
 		}
 		FIKRigGoal Goal;
-		Goals.GetGoalByName(Effector.Goal.Goal, Goal);
+		Goals.GetGoalByName(Effector->GoalName, Goal);
 		Solver.SetEffectorGoal(
-			Effector.IndexInSolver,
+			Effector->IndexInSolver,
 			Goal.FinalBlendedPosition,
 			Goal.FinalBlendedRotation,
 			OffsetAlpha,
-			Effector.StrengthAlpha);
+			Effector->StrengthAlpha);
 	}
 
+	// update settings
+	FPBIKSolverSettings Settings;
+	Settings.Iterations = Iterations;
+	Settings.bAllowStretch = bAllowStretch;
+	Settings.MassMultiplier = MassMultiplier;
+	Settings.bPinRoot = bPinRoot;
+	Settings.bStartSolveFromInputPose = bStartSolveFromInputPose;
+
 	// solve
 	Solver.Solve(Settings);
 
@@ -112,15 +117,205 @@ void UIKRigPBIKSolver::Solve(
 		Solver.GetBoneGlobalTransform(BoneIndex, OutTransform);
 		IKRigSkeleton.CurrentPoseGlobal[BoneIndex] = OutTransform;
 	}
+}
+
+void UIKRigPBIKSolver::UpdateSolverSettings(UIKRigSolver* InSettings)
+{
+	if(UIKRigPBIKSolver* Settings = Cast<UIKRigPBIKSolver>(InSettings))
+	{
+		Iterations = Settings->Iterations;
+		bAllowStretch = Settings->bAllowStretch;
+		MassMultiplier = Settings->MassMultiplier;
+		MinMassMultiplier = Settings->MinMassMultiplier;
+		bPinRoot = Settings->bPinRoot;
+		bStartSolveFromInputPose = Settings->bStartSolveFromInputPose;
+
+		// copy effector settings
+		for (const UIKRig_FBIKEffector* InEffector : Settings->Effectors)
+		{
+			for (UIKRig_FBIKEffector* Effector : Effectors)
+			{
+				if (Effector->GoalName == InEffector->GoalName)
+				{
+					Effector->CopySettings(InEffector);
+					break;
+				}
+			}
+		}
+
+		// copy bone settings
+		for (const UIKRig_PBIKBoneSettings* InBoneSetting : Settings->BoneSettings)
+		{
+			for (UIKRig_PBIKBoneSettings* BoneSetting : BoneSettings)
+			{
+				if (BoneSetting->Bone == InBoneSetting->Bone)
+				{
+					BoneSetting->CopySettings(InBoneSetting);
+					break;	
+				}
+			}
+		}
+	}
+}
+
+void UIKRigPBIKSolver::AddGoal(const UIKRigEffectorGoal& NewGoal)
+{
+	UIKRig_FBIKEffector* NewEffector = NewObject<UIKRig_FBIKEffector>(this, UIKRig_FBIKEffector::StaticClass());
+	NewEffector->GoalName = NewGoal.GoalName;
+	NewEffector->BoneName = NewGoal.BoneName;
+	Effectors.Add(NewEffector);
+}
+
+void UIKRigPBIKSolver::RemoveGoal(const FName& GoalName)
+{
+	// ensure goal even exists in this solver
+	const int32 GoalIndex = GetIndexOfGoal(GoalName);
+	if (GoalIndex == INDEX_NONE)
+	{
+		return;
+	}
+
+	// remove it
+	Effectors.RemoveAt(GoalIndex);
+}
+
+void UIKRigPBIKSolver::RenameGoal(const FName& OldName, const FName& NewName)
+{
+	// ensure goal even exists in this solver
+	const int32 GoalIndex = GetIndexOfGoal(OldName);
+	if (GoalIndex == INDEX_NONE)
+	{
+		return;
+	}
+
+	// rename
+	Effectors[GoalIndex]->GoalName = NewName;
+}
+
+void UIKRigPBIKSolver::SetGoalBone(const FName& GoalName, const FName& NewBoneName)
+{
+	// ensure goal even exists in this solver
+	const int32 GoalIndex = GetIndexOfGoal(GoalName);
+	if (GoalIndex == INDEX_NONE)
+	{
+		return;
+	}
+
+	// rename
+	Effectors[GoalIndex]->BoneName = NewBoneName;
+}
+
+bool UIKRigPBIKSolver::IsGoalConnected(const FName& GoalName) const
+{
+	return GetIndexOfGoal(GoalName) != INDEX_NONE;
+}
+
+void UIKRigPBIKSolver::SetRootBone(const FName& RootBoneName)
+{
+	RootBone = RootBoneName;
+}
+
+UObject* UIKRigPBIKSolver::GetEffectorWithGoal(const FName& GoalName)
+{
+	const int32 GoalIndex = GetIndexOfGoal(GoalName);
+	if (GoalIndex == INDEX_NONE)
+	{
+		return nullptr;
+	}
+
+	return Effectors[GoalIndex];
+}
+
+void UIKRigPBIKSolver::AddBoneSetting(const FName& BoneName)
+{
+	if (GetBoneSetting(BoneName))
+	{
+		return; // already have settings on this bone
+	}
+
+	UIKRig_PBIKBoneSettings* NewBoneSettings = NewObject<UIKRig_PBIKBoneSettings>(this, UIKRig_PBIKBoneSettings::StaticClass());
+	NewBoneSettings->Bone = BoneName;
+	BoneSettings.Add(NewBoneSettings);
+}
+
+void UIKRigPBIKSolver::RemoveBoneSetting(const FName& BoneName)
+{
+	UIKRig_PBIKBoneSettings* BoneSettingToRemove; 
+	for (UIKRig_PBIKBoneSettings* BoneSetting : BoneSettings)
+	{
+		if (BoneSetting->Bone == BoneName)
+		{
+			BoneSettingToRemove = BoneSetting;
+			break; // can only be one with this name
+		}
+	}
+
+	BoneSettings.Remove(BoneSettingToRemove);
+}
+
+UObject* UIKRigPBIKSolver::GetBoneSetting(const FName& BoneName) const
+{
+	for (UIKRig_PBIKBoneSettings* BoneSetting : BoneSettings)
+	{
+		if (BoneSetting && BoneSetting->Bone == BoneName)
+		{
+			return BoneSetting;
+		}
+	}
+	
+	return nullptr;
+}
+
+void UIKRigPBIKSolver::DrawBoneSettings(
+	const FName& BoneName,
+	const FIKRigSkeleton& IKRigSkeleton,
+	FPrimitiveDrawInterface* PDI) const
+{
+	
+}
+
+bool UIKRigPBIKSolver::IsBoneAffectedBySolver(const FName& BoneName, const FIKRigSkeleton& IKRigSkeleton) const
+{
+	for (UIKRig_FBIKEffector* Effector : Effectors)
+	{
+		if (IKRigSkeleton.IsBoneInDirectLineage(Effector->BoneName, BoneName))
+		{
+			return true;
+		}
+	}
 	
-	// do all debug drawing
-	//Debug.Draw(Context.DrawInterface, &Solver);
+	return false;
+}
+
+int32 UIKRigPBIKSolver::GetIndexOfGoal(const FName& OldName) const
+{
+	for (int32 i=0; i<Effectors.Num(); ++i)
+	{
+		if (Effectors[i]->GoalName == OldName)
+		{
+			return i;
+		}
+	}
+
+	return INDEX_NONE;
 }
 
-void UIKRigPBIKSolver::AddGoalsInSolver(TArray<FIKRigEffectorGoal>& OutGoals) const
+void UIKRigPBIKSolver::PostLoad()
 {
-	for (const FIKRig_PBIKEffector& Effector : Effectors)
+	Super::PostLoad();
+	
+	// patch for loading old effectors, we blow them away
+	bool bHasNullEffector = false;
+	for (const UIKRig_FBIKEffector* Effector : Effectors)
 	{
-		AddGoalToArrayNoDuplicates(Effector.Goal,OutGoals);
+		if (!Effector)
+		{
+			bHasNullEffector = true;
+		}
+	}
+
+	if (bHasNullEffector)
+	{
+		Effectors.Empty();
 	}
-}
\ No newline at end of file
+}
diff --git a/Engine/Plugins/Animation/IKRig/Source/IKRig/Private/Solvers/IKRig_SetTransform.cpp b/Engine/Plugins/Animation/IKRig/Source/IKRig/Private/Solvers/IKRig_SetTransform.cpp
--- a/Engine/Plugins/Animation/IKRig/Source/IKRig/Private/Solvers/IKRig_SetTransform.cpp
+++ b/Engine/Plugins/Animation/IKRig/Source/IKRig/Private/Solvers/IKRig_SetTransform.cpp
@@ -5,38 +5,101 @@
 #include "IKRigSkeleton.h"
 
 
+UIKRig_SetTransform::UIKRig_SetTransform()
+{
+	Effector = CreateDefaultSubobject<UIKRig_SetTransformEffector>(TEXT("Effector"));
+}
+
 void UIKRig_SetTransform::Initialize(const FIKRigSkeleton& IKRigSkeleton)
 {
-	BoneIndex = IKRigSkeleton.GetBoneIndexFromName(Effector.Bone);
+	BoneIndex = IKRigSkeleton.GetBoneIndexFromName(Bone);
 }
 
-void UIKRig_SetTransform::Solve(
-	FIKRigSkeleton& IKRigSkeleton,
-	const FIKRigGoalContainer& Goals,
-	FControlRigDrawInterface* InOutDrawInterface)
+void UIKRig_SetTransform::Solve(FIKRigSkeleton& IKRigSkeleton, const FIKRigGoalContainer& Goals)
 {
 	FIKRigGoal OutGoal;
-	if (!Goals.GetGoalByName(Effector.Goal, OutGoal))
+	if (!Goals.GetGoalByName(Goal, OutGoal))
+	{
+		return;
+	}
+
+	// check that settings are such that there is anything to do at all
+	const bool bAnythingEnabled = Effector->bEnablePosition || Effector->bEnableRotation;
+	const bool bHasAlpha = Effector->Alpha > KINDA_SMALL_NUMBER;
+	if (!(bAnythingEnabled && bHasAlpha))
 	{
 		return;
 	}
 
-	FTransform& CurrentTransform = IKRigSkeleton.CurrentPoseGlobal[BoneIndex];
+	FTransform& CurrentTransform = IKRigSkeleton.CurrentPoseGlobal[BoneIndex];	
 
-	if (bEnablePosition)
+	if (Effector->bEnablePosition)
 	{
-		CurrentTransform.SetLocation(OutGoal.FinalBlendedPosition);
+		const FVector TargetPosition = FMath::Lerp(CurrentTransform.GetTranslation(), OutGoal.FinalBlendedPosition, Effector->Alpha);
+		CurrentTransform.SetLocation(TargetPosition);
 	}
 	
-	if (bEnableRotation)
+	if (Effector->bEnableRotation)
 	{
-		CurrentTransform.SetRotation(OutGoal.FinalBlendedRotation);
+		const FQuat TargetRotation = FMath::Lerp(CurrentTransform.GetRotation(), OutGoal.FinalBlendedRotation, Effector->Alpha);
+		CurrentTransform.SetRotation(TargetRotation);
 	}
-
+	
 	IKRigSkeleton.PropagateGlobalPoseBelowBone(BoneIndex);
 }
 
-void UIKRig_SetTransform::AddGoalsInSolver(TArray<FIKRigEffectorGoal>& OutGoals) const
+void UIKRig_SetTransform::UpdateSolverSettings(UIKRigSolver* InSettings)
+{
+	if (UIKRig_SetTransform* Settings = Cast<UIKRig_SetTransform>(InSettings))
+	{
+		Effector->bEnablePosition = Settings->Effector->bEnablePosition;
+		Effector->bEnableRotation = Settings->Effector->bEnableRotation;
+		Effector->Alpha = Settings->Effector->Alpha;
+	}
+}
+
+void UIKRig_SetTransform::AddGoal(const UIKRigEffectorGoal& NewGoal)
+{
+	Goal = NewGoal.GoalName;
+	Bone = NewGoal.BoneName;
+}
+
+void UIKRig_SetTransform::RemoveGoal(const FName& GoalName)
+{
+	if (Goal == GoalName)
+	{
+		Goal = NAME_None;
+		Bone = NAME_None;
+	}	
+}
+
+void UIKRig_SetTransform::RenameGoal(const FName& OldName, const FName& NewName)
+{
+	if (Goal == OldName)
+	{
+		Goal = NewName;
+	}
+}
+
+void UIKRig_SetTransform::SetGoalBone(const FName& GoalName, const FName& NewBoneName)
+{
+	if (Goal == GoalName)
+	{
+		Bone = NewBoneName;
+	}
+}
+
+bool UIKRig_SetTransform::IsGoalConnected(const FName& GoalName) const
+{
+	return Goal == GoalName;
+}
+
+UObject* UIKRig_SetTransform::GetEffectorWithGoal(const FName& GoalName)
+{
+	return Goal == GoalName ? Effector : nullptr;
+}
+
+bool UIKRig_SetTransform::IsBoneAffectedBySolver(const FName& BoneName, const FIKRigSkeleton& IKRigSkeleton) const
 {
-	AddGoalToArrayNoDuplicates(Effector,OutGoals);
+	return IKRigSkeleton.IsBoneInDirectLineage(BoneName, Bone);
 }
diff --git a/Engine/Plugins/Animation/IKRig/Source/IKRig/Public/AnimNodes/AnimNode_IKRetargeter.h b/Engine/Plugins/Animation/IKRig/Source/IKRig/Public/AnimNodes/AnimNode_IKRetargeter.h
--- a/Engine/Plugins/Animation/IKRig/Source/IKRig/Public/AnimNodes/AnimNode_IKRetargeter.h
+++ b/Engine/Plugins/Animation/IKRig/Source/IKRig/Public/AnimNodes/AnimNode_IKRetargeter.h
@@ -3,7 +3,8 @@
 #pragma once
 
 #include "CoreMinimal.h"
-#include "BoneControllers/AnimNode_SkeletalControlBase.h"
+
+#include "Animation/AnimNodeBase.h"
 #include "Retargeter/IKRetargeter.h"
 #include "AnimNode_IKRetargeter.generated.h"
 
diff --git a/Engine/Plugins/Animation/IKRig/Source/IKRig/Public/AnimNodes/AnimNode_IKRig.h b/Engine/Plugins/Animation/IKRig/Source/IKRig/Public/AnimNodes/AnimNode_IKRig.h
--- a/Engine/Plugins/Animation/IKRig/Source/IKRig/Public/AnimNodes/AnimNode_IKRig.h
+++ b/Engine/Plugins/Animation/IKRig/Source/IKRig/Public/AnimNodes/AnimNode_IKRig.h
@@ -5,7 +5,7 @@
 #include "CoreMinimal.h"
 
 #include "IKRigDataTypes.h"
-#include "UObject/ObjectMacros.h"
+#include "IKRigProcessor.h"
 #include "Animation/AnimNodeBase.h"
 #include "AnimNode_IKRig.generated.h"
 
@@ -34,6 +34,9 @@ struct IKRIG_API FAnimNode_IKRig : public FAnimNode_Base
 	UPROPERTY(EditAnywhere, Category = Solver)
 	bool bStartFromRefPose = false;
 
+	/** when true, goals will use the current transforms stored in the IK Rig Definition asset itself */
+	bool bDriveWithSourceAsset = false;
+
 #if WITH_EDITORONLY_DATA
 	/** Toggle drawing of axes to debug joint rotation*/
 	UPROPERTY(EditAnywhere, Category = Solver)
@@ -41,9 +44,10 @@ struct IKRIG_API FAnimNode_IKRig : public FAnimNode_Base
 #endif
 
 private:
-	
+
+	/** IK Rig runtime processor */
 	UPROPERTY(Transient)
-	UIKRigProcessor* IKRigProcessor = nullptr;
+	FIKRigProcessor IKRigProcessor;
 
 	/** a cached list of components on the owning actor that implement the goal creator interface */
 	TArray<IIKGoalCreatorInterface*> GoalCreators;
@@ -58,20 +62,20 @@ public:
 	// FAnimNode_Base interface
 	virtual void GatherDebugData(FNodeDebugData& DebugData) override;
 	virtual void Initialize_AnyThread(const FAnimationInitializeContext& Context) override;
-	virtual void OnInitializeAnimInstance(const FAnimInstanceProxy* InProxy, const UAnimInstance* InAnimInstance) override;
 	virtual void CacheBones_AnyThread(const FAnimationCacheBonesContext& Context)  override;
 	virtual void Evaluate_AnyThread(FPoseContext& Output) override;
 	virtual void Update_AnyThread(const FAnimationUpdateContext& Context) override;
-	virtual bool NeedsOnInitializeAnimInstance() const override { return true; }
 	virtual bool HasPreUpdate() const override { return true; }
 	virtual void PreUpdate(const UAnimInstance* InAnimInstance) override;
 	// End of FAnimNode_Base interface
 
 private:
+	void CopyInputPoseToSolver(FCompactPose& InputPose);
+	void AssignGoalTargets();
+	void CopyOutputPoseToAnimGraph(FCompactPose& OutputPose);
+	
 	bool RebuildGoalList();
 	FName GetGoalName(int32 Index) const;
-
 	void QueueDrawInterface(FAnimInstanceProxy* AnimProxy, const FTransform& ComponentToWorld) const;
-
 	friend class UAnimGraphNode_IKRig;
 };
diff --git a/Engine/Plugins/Animation/IKRig/Source/IKRig/Public/BoneSettings/RotationLimits.h b/Engine/Plugins/Animation/IKRig/Source/IKRig/Public/BoneSettings/RotationLimits.h
deleted file mode 100644
diff --git a/Engine/Plugins/Animation/IKRig/Source/IKRig/Public/IKRigBoneSetting.h b/Engine/Plugins/Animation/IKRig/Source/IKRig/Public/IKRigBoneSetting.h
deleted file mode 100644
diff --git a/Engine/Plugins/Animation/IKRig/Source/IKRig/Public/IKRigDataTypes.h b/Engine/Plugins/Animation/IKRig/Source/IKRig/Public/IKRigDataTypes.h
--- a/Engine/Plugins/Animation/IKRig/Source/IKRig/Public/IKRigDataTypes.h
+++ b/Engine/Plugins/Animation/IKRig/Source/IKRig/Public/IKRigDataTypes.h
@@ -1,37 +1,12 @@
 // Copyright Epic Games, Inc. All Rights Reserved.
 
-
 #pragma once
 
 #include "CoreMinimal.h"
-
-#include "UObject/ObjectMacros.h"
-#include "UObject/Object.h"
-#include "Misc/Guid.h"
+#include "IKRigDefinition.h"
 #include "IKRigDataTypes.generated.h"
 
-struct FIKRigEffectorGoal;
-struct FIKRigHierarchy;
-
-USTRUCT()
-struct IKRIG_API FIKRigEffectorGoal
-{
-	GENERATED_BODY()
-
-	FIKRigEffectorGoal()
-        : Goal("DefaultGoal"),
-          Bone(NAME_None)
-	{
-	}
-
-	UPROPERTY(EditAnywhere, Category = FIKRigEffector)
-	FName Goal;
-
-	UPROPERTY(EditAnywhere, Category = FIKRigEffector)
-	FName Bone;
-
-	bool operator==(const FIKRigEffectorGoal& Other) const { return Goal == Other.Goal; }
-};
+class UIKRigEffectorGoal;
 
 USTRUCT(Blueprintable)
 struct IKRIG_API FIKRigGoal
@@ -96,6 +71,15 @@ struct IKRIG_API FIKRigGoal
 		  FinalBlendedPosition(Position),
 		  FinalBlendedRotation(Rotation){}
 
+	FIKRigGoal(const UIKRigEffectorGoal* InGoal)
+		: Name(InGoal->GoalName),
+		Position(InGoal->CurrentTransform.GetTranslation()),
+		Rotation(InGoal->CurrentTransform.Rotator()),
+		PositionAlpha(InGoal->PositionAlpha),
+        RotationAlpha(InGoal->RotationAlpha),
+        FinalBlendedPosition(ForceInitToZero),
+        FinalBlendedRotation(FQuat::Identity){}
+
 	FString ToString() const
 	{
 		return FString::Printf(TEXT("Name=%s, Pos=(%s, Alpha=%3.3f), Rot=(%s, Alpha=%3.3f)"),
@@ -113,16 +97,20 @@ struct IKRIG_API FIKRigGoalContainer
 public:
 	
 	/** Pre-load all the names of goals (optional, you can just call SetIKGoal to add as needed) */
-	void InitializeGoalsFromNames(const TArray<FIKRigEffectorGoal>& InGoalNames);
+	void InitializeFromGoals(const TArray<UIKRigEffectorGoal*>& InGoals);
 
 	/** Set an IK goal to go to a specific location and rotation (in component space) blended by alpha.
 	 * Will ADD the goal if none exist with the input name. */
 	void SetIKGoal(const FIKRigGoal& InGoal);
 
+	/** Set an IK goal to go to a specific location and rotation (in component space) blended by alpha.
+	* Will ADD the goal if none exist with the input name. */
+	void SetIKGoal(const UIKRigEffectorGoal* InEffectorGoal);
+
 	/** Get an IK goal with the given name. Returns false if no goal is found in the container with the name. */
 	bool GetGoalByName(const FName& InGoalName, FIKRigGoal& OutGoal) const;
 
 	/** Keys are IK Rig Goal names. Values are the Goal data structures.
 	 * These are consumed by an IKRig asset to drive effectors. */
 	TMap<FName, FIKRigGoal> Goals;
-};
\ No newline at end of file
+};
diff --git a/Engine/Plugins/Animation/IKRig/Source/IKRig/Public/IKRigDefinition.h b/Engine/Plugins/Animation/IKRig/Source/IKRig/Public/IKRigDefinition.h
--- a/Engine/Plugins/Animation/IKRig/Source/IKRig/Public/IKRigDefinition.h
+++ b/Engine/Plugins/Animation/IKRig/Source/IKRig/Public/IKRigDefinition.h
@@ -4,15 +4,50 @@
 
 #include "CoreMinimal.h"
 #include "UObject/Object.h"
-#include "IKRigBoneSetting.h"
+#include "Interfaces/Interface_PreviewMeshProvider.h"
 #include "IKRigSkeleton.h"
-#include "IKRigDataTypes.h"
-
 #include "IKRigDefinition.generated.h"
 
-
 class UIKRigSolver;
 
+UCLASS()
+class IKRIG_API UIKRigEffectorGoal : public UObject
+{
+	GENERATED_BODY()
+
+public:
+	
+	UIKRigEffectorGoal() : GoalName("DefaultGoal"), BoneName(NAME_None) {}
+	
+	UIKRigEffectorGoal(const FName& InGoalName, const FName& InBoneName) :  GoalName(InGoalName), BoneName(InBoneName){}
+
+	UPROPERTY(VisibleAnywhere, Category = "Goal Settings")
+	FName GoalName;
+	
+	UPROPERTY(VisibleAnywhere, Category = "Goal Settings")
+	FName BoneName;
+
+	UPROPERTY(EditAnywhere, Category = "Goal Settings", meta = (ClampMin = "0.0", ClampMax = "1.0", UIMin = "0.0", UIMax = "1.0"))
+	float PositionAlpha = 1.0f;
+
+	UPROPERTY(EditAnywhere, Category = "Goal Settings", meta = (ClampMin = "0.0", ClampMax = "1.0", UIMin = "0.0", UIMax = "1.0"))
+	float RotationAlpha = 1.0f;
+
+	UPROPERTY(EditAnywhere, Transient, Category = "Goal Settings")
+	FTransform CurrentTransform;
+	
+	UPROPERTY(VisibleAnywhere, Transient, Category = "Goal Settings")
+	FTransform InitialTransform;
+
+	UPROPERTY(EditAnywhere, Category = "Goal Gizmo", meta = (ClampMin = "0.1", ClampMax = "1000.0", UIMin = "0.1", UIMax = "100.0"))
+	float GizmoSize = 7.0f;
+
+	UPROPERTY(EditAnywhere, Category = "Goal Gizmo",  meta = (ClampMin = "0.0", ClampMax = "10.0", UIMin = "0.0", UIMax = "5.0"))
+	float GizmoThickness = 0.7f;
+
+	bool operator==(const UIKRigEffectorGoal& Other) const { return GoalName == Other.GoalName; }
+};
+
 USTRUCT(Blueprintable)
 struct IKRIG_API FBoneChain
 {
@@ -65,15 +100,16 @@ struct IKRIG_API FRetargetDefinition
 };
 
 UCLASS(Blueprintable)
-class IKRIG_API UIKRigDefinition : public UObject
+class IKRIG_API UIKRigDefinition : public UObject, public IInterface_PreviewMeshProvider
 {
 	GENERATED_BODY()
 
 public:
 
-	/** load a source skeleton to import a hierarchy from */
-	UPROPERTY(VisibleAnywhere, Category = "Skeleton")
-	TSoftObjectPtr<UObject> SourceAsset;
+	/** the skeletal mesh that was used as the source of the skeleton data. also used for preview.
+	 * NOTE: the IK rig may be played back on ANY skeleton that is compatible with it's hierarchy. */
+	UPROPERTY(AssetRegistrySearchable, VisibleAnywhere, Category = "Skeleton")
+	TObjectPtr<class USkeletalMesh> PreviewSkeletalMesh;
 
 	/** hierarchy and bone-pose transforms */
 	UPROPERTY(VisibleAnywhere, Category = "Skeleton")
@@ -81,38 +117,40 @@ public:
 
 	/** stack of solvers, of varying types, executed in serial fashion where
 	output of prior solve is input to the next */
-	UPROPERTY(EditAnywhere, instanced, Category = "IK")
+	UPROPERTY(VisibleAnywhere, instanced, Category = "IK")
 	TArray<UIKRigSolver*> Solvers;
 
-	/** per-bone settings, used by solvers that support them */
-	UPROPERTY(EditAnywhere, instanced, Category = "IK")
-	TArray<UIKRigBoneSetting*> BoneSettings;
+	/** goals, used by solvers that support them */
+	UPROPERTY(VisibleAnywhere, Category = "IK")
+	TArray<UIKRigEffectorGoal*> Goals;
 
 	/** bone chains for animation retargeting */
 	UPROPERTY(EditAnywhere, Category = "Retargeting")
 	FRetargetDefinition RetargetDefinition;
 
-	/** get a list of all goal names used in the solvers */
-	TArray<FIKRigEffectorGoal>& GetEffectorGoals();
-
-	/** get bone associated with a given goal */
-	FName GetBoneNameForGoal(const FName& GoalName);
-
-	/** get name of goal at given index. */
-	FName GetGoalName(int32 GoalIndex);
+	/** editors systems can use this to check if they have most up-to-date settings */
+	int32 GetAssetVersion() const {return AssetVersion;};
+	
+	/** IInterface_PreviewMeshProvider interface */
+	virtual void SetPreviewMesh(USkeletalMesh* PreviewMesh, bool bMarkAsDirty = true) override;
+	virtual USkeletalMesh* GetPreviewMesh() const override;
+	/** END IInterface_PreviewMeshProvider interface */
 
-	// UObject interface
+private:
+	
+	/** UObject interface */
 #if WITH_EDITOR
-	virtual void PostEditChangeProperty(struct FPropertyChangedEvent& PropertyChangedEvent) override;
+	virtual bool Modify(bool bAlwaysMarkDirty = true) override;
 #endif
-	// end UObject interface
+	virtual void PostLoad() override;
+	/** END UObject interface */
 
-private:
+	/** this is incremented whenever the asset is modified in such a way that would require
+	* any IKRigProcessors that are using it to reinitialize with the latest version */
+	int32 AssetVersion = 0;
 
-	UPROPERTY(Transient)
-	TArray<FIKRigEffectorGoal> EffectorGoals;
-	bool bEffectorGoalsDirty = true;
-	void UpdateGoalNameArray();
+	void ResetGoalTransforms() const;
+	FTransform GetGoalInitialTransform(UIKRigEffectorGoal* Goal) const;
 	
 	friend class UIKRigController;
 };
diff --git a/Engine/Plugins/Animation/IKRig/Source/IKRig/Public/IKRigProcessor.h b/Engine/Plugins/Animation/IKRig/Source/IKRig/Public/IKRigProcessor.h
--- a/Engine/Plugins/Animation/IKRig/Source/IKRig/Public/IKRigProcessor.h
+++ b/Engine/Plugins/Animation/IKRig/Source/IKRig/Public/IKRigProcessor.h
@@ -3,9 +3,6 @@
 #pragma once
 
 #include "CoreMinimal.h"
-#include "UObject/ObjectMacros.h"
-#include "UObject/Object.h"
-#include "Drawing/ControlRigDrawInterface.h"
 #include "IKRigDataTypes.h"
 #include "IKRigSkeleton.h"
 
@@ -23,8 +20,8 @@ struct FGoalBone
 	int32 BoneIndex;
 };
 
-UCLASS(BlueprintType)
-class IKRIG_API UIKRigProcessor : public UObject
+USTRUCT()
+struct IKRIG_API FIKRigProcessor
 {
 	GENERATED_BODY()
 
@@ -40,17 +37,16 @@ public:
 	* 5. Copy output transforms with CopyOutputGlobalPoseToArray()
 	* 
 	*/
-
-	static UIKRigProcessor* MakeNewIKRigProcessor(UObject* Outer);
 	 
-	/** setup a new processor to run the given IKRig asset  */
-	void Initialize(UIKRigDefinition* InRigDefinition, const FReferenceSkeleton& RefSkeleton);
+	/** setup a new processor to run the given IKRig asset
+	 *  NOTE!! this function creates new UObjects and consequently MUST be called from the main thread!!	 */
+	void Initialize(UIKRigDefinition* InRigAsset, const FReferenceSkeleton& RefSkeleton, UObject* Outer);
 
 	//
 	// BEGIN UPDATE SEQUENCE FUNCTIONS
 	//
 	// This is the general sequence of function calls to run a typical IK solve:
-	// 
+	//
 	
 	/** Set all bone transforms in global space. This is the pose the IK solve will start from */
 	void SetInputPoseGlobal(const TArray<FTransform>& InGlobalBoneTransforms);
@@ -61,6 +57,9 @@ public:
 	/** Set a named IK goal to go to a specific location and rotation (assumed in component space) blended by separate position/rotation alpha (0-1)*/
 	void SetIKGoal(const FIKRigGoal& Goal);
 
+	/** Set a named IK goal to go to a specific location and rotation (assumed in component space) blended by separate position/rotation alpha (0-1)*/
+	void SetIKGoal(const UIKRigEffectorGoal* Goal);
+
 	/** Run entire stack of solvers */
 	void Solve();
 
@@ -71,44 +70,40 @@ public:
 	// END UPDATE SEQUENCE FUNCTIONS
 	//
 
-	bool IsInitialized() const { return bInitialized; };
+	/** Used to propagate setting values from the source asset at runtime (settings that do not require re-initialization) */
+	void CopyAllInputsFromSourceAssetAtRuntime(UIKRigDefinition* IKRigAsset);
 
-	/** Get interface for debug drawing. */
-	const FControlRigDrawInterface& GetDrawInterface() const { return DrawInterface; }
+	/** checks if the source IKRig asset has been modified in a way that would require reinitialization. */
+	bool NeedsInitialized(UIKRigDefinition* IKRigAsset) const;
 
 	/** Get access to the internal goal data (read only) */
 	const FIKRigGoalContainer& GetGoalContainer() const;
 	
 	/** Get access to the internal skeleton data */
 	FIKRigSkeleton& GetSkeleton();
-
-	/** Get bone for goal */
-	bool GetBoneForGoal(FName GoalName, FGoalBone& OutBone) const;
 	
 private:
 
 	/** Update the final pos/rot of all the goals based on their alpha values. */
 	void BlendGoalsByAlpha();
 
-	/** solving disabled until this flag is true */
-	bool bInitialized = false;
-
 	/** the stack of solvers to run in order */
-	UPROPERTY(transient)
+	UPROPERTY(Transient)
 	TArray<UIKRigSolver*> Solvers;
 
-	/** storage for solver debug drawing data (lines, boxes etc) */
-	UPROPERTY(transient)
-	FControlRigDrawInterface DrawInterface;
-
 	/** the named transforms that solvers use as end effectors */
-	UPROPERTY(transient)
 	FIKRigGoalContainer GoalContainer;
 
 	/** map of goal names to bone names/indices */
 	TMap<FName, FGoalBone> GoalBones;
 
 	/** storage for hierarchy and bone transforms */
-	UPROPERTY(transient)
 	FIKRigSkeleton Skeleton;
+
+	/** solving disabled until this flag is true */
+	bool bInitialized = false;
+	/** which version of the IKRig asset was this instance last initialized with?
+	 * this allows the IKRig asset to undergo modifications at runtime via the editor*/
+	int32 InitializedWithIKRigAssetVersion = -1;
+	int32 LastVersionTried = -1;
 };
\ No newline at end of file
diff --git a/Engine/Plugins/Animation/IKRig/Source/IKRig/Public/IKRigSkeleton.h b/Engine/Plugins/Animation/IKRig/Source/IKRig/Public/IKRigSkeleton.h
--- a/Engine/Plugins/Animation/IKRig/Source/IKRig/Public/IKRigSkeleton.h
+++ b/Engine/Plugins/Animation/IKRig/Source/IKRig/Public/IKRigSkeleton.h
@@ -28,6 +28,8 @@ struct IKRIG_API FIKRigSkeleton
 	TArray<FTransform> RefPoseGlobal;
 
 	void Initialize(const FReferenceSkeleton& RefSkeleton);
+
+	void Reset();
 	
 	int32 GetBoneIndexFromName(const FName InName) const;
 	
@@ -50,6 +52,8 @@ struct IKRIG_API FIKRigSkeleton
 	
 	void PropagateGlobalPoseBelowBone(const int32 BoneIndex);
 
+	bool IsBoneInDirectLineage(const FName& Child, const FName& PotentialParent) const;
+
 	static void NormalizeRotations(TArray<FTransform>& Transforms);
 };
 
diff --git a/Engine/Plugins/Animation/IKRig/Source/IKRig/Public/IKRigSolver.h b/Engine/Plugins/Animation/IKRig/Source/IKRig/Public/IKRigSolver.h
--- a/Engine/Plugins/Animation/IKRig/Source/IKRig/Public/IKRigSolver.h
+++ b/Engine/Plugins/Animation/IKRig/Source/IKRig/Public/IKRigSolver.h
@@ -2,19 +2,16 @@
 
 #pragma once
 
-#include "UObject/ObjectMacros.h"
 #include "UObject/Object.h"
 #include "IKRigDataTypes.h"
-#include "IKRigSolver.generated.h"
 
+#include "IKRigSolver.generated.h"
 
-struct FControlRigDrawInterface;
+class FPrimitiveDrawInterface;
 struct FIKRigGoalContainer;
 struct FIKRigSkeleton;
 
-inline uint32 GetTypeHash(FIKRigEffectorGoal ObjectRef) { return GetTypeHash(ObjectRef.Goal); }
-
-// run time processor 
+// this is the base class for creating your own solver type that integrates into the IK Rig framework/editor.
 UCLASS(abstract, hidecategories = UObject)
 class IKRIG_API UIKRigSolver : public UObject
 {
@@ -22,30 +19,57 @@ class IKRIG_API UIKRigSolver : public UObject
 	
 public:
 
-	UPROPERTY(EditAnywhere, Category = "Settings")
+	UPROPERTY(EditAnywhere, Category = "Solver Settings")
 	bool bEnabled = true;
 
+	//** RUNTIME */
 	/** override to setup internal data based on ref pose */
 	virtual void Initialize(const FIKRigSkeleton& IKRigSkeleton) PURE_VIRTUAL("Init");
-
 	/** override Solve() to evaluate new output pose (InOutGlobalTransform) */
-	virtual void Solve(
-		FIKRigSkeleton& IKRigSkeleton,
-		const FIKRigGoalContainer& Goals,
-		FControlRigDrawInterface* InOutDrawInterface) PURE_VIRTUAL("Solve");
-
-	/** override AddGoalsInSolver() to tell other systems what goals this rig definition uses
-	 * NOTE: only ADD to OutGoals, do not reset or remove */
-	virtual void AddGoalsInSolver(TArray<FIKRigEffectorGoal>& OutGoals) const PURE_VIRTUAL("AddGoalsInSolver");
-
-	#if WITH_EDITOR
-	/** override RenameGoal() when UI renames a goal, you can auto-update any effectors
-	 * that were previously mapped to the old goal name. */
-	virtual void RenameGoal(const FName& OldName, const FName& NewName) {};
-	#endif // WITH_EDITOR
-
-protected:
-	
-	static void AddGoalToArrayNoDuplicates(const FIKRigEffectorGoal& GoalToAdd, TArray<FIKRigEffectorGoal>& OutGoals);
-};
+	virtual void Solve(FIKRigSkeleton& IKRigSkeleton, const FIKRigGoalContainer& Goals) PURE_VIRTUAL("Solve");
+	//** END RUNTIME */
+
+	//** SOLVER SETTINGS */
+	/** override to support RECEIVING modified solver settings from outside systems for editing/UI.
+	 * Note: you can safely cast this to your own solver type and copy any relevant settings at runtime
+	 * This is necessary because at runtime, the IKRigProcessor creates a copy of your solver class
+	 * and the copy must be notified of changes made to the class settings in the source asset.*/
+	virtual void UpdateSolverSettings(UIKRigSolver* InSettings){};
+	//** END SOLVER SETTINGS */
+
+	//** GOALS */
+	/** override to support ADDING a new goal to custom solver */
+	virtual void AddGoal(const UIKRigEffectorGoal& NewGoal) PURE_VIRTUAL("AddGoal");
+	/** override to support REMOVING a goal from custom solver */
+	virtual void RemoveGoal(const FName& GoalName) PURE_VIRTUAL("RemoveGoal");
+	/** override to support RENAMING an existing goal */
+	virtual void RenameGoal(const FName& OldName, const FName& NewName) PURE_VIRTUAL("RenameGoal");
+	/** override to support CHANGING BONE for an existing goal */
+	virtual void SetGoalBone(const FName& GoalName, const FName& NewBoneName) PURE_VIRTUAL("SetGoalBone");
+	/** override to support QUERY for a connected goal */
+	virtual bool IsGoalConnected(const FName& GoalName) const {return false;};
+	/** override to support supplying effector settings to outside systems for editing/UI */
+	virtual UObject* GetEffectorWithGoal(const FName& GoalName) {return nullptr;};
+	//** END GOALS */
 
+	//** ROOT BONE (optional, implement if your solver requires a root bone) */
+	/** override to support SETTING ROOT BONE for the solver */
+	virtual void SetRootBone(const FName& RootBoneName){};
+	virtual bool CanSetRootBone() const { return false; };
+	//** END ROOT BONE */
+
+	//** BONE SETTINGS (optional, implement if your solver supports per-bone settings) */
+	/** override to support ADDING PER-BONE settings for this solver */
+	virtual void AddBoneSetting(const FName& BoneName){};
+	/** override to support ADDING PER-BONE settings for this solver */
+	virtual void RemoveBoneSetting(const FName& BoneName){};
+	/** override to support supplying per-bone settings to outside systems for editing/UI */
+	virtual UObject* GetBoneSetting(const FName& BoneName) const { return nullptr;};
+	/** override to tell systems if this solver supports per-bone settings */
+	virtual bool UsesBoneSettings() const { return false;};
+	/** override to draw custom per-bone settings in the editor viewport */
+	virtual void DrawBoneSettings(const FName& BoneName, const FIKRigSkeleton& IKRigSkeleton, FPrimitiveDrawInterface* PDI) const {};
+	/** return true if the supplied Bone is affected by this solver - this provides UI feedback for user */
+	virtual bool IsBoneAffectedBySolver(const FName& BoneName, const FIKRigSkeleton& IKRigSkeleton) const {return false;};
+	//** END ROOT BONE */
+};
diff --git a/Engine/Plugins/Animation/IKRig/Source/IKRig/Public/Retargeter/IKRetargeter.h b/Engine/Plugins/Animation/IKRig/Source/IKRig/Public/Retargeter/IKRetargeter.h
--- a/Engine/Plugins/Animation/IKRig/Source/IKRig/Public/Retargeter/IKRetargeter.h
+++ b/Engine/Plugins/Animation/IKRig/Source/IKRig/Public/Retargeter/IKRetargeter.h
@@ -298,8 +298,7 @@ public:
 
 private:
 
-	UPROPERTY(Transient)
-	UIKRigProcessor* IKRigProcessor = nullptr;
+	FIKRigProcessor IKRigProcessor;
 	
 	TArray<FRetargetChainPairFK> ChainPairsFK;
 	
diff --git a/Engine/Plugins/Animation/IKRig/Source/IKRig/Public/Solvers/IKRig_PBIKSolver.h b/Engine/Plugins/Animation/IKRig/Source/IKRig/Public/Solvers/IKRig_PBIKSolver.h
--- a/Engine/Plugins/Animation/IKRig/Source/IKRig/Public/Solvers/IKRig_PBIKSolver.h
+++ b/Engine/Plugins/Animation/IKRig/Source/IKRig/Public/Solvers/IKRig_PBIKSolver.h
@@ -2,26 +2,121 @@
 
 #pragma once
 
-
+#include "IKRigDefinition.h"
 #include "IKRigSolver.h"
 #include "Core/PBIKSolver.h"
 #include "PBIK_Shared.h"
 
 #include "IKRig_PBIKSolver.generated.h"
 
-USTRUCT()
-struct FIKRig_PBIKEffector
+UCLASS()
+class IKRIG_API UIKRig_FBIKEffector : public UObject
 {
 	GENERATED_BODY()
 
-	UPROPERTY(EditAnywhere, Category = Settings)
-	FIKRigEffectorGoal Goal;
+public:
+
+	UPROPERTY(VisibleAnywhere, Category = "Full Body IK Effector")
+	FName GoalName;
 
-	UPROPERTY(EditAnywhere, Category = Settings)
+	UPROPERTY(VisibleAnywhere, Category = "Full Body IK Effector")
+	FName BoneName;
+	
+	UPROPERTY(EditAnywhere, Category = "Full Body IK Effector", meta = (ClampMin = "0", ClampMax = "1", UIMin = "0.0", UIMax = "1.0"))
 	float StrengthAlpha = 1.0f;
 
 	UPROPERTY(Transient)
 	int32 IndexInSolver = -1;
+
+	void CopySettings(const UIKRig_FBIKEffector* Other)
+	{
+		StrengthAlpha = Other->StrengthAlpha;
+	}
+};
+
+UCLASS()
+class IKRIG_API UIKRig_PBIKBoneSettings : public UObject
+{
+	GENERATED_BODY()
+
+public:
+	
+	UIKRig_PBIKBoneSettings()
+		: Bone(NAME_None), 
+		X(EPBIKLimitType::Free),
+		Y(EPBIKLimitType::Free),
+		Z(EPBIKLimitType::Free),
+		PreferredAngles(FVector::ZeroVector){}
+
+	UPROPERTY(VisibleAnywhere, Category = Bone, meta = (Constant, CustomWidget = "BoneName"))
+	FName Bone;
+
+	UPROPERTY(EditAnywhere, Category = Stiffness, meta = (ClampMin = "0", ClampMax = "1", UIMin = "0.0", UIMax = "1.0"))
+	float RotationStiffness = 0.0f;
+	UPROPERTY(EditAnywhere, Category = Stiffness, meta = (ClampMin = "0", ClampMax = "1", UIMin = "0.0", UIMax = "1.0"))
+	float PositionStiffness = 0.0f;
+
+	UPROPERTY(EditAnywhere, Category = Limits)
+	EPBIKLimitType X;
+	UPROPERTY(EditAnywhere, Category = Limits, meta = (ClampMin = "-180", ClampMax = "180", UIMin = "-180.0", UIMax = "180.0"))
+	float MinX = 0.0f;
+	UPROPERTY(EditAnywhere, Category = Limits, meta = (ClampMin = "0", ClampMax = "180", UIMin = "0.0", UIMax = "180.0"))
+	float MaxX = 0.0f;
+
+	UPROPERTY(EditAnywhere, Category = Limits)
+	EPBIKLimitType Y;
+	UPROPERTY(EditAnywhere, Category = Limits, meta = (ClampMin = "-180", ClampMax = "180", UIMin = "-180.0", UIMax = "180.0"))
+	float MinY = 0.0f;
+	UPROPERTY(EditAnywhere, Category = Limits, meta = (ClampMin = "0", ClampMax = "180", UIMin = "0.0", UIMax = "180.0"))
+	float MaxY = 0.0f;
+
+	UPROPERTY(EditAnywhere, Category = Limits)
+	EPBIKLimitType Z;
+	UPROPERTY(EditAnywhere, Category = Limits, meta = (ClampMin = "-180", ClampMax = "180", UIMin = "-180.0", UIMax = "180.0"))
+	float MinZ = 0.0f;
+	UPROPERTY(EditAnywhere, Category = Limits, meta = (ClampMin = "0", ClampMax = "180", UIMin = "0.0", UIMax = "180.0"))
+	float MaxZ = 0.0f;
+
+	UPROPERTY(EditAnywhere, Category = PreferredAngles)
+	bool bUsePreferredAngles = false;
+	UPROPERTY(EditAnywhere, Category = PreferredAngles)
+	FVector PreferredAngles;
+
+	void CopyToCoreStruct(PBIK::FBoneSettings& Settings) const
+	{
+		Settings.RotationStiffness = RotationStiffness;
+		Settings.PositionStiffness = PositionStiffness;
+		Settings.X = static_cast<PBIK::ELimitType>(X);
+		Settings.MinX = MinX;
+		Settings.MaxX = MaxX;
+		Settings.Y = static_cast<PBIK::ELimitType>(Y);
+		Settings.MinY = MinY;
+		Settings.MaxY = MaxY;
+		Settings.Z = static_cast<PBIK::ELimitType>(Z);
+		Settings.MinZ = MinZ;
+		Settings.MaxZ = MaxZ;
+		Settings.bUsePreferredAngles = bUsePreferredAngles;
+		Settings.PreferredAngles.Pitch = PreferredAngles.Y;
+		Settings.PreferredAngles.Yaw = PreferredAngles.Z;
+		Settings.PreferredAngles.Roll = PreferredAngles.X;
+	}
+
+	void CopySettings(const UIKRig_PBIKBoneSettings* Other)
+	{	
+		RotationStiffness = Other->RotationStiffness;
+		PositionStiffness = Other->PositionStiffness;
+		X = Other->X;
+		MinX = Other->MinX;
+		MaxX = Other->MaxX;
+		Y = Other->Y;
+		MinY = Other->MinY;
+		MaxY = Other->MaxY;
+		Z = Other->Z;
+		MinZ = Other->MinZ;
+		MaxZ = Other->MaxZ;
+		bUsePreferredAngles = Other->bUsePreferredAngles;
+		PreferredAngles = Other->PreferredAngles;
+	}
 };
 
 UCLASS(EditInlineNew, config = Engine, hidecategories = UObject)
@@ -31,27 +126,71 @@ class IKRIG_API UIKRigPBIKSolver : public UIKRigSolver
 
 public:
 
-	UPROPERTY(EditAnywhere, Category = Settings)
+	UPROPERTY(VisibleAnywhere, Category = "Solver Settings")
 	FName RootBone;
-	
-	UPROPERTY(EditAnywhere, Category = Settings)
-	TArray<FIKRig_PBIKEffector> Effectors;
 
-	UPROPERTY(EditAnywhere, Category = Settings)
-	TArray<FPBIKBoneSetting> BoneSettings;
+	/** High iteration counts can help solve complex joint configurations with competing constraints, but will increase runtime cost. Default is 20. */
+	UPROPERTY(EditAnywhere, Category = SolverSettings, meta = (ClampMin = "0", UIMin = "0.0", UIMax = "200.0"))
+	int32 Iterations = 20;
+
+	/** A global mass multiplier; higher values will make the joints more stiff, but require more iterations. Typical range is 0.0 to 10.0. */
+	UPROPERTY(EditAnywhere, Category = SolverSettings, meta = (ClampMin = "0", UIMin = "0.0", UIMax = "10.0"))
+	float MassMultiplier = 1.0f;
 
-	UPROPERTY(EditAnywhere, Category = Settings)
-	FPBIKSolverSettings Settings;
+	/** Set this as low as possible while keeping the solve stable. Lower values improve convergence of effector targets. Default is 0.2. */
+	UPROPERTY(EditAnywhere, Category = SolverSettings, meta = (ClampMin = "0", UIMin = "0.0", UIMax = "10.0"))
+	float MinMassMultiplier = 0.2f;
 
+	/** If true, joints will translate to reach the effectors; causing bones to lengthen if necessary. Good for cartoon effects. Default is false. */
+	UPROPERTY(EditAnywhere, Category = SolverSettings)
+	bool bAllowStretch = false;
+
+	/** Lock the position and rotation of the solver root bone in-place (at animated position). Useful for partial-body solves. Default is false. */
+	UPROPERTY(EditAnywhere, Category = SolverSettings)
+	bool bPinRoot = false;
+
+	/** When true, the solver is reset each tick to start from the current input pose. If false, incoming animated poses are ignored and the solver starts from the results of the previous solve. Default is true. */
+	UPROPERTY(EditAnywhere, Category = SolverSettings)
+	bool bStartSolveFromInputPose = true;
+	
+	UPROPERTY()
+	TArray<UIKRig_FBIKEffector*> Effectors;
+
+	UPROPERTY()
+	TArray<UIKRig_PBIKBoneSettings*> BoneSettings;
+
+	/** UIKRigSolver interface */
+	// runtime
 	virtual void Initialize(const FIKRigSkeleton& IKRigSkeleton) override;
-	virtual void Solve(
-		FIKRigSkeleton& IKRigSkeleton,
-		const FIKRigGoalContainer& Goals,
-		FControlRigDrawInterface* InOutDrawInterface) override;
-	virtual void AddGoalsInSolver(TArray<FIKRigEffectorGoal>& OutGoals) const override;
+	virtual void Solve(FIKRigSkeleton& IKRigSkeleton, const FIKRigGoalContainer& Goals) override;
+	virtual void UpdateSolverSettings(UIKRigSolver* InSettings) override;
+	// goals
+	virtual void AddGoal(const UIKRigEffectorGoal& NewGoal) override;
+	virtual void RemoveGoal(const FName& GoalName) override;
+	virtual void RenameGoal(const FName& OldName, const FName& NewName) override;
+	virtual void SetGoalBone(const FName& GoalName, const FName& NewBoneName) override;
+	virtual bool IsGoalConnected(const FName& GoalName) const override;
+	virtual void SetRootBone(const FName& RootBoneName) override;
+	virtual UObject* GetEffectorWithGoal(const FName& GoalName) override;
+	// bone settings
+	virtual void AddBoneSetting(const FName& BoneName) override;
+	virtual void RemoveBoneSetting(const FName& BoneName) override;
+	virtual UObject* GetBoneSetting(const FName& BoneName) const override;
+	virtual bool UsesBoneSettings() const override{ return true;};
+	virtual void DrawBoneSettings(const FName& BoneName, const FIKRigSkeleton& IKRigSkeleton, FPrimitiveDrawInterface* PDI) const override;
+	// root bone can be set on this solver
+	virtual bool CanSetRootBone() const override { return true; };
+	virtual bool IsBoneAffectedBySolver(const FName& BoneName, const FIKRigSkeleton& IKRigSkeleton) const override;
+	/** END UIKRigSolver interface */
 
 private:
 
 	FPBIKSolver Solver;
+
+	int32 GetIndexOfGoal(const FName& GoalName) const;
+
+	//** UObject */
+	virtual void PostLoad() override;
 };
 
+
diff --git a/Engine/Plugins/Animation/IKRig/Source/IKRig/Public/Solvers/IKRig_SetTransform.h b/Engine/Plugins/Animation/IKRig/Source/IKRig/Public/Solvers/IKRig_SetTransform.h
--- a/Engine/Plugins/Animation/IKRig/Source/IKRig/Public/Solvers/IKRig_SetTransform.h
+++ b/Engine/Plugins/Animation/IKRig/Source/IKRig/Public/Solvers/IKRig_SetTransform.h
@@ -2,10 +2,30 @@
 
 #pragma once
 
+#include "IKRigDefinition.h"
 #include "IKRigSolver.h"
 
 #include "IKRig_SetTransform.generated.h"
 
+UCLASS()
+class IKRIG_API UIKRig_SetTransformEffector : public UObject
+{
+	GENERATED_BODY()
+
+public:
+
+	/** If true, Goal will drive the translation of the target bone. Default is true. */
+	UPROPERTY(EditAnywhere, Category = "Set Transform Effector")
+	bool bEnablePosition = true;
+
+	/** If true, Goal will drive the rotation of the target bone. Default is true. */
+	UPROPERTY(EditAnywhere, Category = "Set Transform Effector")
+	bool bEnableRotation = true;
+
+	/** Blend the effector on/off. Range is 0-1. Default is 1.0. */
+	UPROPERTY(EditAnywhere, Category = "Set Transform Effector", meta = (ClampMin = "0", ClampMa = "1", UIMin = "0.0", UIMax = "1.0"))
+	float Alpha = 1.0f;
+};
 
 UCLASS(EditInlineNew)
 class IKRIG_API UIKRig_SetTransform : public UIKRigSolver
@@ -14,26 +34,32 @@ class IKRIG_API UIKRig_SetTransform : public UIKRigSolver
 
 public:
 	
-	UPROPERTY(EditAnywhere, Category = "Solver")
-	FIKRigEffectorGoal Effector;
-	
-	UPROPERTY(EditAnywhere, Category = "Solver")
-	bool bEnablePosition = true;
+	UIKRig_SetTransform();
 
-	UPROPERTY(EditAnywhere, Category = "Solver")
-	bool bEnableRotation = true;
+	UPROPERTY(VisibleAnywhere, Category = "Set Transform Settings")
+	FName Goal;
 
-protected:
+	UPROPERTY(VisibleAnywhere, Category = "Set Transform Settings")
+	FName Bone;
 	
+	UPROPERTY(Transient)
+	UIKRig_SetTransformEffector* Effector;
+
+	/** UIKRigSolver interface */
 	virtual void Initialize(const FIKRigSkeleton& IKRigSkeleton) override;
-	virtual void Solve(
-		FIKRigSkeleton& IKRigSkeleton, 
-		const FIKRigGoalContainer& Goals,
-		FControlRigDrawInterface* InOutDrawInterface) override;
-	virtual void AddGoalsInSolver(TArray<FIKRigEffectorGoal>& OutGoals) const override;
+	virtual void Solve(FIKRigSkeleton& IKRigSkeleton, const FIKRigGoalContainer& Goals) override;
+	virtual void UpdateSolverSettings(UIKRigSolver* InSettings) override;
+	virtual void AddGoal(const UIKRigEffectorGoal& NewGoal) override;
+	virtual void RemoveGoal(const FName& GoalName) override;
+	virtual void RenameGoal(const FName& OldName, const FName& NewName) override;
+	virtual void SetGoalBone(const FName& GoalName, const FName& NewBoneName) override;
+	virtual bool IsGoalConnected(const FName& GoalName) const override;
+	virtual UObject* GetEffectorWithGoal(const FName& GoalName) override;
+	virtual bool IsBoneAffectedBySolver(const FName& BoneName, const FIKRigSkeleton& IKRigSkeleton) const override;
+	/** END UIKRigSolver interface */
 
 private:
-	
+
 	int32 BoneIndex;
 };
 
diff --git a/Engine/Plugins/Animation/IKRig/Source/IKRigDeveloper/Private/AnimGraphNode_IKRig.cpp b/Engine/Plugins/Animation/IKRig/Source/IKRigDeveloper/Private/AnimGraphNode_IKRig.cpp
--- a/Engine/Plugins/Animation/IKRig/Source/IKRigDeveloper/Private/AnimGraphNode_IKRig.cpp
+++ b/Engine/Plugins/Animation/IKRig/Source/IKRigDeveloper/Private/AnimGraphNode_IKRig.cpp
@@ -9,7 +9,6 @@
 // UAnimGraphNode_IKRig 
 
 #define LOCTEXT_NAMESPACE "AnimGraphNode_IKRig"
-const FName UAnimGraphNode_IKRig::AnimModeName(TEXT("IKRig.IKRigEditor.IKRigEditMode"));
 
 void UAnimGraphNode_IKRig::Draw(FPrimitiveDrawInterface* PDI, USkeletalMeshComponent* PreviewSkelMeshComp) const
 {
@@ -32,11 +31,6 @@ void UAnimGraphNode_IKRig::CopyNodeDataToPreviewNode(FAnimNode_Base* InPreviewNo
 	FAnimNode_IKRig* IKRigNode = static_cast<FAnimNode_IKRig*>(InPreviewNode);
 }
 
-FEditorModeID UAnimGraphNode_IKRig::GetEditorMode() const
-{
-	return AnimModeName;
-}
-
 void UAnimGraphNode_IKRig::PostEditChangeProperty(struct FPropertyChangedEvent& PropertyChangedEvent)
 {
 	const FName PropertyName = (PropertyChangedEvent.Property ? PropertyChangedEvent.Property->GetFName() : NAME_None);
diff --git a/Engine/Plugins/Animation/IKRig/Source/IKRigDeveloper/Public/AnimGraphNode_IKRig.h b/Engine/Plugins/Animation/IKRig/Source/IKRigDeveloper/Public/AnimGraphNode_IKRig.h
--- a/Engine/Plugins/Animation/IKRig/Source/IKRigDeveloper/Public/AnimGraphNode_IKRig.h
+++ b/Engine/Plugins/Animation/IKRig/Source/IKRigDeveloper/Public/AnimGraphNode_IKRig.h
@@ -30,10 +30,7 @@ public:
 	// UAnimGraphNode_Base interface
 	virtual void PreloadRequiredAssets() override;
 	virtual void CopyNodeDataToPreviewNode(FAnimNode_Base* AnimNode) override;
-	virtual FEditorModeID GetEditorMode() const override;
 	virtual void Draw(FPrimitiveDrawInterface* PDI, USkeletalMeshComponent* PreviewSkelMeshComp) const override;
 	virtual void CustomizePinData(UEdGraphPin* Pin, FName SourcePropertyName, int32 ArrayIndex) const override;
 	// End of UAnimGraphNode_Base interface
-
-	static const FName AnimModeName;
 };
diff --git a/Engine/Plugins/Animation/IKRig/Source/IKRigEditor/IKRigEditor.Build.cs b/Engine/Plugins/Animation/IKRig/Source/IKRigEditor/IKRigEditor.Build.cs
--- a/Engine/Plugins/Animation/IKRig/Source/IKRigEditor/IKRigEditor.Build.cs
+++ b/Engine/Plugins/Animation/IKRig/Source/IKRigEditor/IKRigEditor.Build.cs
@@ -11,28 +11,37 @@ namespace UnrealBuildTool.Rules
             PublicDependencyModuleNames.AddRange(
                 new string[]
                 {
+	                "Core",
+	                "AdvancedPreviewScene",
                 }
             );
 
 			PrivateDependencyModuleNames.AddRange(
                 new string[]
                 {
-                    "Core",
+                    "InputCore",
+                    "EditorFramework",
+                    "UnrealEd",
+                    "ToolMenus",
                     "CoreUObject",
                     "Engine",
-					"UnrealEd",
-					"IKRig",
-					"IKRigDeveloper",
-					"AssetTools",
-					"SlateCore",
-					"Slate",
-					"InputCore",
+                    "Slate",
+                    "SlateCore",
+                    "AssetTools",
+                    "EditorWidgets",
+                    "Kismet",
+                    "KismetWidgets",
+                    "EditorStyle",
+                    "Persona",
+                    "SkeletonEditor",
+                    
 					"PropertyEditor",
-					"EditorStyle",
-					"EditorFramework",
 					"BlueprintGraph",
 					"AnimGraph",
-					"EditorWidgets"
+					"AnimGraphRuntime",
+					
+					"IKRig",
+					"IKRigDeveloper",					
 				}
             );
 
diff --git a/Engine/Plugins/Animation/IKRig/Source/IKRigEditor/Private/AssetTypeActions_IKRetargeter.cpp b/Engine/Plugins/Animation/IKRig/Source/IKRigEditor/Private/AssetTypeActions_IKRetargeter.cpp
deleted file mode 100644
diff --git a/Engine/Plugins/Animation/IKRig/Source/IKRigEditor/Private/AssetTypeActions_IKRigDefinition.cpp b/Engine/Plugins/Animation/IKRig/Source/IKRigEditor/Private/AssetTypeActions_IKRigDefinition.cpp
deleted file mode 100644
diff --git a/Engine/Plugins/Animation/IKRig/Source/IKRigEditor/Private/IKRetargeterFactory.cpp b/Engine/Plugins/Animation/IKRig/Source/IKRigEditor/Private/IKRetargeterFactory.cpp
deleted file mode 100644
diff --git a/Engine/Plugins/Animation/IKRig/Source/IKRigEditor/Private/IKRigController.cpp b/Engine/Plugins/Animation/IKRig/Source/IKRigEditor/Private/IKRigController.cpp
deleted file mode 100644
diff --git a/Engine/Plugins/Animation/IKRig/Source/IKRigEditor/Private/IKRigDefinitionDetails.cpp b/Engine/Plugins/Animation/IKRig/Source/IKRigEditor/Private/IKRigDefinitionDetails.cpp
deleted file mode 100644
diff --git a/Engine/Plugins/Animation/IKRig/Source/IKRigEditor/Private/IKRigDefinitionFactory.cpp b/Engine/Plugins/Animation/IKRig/Source/IKRigEditor/Private/IKRigDefinitionFactory.cpp
deleted file mode 100644
diff --git a/Engine/Plugins/Animation/IKRig/Source/IKRigEditor/Private/IKRigEditMode.cpp b/Engine/Plugins/Animation/IKRig/Source/IKRigEditor/Private/IKRigEditMode.cpp
deleted file mode 100644
diff --git a/Engine/Plugins/Animation/IKRig/Source/IKRigEditor/Private/IKRigEditor.cpp b/Engine/Plugins/Animation/IKRig/Source/IKRigEditor/Private/IKRigEditor.cpp
--- a/Engine/Plugins/Animation/IKRig/Source/IKRigEditor/Private/IKRigEditor.cpp
+++ b/Engine/Plugins/Animation/IKRig/Source/IKRigEditor/Private/IKRigEditor.cpp
@@ -1,21 +1,18 @@
 // Copyright Epic Games, Inc. All Rights Reserved.
 
 #include "IKRigEditor.h"
-#include "Features/IModularFeatures.h"
-#include "Engine/Engine.h"
-#include "Misc/Paths.h"
 #include "Modules/ModuleManager.h"
 #include "Logging/LogMacros.h"
 #include "AssetToolsModule.h"
 #include "PropertyEditorModule.h"
-#include "IKRigDefinitionDetails.h"
 #include "PropertyEditorDelegates.h"
 #include "EditorModeRegistry.h"
-#include "AnimGraphNode_IKRig.h"
-#include "AssetTypeActions_IKRigDefinition.h"
-#include "AssetTypeActions_IKRetargeter.h"
 #include "IKRigDefinition.h"
-#include "IKRigEditMode.h"
+#include "RigEditor/AssetTypeActions_IKRigDefinition.h"
+#include "RetargetEditor/AssetTypeActions_IKRetargeter.h"
+#include "RigEditor/IKRigCommands.h"
+#include "RigEditor/IKRigEditMode.h"
+#include "RigEditor/IKRigSkeletonCommands.h"
 
 IMPLEMENT_MODULE(FIKRigEditor, IKRigEditor)
 
@@ -25,6 +22,9 @@ DEFINE_LOG_CATEGORY_STATIC(LogIKRigEditor, Log, All);
 
 void FIKRigEditor::StartupModule()
 {
+	FIKRigCommands::Register();
+	FIKRigSkeletonCommands::Register();
+	
 	// register IKRigDefinition asset type
 	IKRigDefinitionAssetAction = MakeShareable(new FAssetTypeActions_IKRigDefinition);
 	FModuleManager::LoadModuleChecked<FAssetToolsModule>("AssetTools").Get().RegisterAssetTypeActions(IKRigDefinitionAssetAction.ToSharedRef());
@@ -33,21 +33,21 @@ void FIKRigEditor::StartupModule()
 	IKRetargeterAssetAction = MakeShareable(new FAssetTypeActions_IKRetargeter);
 	FModuleManager::LoadModuleChecked<FAssetToolsModule>("AssetTools").Get().RegisterAssetTypeActions(IKRetargeterAssetAction.ToSharedRef());
 
-	// register details panel customization
-	FPropertyEditorModule& PropertyModule = FModuleManager::GetModuleChecked<FPropertyEditorModule>("PropertyEditor");
-	PropertyModule.RegisterCustomClassLayout(UIKRigDefinition::StaticClass()->GetFName(), FOnGetDetailCustomizationInstance::CreateStatic(&FIKRigDefinitionDetails::MakeInstance));
-
-	// register custom editor mode (TBD)
-	//FEditorModeRegistry::Get().RegisterMode<FIKRigEditMode>(UAnimGraphNode_IKRig::AnimModeName, LOCTEXT("IKRigEditMode", "IKRig"), FSlateIcon(), false);
+	// register custom editor mode
+	FEditorModeRegistry::Get().RegisterMode<FIKRigEditMode>(FIKRigEditMode::ModeName, LOCTEXT("IKRigEditMode", "IKRig"), FSlateIcon(), false);
 }
 
 void FIKRigEditor::ShutdownModule()
 {
-	FEditorModeRegistry::Get().UnregisterMode(UAnimGraphNode_IKRig::AnimModeName);
+	FIKRigCommands::Unregister();
+	FIKRigSkeletonCommands::Unregister();
+	
+	FEditorModeRegistry::Get().UnregisterMode(FIKRigEditMode::ModeName);
 
 	FPropertyEditorModule& PropertyModule = FModuleManager::GetModuleChecked<FPropertyEditorModule>("PropertyEditor");
 	PropertyModule.UnregisterCustomPropertyTypeLayout("IKRigEffector");
 
+	// unregister IKRigDefinition asset action
 	if (IKRigDefinitionAssetAction.IsValid())
 	{
 		if (FModuleManager::Get().IsModuleLoaded("AssetTools"))
@@ -57,6 +57,16 @@ void FIKRigEditor::ShutdownModule()
 		IKRigDefinitionAssetAction.Reset();
 	}
 
+	// unregister IKRetargeter asset action
+	if (IKRetargeterAssetAction.IsValid())
+	{
+		if (FModuleManager::Get().IsModuleLoaded("AssetTools"))
+		{
+			FModuleManager::GetModuleChecked<FAssetToolsModule>("AssetTools").Get().UnregisterAssetTypeActions(IKRetargeterAssetAction.ToSharedRef());
+		}
+		IKRetargeterAssetAction.Reset();
+	}
+
 }
 
 #undef LOCTEXT_NAMESPACE
\ No newline at end of file
diff --git a/Engine/Plugins/Animation/IKRig/Source/IKRigEditor/Private/RetargetEditor/AssetTypeActions_IKRetargeter.cpp b/Engine/Plugins/Animation/IKRig/Source/IKRigEditor/Private/RetargetEditor/AssetTypeActions_IKRetargeter.cpp
new file mode 100644
--- /dev/null
+++ b/Engine/Plugins/Animation/IKRig/Source/IKRigEditor/Private/RetargetEditor/AssetTypeActions_IKRetargeter.cpp
@@ -0,0 +1,18 @@
+// Copyright Epic Games, Inc. All Rights Reserved.
+
+#include "RetargetEditor/AssetTypeActions_IKRetargeter.h"
+#include "Retargeter/IKRetargeter.h"
+
+#define LOCTEXT_NAMESPACE "AssetTypeActions"
+
+UClass* FAssetTypeActions_IKRetargeter::GetSupportedClass() const
+{
+	return UIKRetargeter::StaticClass();
+}
+
+void FAssetTypeActions_IKRetargeter::GetActions(const TArray<UObject*>& InObjects, FToolMenuSection& Section)
+{
+	FAssetTypeActions_Base::GetActions(InObjects, Section);
+}
+
+#undef LOCTEXT_NAMESPACE
diff --git a/Engine/Plugins/Animation/IKRig/Source/IKRigEditor/Private/RetargetEditor/IKRetargeterFactory.cpp b/Engine/Plugins/Animation/IKRig/Source/IKRigEditor/Private/RetargetEditor/IKRetargeterFactory.cpp
new file mode 100644
--- /dev/null
+++ b/Engine/Plugins/Animation/IKRig/Source/IKRigEditor/Private/RetargetEditor/IKRetargeterFactory.cpp
@@ -0,0 +1,51 @@
+// Copyright Epic Games, Inc. All Rights Reserved.
+
+#include "RetargetEditor/IKRetargeterFactory.h"
+#include "IKRig/Public/Retargeter/IKRetargeter.h"
+#include "AssetTypeCategories.h"
+
+#define LOCTEXT_NAMESPACE "IKRetargeterFactory"
+
+
+UIKRetargeterFactory::UIKRetargeterFactory()
+{
+	bCreateNew = true;
+	bEditAfterNew = true;
+	SupportedClass = UIKRetargeter::StaticClass();
+}
+
+UObject* UIKRetargeterFactory::FactoryCreateNew(UClass* Class, UObject* InParent, FName Name, EObjectFlags Flags, UObject* Context, FFeedbackContext* Warn)
+{
+	return NewObject<UIKRetargeter>(InParent, Name, Flags | RF_Transactional);;
+}
+
+bool UIKRetargeterFactory::ShouldShowInNewMenu() const
+{
+	return true;
+}
+
+bool UIKRetargeterFactory::ConfigureProperties()
+{
+	return true;
+}
+
+FText UIKRetargeterFactory::GetDisplayName() const
+{
+	return LOCTEXT("IKRetargeter_DisplayName", "IK Retargeter");
+}
+
+uint32 UIKRetargeterFactory::GetMenuCategories() const
+{
+	return EAssetTypeCategories::Animation;
+}
+
+FText UIKRetargeterFactory::GetToolTip() const
+{
+	return LOCTEXT("IKRetargeter_Tooltip", "Defines a pair of Source/Target Retarget Rigs and the mapping between them.");
+}
+
+FString UIKRetargeterFactory::GetDefaultNewAssetName() const
+{
+	return FString(TEXT("NewIKRetargeter"));
+}
+#undef LOCTEXT_NAMESPACE
diff --git a/Engine/Plugins/Animation/IKRig/Source/IKRigEditor/Private/RigEditor/AssetTypeActions_IKRigDefinition.cpp b/Engine/Plugins/Animation/IKRig/Source/IKRigEditor/Private/RigEditor/AssetTypeActions_IKRigDefinition.cpp
new file mode 100644
--- /dev/null
+++ b/Engine/Plugins/Animation/IKRig/Source/IKRigEditor/Private/RigEditor/AssetTypeActions_IKRigDefinition.cpp
@@ -0,0 +1,44 @@
+// Copyright Epic Games, Inc. All Rights Reserved.
+
+#include "RigEditor/AssetTypeActions_IKRigDefinition.h"
+#include "IKRigDefinition.h"
+#include "RigEditor/IKRigToolkit.h"
+#include "ThumbnailRendering/SceneThumbnailInfo.h"
+
+#define LOCTEXT_NAMESPACE "AssetTypeActions"
+
+
+UClass* FAssetTypeActions_IKRigDefinition::GetSupportedClass() const
+{
+	return UIKRigDefinition::StaticClass();
+}
+
+UThumbnailInfo* FAssetTypeActions_IKRigDefinition::GetThumbnailInfo(UObject* Asset) const
+{
+	UIKRigDefinition* IKRig = CastChecked<UIKRigDefinition>(Asset);
+	return NewObject<USceneThumbnailInfo>(IKRig, NAME_None, RF_Transactional);
+}
+
+void FAssetTypeActions_IKRigDefinition::GetActions(const TArray<UObject*>& InObjects, FToolMenuSection& Section)
+{
+	FAssetTypeActions_Base::GetActions(InObjects, Section);
+}
+
+
+void FAssetTypeActions_IKRigDefinition::OpenAssetEditor(
+	const TArray<UObject*>& InObjects,
+	TSharedPtr<IToolkitHost> EditWithinLevelEditor)
+{
+	const EToolkitMode::Type Mode = EditWithinLevelEditor.IsValid() ? EToolkitMode::WorldCentric : EToolkitMode::Standalone;
+    
+    for (auto ObjIt = InObjects.CreateConstIterator(); ObjIt; ++ObjIt)
+    {
+    	if (UIKRigDefinition* Asset = Cast<UIKRigDefinition>(*ObjIt))
+    	{
+    		TSharedRef<FIKRigEditorToolkit> NewEditor(new FIKRigEditorToolkit());
+    		NewEditor->InitAssetEditor(Mode, EditWithinLevelEditor, Asset);
+    	}
+    }
+}
+
+#undef LOCTEXT_NAMESPACE
diff --git a/Engine/Plugins/Animation/IKRig/Source/IKRigEditor/Private/RigEditor/IKRigAnimInstance.cpp b/Engine/Plugins/Animation/IKRig/Source/IKRigEditor/Private/RigEditor/IKRigAnimInstance.cpp
new file mode 100644
--- /dev/null
+++ b/Engine/Plugins/Animation/IKRig/Source/IKRigEditor/Private/RigEditor/IKRigAnimInstance.cpp
@@ -0,0 +1,21 @@
+// Copyright Epic Games, Inc. All Rights Reserved.
+
+#include "RigEditor/IKRigAnimInstance.h"
+#include "RigEditor/IKRigAnimInstanceProxy.h"
+
+UIKRigAnimInstance::UIKRigAnimInstance(const FObjectInitializer& ObjectInitializer)
+	: Super(ObjectInitializer)
+{
+	bUseMultiThreadedAnimationUpdate = false;
+}
+
+void UIKRigAnimInstance::SetIKRigAsset(UIKRigDefinition* InIKRigAsset)
+{
+	FIKRigAnimInstanceProxy& Proxy = GetProxyOnGameThread<FIKRigAnimInstanceProxy>();
+	Proxy.SetIKRigAsset(InIKRigAsset);
+}
+
+FAnimInstanceProxy* UIKRigAnimInstance::CreateAnimInstanceProxy()
+{
+	return new FIKRigAnimInstanceProxy(this);
+}
diff --git a/Engine/Plugins/Animation/IKRig/Source/IKRigEditor/Private/RigEditor/IKRigAnimInstanceProxy.cpp b/Engine/Plugins/Animation/IKRig/Source/IKRigEditor/Private/RigEditor/IKRigAnimInstanceProxy.cpp
new file mode 100644
--- /dev/null
+++ b/Engine/Plugins/Animation/IKRig/Source/IKRigEditor/Private/RigEditor/IKRigAnimInstanceProxy.cpp
@@ -0,0 +1,59 @@
+// Copyright Epic Games, Inc. All Rights Reserved.
+
+#include "RigEditor/IKRigAnimInstanceProxy.h"
+#include "RigEditor/IKRigAnimInstance.h"
+
+FIKRigAnimInstanceProxy::FIKRigAnimInstanceProxy()
+{
+}
+
+FIKRigAnimInstanceProxy::FIKRigAnimInstanceProxy(UAnimInstance* InAnimInstance)
+	: FAnimPreviewInstanceProxy(InAnimInstance)
+{
+}
+
+FIKRigAnimInstanceProxy::~FIKRigAnimInstanceProxy()
+{
+}
+
+void FIKRigAnimInstanceProxy::Initialize(UAnimInstance* InAnimInstance)
+{
+	FAnimInstanceProxy::Initialize(InAnimInstance);
+	IKRigNode.Source.SetLinkNode(&SingleNode);
+
+	// force this instance of the IK Rig evaluation to copy setting from the source IK Rig asset
+	IKRigNode.bDriveWithSourceAsset = true; 
+}
+
+bool FIKRigAnimInstanceProxy::Evaluate(FPoseContext& Output)
+{
+	IKRigNode.Evaluate_AnyThread(Output);
+	return true;
+}
+
+FAnimNode_Base* FIKRigAnimInstanceProxy::GetCustomRootNode()
+{
+	return &IKRigNode;
+}
+
+void FIKRigAnimInstanceProxy::GetCustomNodes(TArray<FAnimNode_Base*>& OutNodes)
+{
+	OutNodes.Add(&IKRigNode);
+}
+
+void FIKRigAnimInstanceProxy::UpdateAnimationNode(const FAnimationUpdateContext& InContext)
+{
+	if (CurrentAsset != nullptr)
+	{
+		FAnimPreviewInstanceProxy::UpdateAnimationNode(InContext);
+	}
+	else
+	{
+		IKRigNode.Update_AnyThread(InContext);
+	}
+}
+
+void FIKRigAnimInstanceProxy::SetIKRigAsset(UIKRigDefinition* InIKRigAsset)
+{
+	IKRigNode.RigDefinitionAsset = InIKRigAsset;
+}
diff --git a/Engine/Plugins/Animation/IKRig/Source/IKRigEditor/Private/RigEditor/IKRigCommands.cpp b/Engine/Plugins/Animation/IKRig/Source/IKRigEditor/Private/RigEditor/IKRigCommands.cpp
new file mode 100644
--- /dev/null
+++ b/Engine/Plugins/Animation/IKRig/Source/IKRigEditor/Private/RigEditor/IKRigCommands.cpp
@@ -0,0 +1,12 @@
+// Copyright Epic Games, Inc. All Rights Reserved.
+
+#include "RigEditor/IKRigCommands.h"
+
+#define LOCTEXT_NAMESPACE "IKRigCommands"
+
+void FIKRigCommands::RegisterCommands()
+{
+	UI_COMMAND(Reset, "Reset", "Reset state of the rig and goals to initial pose.", EUserInterfaceActionType::Button, FInputChord());
+}
+
+#undef LOCTEXT_NAMESPACE
diff --git a/Engine/Plugins/Animation/IKRig/Source/IKRigEditor/Private/RigEditor/IKRigController.cpp b/Engine/Plugins/Animation/IKRig/Source/IKRigEditor/Private/RigEditor/IKRigController.cpp
new file mode 100644
--- /dev/null
+++ b/Engine/Plugins/Animation/IKRig/Source/IKRigEditor/Private/RigEditor/IKRigController.cpp
@@ -0,0 +1,606 @@
+// Copyright Epic Games, Inc. All Rights Reserved.
+
+#include "RigEditor/IKRigController.h"
+
+#include "IKRigDefinition.h"
+#include "IKRigSolver.h"
+
+#include "Engine/SkeletalMesh.h"
+#include "ScopedTransaction.h"
+
+#define LOCTEXT_NAMESPACE	"IKRigController"
+
+
+TMap<UIKRigDefinition*, UIKRigController*> UIKRigController::AssetToControllerMap;
+
+UIKRigController* UIKRigController::GetIKRigController(UIKRigDefinition* InIKRigDefinition)
+{
+	if (!InIKRigDefinition)
+	{
+		return nullptr;
+	}
+	
+	UIKRigController** Controller = AssetToControllerMap.Find(InIKRigDefinition);
+	if (Controller)
+	{
+		return *Controller;
+	}
+
+	UIKRigController* NewController = NewObject<UIKRigController>();
+	AssetToControllerMap.Add(InIKRigDefinition) = NewController;
+	NewController->IKRigAsset = InIKRigDefinition;
+	return NewController;
+}
+
+UIKRigDefinition* UIKRigController::GetAsset() const
+{
+	return IKRigAsset;
+}
+
+void UIKRigController::AddBoneSetting(const FName& BoneName, int32 SolverIndex) const
+{
+	UIKRigSolver* Solver = GetSolver(SolverIndex);
+	if (!Solver)
+	{
+		return; // solver doesn't exist
+	}
+
+	if (IKRigAsset->Skeleton.GetBoneIndexFromName(BoneName) == INDEX_NONE)
+	{
+		return; // bone doesn't exist
+	}
+
+	FScopedTransaction Transaction(LOCTEXT("AddBoneSetting_Label", "Add Bone Setting"));
+	IKRigAsset->Modify();
+
+	Solver->AddBoneSetting(BoneName);
+}
+
+bool UIKRigController::CanAddBoneSetting(const FName& BoneName, int32 SolverIndex) const
+{
+	UIKRigSolver* Solver = GetSolver(SolverIndex);
+	if (!Solver)
+	{
+		return false; // solver doesn't exist
+	}
+
+	if (IKRigAsset->Skeleton.GetBoneIndexFromName(BoneName) == INDEX_NONE)
+	{
+		return false; // bone doesn't exist
+	}
+
+	if (!IKRigAsset->Solvers[SolverIndex]->UsesBoneSettings())
+	{
+		return false; // solver doesn't support per-bone settings
+	}
+
+	// returns true if the solver in question does NOT already have a settings object for this bone 
+	return !static_cast<bool>(IKRigAsset->Solvers[SolverIndex]->GetBoneSetting(BoneName));
+}
+
+void UIKRigController::RemoveBoneSetting(const FName& BoneName, int32 SolverIndex) const
+{
+	UIKRigSolver* Solver = GetSolver(SolverIndex);
+	if (!Solver)
+	{
+		return; // solver doesn't exist
+	}
+
+	if (IKRigAsset->Skeleton.GetBoneIndexFromName(BoneName) == INDEX_NONE)
+	{
+		return; // bone doesn't exist
+	}
+
+	FScopedTransaction Transaction(LOCTEXT("AddBoneSetting_Label", "Add Bone Setting"));
+	IKRigAsset->Modify();
+
+	Solver->RemoveBoneSetting(BoneName);
+}
+
+bool UIKRigController::CanRemoveBoneSetting(const FName& BoneName, int32 SolverIndex) const
+{
+	UIKRigSolver* Solver = GetSolver(SolverIndex);
+	if (!Solver)
+	{
+		return false; // solver doesn't exist
+	}
+
+	if (!Solver->UsesBoneSettings())
+	{
+		return false; // solver doesn't use bone settings
+	}
+
+	if (IKRigAsset->Skeleton.GetBoneIndexFromName(BoneName) == INDEX_NONE)
+	{
+		return false; // bone doesn't exist
+	}
+
+	if (!Solver->GetBoneSetting(BoneName))
+	{
+		return false; // solver doesn't have any settings for this bone
+	}
+
+	return true;
+}
+
+UObject* UIKRigController::GetSettingsForBone(const FName& BoneName, int32 SolverIndex) const
+{
+	UIKRigSolver* Solver = GetSolver(SolverIndex);
+	if (!Solver)
+	{
+		return nullptr; // solver doesn't exist
+	}
+	
+	return Solver->GetBoneSetting(BoneName);
+}
+
+bool UIKRigController::DoesBoneHaveSettings(const FName& BoneName) const
+{
+	if (IKRigAsset->Skeleton.GetBoneIndexFromName(BoneName) == INDEX_NONE)
+	{
+		return false; // bone doesn't exist
+	}
+	
+	for (UIKRigSolver* Solver : IKRigAsset->Solvers)
+	{
+		if (UObject* BoneSetting = Solver->GetBoneSetting(BoneName))
+		{
+			return true;
+		}
+	}
+
+	return false;
+}
+
+void UIKRigController::AddReferencedObjects(UObject* InThis, FReferenceCollector& Collector)
+{
+	// since static member, we just add only for default object
+	if (InThis && InThis->IsTemplate())
+	{
+		for (auto Iter = AssetToControllerMap.CreateIterator(); Iter; ++Iter)
+		{
+			// we add controllers to the list. so that it doesn't GCed
+			// when do we delete from the list?
+			UIKRigController* Obj = Iter.Value();
+			Collector.AddReferencedObject(Obj);
+		}
+	}
+}
+
+// -------------------------------------------------------
+// SKELETON
+//
+
+void UIKRigController::SetSourceSkeletalMesh(USkeletalMesh* SkeletalMesh, bool bReImportBones) const
+{
+	FScopedTransaction Transaction(LOCTEXT("SetSkeletalMesh_Label", "Set Skeletal Mesh"));
+	IKRigAsset->Modify();
+
+	// update stored skeletal mesh used for previewing results
+	IKRigAsset->PreviewSkeletalMesh = SkeletalMesh;
+
+	if (!bReImportBones)
+	{
+		return;
+	}
+
+	// reimport the skeleton data
+	if (SkeletalMesh)
+	{
+		SetSkeleton(SkeletalMesh->GetRefSkeleton());		
+	}else
+	{
+		IKRigAsset->Skeleton.Reset();
+	}
+}
+
+void UIKRigController::SetSkeleton(const FReferenceSkeleton& InSkeleton) const
+{
+	FScopedTransaction Transaction(LOCTEXT("SetSkeleton_Label", "Set Skeleton"));
+	IKRigAsset->Modify();
+
+	IKRigAsset->Skeleton.Initialize(InSkeleton);
+}
+
+FIKRigSkeleton& UIKRigController::GetSkeleton() const
+{
+	return IKRigAsset->Skeleton;
+}
+
+USkeletalMesh* UIKRigController::GetSourceSkeletalMesh() const
+{
+	return IKRigAsset->PreviewSkeletalMesh.Get();
+}
+
+
+// -------------------------------------------------------
+// SOLVERS
+//
+
+int32 UIKRigController::AddSolver(TSubclassOf<UIKRigSolver> InIKRigSolverClass) const
+{
+	check(IKRigAsset)
+
+	FScopedTransaction Transaction(LOCTEXT("AddSolver_Label", "Add Solver"));
+	IKRigAsset->Modify();
+
+	UIKRigSolver* NewSolver = NewObject<UIKRigSolver>(IKRigAsset, InIKRigSolverClass);
+	check(NewSolver);
+
+	return IKRigAsset->Solvers.Add(NewSolver);
+}
+
+int32 UIKRigController::GetNumSolvers() const
+{
+	check(IKRigAsset)
+	return IKRigAsset->Solvers.Num();
+}
+
+UIKRigSolver* UIKRigController::GetSolver(int32 Index) const
+{
+	check(IKRigAsset)
+
+	if (IKRigAsset->Solvers.IsValidIndex(Index))
+	{
+		return IKRigAsset->Solvers[Index];
+	}
+	
+	return nullptr;
+}
+
+void UIKRigController::RemoveSolver(UIKRigSolver* SolverToDelete) const
+{
+	if (!(IKRigAsset && SolverToDelete))
+	{
+		return;
+	}
+
+	FScopedTransaction Transaction(LOCTEXT("RemoveSolver_Label", "Remove Solver"));
+	IKRigAsset->Modify();
+
+	IKRigAsset->Solvers.Remove(SolverToDelete);
+}
+
+bool UIKRigController::MoveSolverInStack(int32 SolverToMoveIndex, int32 TargetSolverIndex) const
+{
+	if (!IKRigAsset->Solvers.IsValidIndex(SolverToMoveIndex))
+	{
+		return false;
+	}
+
+	if (!IKRigAsset->Solvers.IsValidIndex(TargetSolverIndex))
+	{
+		return false;
+	}
+
+	FScopedTransaction Transaction(LOCTEXT("ReorderSolver_Label", "Reorder Solvers"));
+	IKRigAsset->Modify();
+
+	UIKRigSolver* SolverToMove = IKRigAsset->Solvers[SolverToMoveIndex];
+	IKRigAsset->Solvers.Insert(SolverToMove, TargetSolverIndex + 1);
+	const int32 SolverToRemove = TargetSolverIndex > SolverToMoveIndex ? SolverToMoveIndex : SolverToMoveIndex + 1;
+	IKRigAsset->Solvers.RemoveAt(SolverToRemove);
+	return true;
+}
+
+void UIKRigController::SetRootBone(const FName& RootBoneName, int32 SolverIndex) const
+{
+	if (!IKRigAsset->Solvers.IsValidIndex(SolverIndex))
+	{
+		return; // solver doesn't exist
+	}
+
+	if (IKRigAsset->Skeleton.GetBoneIndexFromName(RootBoneName) == INDEX_NONE)
+	{
+		return; // bone doesn't exist
+	}
+
+	FScopedTransaction Transaction(LOCTEXT("SetRootBone_Label", "Set Root Bone"));
+	IKRigAsset->Modify();
+
+	IKRigAsset->Solvers[SolverIndex]->SetRootBone(RootBoneName);
+}
+
+const TArray<UIKRigSolver*>& UIKRigController::GetSolverArray() const
+{
+	return IKRigAsset->Solvers;
+}
+
+// -------------------------------------------------------
+// GOALS
+//
+
+UIKRigEffectorGoal* UIKRigController::AddNewGoal(const FName& GoalName, const FName& BoneName) const
+{
+	if (GetGoalIndex(GoalName) != INDEX_NONE)
+	{
+		return nullptr; // goal already exists!
+	}
+	
+	FScopedTransaction Transaction(LOCTEXT("AddSetting_Label", "Add Bone Setting"));
+	IKRigAsset->Modify();
+
+	UIKRigEffectorGoal* NewGoal = NewObject<UIKRigEffectorGoal>(GetAsset(), UIKRigEffectorGoal::StaticClass());
+	NewGoal->BoneName = BoneName;
+	NewGoal->GoalName = GoalName;
+	IKRigAsset->Goals.Add(NewGoal);
+
+	// set initial transform
+	NewGoal->InitialTransform = IKRigAsset->GetGoalInitialTransform(NewGoal);
+	NewGoal->CurrentTransform = NewGoal->InitialTransform;
+	
+	return NewGoal;
+}
+
+bool UIKRigController::RemoveGoal(const FName& GoalName) const
+{
+	const int32 GoalIndex = GetGoalIndex(GoalName);
+	if (GoalIndex == INDEX_NONE)
+	{
+		return false; // can't remove goal we don't have
+	}
+
+	FScopedTransaction Transaction(LOCTEXT("RemoveGoal_Label", "Remove Goal"));
+	IKRigAsset->Modify();
+
+	// remove from all the solvers
+	const FName& GoalToRemove = IKRigAsset->Goals[GoalIndex]->GoalName;
+	for (UIKRigSolver* Solver : IKRigAsset->Solvers)
+	{
+		Solver->RemoveGoal(GoalToRemove);
+	}
+
+	// remove from core system
+	IKRigAsset->Goals.RemoveAt(GoalIndex);
+	
+	return true;
+}
+
+FName UIKRigController::RenameGoal(const FName& OldName, const FName& PotentialNewName) const
+{
+	// sanitize the potential new name
+	FString CleanName = PotentialNewName.ToString();
+	SanitizeGoalName(CleanName);
+	const FName NewName = FName(CleanName);
+
+	// validate new name
+	const int32 ExistingGoalIndex = GetGoalIndex(NewName);
+	if (ExistingGoalIndex != INDEX_NONE)
+	{
+		return NAME_None; // name already in use, can't use that
+	}
+	const int32 GoalIndex = GetGoalIndex(OldName);
+	if (GoalIndex == INDEX_NONE)
+	{
+		return NAME_None; // can't rename goal we don't have
+	}
+	
+	FScopedTransaction Transaction(LOCTEXT("RenameGoal_Label", "Rename Goal"));
+	IKRigAsset->Modify();
+
+	// rename in core
+	IKRigAsset->Goals[GoalIndex]->GoalName = NewName;
+
+	// rename in solvers
+	for (UIKRigSolver* Solver : IKRigAsset->Solvers)
+	{
+		Solver->RenameGoal(OldName, NewName);
+	}
+
+	return NewName;
+}
+
+bool UIKRigController::SetGoalBone(const FName& GoalName, const FName& NewBoneName) const
+{
+	const int32 GoalIndex = GetGoalIndex(GoalName);
+	if (GoalIndex == INDEX_NONE)
+	{
+		return false; // goal doesn't exist in the rig
+	}
+
+	const int32 BoneIndex = GetSkeleton().GetBoneIndexFromName(NewBoneName);
+	if (BoneIndex == INDEX_NONE)
+	{
+		return false; // bone does not exist in the skeleton
+	}
+
+	if (GetBoneForGoal(GoalName) == NewBoneName)
+	{
+		return false; // goal is already using this bone
+	}
+	
+	FScopedTransaction Transaction(LOCTEXT("RenameGoal_Label", "Rename Goal"));
+	IKRigAsset->Modify();
+
+	// update goal
+	IKRigAsset->Goals[GoalIndex]->BoneName = NewBoneName;
+	
+	// update in solvers
+	for (UIKRigSolver* Solver : IKRigAsset->Solvers)
+	{
+		Solver->SetGoalBone(GoalName, NewBoneName);
+	}
+
+	// update initial transforms
+	IKRigAsset->ResetGoalTransforms();
+	
+	return true;
+}
+
+FName UIKRigController::GetBoneForGoal(const FName& GoalName) const
+{
+	for (const UIKRigEffectorGoal* Goal : IKRigAsset->Goals)
+	{
+		if (Goal->GoalName == GoalName)
+		{
+			return Goal->BoneName;
+		}
+	}
+	
+	return NAME_None;
+}
+
+bool UIKRigController::ConnectGoalToSolver(const UIKRigEffectorGoal& Goal, int32 SolverIndex) const
+{
+	// can't add goal that is not present in the core
+	check(GetGoalIndex(Goal.GoalName) != INDEX_NONE);
+	// can't add goal to a solver with an invalid index
+	check(IKRigAsset->Solvers.IsValidIndex(SolverIndex))
+
+	FScopedTransaction Transaction(LOCTEXT("AddGoalToSolver_Label", "Add Goal"));
+	IKRigAsset->Modify();
+	
+	IKRigAsset->Solvers[SolverIndex]->AddGoal(Goal);
+	return true;
+}
+
+bool UIKRigController::DisconnectGoalFromSolver(const FName& GoalToRemove, int32 SolverIndex) const
+{
+	// can't remove goal that is not present in the core
+	check(GetGoalIndex(GoalToRemove) != INDEX_NONE);
+	// can't remove goal from a solver with an invalid index
+	check(IKRigAsset->Solvers.IsValidIndex(SolverIndex))
+
+    FScopedTransaction Transaction(LOCTEXT("RemoveGoalFromSolver_Label", "Remove Goal"));
+	IKRigAsset->Modify();
+	
+	IKRigAsset->Solvers[SolverIndex]->RemoveGoal(GoalToRemove);
+	return true;
+}
+
+bool UIKRigController::IsGoalConnectedToSolver(const FName& Goal, int32 SolverIndex) const
+{
+	if (!IKRigAsset->Solvers.IsValidIndex(SolverIndex))
+	{
+		return false;
+	}
+
+	return IKRigAsset->Solvers[SolverIndex]->IsGoalConnected(Goal);
+}
+
+const TArray<UIKRigEffectorGoal*>& UIKRigController::GetAllGoals() const
+{
+	return IKRigAsset->Goals;
+}
+
+UIKRigEffectorGoal* UIKRigController::GetGoal(int32 GoalIndex) const
+{
+	if (!IKRigAsset->Goals.IsValidIndex(GoalIndex))
+	{
+		return nullptr;
+	}
+
+	return IKRigAsset->Goals[GoalIndex];
+}
+
+UIKRigEffectorGoal* UIKRigController::GetGoal(const FName& GoalName) const
+{
+	const int32 GoalIndex = GetGoalIndex(GoalName);
+	if (GoalIndex == INDEX_NONE)
+	{
+		return nullptr;
+	}
+	
+	return IKRigAsset->Goals[GoalIndex];
+}
+
+UObject* UIKRigController::GetEffectorForGoal(const FName& GoalName, int32 SolverIndex) const
+{
+	const int32 GoalIndex = GetGoalIndex(GoalName);
+	if (GoalIndex == INDEX_NONE)
+	{
+		return nullptr; // no goal with that index
+	}
+
+	if (!IKRigAsset->Solvers.IsValidIndex(SolverIndex))
+	{
+		return nullptr; // no solver with that index
+	}
+
+	return IKRigAsset->Solvers[SolverIndex]->GetEffectorWithGoal(GoalName);
+}
+
+FTransform UIKRigController::GetGoalInitialTransform(const FName& GoalName) const
+{
+	if(UIKRigEffectorGoal* Goal = GetGoal(GoalName))
+	{
+		return Goal->InitialTransform;
+	}
+	
+	return FTransform::Identity; // no goal with that name
+}
+
+FTransform UIKRigController::GetGoalCurrentTransform(const FName& GoalName) const
+{
+	if(UIKRigEffectorGoal* Goal = GetGoal(GoalName))
+	{
+		return Goal->CurrentTransform;
+	}
+	
+	return FTransform::Identity; // no goal with that name
+}
+
+void UIKRigController::SetGoalInitialTransform(const FName& GoalName, const FTransform& Transform) const
+{
+	UIKRigEffectorGoal* Goal = GetGoal(GoalName);
+	check(Goal);
+
+	Goal->InitialTransform = Transform;
+}
+
+void UIKRigController::SetGoalCurrentTransform(const FName& GoalName, const FTransform& Transform) const
+{
+	UIKRigEffectorGoal* Goal = GetGoal(GoalName);
+	check(Goal);
+
+	Goal->CurrentTransform = Transform;
+}
+
+void UIKRigController::SanitizeGoalName(FString& InOutName)
+{
+	for (int32 i = 0; i < InOutName.Len(); ++i)
+	{
+		TCHAR& C = InOutName[i];
+
+		const bool bGoodChar =
+            ((C >= 'A') && (C <= 'Z')) || ((C >= 'a') && (C <= 'z')) ||		// A-Z (upper and lowercase) anytime
+            (C == '_') || (C == '-') || (C == '.') ||						// _  - . anytime
+            ((i > 0) && (C >= '0') && (C <= '9'));							// 0-9 after the first character
+
+		if (!bGoodChar)
+		{
+			C = '_';
+		}
+	}
+
+	const int32 MaxNameLength = 20;
+	if (InOutName.Len() > MaxNameLength)
+	{
+		InOutName.LeftChopInline(InOutName.Len() - MaxNameLength);
+	}
+}
+
+int32 UIKRigController::GetGoalIndex(const FName& GoalName) const
+{
+	for (int32 i=0; i<IKRigAsset->Goals.Num(); ++i)
+	{	
+		if (IKRigAsset->Goals[i]->GoalName == GoalName)
+		{
+			return i;
+		}
+	}
+
+	return INDEX_NONE;
+}
+
+FName UIKRigController::GetGoalName(const int32& GoalIndex) const
+{
+	if (!IKRigAsset->Goals.IsValidIndex(GoalIndex))
+	{
+		return NAME_None;
+	}
+
+	return IKRigAsset->Goals[GoalIndex]->GoalName;
+}
+
+#undef LOCTEXT_NAMESPACE
diff --git a/Engine/Plugins/Animation/IKRig/Source/IKRigEditor/Private/RigEditor/IKRigDefinitionFactory.cpp b/Engine/Plugins/Animation/IKRig/Source/IKRigEditor/Private/RigEditor/IKRigDefinitionFactory.cpp
new file mode 100644
--- /dev/null
+++ b/Engine/Plugins/Animation/IKRig/Source/IKRigEditor/Private/RigEditor/IKRigDefinitionFactory.cpp
@@ -0,0 +1,51 @@
+// Copyright Epic Games, Inc. All Rights Reserved.
+
+#include "RigEditor/IKRigDefinitionFactory.h"
+#include "IKRigDefinition.h"
+#include "AssetTypeCategories.h"
+
+#define LOCTEXT_NAMESPACE "IKRigDefinitionFactory"
+
+
+UIKRigDefinitionFactory::UIKRigDefinitionFactory()
+{
+	bCreateNew = true;
+	bEditAfterNew = true;
+	SupportedClass = UIKRigDefinition::StaticClass();
+}
+
+UObject* UIKRigDefinitionFactory::FactoryCreateNew(UClass* Class, UObject* InParent, FName Name, EObjectFlags Flags, UObject* Context, FFeedbackContext* Warn)
+{
+	return NewObject<UIKRigDefinition>(InParent, Name, Flags | RF_Transactional);;
+}
+
+bool UIKRigDefinitionFactory::ShouldShowInNewMenu() const
+{
+	return true;
+}
+
+bool UIKRigDefinitionFactory::ConfigureProperties()
+{
+	return true;
+}
+
+FText UIKRigDefinitionFactory::GetDisplayName() const
+{
+	return LOCTEXT("IKRigDefinition_DisplayName", "IK Rig");
+}
+
+uint32 UIKRigDefinitionFactory::GetMenuCategories() const
+{
+	return EAssetTypeCategories::Animation;
+}
+
+FText UIKRigDefinitionFactory::GetToolTip() const
+{
+	return LOCTEXT("IKRigDefinition_Tooltip", "Defines a set of IK Solvers and Effectors to pose a skeleton with Goals.");
+}
+
+FString UIKRigDefinitionFactory::GetDefaultNewAssetName() const
+{
+	return FString(TEXT("NewIKRig"));
+}
+#undef LOCTEXT_NAMESPACE
diff --git a/Engine/Plugins/Animation/IKRig/Source/IKRigEditor/Private/RigEditor/IKRigEditMode.cpp b/Engine/Plugins/Animation/IKRig/Source/IKRigEditor/Private/RigEditor/IKRigEditMode.cpp
new file mode 100644
--- /dev/null
+++ b/Engine/Plugins/Animation/IKRig/Source/IKRigEditor/Private/RigEditor/IKRigEditMode.cpp
@@ -0,0 +1,224 @@
+// Copyright Epic Games, Inc. All Rights Reserved.
+
+#include "RigEditor/IKRigEditMode.h"
+
+#include "EditorViewportClient.h"
+#include "AssetEditorModeManager.h"
+#include "IPersonaPreviewScene.h"
+#include "Animation/DebugSkelMeshComponent.h"
+#include "Materials/MaterialInstanceDynamic.h"
+#include "RigEditor/IKRigHitProxies.h"
+#include "RigEditor/IKRigToolkit.h"
+
+FName FIKRigEditMode::ModeName("IKRigAssetEditMode");
+
+FIKRigEditMode::FIKRigEditMode()
+{
+}
+
+bool FIKRigEditMode::GetCameraTarget(FSphere& OutTarget) const
+{
+	// target union of selected goals
+	const TArray<FName>& OutGoalNames = EditorController.Pin()->GetSelectedGoals();
+	if (!OutGoalNames.IsEmpty())
+	{
+		TArray<FVector> GoalPoints;
+		for (const auto& GoalName : OutGoalNames)
+		{
+			FTransform GoalTransform = EditorController.Pin()->AssetController->GetGoalInitialTransform(GoalName);
+			GoalPoints.Add(GoalTransform.GetLocation());
+		}
+
+		// create a sphere that contains all the goal points
+		OutTarget = FSphere(&GoalPoints[0], GoalPoints.Num());
+		return true;
+	}
+
+	// target skeletal mesh
+	if (EditorController.Pin()->SkelMeshComponent)
+	{
+		OutTarget = EditorController.Pin()->SkelMeshComponent->Bounds.GetSphere();
+		return true;
+	}
+	
+	return false;
+}
+
+IPersonaPreviewScene& FIKRigEditMode::GetAnimPreviewScene() const
+{
+	return *static_cast<IPersonaPreviewScene*>(static_cast<FAssetEditorModeManager*>(Owner)->GetPreviewScene());
+}
+
+void FIKRigEditMode::GetOnScreenDebugInfo(TArray<FText>& OutDebugInfo) const
+{
+	// todo: provide warnings from solvers
+}
+
+void FIKRigEditMode::Render(const FSceneView* View, FViewport* Viewport, FPrimitiveDrawInterface* PDI)
+{
+	FEdMode::Render(View, Viewport, PDI);
+
+	UIKRigController* Controller = EditorController.Pin()->AssetController;
+	TArray<UIKRigEffectorGoal*> Goals = Controller->GetAllGoals();
+	for (const UIKRigEffectorGoal* Goal : Goals)
+	{
+		const bool bIsSelected = EditorController.Pin()->IsGoalSelected(Goal->GoalName);
+		PDI->SetHitProxy(new HIKRigEditorGoalProxy(Goal->GoalName));
+		GoalDrawer.DrawGoal(PDI, Goal, bIsSelected);
+		PDI->SetHitProxy(NULL);
+	}
+}
+
+bool FIKRigEditMode::AllowWidgetMove()
+{
+	return EditorController.Pin()->GetNumSelectedGoals() > 0;
+}
+
+bool FIKRigEditMode::ShouldDrawWidget() const
+{
+	return EditorController.Pin()->GetNumSelectedGoals() > 0;
+}
+
+bool FIKRigEditMode::UsesTransformWidget() const
+{
+	return EditorController.Pin()->GetNumSelectedGoals() > 0;
+}
+
+bool FIKRigEditMode::UsesTransformWidget(UE::Widget::EWidgetMode CheckMode) const
+{
+	return EditorController.Pin()->GetNumSelectedGoals() > 0;
+}
+
+FVector FIKRigEditMode::GetWidgetLocation() const
+{
+	const TArray<FName>& OutGoalNames = EditorController.Pin()->GetSelectedGoals();
+	if (OutGoalNames.IsEmpty())
+	{
+		return FVector::ZeroVector; 
+	}
+
+	return EditorController.Pin()->AssetController->GetGoalCurrentTransform(OutGoalNames.Last()).GetTranslation();
+}
+
+bool FIKRigEditMode::HandleClick(FEditorViewportClient* InViewportClient, HHitProxy* HitProxy, const FViewportClick& Click)
+{
+	// check for selections
+	if (Click.GetKey() == EKeys::LeftMouseButton)
+	{
+		// selected goal
+		if (HitProxy && HitProxy->IsA(HIKRigEditorGoalProxy::StaticGetType()))
+		{
+			HIKRigEditorGoalProxy* GoalProxy = static_cast<HIKRigEditorGoalProxy*>(HitProxy);
+			const bool bReplaceSelection = !(InViewportClient->IsCtrlPressed() || InViewportClient->IsShiftPressed());
+			EditorController.Pin()->HandleGoalSelectedInViewport(GoalProxy->GoalName, bReplaceSelection);
+			return true;
+		}else
+		{
+			// clicking in empty space clears selected goals
+			const bool bReplaceSelection = true;
+			EditorController.Pin()->HandleGoalSelectedInViewport(NAME_None, bReplaceSelection);
+		}
+	}
+	
+	return false;
+}
+
+bool FIKRigEditMode::StartTracking(FEditorViewportClient* InViewportClient, FViewport* InViewport)
+{
+	const TArray<FName>& SelectedGoalNames = EditorController.Pin()->GetSelectedGoals();
+	if (SelectedGoalNames.IsEmpty())
+	{
+		return false; // no goals selected to manipulate
+	}
+
+	const EAxisList::Type CurrentAxis = InViewportClient->GetCurrentWidgetAxis();
+	if (CurrentAxis == EAxisList::None)
+	{
+		return false; // not manipulating a required axis
+	}
+	
+	EditorController.Pin()->bManipulatingGoals = true;
+	return true;
+}
+
+bool FIKRigEditMode::EndTracking(FEditorViewportClient* InViewportClient, FViewport* InViewport)
+{
+	if (!EditorController.Pin()->bManipulatingGoals)
+	{
+		return false; // not handled
+	}
+
+	EditorController.Pin()->bManipulatingGoals = false;
+	return true;
+}
+
+bool FIKRigEditMode::InputDelta(FEditorViewportClient* InViewportClient, FViewport* InViewport, FVector& InDrag, FRotator& InRot, FVector& InScale)
+{
+	if (!EditorController.Pin()->bManipulatingGoals)
+	{
+		return false; // not handled
+	}
+
+	// translate goals
+	if(InViewportClient->GetWidgetMode() == UE::Widget::WM_Translate)
+	{
+		const TArray<FName>& SelectedGoalNames = EditorController.Pin()->GetSelectedGoals();
+		UIKRigController* Controller = EditorController.Pin()->AssetController;
+		for (const FName& GoalName : SelectedGoalNames)
+		{
+			FTransform CurrentTransform = Controller->GetGoal(GoalName)->CurrentTransform;
+			CurrentTransform.AddToTranslation(InDrag);
+			Controller->SetGoalCurrentTransform(GoalName, CurrentTransform);
+		}
+	}
+
+	// rotate goals
+	if(InViewportClient->GetWidgetMode() == UE::Widget::WM_Rotate)
+	{
+		const TArray<FName>& SelectedGoalNames = EditorController.Pin()->GetSelectedGoals();
+		UIKRigController* Controller = EditorController.Pin()->AssetController;
+		for (const FName& GoalName : SelectedGoalNames)
+		{
+			FTransform CurrentTransform = Controller->GetGoal(GoalName)->CurrentTransform;
+			FQuat CurrentRotation = CurrentTransform.GetRotation();
+			CurrentRotation = (InRot.Quaternion() * CurrentRotation);
+			CurrentTransform.SetRotation(CurrentRotation);
+			Controller->SetGoalCurrentTransform(GoalName, CurrentTransform);
+		}
+	}
+	
+	return true;
+}
+
+bool FIKRigEditMode::GetCustomDrawingCoordinateSystem(FMatrix& InMatrix, void* InData)
+{
+	const TArray<FName>& SelectedGoals = EditorController.Pin()->GetSelectedGoals();
+	if (SelectedGoals.IsEmpty())
+	{
+		return false; // nothing selected to manipulate
+	}
+
+	if (UIKRigEffectorGoal* Goal = EditorController.Pin()->AssetController->GetGoal(SelectedGoals[0]))
+	{
+		InMatrix = Goal->CurrentTransform.ToMatrixNoScale().RemoveTranslation();
+		return true;
+	}
+
+	return false;
+}
+
+bool FIKRigEditMode::GetCustomInputCoordinateSystem(FMatrix& InMatrix, void* InData)
+{
+	return GetCustomDrawingCoordinateSystem(InMatrix, InData);
+}
+
+void FIKRigEditMode::Tick(FEditorViewportClient* ViewportClient, float DeltaTime)
+{
+	FEdMode::Tick(ViewportClient, DeltaTime);
+}
+
+void FIKRigEditMode::DrawHUD(FEditorViewportClient* ViewportClient, FViewport* Viewport, const FSceneView* View, FCanvas* Canvas)
+{
+	FEdMode::DrawHUD(ViewportClient, Viewport, View, Canvas);
+}
+
diff --git a/Engine/Plugins/Animation/IKRig/Source/IKRigEditor/Private/RigEditor/IKRigEditorController.cpp b/Engine/Plugins/Animation/IKRig/Source/IKRigEditor/Private/RigEditor/IKRigEditorController.cpp
new file mode 100644
--- /dev/null
+++ b/Engine/Plugins/Animation/IKRig/Source/IKRigEditor/Private/RigEditor/IKRigEditorController.cpp
@@ -0,0 +1,203 @@
+// Copyright Epic Games, Inc. All Rights Reserved.
+
+#include "RigEditor/IKRigEditorController.h"
+
+#include "SKismetInspector.h"
+#include "Dialogs/Dialogs.h"
+
+#include "RigEditor/IKRigController.h"
+#include "RigEditor/SIKRigSkeleton.h"
+#include "RigEditor/SIKRigSolverStack.h"
+#include "RigEditor/IKRigAnimInstance.h"
+
+#define LOCTEXT_NAMESPACE "IKRigEditorController"
+
+void FIKRigEditorController::Reset() const
+{
+	const TArray<UIKRigEffectorGoal*>& Goals = AssetController->GetAllGoals();
+	for (UIKRigEffectorGoal* Goal : Goals)
+	{
+		Goal->CurrentTransform = Goal->InitialTransform;
+	}
+}
+
+bool FIKRigEditorController::IsGoalSelected(const FName& GoalName)
+{
+	return SelectedGoals.Contains(GoalName);
+}
+
+int32 FIKRigEditorController::GetNumSelectedGoals()
+{
+	return SelectedGoals.Num();
+}
+
+void FIKRigEditorController::HandleGoalSelectedInViewport(const FName& GoalName, bool bReplace)
+{
+	if (bReplace)
+	{
+		SelectedGoals.Reset();
+		if (GoalName != NAME_None)
+		{
+			SelectedGoals.Add(GoalName);
+		}
+	}else
+	{
+		if (GoalName != NAME_None)
+		{
+			const bool bAlreadySelected = SelectedGoals.Contains(GoalName);
+			if (bAlreadySelected)
+			{
+				SelectedGoals.Remove(GoalName);	
+			}else
+			{
+				SelectedGoals.Add(GoalName);
+			}
+		}	
+	}
+	
+	SkeletonView->SetSelectedGoalsFromViewport(SelectedGoals);
+	ShowDetailsForGoal(GoalName);
+}
+
+void FIKRigEditorController::HandleGoalsSelectedInTreeView(const TArray<FName>& GoalNames)
+{
+	SelectedGoals = GoalNames;
+}
+
+void FIKRigEditorController::GetSelectedSolvers(TArray<TSharedPtr<FSolverStackElement>>& OutSelectedSolvers)
+{
+	OutSelectedSolvers.Reset();
+	OutSelectedSolvers.Append(SolverStackView->ListView->GetSelectedItems());
+}
+
+int32 FIKRigEditorController::GetSelectedSolverIndex()
+{
+	TArray<TSharedPtr<FSolverStackElement>> SelectedSolvers = SolverStackView->ListView->GetSelectedItems();
+	if (SelectedSolvers.IsEmpty())
+	{
+		return INDEX_NONE;
+	}
+
+	return SelectedSolvers[0]->IndexInStack;
+}
+
+void FIKRigEditorController::PromptToAddSolverAtStartup() const
+{
+	if (AssetController->GetNumSolvers() > 0)
+	{
+		return;
+	}
+
+	// prompt user to add a default solver
+	FIKRigAddFirstSolverSettings Settings;
+	const TSharedPtr<FStructOnScope> StructToDisplay = MakeShareable(new FStructOnScope(FIKRigAddFirstSolverSettings::StaticStruct(), (uint8*)&Settings));
+	TSharedRef<SKismetInspector> KismetInspector = SNew(SKismetInspector);
+	KismetInspector->ShowSingleStruct(StructToDisplay);
+
+	SGenericDialogWidget::OpenDialog(
+        LOCTEXT("EditorController_IKRigFirstSolver", "Add Default Solver"), 
+        KismetInspector, 
+        SGenericDialogWidget::FArguments(), 
+        true);
+
+	if (Settings.SolverType != nullptr)
+	{
+		AssetController->AddSolver(Settings.SolverType);
+		SolverStackView->RefreshStackView();
+	}
+}
+
+bool FIKRigEditorController::IsElementConnectedToSolver(TSharedRef<FIKRigTreeElement> TreeElement, int32 SolverIndex)
+{
+	if (!AssetController->GetSolverArray().IsValidIndex(SolverIndex))
+	{
+		return false; // not a valid solver index
+	}
+
+	UIKRigSolver* Solver = AssetController->GetSolver(SolverIndex);
+	if (TreeElement->ElementType == IKRigTreeElementType::BONE)
+	{
+		// is this bone affected by this solver?
+		return Solver->IsBoneAffectedBySolver(TreeElement->Key, AssetController->GetSkeleton());
+	}
+
+	if (TreeElement->ElementType == IKRigTreeElementType::BONE_SETTINGS)
+	{
+		// is this bone setting belonging to the solver?
+		return (Solver->GetBoneSetting(TreeElement->BoneSettingBoneName) != nullptr);
+	}
+
+	if (TreeElement->ElementType == IKRigTreeElementType::GOAL)
+	{
+		// is goal connected to the solver?
+		return AssetController->IsGoalConnectedToSolver(TreeElement->Key, SolverIndex);
+	}
+
+	if (TreeElement->ElementType == IKRigTreeElementType::EFFECTOR)
+	{
+		// is this an effector for this solver?
+		return TreeElement->EffectorSolverIndex == SolverIndex;
+	}
+
+	checkNoEntry();
+	return false;
+}
+
+bool FIKRigEditorController::IsElementConnectedToAnySolver(TSharedRef<FIKRigTreeElement> TreeElement)
+{
+	const int32 NumSolvers = AssetController->GetNumSolvers();
+	for (int32 SolverIndex=0; SolverIndex<NumSolvers; ++SolverIndex)
+	{
+		if (IsElementConnectedToSolver(TreeElement, SolverIndex))
+		{
+			return true;
+		}
+	}
+
+	return false;
+}
+
+void FIKRigEditorController::ShowDetailsForBone(const FName BoneName)
+{
+	ShowEmptyDetails();
+	// this can be useful during development/debugging, but is slow/confusing so should not be left enabled
+	//DetailsView->SetObject(AssetController->GetAsset());
+}
+
+void FIKRigEditorController::ShowDetailsForBoneSettings(const FName BoneName, int32 SolverIndex)
+{
+	if (UObject* BoneSettings = AssetController->GetSettingsForBone(BoneName, SolverIndex))
+	{
+		DetailsView->SetObject(BoneSettings);
+	}
+}
+
+void FIKRigEditorController::ShowDetailsForGoal(const FName GoalName)
+{
+	UIKRigEffectorGoal* Goal = AssetController->GetGoal(GoalName);
+	DetailsView->SetObject(Goal);
+}
+
+void FIKRigEditorController::ShowDetailsForEffector(const FName GoalName, const int32 SolverIndex)
+{
+	// get solver that owns this effector
+	if (UIKRigSolver* SolverWithEffector = AssetController->GetSolver(SolverIndex))
+	{
+		if (UObject* EffectorSettings = SolverWithEffector->GetEffectorWithGoal(GoalName))
+		{
+			DetailsView->SetObject(EffectorSettings);
+		}
+	}
+}
+
+void FIKRigEditorController::ShowDetailsForSolver(const int32 SolverIndex)
+{
+	DetailsView->SetObject(AssetController->GetSolver(SolverIndex));
+}
+
+void FIKRigEditorController::ShowEmptyDetails()
+{
+	DetailsView->SetObject(nullptr);
+}
+
+#undef LOCTEXT_NAMESPACE
diff --git a/Engine/Plugins/Animation/IKRig/Source/IKRigEditor/Private/RigEditor/IKRigHitProxies.cpp b/Engine/Plugins/Animation/IKRig/Source/IKRigEditor/Private/RigEditor/IKRigHitProxies.cpp
new file mode 100644
--- /dev/null
+++ b/Engine/Plugins/Animation/IKRig/Source/IKRigEditor/Private/RigEditor/IKRigHitProxies.cpp
@@ -0,0 +1,6 @@
+// Copyright Epic Games, Inc. All Rights Reserved.
+
+#include "RigEditor/IKRigHitProxies.h"
+
+IMPLEMENT_HIT_PROXY(HIKRigEditorBoneProxy,HHitProxy);
+IMPLEMENT_HIT_PROXY(HIKRigEditorGoalProxy,HHitProxy);
\ No newline at end of file
diff --git a/Engine/Plugins/Animation/IKRig/Source/IKRigEditor/Private/RigEditor/IKRigMode.cpp b/Engine/Plugins/Animation/IKRig/Source/IKRigEditor/Private/RigEditor/IKRigMode.cpp
new file mode 100644
--- /dev/null
+++ b/Engine/Plugins/Animation/IKRig/Source/IKRigEditor/Private/RigEditor/IKRigMode.cpp
@@ -0,0 +1,102 @@
+// Copyright Epic Games, Inc. All Rights Reserved.
+
+#include "RigEditor/IKRigMode.h"
+#include "RigEditor/IKRigToolkit.h"
+#include "IPersonaPreviewScene.h"
+#include "PersonaModule.h"
+#include "ISkeletonEditorModule.h"
+#include "Modules/ModuleManager.h"
+#include "PersonaTabs.h"
+#include "RigEditor/IKRigSkeletonTabSummoner.h"
+#include "RigEditor/IKRigSolverStackTabSummoner.h"
+
+#define LOCTEXT_NAMESPACE "IKRigMode"
+
+static const FName IKRigSolverStackName("IKRigSolverStack");
+
+FIKRigMode::FIKRigMode(
+	TSharedRef<FWorkflowCentricApplication> InHostingApp,  
+	TSharedRef<IPersonaPreviewScene> InPreviewScene)
+	: FApplicationMode(IKRigEditorModes::IKRigEditorMode)
+{
+	IKRigEditorPtr = StaticCastSharedRef<FIKRigEditorToolkit>(InHostingApp);
+	TSharedRef<FIKRigEditorToolkit> IKRigEditor = StaticCastSharedRef<FIKRigEditorToolkit>(InHostingApp);
+
+	FPersonaViewportArgs ViewportArgs(InPreviewScene);
+	ViewportArgs.bAlwaysShowTransformToolbar = true;
+	ViewportArgs.bShowStats = false;
+	ViewportArgs.bShowTurnTable = false;
+	ViewportArgs.ContextName = TEXT("IKRigEditor.Viewport");
+
+	// register Persona tabs
+	FPersonaModule& PersonaModule = FModuleManager::LoadModuleChecked<FPersonaModule>("Persona");
+	TabFactories.RegisterFactory(PersonaModule.CreatePersonaViewportTabFactory(InHostingApp, ViewportArgs));
+	TabFactories.RegisterFactory(PersonaModule.CreateAdvancedPreviewSceneTabFactory(InHostingApp, InPreviewScene));
+	TabFactories.RegisterFactory(PersonaModule.CreateDetailsTabFactory(InHostingApp, FOnDetailsCreated::CreateSP(&IKRigEditor.Get(), &FIKRigEditorToolkit::HandleDetailsCreated)));
+
+	// register custom tabs
+	TabFactories.RegisterFactory(MakeShared<FIKRigSkeletonTabSummoner>(IKRigEditor));
+	TabFactories.RegisterFactory(MakeShared<FIKRigSolverStackTabSummoner>(IKRigEditor));
+
+	// create tab layout
+	TabLayout = FTabManager::NewLayout("Standalone_IKRigEditor_Layout_v1.118")
+		->AddArea
+		(
+			FTabManager::NewPrimaryArea()
+			->SetOrientation(Orient_Vertical)
+			->Split
+			(
+				FTabManager::NewSplitter()
+				->SetSizeCoefficient(0.9f)
+				->SetOrientation(Orient_Horizontal)
+				->Split
+				(
+				    FTabManager::NewSplitter()
+				    ->SetSizeCoefficient(0.2f)
+				    ->SetOrientation(Orient_Vertical)
+				    ->Split
+				    (
+					    FTabManager::NewStack()
+					    ->SetSizeCoefficient(0.6f)
+					    ->AddTab(FIKRigSkeletonTabSummoner::TabID, ETabState::OpenedTab)
+					)
+					->Split
+					(
+					    FTabManager::NewStack()
+					    ->SetSizeCoefficient(0.4f)
+					    ->AddTab(IKRigSolverStackName, ETabState::OpenedTab)
+					)
+				)
+				->Split
+				(
+					FTabManager::NewStack()
+					->SetSizeCoefficient(0.6f)
+					->SetHideTabWell(true)
+					->AddTab(FPersonaTabs::PreviewViewportID, ETabState::OpenedTab)
+				)
+				->Split
+				(
+					FTabManager::NewStack()
+                    ->SetSizeCoefficient(0.6f)
+                    ->AddTab(FPersonaTabs::DetailsID, ETabState::OpenedTab)
+                    ->AddTab(FPersonaTabs::AdvancedPreviewSceneSettingsID, ETabState::OpenedTab)
+                    ->SetForegroundTab(FPersonaTabs::DetailsID)
+				)
+			)
+		);
+
+	PersonaModule.OnRegisterTabs().Broadcast(TabFactories, InHostingApp);
+	LayoutExtender = MakeShared<FLayoutExtender>();
+	PersonaModule.OnRegisterLayoutExtensions().Broadcast(*LayoutExtender.Get());
+	TabLayout->ProcessExtensions(*LayoutExtender.Get());
+}
+
+void FIKRigMode::RegisterTabFactories(TSharedPtr<FTabManager> InTabManager)
+{
+	TSharedPtr<FIKRigEditorToolkit> IKRigEditor = IKRigEditorPtr.Pin();
+	IKRigEditor->RegisterTabSpawners(InTabManager.ToSharedRef());
+	IKRigEditor->PushTabFactories(TabFactories);
+	FApplicationMode::RegisterTabFactories(InTabManager);
+}
+
+#undef LOCTEXT_NAMESPACE
diff --git a/Engine/Plugins/Animation/IKRig/Source/IKRigEditor/Private/RigEditor/IKRigSkeletonCommands.cpp b/Engine/Plugins/Animation/IKRig/Source/IKRigEditor/Private/RigEditor/IKRigSkeletonCommands.cpp
new file mode 100644
--- /dev/null
+++ b/Engine/Plugins/Animation/IKRig/Source/IKRigEditor/Private/RigEditor/IKRigSkeletonCommands.cpp
@@ -0,0 +1,18 @@
+// Copyright Epic Games, Inc. All Rights Reserved.
+
+#include "RigEditor/IKRigSkeletonCommands.h"
+
+#define LOCTEXT_NAMESPACE "ControlRigHierarchyCommands"
+
+void FIKRigSkeletonCommands::RegisterCommands()
+{
+	UI_COMMAND(NewGoal, "New IK Goal", "Add new IK goal at the selected bone.", EUserInterfaceActionType::Button, FInputChord());
+	UI_COMMAND(Delete, "Delete Goal or Effector", "Delete the selected IK Goal or Effector.", EUserInterfaceActionType::Button, FInputChord(EKeys::Delete));
+	UI_COMMAND(ConnectGoalToSolvers, "Connect Goal to Selected Solvers", "Make the selected goal an effector in the selected solvers.", EUserInterfaceActionType::Button, FInputChord());
+	UI_COMMAND(DisconnectGoalFromSolvers, "Disconnect Goal from Selected Solvers", "Remove effectors from the selected solvers that use this goal.", EUserInterfaceActionType::Button, FInputChord());
+	UI_COMMAND(SetRootBoneOnSolvers, "Set Root Bone on Selected Solvers", "Set the Root Bone setting on the selected solvers.", EUserInterfaceActionType::Button, FInputChord());
+	UI_COMMAND(AddBoneSettings, "Add Settings to Selected Bone", "Apply settings to the selected bone for all selected solvers (defined per-solver for limits, stiffness etc).", EUserInterfaceActionType::Button, FInputChord());
+	UI_COMMAND(RemoveBoneSettings, "Remove Settings on Selected Bone", "Remove all settings on the selected bone in all selected solvers.", EUserInterfaceActionType::Button, FInputChord());
+}
+
+#undef LOCTEXT_NAMESPACE
diff --git a/Engine/Plugins/Animation/IKRig/Source/IKRigEditor/Private/RigEditor/IKRigSkeletonTabSummoner.cpp b/Engine/Plugins/Animation/IKRig/Source/IKRigEditor/Private/RigEditor/IKRigSkeletonTabSummoner.cpp
new file mode 100644
--- /dev/null
+++ b/Engine/Plugins/Animation/IKRig/Source/IKRigEditor/Private/RigEditor/IKRigSkeletonTabSummoner.cpp
@@ -0,0 +1,37 @@
+// Copyright Epic Games, Inc. All Rights Reserved.
+
+#include "RigEditor/IKRigSkeletonTabSummoner.h"
+
+#include "IDocumentation.h"
+#include "RigEditor/IKRigEditorStyle.h"
+#include "RigEditor/IKRigToolkit.h"
+#include "RigEditor/SIKRigSkeleton.h"
+
+#define LOCTEXT_NAMESPACE "IKRigSkeletonTabSummoner"
+
+const FName FIKRigSkeletonTabSummoner::TabID(TEXT("IKRigSkeleton"));
+
+FIKRigSkeletonTabSummoner::FIKRigSkeletonTabSummoner(const TSharedRef<FIKRigEditorToolkit>& InIKRigEditor)
+	: FWorkflowTabFactory(TabID, InIKRigEditor)
+	, IKRigEditor(InIKRigEditor)
+{
+	bIsSingleton = true; // only allow a single instance of this tab
+	
+	TabLabel = LOCTEXT("IKRigSkeletonTabLabel", "Hierarchy");
+	TabIcon = FSlateIcon(FIKRigEditorStyle::Get().GetStyleSetName(), "IKRig.TabIcon");
+
+	ViewMenuDescription = LOCTEXT("IKRigSkeleton_ViewMenu_Desc", "Hierarchy");
+	ViewMenuTooltip = LOCTEXT("IKRigSkeleton_ViewMenu_ToolTip", "Show the IK Rig Hierarchy Tab");
+}
+
+TSharedPtr<SToolTip> FIKRigSkeletonTabSummoner::CreateTabToolTipWidget(const FWorkflowTabSpawnInfo& Info) const
+{
+	return  IDocumentation::Get()->CreateToolTip(LOCTEXT("IKRigSkeletonTooltip", "The IK Rig Hierarchy tab lets you view the rig hierarchy."), NULL, TEXT("Shared/Editors/Persona"), TEXT("IKRigSkeleton_Window"));
+}
+
+TSharedRef<SWidget> FIKRigSkeletonTabSummoner::CreateTabBody(const FWorkflowTabSpawnInfo& Info) const
+{
+	return SNew(SIKRigSkeleton, IKRigEditor.Pin()->GetController());
+}
+
+#undef LOCTEXT_NAMESPACE 
diff --git a/Engine/Plugins/Animation/IKRig/Source/IKRigEditor/Private/RigEditor/IKRigSolverStackTabSummoner.cpp b/Engine/Plugins/Animation/IKRig/Source/IKRigEditor/Private/RigEditor/IKRigSolverStackTabSummoner.cpp
new file mode 100644
--- /dev/null
+++ b/Engine/Plugins/Animation/IKRig/Source/IKRigEditor/Private/RigEditor/IKRigSolverStackTabSummoner.cpp
@@ -0,0 +1,37 @@
+// Copyright Epic Games, Inc. All Rights Reserved.
+
+#include "RigEditor/IKRigSolverStackTabSummoner.h"
+
+#include "IDocumentation.h"
+#include "RigEditor/IKRigEditorStyle.h"
+#include "RigEditor/IKRigToolkit.h"
+#include "RigEditor/SIKRigSolverStack.h"
+
+#define LOCTEXT_NAMESPACE "SolverStackTabSummoner"
+
+const FName FIKRigSolverStackTabSummoner::TabID(TEXT("IKRigSolverStack"));
+
+FIKRigSolverStackTabSummoner::FIKRigSolverStackTabSummoner(const TSharedRef<FIKRigEditorToolkit>& InIKRigEditor)
+	: FWorkflowTabFactory(TabID, InIKRigEditor)
+	, IKRigEditor(InIKRigEditor)
+{
+	bIsSingleton = true; // only allow a single instance of this tab
+	
+	TabLabel = LOCTEXT("IKRigSolverStackTabLabel", "Solver Stack");
+	TabIcon = FSlateIcon(FIKRigEditorStyle::Get().GetStyleSetName(), "IKRig.TabIcon");
+
+	ViewMenuDescription = LOCTEXT("IKRigSolverStack_ViewMenu_Desc", "IK Rig Solver Stack");
+	ViewMenuTooltip = LOCTEXT("IKRigSolverStack_ViewMenu_ToolTip", "Show the IK Rig Solver Stack Tab");
+}
+
+TSharedPtr<SToolTip> FIKRigSolverStackTabSummoner::CreateTabToolTipWidget(const FWorkflowTabSpawnInfo& Info) const
+{
+	return  IDocumentation::Get()->CreateToolTip(LOCTEXT("IKRigSolverStackTooltip", "A stack of IK solvers executed from top to bottom. These are associated with goals to affect the skeleton."), NULL, TEXT("Shared/Editors/Persona"), TEXT("IKRigSolverStack_Window"));
+}
+
+TSharedRef<SWidget> FIKRigSolverStackTabSummoner::CreateTabBody(const FWorkflowTabSpawnInfo& Info) const
+{
+	return SNew(SIKRigSolverStack, IKRigEditor.Pin()->GetController());
+}
+
+#undef LOCTEXT_NAMESPACE 
diff --git a/Engine/Plugins/Animation/IKRig/Source/IKRigEditor/Private/RigEditor/IKRigToolkit.cpp b/Engine/Plugins/Animation/IKRig/Source/IKRigEditor/Private/RigEditor/IKRigToolkit.cpp
new file mode 100644
--- /dev/null
+++ b/Engine/Plugins/Animation/IKRig/Source/IKRigEditor/Private/RigEditor/IKRigToolkit.cpp
@@ -0,0 +1,249 @@
+// Copyright Epic Games, Inc. All Rights Reserved.
+
+#include "RigEditor/IKRigToolkit.h"
+
+#include "AnimationEditorPreviewActor.h"
+#include "EditorModeManager.h"
+#include "GameFramework/WorldSettings.h"
+#include "Modules/ModuleManager.h"
+#include "PersonaModule.h"
+#include "IPersonaToolkit.h"
+#include "IAssetFamily.h"
+#include "ISkeletonEditorModule.h"
+#include "Preferences/PersonaOptions.h"
+#include "AnimCustomInstanceHelper.h"
+#include "Animation/DebugSkelMeshComponent.h"
+#include "Framework/MultiBox/MultiBoxBuilder.h"
+
+#include "IKRigDefinition.h"
+#include "RigEditor/IKRigAnimInstance.h"
+#include "RigEditor/IKRigCommands.h"
+#include "RigEditor/IKRigEditMode.h"
+#include "RigEditor/IKRigMode.h"
+#include "RigEditor/IKRigEditorController.h"
+
+#define LOCTEXT_NAMESPACE "IKRigEditorToolkit"
+
+const FName IKRigEditorModes::IKRigEditorMode("IKRigEditorMode");
+const FName IKRigEditorAppName = FName(TEXT("IKRigEditorApp"));
+
+FIKRigEditorToolkit::FIKRigEditorToolkit()
+	: EditorController(MakeShared<FIKRigEditorController>())
+{
+}
+
+FIKRigEditorToolkit::~FIKRigEditorToolkit()
+{
+}
+
+void FIKRigEditorToolkit::InitAssetEditor(
+	const EToolkitMode::Type Mode,
+    const TSharedPtr<IToolkitHost>& InitToolkitHost, 
+    UIKRigDefinition* IKRigAsset)
+{
+	EditorController->AssetController = UIKRigController::GetIKRigController(IKRigAsset);
+	EditorController->EditorToolkit = SharedThis(this);
+
+	BindCommands();
+	
+	FPersonaToolkitArgs PersonaToolkitArgs;
+	PersonaToolkitArgs.OnPreviewSceneCreated = FOnPreviewSceneCreated::FDelegate::CreateSP(this, &FIKRigEditorToolkit::HandlePreviewSceneCreated);
+	
+	FPersonaModule& PersonaModule = FModuleManager::LoadModuleChecked<FPersonaModule>("Persona");
+	PersonaToolkit = PersonaModule.CreatePersonaToolkit(IKRigAsset, PersonaToolkitArgs);
+	
+	// when/if preview mesh is changed, we need to reinitialize the anim instance
+    PersonaToolkit->GetPreviewScene()->RegisterOnPreviewMeshChanged(FOnPreviewMeshChanged::CreateSP(this, &FIKRigEditorToolkit::HandlePreviewMeshChanged));
+
+	TSharedRef<IAssetFamily> AssetFamily = PersonaModule.CreatePersonaAssetFamily(IKRigAsset);
+	AssetFamily->RecordAssetOpened(FAssetData(IKRigAsset));
+
+	const bool bCreateDefaultStandaloneMenu = true;
+	const bool bCreateDefaultToolbar = true;
+	FAssetEditorToolkit::InitAssetEditor(
+		Mode, 
+		InitToolkitHost, 
+		IKRigEditorAppName, 
+		FTabManager::FLayout::NullLayout, 
+		bCreateDefaultStandaloneMenu, 
+		bCreateDefaultToolbar, 
+		IKRigAsset);
+
+	AddApplicationMode(
+		IKRigEditorModes::IKRigEditorMode,
+		MakeShareable(new FIKRigMode(SharedThis(this), PersonaToolkit->GetPreviewScene())));
+
+	SetCurrentMode(IKRigEditorModes::IKRigEditorMode);
+
+	GetEditorModeManager().SetDefaultMode(FIKRigEditMode::ModeName);
+	GetEditorModeManager().ActivateMode(FPersonaEditModes::SkeletonSelection);
+	GetEditorModeManager().ActivateMode(FIKRigEditMode::ModeName);
+	static_cast<FIKRigEditMode*>(GetEditorModeManager().GetActiveMode(FIKRigEditMode::ModeName))->SetEditorController(EditorController);
+
+	ExtendToolbar();
+	RegenerateMenusAndToolbars();
+}
+
+void FIKRigEditorToolkit::RegisterTabSpawners(const TSharedRef<FTabManager>& InTabManager)
+{
+	WorkspaceMenuCategory = InTabManager->AddLocalWorkspaceMenuCategory(LOCTEXT("WorkspaceMenu_IKRigEditor", "IK Rig Editor"));
+	auto WorkspaceMenuCategoryRef = WorkspaceMenuCategory.ToSharedRef();
+
+	FAssetEditorToolkit::RegisterTabSpawners(InTabManager);
+}
+
+void FIKRigEditorToolkit::UnregisterTabSpawners(const TSharedRef<FTabManager>& InTabManager)
+{
+	FAssetEditorToolkit::UnregisterTabSpawners(InTabManager);
+}
+
+void FIKRigEditorToolkit::BindCommands()
+{
+	const FIKRigCommands& Commands = FIKRigCommands::Get();
+
+	ToolkitCommands->MapAction(
+        Commands.Reset,
+        FExecuteAction::CreateSP(this, &FIKRigEditorToolkit::HandleReset),
+		EUIActionRepeatMode::RepeatDisabled);
+}
+
+void FIKRigEditorToolkit::ExtendToolbar()
+{
+	TSharedPtr<FExtender> ToolbarExtender = MakeShareable(new FExtender);
+
+	AddToolbarExtender(ToolbarExtender);
+
+	ToolbarExtender->AddToolBarExtension(
+        "Asset",
+        EExtensionHook::After,
+        GetToolkitCommands(),
+        FToolBarExtensionDelegate::CreateSP(this, &FIKRigEditorToolkit::FillToolbar)
+    );
+}
+
+void FIKRigEditorToolkit::FillToolbar(FToolBarBuilder& ToolbarBuilder)
+{
+	ToolbarBuilder.BeginSection("Reset");
+	{
+		ToolbarBuilder.AddToolBarButton(FIKRigCommands::Get().Reset,
+			NAME_None, TAttribute<FText>(), TAttribute<FText>(), FSlateIcon(FIKRigEditorStyle::Get().GetStyleSetName(), "IKRig.Reset"));
+	}
+	ToolbarBuilder.EndSection();
+}
+
+FName FIKRigEditorToolkit::GetToolkitFName() const
+{
+	return FName("IKRigEditor");
+}
+
+FText FIKRigEditorToolkit::GetBaseToolkitName() const
+{
+	return LOCTEXT("IKRigEditorAppLabel", "IK Rig Editor");
+}
+
+FText FIKRigEditorToolkit::GetToolkitName() const
+{
+	const bool bDirtyState = EditorController->AssetController->GetAsset()->GetOutermost()->IsDirty();
+
+	FFormatNamedArguments Args;
+	Args.Add(TEXT("AssetName"), FText::FromString(EditorController->AssetController->GetAsset()->GetName()));
+	Args.Add(TEXT("DirtyState"), bDirtyState ? FText::FromString(TEXT("*")) : FText::GetEmpty());
+	return FText::Format(LOCTEXT("DemoEditorToolkitName", "{AssetName}{DirtyState}"), Args);
+}
+
+FLinearColor FIKRigEditorToolkit::GetWorldCentricTabColorScale() const
+{
+	return FLinearColor::White;
+}
+
+FString FIKRigEditorToolkit::GetWorldCentricTabPrefix() const
+{
+	return TEXT("IKRigEditor");
+}
+
+void FIKRigEditorToolkit::AddReferencedObjects(FReferenceCollector& Collector)
+{
+	UIKRigDefinition* Asset = EditorController->AssetController->GetAsset();
+	Collector.AddReferencedObject(Asset);
+}
+
+TStatId FIKRigEditorToolkit::GetStatId() const
+{
+	RETURN_QUICK_DECLARE_CYCLE_STAT(FIKRigEditorToolkit, STATGROUP_Tickables);
+}
+
+void FIKRigEditorToolkit::HandlePreviewSceneCreated(const TSharedRef<IPersonaPreviewScene>& InPersonaPreviewScene)
+{
+	AAnimationEditorPreviewActor* Actor = InPersonaPreviewScene->GetWorld()->SpawnActor<AAnimationEditorPreviewActor>(AAnimationEditorPreviewActor::StaticClass(), FTransform::Identity);
+	Actor->SetFlags(RF_Transient);
+	InPersonaPreviewScene->SetActor(Actor);
+	
+	// create the preview skeletal mesh component
+	EditorController->SkelMeshComponent = NewObject<UDebugSkelMeshComponent>(Actor);
+
+	// setup an apply an anim instance to the skeletal mesh component
+	UIKRigAnimInstance* IKRigAnimInstance = NewObject<UIKRigAnimInstance>(EditorController->SkelMeshComponent, TEXT("IKRigAnimScriptInstance"));
+	SetupAnimInstance(IKRigAnimInstance);
+
+	// set the skeletal mesh on the component
+	// NOTE: this must be done AFTER setting the AnimInstance so that the correct root anim node is loaded
+	USkeletalMesh* Mesh = EditorController->AssetController->GetSourceSkeletalMesh();
+	EditorController->SkelMeshComponent->SetSkeletalMesh(Mesh);
+
+	// apply mesh to the preview scene
+	InPersonaPreviewScene->SetPreviewMeshComponent(EditorController->SkelMeshComponent);
+	InPersonaPreviewScene->AddComponent(EditorController->SkelMeshComponent, FTransform::Identity);
+	InPersonaPreviewScene->SetAdditionalMeshesSelectable(false);
+	InPersonaPreviewScene->SetPreviewMesh(Mesh);
+}
+
+void FIKRigEditorToolkit::HandlePreviewMeshChanged(USkeletalMesh* InOldSkeletalMesh, USkeletalMesh* InNewSkeletalMesh)
+{
+	if (!InOldSkeletalMesh)
+	{
+		return; // first time setup
+	}
+	
+	// update asset with new skeletal mesh (will copy new skeleton data)
+	EditorController->AssetController->SetSourceSkeletalMesh(InNewSkeletalMesh, false);
+
+	// update anim instance to use new skeletal mesh
+	// this is required so that the bone containers passed around during update/eval are correctly sized
+	UDebugSkelMeshComponent* EditorSkelComp = Cast<UDebugSkelMeshComponent>(GetPersonaToolkit()->GetPreviewScene()->GetPreviewMeshComponent());
+	if (EditorSkelComp)
+	{
+		bool bWasCreated = false;
+		UIKRigAnimInstance* AnimInstance = FAnimCustomInstanceHelper::BindToSkeletalMeshComponent<UIKRigAnimInstance>(EditorSkelComp , bWasCreated);
+		SetupAnimInstance(AnimInstance);
+	}
+
+	EditorController->SkelMeshComponent->SetSkeletalMesh(InNewSkeletalMesh);
+}
+
+void FIKRigEditorToolkit::HandleDetailsCreated(const TSharedRef<class IDetailsView>& InDetailsView)
+{
+	EditorController->DetailsView = InDetailsView;
+	EditorController->DetailsView->OnFinishedChangingProperties().AddSP(this, &FIKRigEditorToolkit::OnFinishedChangingDetails);
+	EditorController->ShowEmptyDetails();
+}
+
+void FIKRigEditorToolkit::OnFinishedChangingDetails(
+    const FPropertyChangedEvent& PropertyChangedEvent)
+{
+	UE_LOG(LogTemp, Log, TEXT("FIKRigEditorToolkit::OnFinishedChangingProperties"));
+}
+
+void FIKRigEditorToolkit::HandleReset()
+{
+	EditorController.Get().Reset();
+}
+
+void FIKRigEditorToolkit::SetupAnimInstance(UIKRigAnimInstance* InAnimInstance)
+{
+	InAnimInstance->SetIKRigAsset(EditorController->AssetController->GetAsset());
+	EditorController->SkelMeshComponent->PreviewInstance = InAnimInstance;
+	InAnimInstance->InitializeAnimation();
+}
+
+
+#undef LOCTEXT_NAMESPACE
diff --git a/Engine/Plugins/Animation/IKRig/Source/IKRigEditor/Private/RigEditor/SIKRigSkeleton.cpp b/Engine/Plugins/Animation/IKRig/Source/IKRigEditor/Private/RigEditor/SIKRigSkeleton.cpp
new file mode 100644
--- /dev/null
+++ b/Engine/Plugins/Animation/IKRig/Source/IKRigEditor/Private/RigEditor/SIKRigSkeleton.cpp
@@ -0,0 +1,1105 @@
+// Copyright Epic Games, Inc. All Rights Reserved.
+
+#include "RigEditor/SIKRigSkeleton.h"
+
+#include "IKRigSolver.h"
+#include "Engine/SkeletalMesh.h"
+#include "IPersonaToolkit.h"
+#include "SKismetInspector.h"
+#include "Dialogs/Dialogs.h"
+#include "SEditorHeaderButton.h"
+#include "RigEditor/IKRigEditorController.h"
+#include "RigEditor/IKRigEditorStyle.h"
+#include "RigEditor/IKRigSkeletonCommands.h"
+#include "RigEditor/IKRigToolkit.h"
+#include "RigEditor/SIKRigSolverStack.h"
+#include "Widgets/Text/SInlineEditableTextBlock.h"
+#include "Framework/MultiBox/MultiBoxBuilder.h"
+#include "Framework/Commands/UICommandList.h"
+
+#define LOCTEXT_NAMESPACE "SIKRigSkeleton"
+
+FIKRigTreeElement::FIKRigTreeElement(const FName& InKey, IKRigTreeElementType InType)
+{
+	Key = InKey;
+	ElementType = InType;
+}
+
+TSharedRef<ITableRow> FIKRigTreeElement::MakeTreeRowWidget(
+	TSharedRef<FIKRigEditorController> InEditorController,
+	const TSharedRef<STableViewBase>& InOwnerTable,
+	TSharedRef<FIKRigTreeElement> InRigTreeElement,
+	TSharedRef<FUICommandList> InCommandList,
+	TSharedPtr<SIKRigSkeleton> InSkeleton)
+{
+	return SNew(SIKRigSkeletonItem, InEditorController, InOwnerTable, InRigTreeElement, InCommandList, InSkeleton);
+}
+
+void FIKRigTreeElement::RequestRename()
+{
+	OnRenameRequested.ExecuteIfBound();
+}
+
+void SIKRigSkeletonItem::Construct(
+	const FArguments& InArgs,
+	TSharedRef<FIKRigEditorController> InEditorController,
+    const TSharedRef<STableViewBase>& OwnerTable,
+    TSharedRef<FIKRigTreeElement> InRigTreeElement,
+    TSharedRef<FUICommandList> InCommandList,
+    TSharedPtr<SIKRigSkeleton> InSkeleton)
+{
+	WeakRigTreeElement = InRigTreeElement;
+	EditorController = InEditorController;
+
+	// is this element affected by the selected solver?
+	bool bIsConnectedToSelectedSolver;
+	const int32 SelectedSolver = EditorController.Pin()->GetSelectedSolverIndex();
+	if (SelectedSolver == INDEX_NONE)
+	{
+		bIsConnectedToSelectedSolver = EditorController.Pin()->IsElementConnectedToAnySolver(InRigTreeElement);
+	}else
+	{
+		bIsConnectedToSelectedSolver = EditorController.Pin()->IsElementConnectedToSolver(InRigTreeElement, SelectedSolver);
+	}
+	
+	// determine text style
+	FSlateFontInfo TextStyle = FEditorStyle::GetWidgetStyle<FTextBlockStyle>("SkeletonTree.NormalFont").Font;
+	// highlight elements connected to the selected solver
+	if (bIsConnectedToSelectedSolver)
+	{
+		TextStyle = FIKRigEditorStyle::Get().GetWidgetStyle<FTextBlockStyle>("IKRig.Tree.AffectedBoneText").Font;
+	}
+		
+	// determine which icon to use for tree element
+	const FSlateBrush* Brush;
+	switch(InRigTreeElement->ElementType)
+	{
+		case IKRigTreeElementType::BONE:
+			if (bIsConnectedToSelectedSolver)
+			{
+				Brush = FAppStyle::Get().GetBrush("SkeletonTree.Bone");
+			}else
+			{
+				Brush = FAppStyle::Get().GetBrush("SkeletonTree.BoneNonWeighted");
+			}
+			break;
+		case IKRigTreeElementType::BONE_SETTINGS:
+			Brush = FIKRigEditorStyle::Get().GetBrush("IKRig.Tree.BoneWithSettings");
+			break;
+		case IKRigTreeElementType::GOAL:
+			Brush = FIKRigEditorStyle::Get().GetBrush("IKRig.Tree.Goal");
+			break;
+		case IKRigTreeElementType::EFFECTOR:
+			Brush = FIKRigEditorStyle::Get().GetBrush("IKRig.Tree.Effector");
+			break;
+		default:
+			checkNoEntry();
+	}
+	 
+	TSharedPtr<SHorizontalBox> HorizontalBox;
+	STableRow<TSharedPtr<FIKRigTreeElement>>::Construct(
+        STableRow<TSharedPtr<FIKRigTreeElement>>::FArguments()
+        .ShowWires(true)
+        .OnDragDetected(InSkeleton.Get(), &SIKRigSkeleton::OnDragDetected)
+        .OnCanAcceptDrop(InSkeleton.Get(), &SIKRigSkeleton::OnCanAcceptDrop)
+        .OnAcceptDrop(InSkeleton.Get(), &SIKRigSkeleton::OnAcceptDrop)
+        .Content()
+        [
+            SAssignNew(HorizontalBox, SHorizontalBox)
+            + SHorizontalBox::Slot()
+            .MaxWidth(18)
+            .FillWidth(1.0)
+            .HAlign(HAlign_Left)
+            .VAlign(VAlign_Center)
+            [
+                SNew(SImage)
+                .Image(Brush)
+            ]
+        ], OwnerTable);
+	
+	if (InRigTreeElement->ElementType == IKRigTreeElementType::BONE)
+	{
+		HorizontalBox->AddSlot()
+        .AutoWidth()
+        .VAlign(VAlign_Center)
+        [
+			SNew(STextBlock)
+			.Text(this, &SIKRigSkeletonItem::GetName)
+			.Font(TextStyle)
+        ];
+	}else
+	{
+		TSharedPtr< SInlineEditableTextBlock > InlineWidget;
+		HorizontalBox->AddSlot()
+        .AutoWidth()
+        .VAlign(VAlign_Center)
+        [
+	        SAssignNew(InlineWidget, SInlineEditableTextBlock)
+		    .Text(this, &SIKRigSkeletonItem::GetName)
+		    .OnTextCommitted(this, &SIKRigSkeletonItem::OnNameCommitted)
+		    .MultiLine(false)
+        ];
+		InRigTreeElement->OnRenameRequested.BindSP(InlineWidget.Get(), &SInlineEditableTextBlock::EnterEditingMode);
+	}
+}
+
+void SIKRigSkeletonItem::OnNameCommitted(const FText& InText, ETextCommit::Type InCommitType) const
+{
+	check(WeakRigTreeElement.IsValid());
+
+	if (!(InCommitType == ETextCommit::OnEnter || InCommitType == ETextCommit::OnUserMovedFocus))
+	{
+		return; // make sure user actually intends to commit a name change
+	}
+	
+	const FName OldName = WeakRigTreeElement.Pin()->Key;
+	const FName PotentialNewName = FName(InText.ToString());
+	const FName NewName = EditorController.Pin()->AssetController->RenameGoal(OldName, PotentialNewName);
+	if (NewName != NAME_None)
+	{
+		WeakRigTreeElement.Pin()->Key = NewName;
+	}
+
+	EditorController.Pin()->SkeletonView->RefreshTreeView();
+}
+
+FText SIKRigSkeletonItem::GetName() const
+{
+	return (FText::FromName(WeakRigTreeElement.Pin()->Key));
+}
+
+TSharedRef<FIKRigSkeletonDragDropOp> FIKRigSkeletonDragDropOp::New(TWeakPtr<FIKRigTreeElement> InElement)
+{
+	TSharedRef<FIKRigSkeletonDragDropOp> Operation = MakeShared<FIKRigSkeletonDragDropOp>();
+	Operation->Element = InElement;
+	Operation->Construct();
+	return Operation;
+}
+
+TSharedPtr<SWidget> FIKRigSkeletonDragDropOp::GetDefaultDecorator() const
+{
+	return SNew(SBorder)
+        .Visibility(EVisibility::Visible)
+        .BorderImage(FEditorStyle::GetBrush("Menu.Background"))
+        [
+            SNew(STextBlock)
+            .Text(FText::FromString(Element.Pin()->Key.ToString()))
+        ];
+}
+
+void SIKRigSkeleton::Construct(const FArguments& InArgs, TSharedRef<FIKRigEditorController> InEditorController)
+{
+	EditorController = InEditorController;
+	EditorController.Pin()->SkeletonView = SharedThis(this);
+	CommandList = MakeShared<FUICommandList>();
+	BindCommands();
+	
+	ChildSlot
+    [
+        SNew(SVerticalBox)
+        +SVerticalBox::Slot()
+        .AutoHeight()
+        .VAlign(VAlign_Top)
+        .Padding(0.0f)
+        [
+            SNew(SBorder)
+            .Padding(0.0f)
+            .BorderImage(FEditorStyle::GetBrush("DetailsView.CategoryTop"))
+            .BorderBackgroundColor(FLinearColor(0.6f, 0.6f, 0.6f, 1.0f))
+            [
+                SNew(SVerticalBox)
+                + SVerticalBox::Slot()
+                .AutoHeight()
+                .VAlign(VAlign_Top)
+                [
+                    SNew(SHorizontalBox)
+                    .Visibility(this, &SIKRigSkeleton::IsImportButtonVisible)
+                    + SHorizontalBox::Slot()
+                    .VAlign(VAlign_Center)
+                    .HAlign(HAlign_Left)
+                    .FillWidth(1.f)
+                    .Padding(3.0f, 1.0f)
+                    [
+						SNew(SEditorHeaderButton)
+                        .Icon(FAppStyle::Get().GetBrush("Icons.Plus"))
+                        .Text(LOCTEXT("ImportSkeletonLabel", "Import Skeleton"))
+                        .ToolTipText(LOCTEXT("ImportSkeletonToolTip", "Import a skeletal mesh to apply IK to."))
+                        .OnClicked(FOnClicked::CreateSP(this, &SIKRigSkeleton::OnImportSkeletonClicked))
+                    ]
+                ]
+            ]
+        ]
+        +SVerticalBox::Slot()
+        .Padding(0.0f, 0.0f)
+        [
+            SNew(SBorder)
+            .Padding(2.0f)
+            .BorderImage(FEditorStyle::GetBrush("SCSEditor.TreePanel"))
+            [
+                SAssignNew(TreeView, SIKRigSkeletonTreeView)
+                .TreeItemsSource(&RootElements)
+                .SelectionMode(ESelectionMode::Multi)
+                .OnGenerateRow(this, &SIKRigSkeleton::MakeTableRowWidget)
+                .OnGetChildren(this, &SIKRigSkeleton::HandleGetChildrenForTree)
+                .OnSelectionChanged(this, &SIKRigSkeleton::OnSelectionChanged)
+                .OnContextMenuOpening(this, &SIKRigSkeleton::CreateContextMenu)
+                .OnMouseButtonClick(this, &SIKRigSkeleton::OnItemClicked)
+                .OnMouseButtonDoubleClick(this, &SIKRigSkeleton::OnItemDoubleClicked)
+                .OnSetExpansionRecursive(this, &SIKRigSkeleton::OnSetExpansionRecursive)
+                .HighlightParentNodesForSelection(true)
+                .ItemHeight(24)
+            ]
+        ]
+    ];
+
+	const bool IsInitialSetup = true;
+	RefreshTreeView(IsInitialSetup);
+}
+
+void SIKRigSkeleton::SetSelectedGoalsFromViewport(const TArray<FName>& GoalNames)
+{
+	if (GoalNames.IsEmpty())
+	{
+		TreeView->ClearSelection();
+		return;
+	}
+	
+	for (const auto& Item : AllElements)
+	{
+		if (GoalNames.Contains(Item->Key))
+		{
+			TreeView->SetSelection(Item, ESelectInfo::Direct);
+		}
+	}
+}
+
+void SIKRigSkeleton::BindCommands()
+{
+	const FIKRigSkeletonCommands& Commands = FIKRigSkeletonCommands::Get();
+
+	CommandList->MapAction(Commands.NewGoal,
+        FExecuteAction::CreateSP(this, &SIKRigSkeleton::HandleNewGoal),
+        FCanExecuteAction::CreateSP(this, &SIKRigSkeleton::CanAddNewGoal));
+	
+	CommandList->MapAction(Commands.Delete,
+        FExecuteAction::CreateSP(this, &SIKRigSkeleton::HandleDeleteItem),
+        FCanExecuteAction::CreateSP(this, &SIKRigSkeleton::CanDeleteItem));
+
+	CommandList->MapAction(Commands.ConnectGoalToSolvers,
+        FExecuteAction::CreateSP(this, &SIKRigSkeleton::HandleConnectGoalToSolver),
+        FCanExecuteAction::CreateSP(this, &SIKRigSkeleton::CanConnectGoalToSolvers));
+
+	CommandList->MapAction(Commands.DisconnectGoalFromSolvers,
+        FExecuteAction::CreateSP(this, &SIKRigSkeleton::HandleDisconnectGoalFromSolver),
+        FCanExecuteAction::CreateSP(this, &SIKRigSkeleton::CanDisconnectGoalFromSolvers));
+
+	CommandList->MapAction(Commands.SetRootBoneOnSolvers,
+        FExecuteAction::CreateSP(this, &SIKRigSkeleton::HandleSetRootBoneOnSolvers),
+        FCanExecuteAction::CreateSP(this, &SIKRigSkeleton::CanSetRootBoneOnSolvers));
+
+	CommandList->MapAction(Commands.AddBoneSettings,
+        FExecuteAction::CreateSP(this, &SIKRigSkeleton::HandleAddBoneSettings),
+        FCanExecuteAction::CreateSP(this, &SIKRigSkeleton::CanAddBoneSettings));
+
+	CommandList->MapAction(Commands.RemoveBoneSettings,
+        FExecuteAction::CreateSP(this, &SIKRigSkeleton::HandleRemoveBoneSettings),
+        FCanExecuteAction::CreateSP(this, &SIKRigSkeleton::CanRemoveBoneSettings));
+}
+
+void SIKRigSkeleton::FillContextMenu(FMenuBuilder& MenuBuilder)
+{
+	const FIKRigSkeletonCommands& Actions = FIKRigSkeletonCommands::Get();
+
+	const TArray<TSharedPtr<FIKRigTreeElement>> SelectedItems = TreeView->GetSelectedItems();
+	if (SelectedItems.IsEmpty())
+	{
+		return;
+	}
+
+	if (SelectedItems.Num() > 1)
+	{
+		return; // TODO create limb definitions from multiple selected bones
+	}
+
+	MenuBuilder.BeginSection("AddRemoveGoals", LOCTEXT("AddRemoveGoalOperations", "Goals"));
+	MenuBuilder.AddMenuEntry(Actions.NewGoal);
+	MenuBuilder.AddMenuEntry(Actions.Delete);
+	MenuBuilder.EndSection();
+	
+	MenuBuilder.BeginSection("ConnectGoals", LOCTEXT("ConnectGoalOperations", "Connect Goals To Solvers"));
+	MenuBuilder.AddMenuEntry(Actions.ConnectGoalToSolvers);
+	MenuBuilder.AddMenuEntry(Actions.DisconnectGoalFromSolvers);
+	MenuBuilder.EndSection();
+
+	MenuBuilder.BeginSection("BoneSettings", LOCTEXT("BoneSettingsOperations", "Bone Settings"));
+	MenuBuilder.AddMenuEntry(Actions.AddBoneSettings);
+	MenuBuilder.AddMenuEntry(Actions.RemoveBoneSettings);
+	MenuBuilder.EndSection();
+
+	MenuBuilder.BeginSection("RootBone", LOCTEXT("RootBoneOperations", "Root Bone"));
+	MenuBuilder.AddMenuEntry(Actions.SetRootBoneOnSolvers);
+	MenuBuilder.EndSection();
+}
+
+void SIKRigSkeleton::HandleNewGoal()
+{
+	const TArray<TSharedPtr<FIKRigTreeElement>> SelectedItems = TreeView.Get()->GetSelectedItems();
+	if (SelectedItems.IsEmpty())
+	{
+		return; // have to select something
+	}
+
+	TSharedPtr<FIKRigTreeElement> SelectedItem = SelectedItems[0];
+	if (SelectedItem->ElementType != IKRigTreeElementType::BONE)
+	{
+		return; // can only add goals to bones
+	}
+
+	UIKRigController* Controller = EditorController.Pin()->AssetController;
+
+	// build default name for the new goal
+	const FName BoneName = SelectedItem->Key;
+	const FName NewGoalName = FName(BoneName.ToString() + "_Goal");
+
+	// create a new goal
+	UIKRigEffectorGoal* NewGoal = Controller->AddNewGoal(NewGoalName, BoneName);
+
+	if (NewGoal)
+	{
+		// get selected solvers
+		TArray<TSharedPtr<FSolverStackElement>> SelectedSolvers;
+		EditorController.Pin()->GetSelectedSolvers(SelectedSolvers);
+		// connect the new goal to all the selected solvers
+		for (const TSharedPtr<FSolverStackElement>& SolverElement : SelectedSolvers)
+		{
+			Controller->ConnectGoalToSolver(*NewGoal, SolverElement->IndexInStack);	
+		}
+	}
+	
+	// create new goal tree element and update view
+	RefreshTreeView();
+
+	// show goal in details view
+	EditorController.Pin()->ShowDetailsForGoal(NewGoalName);
+}
+
+bool SIKRigSkeleton::CanAddNewGoal()
+{
+	// is anything selected?
+	TArray<TSharedPtr<FIKRigTreeElement>> SelectedItems = TreeView->GetSelectedItems();
+	if (SelectedItems.IsEmpty())
+	{
+		return false;
+	}
+
+	// can only add goals to selected bones
+	for (const auto&Item : SelectedItems)
+	{
+		if (Item->ElementType != IKRigTreeElementType::BONE)
+		{
+			return false;
+		}
+	}
+
+	return true;
+}
+
+void SIKRigSkeleton::HandleDeleteItem()
+{
+	TArray<TSharedPtr<FIKRigTreeElement>> SelectedItems = TreeView->GetSelectedItems();
+	TSharedPtr<FIKRigTreeElement> ParentOfDeletedGoal;
+	for (const auto&Item : SelectedItems)
+	{
+		if (Item->ElementType == IKRigTreeElementType::GOAL)
+		{
+			EditorController.Pin()->AssetController->RemoveGoal(Item->Key);
+			
+		}else if (Item->ElementType == IKRigTreeElementType::EFFECTOR)
+		{
+			EditorController.Pin()->AssetController->DisconnectGoalFromSolver(Item->EffectorGoalName, Item->EffectorSolverIndex);
+		}
+	}
+
+	EditorController.Pin()->ShowEmptyDetails();
+	RefreshTreeView();
+}
+bool SIKRigSkeleton::CanDeleteItem() const
+{
+	// is anything selected?
+	TArray<TSharedPtr<FIKRigTreeElement>> SelectedItems = TreeView->GetSelectedItems();
+	if (SelectedItems.IsEmpty())
+	{
+		return false;
+	}
+
+	// are all selected items goals or effectors?
+	for (const auto&Item : SelectedItems)
+	{
+		const bool bIsGoal = Item->ElementType == IKRigTreeElementType::GOAL;
+		const bool bIsEffector = Item->ElementType == IKRigTreeElementType::EFFECTOR;
+		if (!(bIsGoal || bIsEffector))
+		{
+			return false;
+		}
+	}
+
+	return true;
+}
+
+void SIKRigSkeleton::HandleConnectGoalToSolver()
+{
+	const bool bConnect = true; //connect
+	ConnectSelectedGoalsToSelectedSolvers(bConnect);
+}
+
+void SIKRigSkeleton::HandleDisconnectGoalFromSolver()
+{
+	const bool bConnect = false; //disconnect
+	ConnectSelectedGoalsToSelectedSolvers(bConnect);
+}
+
+bool SIKRigSkeleton::CanConnectGoalToSolvers() const
+{
+	const bool bCountOnlyConnected = false;
+	const int32 NumDisconnected = GetNumSelectedGoalToSolverConnections(bCountOnlyConnected);
+	return NumDisconnected > 0;
+}
+
+bool SIKRigSkeleton::CanDisconnectGoalFromSolvers() const
+{
+	const bool bCountOnlyConnected = true;
+	const int32 NumConnected = GetNumSelectedGoalToSolverConnections(bCountOnlyConnected);
+	return NumConnected > 0;
+}
+
+void SIKRigSkeleton::ConnectSelectedGoalsToSelectedSolvers(bool bConnect)
+{
+	TArray<TSharedPtr<FIKRigTreeElement>> SelectedGoals;
+	GetSelectedGoals(SelectedGoals);
+	TArray<TSharedPtr<FSolverStackElement>> SelectedSolvers;
+	EditorController.Pin()->GetSelectedSolvers(SelectedSolvers);
+
+	UIKRigController* Controller = EditorController.Pin()->AssetController;
+	for (const TSharedPtr<FIKRigTreeElement>& GoalElement : SelectedGoals)
+	{
+		const FName GoalName = GoalElement->Key;
+		const int32 GoalIndex = Controller->GetGoalIndex(GoalName);
+		check(GoalIndex != INDEX_NONE);
+		const UIKRigEffectorGoal& EffectorGoal = *Controller->GetGoal(GoalIndex);
+		for (const TSharedPtr<FSolverStackElement>& SolverElement : SelectedSolvers)
+		{
+			if (bConnect)
+			{
+				Controller->ConnectGoalToSolver(EffectorGoal, SolverElement->IndexInStack);	
+			}else
+			{
+				Controller->DisconnectGoalFromSolver(EffectorGoal.GoalName, SolverElement->IndexInStack);	
+			}
+		}
+	}
+
+	// add/remove new effector under goal in skeleton view
+	RefreshTreeView();
+}
+
+int32 SIKRigSkeleton::GetNumSelectedGoalToSolverConnections(bool bCountOnlyConnected) const
+{
+	TArray<TSharedPtr<FIKRigTreeElement>> SelectedGoals;
+	GetSelectedGoals(SelectedGoals);
+	TArray<TSharedPtr<FSolverStackElement>> SelectedSolvers;
+	EditorController.Pin()->GetSelectedSolvers(SelectedSolvers);
+
+	int32 NumMatched = 0;
+	for (const TSharedPtr<FIKRigTreeElement>& Goal : SelectedGoals)
+	{
+		for (const TSharedPtr<FSolverStackElement>& Solver : SelectedSolvers)
+		{
+			const bool bIsConnected = EditorController.Pin()->AssetController->IsGoalConnectedToSolver(Goal->Key, Solver->IndexInStack);
+			if (bIsConnected == bCountOnlyConnected)
+			{
+				++NumMatched;
+			}
+		}
+	}
+
+	return NumMatched;
+}
+
+void SIKRigSkeleton::HandleSetRootBoneOnSolvers()
+{
+    // get name of selected root bone
+    TArray<TSharedPtr<FIKRigTreeElement>> SelectedBones;
+	GetSelectedBones(SelectedBones);
+	const FName RootBoneName = SelectedBones[0]->Key;
+
+	// apply to all selected solvers (ignored on solvers that don't accept a root bone)
+	UIKRigController* Controller = EditorController.Pin()->AssetController;
+	TArray<TSharedPtr<FSolverStackElement>> SelectedSolvers;
+	EditorController.Pin()->GetSelectedSolvers(SelectedSolvers);
+	int32 SolverToShow = 0;
+	for (const auto& Solver : SelectedSolvers)
+	{
+		Controller->SetRootBone(RootBoneName, Solver->IndexInStack);
+		SolverToShow = Solver->IndexInStack;
+	}
+
+	// show solver that had it's root bone updated
+	EditorController.Pin()->ShowDetailsForSolver(SolverToShow);
+	
+	// show new icon when bone has settings applied
+	RefreshTreeView();
+}
+
+bool SIKRigSkeleton::CanSetRootBoneOnSolvers()
+{
+	// must have at least 1 bone selected
+	TArray<TSharedPtr<FIKRigTreeElement>> SelectedBones;
+	GetSelectedBones(SelectedBones);
+	if (SelectedBones.Num() != 1)
+	{
+		return false;
+	}
+
+	// must have at least 1 solver selected that accepts root bones
+	UIKRigController* Controller = EditorController.Pin()->AssetController;
+	TArray<TSharedPtr<FSolverStackElement>> SelectedSolvers;
+	EditorController.Pin()->GetSelectedSolvers(SelectedSolvers);
+	for (const auto& Solver : SelectedSolvers)
+	{
+		if (Controller->GetSolver(Solver->IndexInStack)->CanSetRootBone())
+		{
+			return true;
+		}
+	}
+	
+	return false;
+}
+
+void SIKRigSkeleton::HandleAddBoneSettings()
+{
+	// get selected bones
+	TArray<TSharedPtr<FIKRigTreeElement>> SelectedBones;
+	GetSelectedBones(SelectedBones);
+	
+	// add settings for bone on all selected solvers (ignored if already present)
+	UIKRigController* Controller = EditorController.Pin()->AssetController;
+	TArray<TSharedPtr<FSolverStackElement>> SelectedSolvers;
+	EditorController.Pin()->GetSelectedSolvers(SelectedSolvers);
+	FName BoneNameForSettings;
+	int32 SolverIndex = INDEX_NONE;
+	for (const TSharedPtr<FIKRigTreeElement>& BoneItem : SelectedBones)
+	{
+		for (const TSharedPtr<FSolverStackElement>& Solver : SelectedSolvers)
+		{
+			Controller->AddBoneSetting(BoneItem->Key, Solver->IndexInStack);
+			BoneNameForSettings = BoneItem->Key;
+			SolverIndex = Solver->IndexInStack;
+        }
+	}
+
+	EditorController.Pin()->ShowDetailsForBoneSettings(BoneNameForSettings, SolverIndex);
+
+	// show new icon when bone has settings applied
+	RefreshTreeView();
+}
+
+bool SIKRigSkeleton::CanAddBoneSettings()
+{
+    // must have at least 1 bone selected
+    TArray<TSharedPtr<FIKRigTreeElement>> SelectedBones;
+	GetSelectedBones(SelectedBones);
+	if (SelectedBones.IsEmpty())
+	{
+		return false;
+	}
+
+	// must have at least 1 solver selected that does not already have a bone setting for the selected bones
+	UIKRigController* Controller = EditorController.Pin()->AssetController;
+	TArray<TSharedPtr<FSolverStackElement>> SelectedSolvers;
+	EditorController.Pin()->GetSelectedSolvers(SelectedSolvers);
+	for (const TSharedPtr<FIKRigTreeElement>& BoneItem : SelectedBones)
+	{
+		for (const TSharedPtr<FSolverStackElement>& Solver : SelectedSolvers)
+		{
+			if (Controller->CanAddBoneSetting(BoneItem->Key, Solver->IndexInStack))
+			{
+				return true;
+			}
+        }
+	}
+	
+	return false;
+}
+
+void SIKRigSkeleton::HandleRemoveBoneSettings()
+{
+	// get selected bones
+	TArray<TSharedPtr<FIKRigTreeElement>> SelectedBones;
+	GetSelectedBones(SelectedBones);
+	
+	// add settings for bone on all selected solvers (ignored if already present)
+	UIKRigController* Controller = EditorController.Pin()->AssetController;
+	TArray<TSharedPtr<FSolverStackElement>> SelectedSolvers;
+	EditorController.Pin()->GetSelectedSolvers(SelectedSolvers);
+	FName BoneToShowInDetailsView;
+	for (const TSharedPtr<FIKRigTreeElement>& BoneItem : SelectedBones)
+	{
+		for (const TSharedPtr<FSolverStackElement>& Solver : SelectedSolvers)
+		{
+			Controller->RemoveBoneSetting(BoneItem->Key, Solver->IndexInStack);
+			BoneToShowInDetailsView = BoneItem->Key;
+		}
+	}
+
+	EditorController.Pin()->ShowDetailsForBone(BoneToShowInDetailsView);
+
+	// show new icon when bone has settings applied
+	RefreshTreeView();
+}
+
+bool SIKRigSkeleton::CanRemoveBoneSettings()
+{
+    // must have at least 1 bone selected
+    TArray<TSharedPtr<FIKRigTreeElement>> SelectedBones;
+	GetSelectedBones(SelectedBones);
+	if (SelectedBones.IsEmpty())
+	{
+		return false;
+	}
+
+	// must have at least 1 solver selected that has a bone setting for 1 of the selected bones
+	UIKRigController* Controller = EditorController.Pin()->AssetController;
+	TArray<TSharedPtr<FSolverStackElement>> SelectedSolvers;
+	EditorController.Pin()->GetSelectedSolvers(SelectedSolvers);
+	for (const TSharedPtr<FIKRigTreeElement>& BoneItem : SelectedBones)
+	{
+		for (const TSharedPtr<FSolverStackElement>& Solver : SelectedSolvers)
+		{
+			if (Controller->CanRemoveBoneSetting(BoneItem->Key, Solver->IndexInStack))
+			{
+				return true;
+			}
+		}
+	}
+	
+	return false;
+}
+
+void SIKRigSkeleton::GetSelectedBones(TArray<TSharedPtr<FIKRigTreeElement>>& OutBoneItems)
+{
+	const TArray<TSharedPtr<FIKRigTreeElement>> SelectedItems = TreeView->GetSelectedItems();
+	for (const auto& Item : SelectedItems)
+	{
+		if (Item->ElementType == IKRigTreeElementType::BONE)
+		{
+			OutBoneItems.Add(Item);
+        }
+	}
+}
+
+void SIKRigSkeleton::GetSelectedGoals(TArray<TSharedPtr<FIKRigTreeElement>>& OutSelectedGoals) const
+{
+	OutSelectedGoals.Reset();
+	TArray<TSharedPtr<FIKRigTreeElement>> SelectedItems = TreeView->GetSelectedItems();
+	for (const auto& Item : SelectedItems)
+	{
+		if (Item->ElementType == IKRigTreeElementType::GOAL)
+		{
+			OutSelectedGoals.Add(Item);
+		}
+	}
+}
+
+void SIKRigSkeleton::HandleRenameElement() const
+{
+	TArray<TSharedPtr<FIKRigTreeElement>> SelectedGoals;
+	GetSelectedGoals(SelectedGoals);
+	if (SelectedGoals.Num() != 1)
+	{
+		return;
+	}
+	
+	SelectedGoals[0]->RequestRename();
+}
+
+FReply SIKRigSkeleton::OnImportSkeletonClicked()
+{
+	FIKRigSkeletonImportSettings Settings;
+	const TSharedPtr<FStructOnScope> StructToDisplay = MakeShareable(new FStructOnScope(FIKRigSkeletonImportSettings::StaticStruct(), (uint8*)&Settings));
+	TSharedRef<SKismetInspector> KismetInspector = SNew(SKismetInspector);
+	KismetInspector->ShowSingleStruct(StructToDisplay);
+
+	SGenericDialogWidget::OpenDialog(LOCTEXT("IKRigSkeletonImport", "Import Skeleton"), KismetInspector, SGenericDialogWidget::FArguments(), true);
+
+	if (Settings.Mesh != nullptr)
+	{
+		ImportSkeleton(FAssetData(Settings.Mesh));
+		EditorController.Pin()->PromptToAddSolverAtStartup();
+	}
+	
+	return FReply::Handled();
+}
+
+EVisibility SIKRigSkeleton::IsImportButtonVisible() const
+{
+	if (UIKRigController* Controller = EditorController.Pin()->AssetController)
+	{
+		if (Controller->GetSkeleton().BoneNames.Num() > 0)
+		{
+			return EVisibility::Collapsed;
+		}
+	}
+	
+	return EVisibility::Visible;
+}
+
+void SIKRigSkeleton::RefreshTreeView(bool IsInitialSetup)
+{
+	// save expansion state
+	TreeView->SaveAndClearSparseItemInfos();
+
+	// reset all tree items
+	RootElements.Reset();
+	AllElements.Reset();
+
+	// validate we have a skeleton to load
+	UIKRigController* Controller = EditorController.Pin()->AssetController;
+	FIKRigSkeleton& Skeleton = Controller->GetSkeleton();
+	if (Skeleton.BoneNames.IsEmpty())
+	{
+		TreeView->RequestTreeRefresh();
+		return;
+	}
+
+	// create all bone elements
+	for (const FName BoneName : Skeleton.BoneNames)
+	{
+		TSharedPtr<FIKRigTreeElement> NewItem = MakeShared<FIKRigTreeElement>(BoneName, IKRigTreeElementType::BONE);
+		AllElements.Add(NewItem);
+	}
+
+	// store children/parent pointers on all bone elements
+	for (int32 BoneIndex=0; BoneIndex<Skeleton.BoneNames.Num(); ++BoneIndex)
+	{
+		const int32 ParentIndex = Skeleton.ParentIndices[BoneIndex];
+		if (ParentIndex < 0)
+		{
+			// store root element
+			RootElements.Add(AllElements[BoneIndex]);
+		}else
+		{
+			// store pointer to child on parent
+			AllElements[ParentIndex]->Children.Add(AllElements[BoneIndex]);
+			// store pointer to parent on child
+			AllElements[BoneIndex]->Parent = AllElements[ParentIndex];
+		}
+	}
+
+	// create all bone settings elements
+	for (int32 SolverIndex=0; SolverIndex<Controller->GetNumSolvers(); ++SolverIndex)
+	{
+		UIKRigSolver* Solver = Controller->GetSolver(SolverIndex);
+		if (!Solver->UsesBoneSettings())
+		{
+			continue;
+		}
+		for (int32 BoneIndex=0; BoneIndex<Skeleton.BoneNames.Num(); ++BoneIndex)
+		{
+			if (Solver->GetBoneSetting(Skeleton.BoneNames[BoneIndex]))
+			{
+				const FName DisplayName = FName("Bone Settings for: " + Solver->GetName());
+				TSharedPtr<FIKRigTreeElement> SettingsItem = MakeShared<FIKRigTreeElement>(DisplayName, IKRigTreeElementType::BONE_SETTINGS);
+				SettingsItem->BoneSettingBoneName = Skeleton.BoneNames[BoneIndex];
+				SettingsItem->BoneSettingsSolverIndex = SolverIndex;
+				AllElements.Add(SettingsItem);
+				// store hierarchy pointers for item
+				AllElements[BoneIndex]->Children.Add(SettingsItem);
+				SettingsItem->Parent = AllElements[BoneIndex];
+			}
+		}
+    }
+	
+	// create all goal and effector elements
+	TArray<UIKRigEffectorGoal*> Goals = Controller->GetAllGoals();
+	for (const UIKRigEffectorGoal* Goal : Goals)
+	{
+		// make new element for goal
+		const int32 BoneIndex = Skeleton.GetBoneIndexFromName(Goal->BoneName);
+		check(BoneIndex!=INDEX_NONE);
+		TSharedPtr<FIKRigTreeElement> GoalItem = MakeShared<FIKRigTreeElement>(Goal->GoalName, IKRigTreeElementType::GOAL);
+		const int32 GoalItemIndex = AllElements.Add(GoalItem);
+
+		// store hierarchy pointers for goal
+		AllElements[BoneIndex]->Children.Add(GoalItem);
+		GoalItem->Parent = AllElements[BoneIndex];
+
+		// add all effectors connected to this goal
+		for (int32 SolverIndex=0; SolverIndex<Controller->GetNumSolvers(); ++SolverIndex)
+		{
+			if (UObject* Effector = Controller->GetEffectorForGoal(Goal->GoalName, SolverIndex))
+			{
+				// make new element for effector
+				const UIKRigSolver* Solver = Controller->GetSolver(SolverIndex);
+				const FName DisplayName = FName("Effector for: " + Solver->GetName());
+				TSharedPtr<FIKRigTreeElement> EffectorItem = MakeShared<FIKRigTreeElement>(DisplayName, IKRigTreeElementType::EFFECTOR);
+				EffectorItem->EffectorSolverIndex = SolverIndex;
+				EffectorItem->EffectorGoalName = Goal->GoalName;
+				AllElements.Add(EffectorItem);
+				EffectorItem->Parent = GoalItem;
+				GoalItem->Children.Add(EffectorItem);
+			}
+		}
+	}
+
+	// restore expansion state
+	for (const TSharedPtr<FIKRigTreeElement>& Element : AllElements)
+	{
+		TreeView->RestoreSparseItemInfos(Element);
+	}
+
+	// expand all elements upon the initial construction of the tree
+	if (IsInitialSetup)
+	{
+		for (TSharedPtr<FIKRigTreeElement> RootElement : RootElements)
+		{
+			SetExpansionRecursive(RootElement, false, true);
+		}
+	}
+	
+	TreeView->RequestTreeRefresh();
+}
+
+TSharedRef<ITableRow> SIKRigSkeleton::MakeTableRowWidget(
+	TSharedPtr<FIKRigTreeElement> InItem,
+	const TSharedRef<STableViewBase>& OwnerTable)
+{
+	return InItem->MakeTreeRowWidget(EditorController.Pin().ToSharedRef(), OwnerTable, InItem.ToSharedRef(), CommandList.ToSharedRef(), SharedThis(this));
+}
+
+void SIKRigSkeleton::HandleGetChildrenForTree(
+	TSharedPtr<FIKRigTreeElement> InItem,
+	TArray<TSharedPtr<FIKRigTreeElement>>& OutChildren)
+{
+	OutChildren = InItem.Get()->Children;
+}
+
+void SIKRigSkeleton::OnSelectionChanged(TSharedPtr<FIKRigTreeElement> Selection, ESelectInfo::Type SelectInfo)
+{
+	// gate any selection changes NOT made by user clicking mouse
+	if (SelectInfo == ESelectInfo::OnMouseClick)
+	{
+		TArray<TSharedPtr<FIKRigTreeElement>> SelectedGoals;
+		GetSelectedGoals(SelectedGoals);
+		TArray<FName> SelectedGoalNames;
+		for (const auto& Goal : SelectedGoals)
+		{
+			SelectedGoalNames.Add(Goal->Key);
+		}
+		EditorController.Pin()->HandleGoalsSelectedInTreeView(SelectedGoalNames);
+	}
+}
+
+TSharedPtr<SWidget> SIKRigSkeleton::CreateContextMenu()
+{
+	const bool CloseAfterSelection = true;
+	FMenuBuilder MenuBuilder(CloseAfterSelection, CommandList);
+	FillContextMenu(MenuBuilder);
+	return MenuBuilder.MakeWidget();
+}
+
+void SIKRigSkeleton::OnItemClicked(TSharedPtr<FIKRigTreeElement> InItem)
+{
+	// update details view
+	if (InItem->ElementType == IKRigTreeElementType::BONE)
+	{
+		EditorController.Pin()->ShowDetailsForBone(InItem->Key);
+	}else if (InItem->ElementType == IKRigTreeElementType::GOAL)
+	{
+		EditorController.Pin()->ShowDetailsForGoal(InItem->Key);
+	}else if (InItem->ElementType == IKRigTreeElementType::EFFECTOR)
+	{
+		EditorController.Pin()->ShowDetailsForEffector(InItem->EffectorGoalName, InItem->EffectorSolverIndex);
+	}else if (InItem->ElementType == IKRigTreeElementType::BONE_SETTINGS)
+	{
+		EditorController.Pin()->ShowDetailsForBoneSettings(InItem->BoneSettingBoneName, InItem->BoneSettingsSolverIndex);
+	}
+
+	// to rename an item, you have to select it first, then click on it again within a time limit (slow double click)
+	const bool ClickedOnSameItem = TreeView->LastSelected.Pin().Get() == InItem.Get();
+	const uint32 CurrentCycles = FPlatformTime::Cycles();
+	const double SecondsPassed = static_cast<double>(CurrentCycles - TreeView->LastClickCycles) * FPlatformTime::GetSecondsPerCycle();
+	if (ClickedOnSameItem && SecondsPassed > 0.25f)
+	{
+		RegisterActiveTimer(0.f,
+            FWidgetActiveTimerDelegate::CreateLambda([this](double, float) {
+            HandleRenameElement();
+            return EActiveTimerReturnType::Stop;
+        }));
+	}
+
+	TreeView->LastClickCycles = CurrentCycles;
+	TreeView->LastSelected = InItem;
+}
+
+void SIKRigSkeleton::OnItemDoubleClicked(TSharedPtr<FIKRigTreeElement> InItem)
+{
+	if (TreeView->IsItemExpanded(InItem))
+	{
+		SetExpansionRecursive(InItem, false, false);
+	}
+	else
+	{
+		SetExpansionRecursive(InItem, false, true);
+	}
+}
+
+void SIKRigSkeleton::OnSetExpansionRecursive(TSharedPtr<FIKRigTreeElement> InItem, bool bShouldBeExpanded)
+{
+	SetExpansionRecursive(InItem, false, bShouldBeExpanded);
+}
+
+void SIKRigSkeleton::SetExpansionRecursive(
+	TSharedPtr<FIKRigTreeElement> InElement,
+	bool bTowardsParent,
+    bool bShouldBeExpanded)
+{
+	TreeView->SetItemExpansion(InElement, bShouldBeExpanded);
+    
+    if (bTowardsParent)
+    {
+    	if (InElement->Parent.Get())
+    	{
+    		SetExpansionRecursive(InElement->Parent, bTowardsParent, bShouldBeExpanded);
+    	}
+    }
+    else
+    {
+    	for (int32 ChildIndex = 0; ChildIndex < InElement->Children.Num(); ++ChildIndex)
+    	{
+    		SetExpansionRecursive(InElement->Children[ChildIndex], bTowardsParent, bShouldBeExpanded);
+    	}
+    }
+}
+
+void SIKRigSkeleton::ImportSkeleton(const FAssetData& InAssetData)
+{
+	USkeletalMesh* Mesh = Cast<USkeletalMesh>(InAssetData.GetAsset());
+	if (!Mesh)
+	{
+		return;
+	}
+	
+	// load the skeletal mesh in the viewport
+	TSharedRef<IPersonaToolkit> Persona = EditorController.Pin()->EditorToolkit.Pin()->GetPersonaToolkit();
+	Persona->SetPreviewMesh(Mesh, true);
+	Persona->GetPreviewScene()->FocusViews();
+
+	// import the skeleton data into the IKRig asset
+	EditorController.Pin()->AssetController->SetSourceSkeletalMesh(Mesh, true);
+
+	// update the bone hierarchy view
+	RefreshTreeView(true);
+}
+
+FReply SIKRigSkeleton::OnDragDetected(const FGeometry& MyGeometry, const FPointerEvent& MouseEvent)
+{
+	const TArray<TSharedPtr<FIKRigTreeElement>> SelectedItems = TreeView.Get()->GetSelectedItems();
+	if (SelectedItems.Num() != 1)
+	{
+		return FReply::Unhandled();
+	}
+
+	const TSharedPtr<FIKRigTreeElement> DraggedElement = SelectedItems[0];
+	if (DraggedElement->ElementType != IKRigTreeElementType::GOAL)
+	{
+		return FReply::Unhandled();
+	}
+	
+	if (MouseEvent.IsMouseButtonDown(EKeys::LeftMouseButton))
+	{
+		const TSharedRef<FIKRigSkeletonDragDropOp> DragDropOp = FIKRigSkeletonDragDropOp::New(DraggedElement);
+		return FReply::Handled().BeginDragDrop(DragDropOp);
+	}
+
+	return FReply::Unhandled();
+}
+
+TOptional<EItemDropZone> SIKRigSkeleton::OnCanAcceptDrop(
+	const FDragDropEvent& DragDropEvent,
+	EItemDropZone DropZone,
+    TSharedPtr<FIKRigTreeElement> TargetItem)
+{
+	TOptional<EItemDropZone> ReturnedDropZone;
+	
+	const TSharedPtr<FIKRigSkeletonDragDropOp> DragDropOp = DragDropEvent.GetOperationAs<FIKRigSkeletonDragDropOp>();
+	if (DragDropOp.IsValid())
+	{
+		if (TargetItem.Get()->ElementType == IKRigTreeElementType::BONE)
+        {
+        	ReturnedDropZone = EItemDropZone::BelowItem;	
+        }
+	}
+	
+	return ReturnedDropZone;
+}
+
+FReply SIKRigSkeleton::OnAcceptDrop(
+	const FDragDropEvent& DragDropEvent,
+	EItemDropZone DropZone,
+    TSharedPtr<FIKRigTreeElement> TargetItem)
+{
+	const TSharedPtr<FIKRigSkeletonDragDropOp> DragDropOp = DragDropEvent.GetOperationAs<FIKRigSkeletonDragDropOp>();
+	if (!DragDropOp.IsValid())
+	{
+		return FReply::Unhandled();
+	}
+
+	const FIKRigTreeElement& DraggedElement = *DragDropOp.Get()->Element.Pin().Get();
+	UIKRigController* Controller = EditorController.Pin()->AssetController;
+	const bool bWasReparented = Controller->SetGoalBone(DraggedElement.Key, TargetItem.Get()->Key);
+	if (bWasReparented)
+	{
+		RefreshTreeView();
+	}
+	
+	return FReply::Handled();
+}
+
+FReply SIKRigSkeleton::OnKeyDown(const FGeometry& MyGeometry, const FKeyEvent& InKeyEvent)
+{
+	const FKey Key = InKeyEvent.GetKey();
+
+	// handle deleting selected goals
+	if (Key == EKeys::Delete)
+	{
+		TArray<TSharedPtr<FIKRigTreeElement>> SelectedItems = TreeView->GetSelectedItems();
+		for (const auto& SelectedItem : SelectedItems)
+		{
+			if (SelectedItem->ElementType == IKRigTreeElementType::GOAL)
+			{
+				UIKRigController* Controller = EditorController.Pin()->AssetController;
+				Controller->RemoveGoal(SelectedItem->Key);
+			}
+
+			if (SelectedItem->ElementType == IKRigTreeElementType::EFFECTOR)
+			{
+				UIKRigController* Controller = EditorController.Pin()->AssetController;
+				Controller->DisconnectGoalFromSolver(SelectedItem->EffectorGoalName, SelectedItem->EffectorSolverIndex);
+			}
+			
+			if (SelectedItem->ElementType == IKRigTreeElementType::BONE_SETTINGS)
+			{
+				UIKRigController* Controller = EditorController.Pin()->AssetController;
+				Controller->RemoveBoneSetting(SelectedItem->BoneSettingBoneName, SelectedItem->BoneSettingsSolverIndex);
+			}
+		}
+
+		RefreshTreeView();
+		
+		return FReply::Handled();
+	}
+	
+	return FReply::Unhandled();
+}
+
+#undef LOCTEXT_NAMESPACE
\ No newline at end of file
diff --git a/Engine/Plugins/Animation/IKRig/Source/IKRigEditor/Private/RigEditor/SIKRigSolverStack.cpp b/Engine/Plugins/Animation/IKRig/Source/IKRigEditor/Private/RigEditor/SIKRigSolverStack.cpp
new file mode 100644
--- /dev/null
+++ b/Engine/Plugins/Animation/IKRig/Source/IKRigEditor/Private/RigEditor/SIKRigSolverStack.cpp
@@ -0,0 +1,324 @@
+// Copyright Epic Games, Inc. All Rights Reserved.
+
+#include "RigEditor/SIKRigSolverStack.h"
+
+#include "RigEditor/IKRigEditorStyle.h"
+#include "RigEditor/IKRigToolkit.h"
+#include "RigEditor/IKRigEditorController.h"
+#include "IKRigSolver.h"
+
+#include "Framework/Commands/UICommandList.h"
+#include "Framework/MultiBox/MultiBoxBuilder.h"
+#include "IPersonaToolkit.h"
+#include "SKismetInspector.h"
+#include "Dialogs/Dialogs.h"
+#include "Widgets/Text/SInlineEditableTextBlock.h"
+#include "SEditorHeaderButton.h"
+
+#define LOCTEXT_NAMESPACE "SIKRigSolverStack"
+
+TSharedRef<ITableRow> FSolverStackElement::MakeListRowWidget(
+	const TSharedRef<STableViewBase>& InOwnerTable,
+	TSharedRef<FSolverStackElement> InStackElement,
+	TSharedPtr<SIKRigSolverStack> InSolverStack)
+{
+	return SNew(SIKRigSolverStackItem, InOwnerTable, InStackElement, InSolverStack);
+}
+
+void SIKRigSolverStackItem::Construct(
+	const FArguments& InArgs,
+	const TSharedRef<STableViewBase>& OwnerTable,
+	TSharedRef<FSolverStackElement> InStackElement,
+	TSharedPtr<SIKRigSolverStack> InSolverStack)
+{
+	STableRow<TSharedPtr<FSolverStackElement>>::Construct(
+        STableRow<TSharedPtr<FSolverStackElement>>::FArguments()
+        .OnDragDetected(InSolverStack.Get(), &SIKRigSolverStack::OnDragDetected)
+        .OnCanAcceptDrop(InSolverStack.Get(), &SIKRigSolverStack::OnCanAcceptDrop)
+        .OnAcceptDrop(InSolverStack.Get(), &SIKRigSolverStack::OnAcceptDrop)
+        .Content()
+        [
+            SNew(SHorizontalBox)
+            + SHorizontalBox::Slot()
+            .MaxWidth(18)
+			.FillWidth(1.0)
+			.HAlign(HAlign_Left)
+			.VAlign(VAlign_Center)
+            [
+                SNew(SImage)
+                .Image(FIKRigEditorStyle::Get().GetBrush("IKRig.DragSolver"))
+            ]
+     
+			+ SHorizontalBox::Slot()
+            .VAlign(VAlign_Fill)
+            .HAlign(HAlign_Fill)
+            .Padding(3.0f, 1.0f)
+            [
+				SNew(STextBlock)
+				.Text(InStackElement->DisplayName)
+				.TextStyle(FEditorStyle::Get(), "NormalText.Important")
+            ]
+        ], OwnerTable);
+}
+
+TSharedRef<FIKRigSolverStackDragDropOp> FIKRigSolverStackDragDropOp::New(TWeakPtr<FSolverStackElement> InElement)
+{
+	TSharedRef<FIKRigSolverStackDragDropOp> Operation = MakeShared<FIKRigSolverStackDragDropOp>();
+	Operation->Element = InElement;
+	Operation->Construct();
+	return Operation;
+}
+
+TSharedPtr<SWidget> FIKRigSolverStackDragDropOp::GetDefaultDecorator() const
+{
+	return SNew(SBorder)
+        .Visibility(EVisibility::Visible)
+        .BorderImage(FEditorStyle::GetBrush("Menu.Background"))
+        [
+            SNew(STextBlock)
+            .Text(Element.Pin()->DisplayName)
+        ];
+}
+
+SIKRigSolverStack::~SIKRigSolverStack()
+{
+	
+}
+
+void SIKRigSolverStack::Construct(const FArguments& InArgs, TSharedRef<FIKRigEditorController> InEditorController)
+{
+	EditorController = InEditorController;
+	EditorController.Pin()->SolverStackView = SharedThis(this);
+	
+	CommandList = MakeShared<FUICommandList>();
+
+	ChildSlot
+    [
+        SNew(SVerticalBox)
+        +SVerticalBox::Slot()
+        .AutoHeight()
+        .VAlign(VAlign_Top)
+        .Padding(0.0f)
+        [
+            SNew(SBorder)
+            .Padding(0.0f)
+            .BorderImage(FEditorStyle::GetBrush("DetailsView.CategoryTop"))
+            .BorderBackgroundColor(FLinearColor(0.6f, 0.6f, 0.6f, 1.0f))
+            [
+                SNew(SVerticalBox)
+                + SVerticalBox::Slot()
+                .AutoHeight()
+                .VAlign(VAlign_Top)
+                [
+                    SNew(SHorizontalBox)
+                    + SHorizontalBox::Slot()
+                    .VAlign(VAlign_Center)
+                    .HAlign(HAlign_Left)
+                    .FillWidth(1.f)
+                    .Padding(3.0f, 1.0f)
+                    [
+	                    SNew(SEditorHeaderButton)
+				        .Icon(FAppStyle::Get().GetBrush("Icons.Plus"))
+				        .Text(LOCTEXT("AddNewSolverLabel", "Add New Solver"))
+				        .ToolTipText(LOCTEXT("AddNewToolTip", "Add a new IK solver to the rig."))
+				        .IsEnabled(this, &SIKRigSolverStack::IsAddSolverEnabled)
+				        .OnGetMenuContent(this, &SIKRigSolverStack::CreateAddNewMenuWidget)
+                    ]
+                ]
+            ]
+        ]
+
+        +SVerticalBox::Slot()
+        .Padding(0.0f, 0.0f)
+        [
+			SAssignNew( ListView, SSolverStackListViewType )
+			.SelectionMode(ESelectionMode::Multi)
+			.IsEnabled(this, &SIKRigSolverStack::IsAddSolverEnabled)
+			.ListItemsSource( &ListViewItems )
+			.OnGenerateRow( this, &SIKRigSolverStack::MakeListRowWidget )
+			.OnMouseButtonClick(this, &SIKRigSolverStack::OnItemClicked)
+        ]
+    ];
+
+	RefreshStackView();
+}
+
+TSharedRef<SWidget> SIKRigSolverStack::CreateAddNewMenuWidget()
+{
+	const bool bShouldCloseWindowAfterMenuSelection = true;
+	FMenuBuilder MenuBuilder(bShouldCloseWindowAfterMenuSelection, CommandList);
+
+	BuildAddNewMenu(MenuBuilder);
+
+	return MenuBuilder.MakeWidget();
+}
+
+void SIKRigSolverStack::BuildAddNewMenu(FMenuBuilder& MenuBuilder)
+{
+	MenuBuilder.BeginSection("AddNewSolver", LOCTEXT("AddOperations", "Add New Solver"));
+	
+	// add menu option to create each solver type
+	TArray<UClass*> SolverClasses;
+	for(TObjectIterator<UClass> ClassIt; ClassIt; ++ClassIt)
+	{
+		UClass* Class = *ClassIt;
+
+		if(Class->IsChildOf(UIKRigSolver::StaticClass()) && !Class->HasAnyClassFlags(CLASS_Abstract))
+		{
+			FUIAction Action = FUIAction( FExecuteAction::CreateSP(this, &SIKRigSolverStack::AddNewSolver, Class));
+			MenuBuilder.AddMenuEntry(FText::FromString(Class->GetName()), FText::GetEmpty(), FSlateIcon(), Action);
+		}
+	}
+	
+	MenuBuilder.EndSection();
+}
+
+bool SIKRigSolverStack::IsAddSolverEnabled() const
+{
+	if (UIKRigController* Controller = EditorController.Pin()->AssetController)
+	{
+		if (Controller->GetSkeleton().BoneNames.Num() > 0)
+		{
+			return true;
+		}
+	}
+	
+	return false;
+}
+
+void SIKRigSolverStack::AddNewSolver(UClass* Class)
+{
+	if (UIKRigController* Controller = EditorController.Pin()->AssetController)
+	{
+		// add the solver
+		const int32 NewSolverIndex = Controller->AddSolver(Class);
+	}
+
+	RefreshStackView();
+	EditorController.Pin()->SkeletonView->RefreshTreeView(); // update solver indices in effector items
+}
+
+void SIKRigSolverStack::RefreshStackView()
+{
+	ListViewItems.Reset();
+	UIKRigController* Controller = EditorController.Pin()->AssetController;
+	const int32 NumSolvers = Controller->GetNumSolvers();
+	for (int32 i=0; i<NumSolvers; ++i)
+	{
+		UIKRigSolver* Solver = Controller->GetSolver(i);
+		const FText DisplayName = Solver ? FText::FromString(Solver->GetName()) : FText::FromString("Unknown Solver");
+		TSharedPtr<FSolverStackElement> SolverItem = FSolverStackElement::Make(DisplayName, i);
+		ListViewItems.Add(SolverItem);
+	}
+
+	// select first item if none others selected
+	if (ListViewItems.Num() > 0 && ListView->GetNumItemsSelected() == 0)
+	{
+		ListView->SetSelection(ListViewItems[0]);
+	}
+
+	ListView->RequestListRefresh();
+}
+
+TSharedRef<ITableRow> SIKRigSolverStack::MakeListRowWidget(
+	TSharedPtr<FSolverStackElement> InElement,
+    const TSharedRef<STableViewBase>& OwnerTable)
+{
+	return InElement->MakeListRowWidget(
+		OwnerTable,
+		InElement.ToSharedRef(),
+		SharedThis(this));
+}
+
+FReply SIKRigSolverStack::OnDragDetected(
+	const FGeometry& MyGeometry,
+	const FPointerEvent& MouseEvent)
+{
+	const TArray<TSharedPtr<FSolverStackElement>> SelectedItems = ListView.Get()->GetSelectedItems();
+	if (SelectedItems.Num() != 1)
+	{
+		return FReply::Unhandled();
+	}
+	
+	if (MouseEvent.IsMouseButtonDown(EKeys::LeftMouseButton))
+	{
+		const TSharedPtr<FSolverStackElement> DraggedElement = SelectedItems[0];
+		const TSharedRef<FIKRigSolverStackDragDropOp> DragDropOp = FIKRigSolverStackDragDropOp::New(DraggedElement);
+		return FReply::Handled().BeginDragDrop(DragDropOp);
+	}
+
+	return FReply::Unhandled();
+}
+
+void SIKRigSolverStack::OnItemClicked(TSharedPtr<FSolverStackElement> InItem)
+{
+	EditorController.Pin()->ShowDetailsForSolver(InItem.Get()->IndexInStack);
+	EditorController.Pin()->SkeletonView->RefreshTreeView(); // update filter showing which bones are affected
+}
+
+TOptional<EItemDropZone> SIKRigSolverStack::OnCanAcceptDrop(
+	const FDragDropEvent& DragDropEvent,
+	EItemDropZone DropZone,
+	TSharedPtr<FSolverStackElement> TargetItem)
+{
+	TOptional<EItemDropZone> ReturnedDropZone;
+	
+	const TSharedPtr<FIKRigSolverStackDragDropOp> DragDropOp = DragDropEvent.GetOperationAs<FIKRigSolverStackDragDropOp>();
+	if (DragDropOp.IsValid())
+	{
+		ReturnedDropZone = EItemDropZone::BelowItem;	
+	}
+	
+	return ReturnedDropZone;
+}
+
+FReply SIKRigSolverStack::OnAcceptDrop(
+	const FDragDropEvent& DragDropEvent,
+	EItemDropZone DropZone,
+	TSharedPtr<FSolverStackElement> TargetItem)
+{
+	const TSharedPtr<FIKRigSolverStackDragDropOp> DragDropOp = DragDropEvent.GetOperationAs<FIKRigSolverStackDragDropOp>();
+	if (!DragDropOp.IsValid())
+	{
+		return FReply::Unhandled();
+	}
+
+	const FSolverStackElement& DraggedElement = *DragDropOp.Get()->Element.Pin().Get();
+	UIKRigController* Controller = EditorController.Pin()->AssetController;
+	const bool bWasReparented = Controller->MoveSolverInStack(DraggedElement.IndexInStack, TargetItem.Get()->IndexInStack);
+	if (bWasReparented)
+	{
+		RefreshStackView();
+		EditorController.Pin()->SkeletonView->RefreshTreeView(); // update solver indices in effector items
+	}
+	
+	return FReply::Handled();
+}
+
+FReply SIKRigSolverStack::OnKeyDown(const FGeometry& MyGeometry, const FKeyEvent& InKeyEvent)
+{
+	const FKey Key = InKeyEvent.GetKey();
+
+	// handle deleting selected solver
+	if (Key == EKeys::Delete)
+	{
+		TArray<TSharedPtr<FSolverStackElement>> SelectedItems = ListView->GetSelectedItems();
+		if (SelectedItems.IsEmpty())
+		{
+			return FReply::Unhandled();
+		}
+
+		UIKRigController* Controller = EditorController.Pin()->AssetController;
+		UIKRigSolver* SolverToRemove = Controller->GetSolver(SelectedItems[0]->IndexInStack);
+		Controller->RemoveSolver(SolverToRemove);
+
+		RefreshStackView();
+		EditorController.Pin()->SkeletonView->RefreshTreeView(); // update solver indices in effector items
+		
+		return FReply::Handled();
+	}
+	
+	return FReply::Unhandled();
+}
+
+#undef LOCTEXT_NAMESPACE
diff --git a/Engine/Plugins/Animation/IKRig/Source/IKRigEditor/Public/AssetTypeActions_IKRetargeter.h b/Engine/Plugins/Animation/IKRig/Source/IKRigEditor/Public/AssetTypeActions_IKRetargeter.h
deleted file mode 100644
diff --git a/Engine/Plugins/Animation/IKRig/Source/IKRigEditor/Public/AssetTypeActions_IKRigDefinition.h b/Engine/Plugins/Animation/IKRig/Source/IKRigEditor/Public/AssetTypeActions_IKRigDefinition.h
deleted file mode 100644
diff --git a/Engine/Plugins/Animation/IKRig/Source/IKRigEditor/Public/IKRetargeterFactory.h b/Engine/Plugins/Animation/IKRig/Source/IKRigEditor/Public/IKRetargeterFactory.h
deleted file mode 100644
diff --git a/Engine/Plugins/Animation/IKRig/Source/IKRigEditor/Public/IKRigController.h b/Engine/Plugins/Animation/IKRig/Source/IKRigEditor/Public/IKRigController.h
deleted file mode 100644
diff --git a/Engine/Plugins/Animation/IKRig/Source/IKRigEditor/Public/IKRigDefinitionDetails.h b/Engine/Plugins/Animation/IKRig/Source/IKRigEditor/Public/IKRigDefinitionDetails.h
deleted file mode 100644
diff --git a/Engine/Plugins/Animation/IKRig/Source/IKRigEditor/Public/IKRigDefinitionFactory.h b/Engine/Plugins/Animation/IKRig/Source/IKRigEditor/Public/IKRigDefinitionFactory.h
deleted file mode 100644
diff --git a/Engine/Plugins/Animation/IKRig/Source/IKRigEditor/Public/IKRigEditMode.h b/Engine/Plugins/Animation/IKRig/Source/IKRigEditor/Public/IKRigEditMode.h
deleted file mode 100644
diff --git a/Engine/Plugins/Animation/IKRig/Source/IKRigEditor/Public/RetargetEditor/AssetTypeActions_IKRetargeter.h b/Engine/Plugins/Animation/IKRig/Source/IKRigEditor/Public/RetargetEditor/AssetTypeActions_IKRetargeter.h
new file mode 100644
--- /dev/null
+++ b/Engine/Plugins/Animation/IKRig/Source/IKRigEditor/Public/RetargetEditor/AssetTypeActions_IKRetargeter.h
@@ -0,0 +1,20 @@
+// Copyright Epic Games, Inc. All Rights Reserved.
+
+#pragma once
+
+#include "CoreMinimal.h"
+#include "Toolkits/IToolkitHost.h"
+#include "AssetTypeActions_Base.h"
+
+class FAssetTypeActions_IKRetargeter : public FAssetTypeActions_Base
+{
+public:
+	// IAssetTypeActions Implementation
+	virtual FText GetName() const override { return NSLOCTEXT("AssetTypeActions", "AssetTypeActions_IKRetargeter", "IK Retargeter"); }
+	virtual void GetActions(const TArray<UObject*>& InObjects, struct FToolMenuSection& Section) override;
+	virtual FColor GetTypeColor() const override { return FColor(255, 192, 128); }
+	virtual UClass* GetSupportedClass() const override;
+	//virtual void OpenAssetEditor(const TArray<UObject*>& InObjects, TSharedPtr<class IToolkitHost> EditWithinLevelEditor = TSharedPtr<IToolkitHost>()) override;
+	virtual uint32 GetCategories() override { return EAssetTypeCategories::Animation; }
+	//virtual UThumbnailInfo* GetThumbnailInfo(UObject* Asset) const override;
+};
diff --git a/Engine/Plugins/Animation/IKRig/Source/IKRigEditor/Public/RetargetEditor/IKRetargeterFactory.h b/Engine/Plugins/Animation/IKRig/Source/IKRigEditor/Public/RetargetEditor/IKRetargeterFactory.h
new file mode 100644
--- /dev/null
+++ b/Engine/Plugins/Animation/IKRig/Source/IKRigEditor/Public/RetargetEditor/IKRetargeterFactory.h
@@ -0,0 +1,28 @@
+// Copyright Epic Games, Inc. All Rights Reserved.
+
+#pragma once
+
+#include "CoreMinimal.h"
+#include "UObject/ObjectMacros.h"
+#include "Factories/Factory.h"
+#include "IKRetargeterFactory.generated.h"
+
+
+UCLASS(BlueprintType, hidecategories=Object)
+class UIKRetargeterFactory : public UFactory
+{
+	GENERATED_BODY()
+
+public:
+
+	UIKRetargeterFactory();
+
+	// UFactory Interface
+	virtual FText GetDisplayName() const override;
+	virtual uint32 GetMenuCategories() const override;
+	virtual FText GetToolTip() const override;
+	virtual FString GetDefaultNewAssetName() const override;
+	virtual UObject* FactoryCreateNew(UClass* Class, UObject* InParent, FName Name, EObjectFlags Flags, UObject* Context, FFeedbackContext* Warn) override;
+	virtual bool ConfigureProperties() override;
+	virtual bool ShouldShowInNewMenu() const override;
+};
diff --git a/Engine/Plugins/Animation/IKRig/Source/IKRigEditor/Public/RigEditor/AssetTypeActions_IKRigDefinition.h b/Engine/Plugins/Animation/IKRig/Source/IKRigEditor/Public/RigEditor/AssetTypeActions_IKRigDefinition.h
new file mode 100644
--- /dev/null
+++ b/Engine/Plugins/Animation/IKRig/Source/IKRigEditor/Public/RigEditor/AssetTypeActions_IKRigDefinition.h
@@ -0,0 +1,21 @@
+// Copyright Epic Games, Inc. All Rights Reserved.
+
+#pragma once
+
+#include "CoreMinimal.h"
+#include "Toolkits/IToolkitHost.h"
+#include "AssetTypeActions_Base.h"
+
+class FAssetTypeActions_IKRigDefinition : public FAssetTypeActions_Base
+{
+public:
+	// IAssetTypeActions Implementation
+	virtual FText GetName() const override { return NSLOCTEXT("AssetTypeActions", "AssetTypeActions_IKRigDefinition", "IK Rig"); }
+	virtual void GetActions(const TArray<UObject*>& InObjects, struct FToolMenuSection& Section) override;
+	virtual FColor GetTypeColor() const override { return FColor(255, 192, 128); }
+	virtual UClass* GetSupportedClass() const override;
+	virtual void OpenAssetEditor(const TArray<UObject*>& InObjects, TSharedPtr<class IToolkitHost> EditWithinLevelEditor = TSharedPtr<IToolkitHost>()) override;
+	virtual uint32 GetCategories() override { return EAssetTypeCategories::Animation; }
+	virtual UThumbnailInfo* GetThumbnailInfo(UObject* Asset) const override;
+	// END IAssetTypeActions Implementation
+};
diff --git a/Engine/Plugins/Animation/IKRig/Source/IKRigEditor/Public/RigEditor/IKRigAnimInstance.h b/Engine/Plugins/Animation/IKRig/Source/IKRigEditor/Public/RigEditor/IKRigAnimInstance.h
new file mode 100644
--- /dev/null
+++ b/Engine/Plugins/Animation/IKRig/Source/IKRigEditor/Public/RigEditor/IKRigAnimInstance.h
@@ -0,0 +1,22 @@
+// Copyright Epic Games, Inc. All Rights Reserved.
+
+#pragma once
+
+#include "AnimPreviewInstance.h"
+
+#include "IKRigAnimInstance.generated.h"
+
+class UIKRigDefinition;
+
+UCLASS(transient, NotBlueprintable)
+class UIKRigAnimInstance : public UAnimPreviewInstance
+{
+	GENERATED_UCLASS_BODY()
+
+public:
+
+	void SetIKRigAsset(UIKRigDefinition* IKRigAsset);
+
+protected:
+	virtual FAnimInstanceProxy* CreateAnimInstanceProxy() override;
+};
diff --git a/Engine/Plugins/Animation/IKRig/Source/IKRigEditor/Public/RigEditor/IKRigAnimInstanceProxy.h b/Engine/Plugins/Animation/IKRig/Source/IKRigEditor/Public/RigEditor/IKRigAnimInstanceProxy.h
new file mode 100644
--- /dev/null
+++ b/Engine/Plugins/Animation/IKRig/Source/IKRigEditor/Public/RigEditor/IKRigAnimInstanceProxy.h
@@ -0,0 +1,39 @@
+// Copyright Epic Games, Inc. All Rights Reserved.
+
+#pragma once
+
+#include "AnimPreviewInstance.h"
+#include "Animation/AnimNode_LinkedInputPose.h"
+#include "AnimNodes/AnimNode_IKRig.h"
+
+#include "IKRigAnimInstanceProxy.generated.h"
+
+/** Proxy override for this UAnimInstance-derived class */
+USTRUCT()
+struct FIKRigAnimInstanceProxy : public FAnimPreviewInstanceProxy
+{
+	GENERATED_BODY()
+
+public:
+	
+	FIKRigAnimInstanceProxy();
+	FIKRigAnimInstanceProxy(UAnimInstance* InAnimInstance);
+	virtual ~FIKRigAnimInstanceProxy();
+
+	/** FAnimPreviewInstanceProxy interface */
+	virtual void Initialize(UAnimInstance* InAnimInstance) override;
+	virtual bool Evaluate(FPoseContext& Output) override;
+	virtual void UpdateAnimationNode(const FAnimationUpdateContext& InContext) override;
+	/** END FAnimPreviewInstanceProxy interface */
+
+	/** FAnimInstanceProxy interface */
+	/** Called when the anim instance is being initialized. These nodes can be provided */
+	virtual FAnimNode_Base* GetCustomRootNode() override;
+	virtual void GetCustomNodes(TArray<FAnimNode_Base*>& OutNodes) override;
+	/** END FAnimInstanceProxy interface */
+
+	void SetIKRigAsset(UIKRigDefinition* InIKRigAsset);
+
+	UPROPERTY(Transient)
+	FAnimNode_IKRig IKRigNode;
+};
\ No newline at end of file
diff --git a/Engine/Plugins/Animation/IKRig/Source/IKRigEditor/Public/RigEditor/IKRigCommands.h b/Engine/Plugins/Animation/IKRig/Source/IKRigEditor/Public/RigEditor/IKRigCommands.h
new file mode 100644
--- /dev/null
+++ b/Engine/Plugins/Animation/IKRig/Source/IKRigEditor/Public/RigEditor/IKRigCommands.h
@@ -0,0 +1,26 @@
+// Copyright Epic Games, Inc. All Rights Reserved.
+
+#pragma once
+
+#include "CoreMinimal.h"
+#include "Framework/Commands/Commands.h"
+#include "IKRigEditorStyle.h"
+
+class FIKRigCommands : public TCommands<FIKRigCommands>
+{
+public:
+	FIKRigCommands() : TCommands<FIKRigCommands>
+	(
+		"IKRig",
+		NSLOCTEXT("Contexts", "IKRig", "IK Rig"),
+		NAME_None, // "MainFrame"
+		FIKRigEditorStyle::Get().GetStyleSetName() // Icon Style Set
+	)
+	{}
+	
+	/** reset whole system to initial state */
+	TSharedPtr< FUICommandInfo > Reset;
+
+	/** initialize commands */
+	virtual void RegisterCommands() override;
+};
diff --git a/Engine/Plugins/Animation/IKRig/Source/IKRigEditor/Public/RigEditor/IKRigController.h b/Engine/Plugins/Animation/IKRig/Source/IKRigEditor/Public/RigEditor/IKRigController.h
new file mode 100644
--- /dev/null
+++ b/Engine/Plugins/Animation/IKRig/Source/IKRigEditor/Public/RigEditor/IKRigController.h
@@ -0,0 +1,92 @@
+// Copyright Epic Games, Inc. All Rights Reserved.
+
+#pragma once
+
+#include "CoreMinimal.h"
+#include "UObject/Object.h"
+
+#include "IKRigController.generated.h"
+
+struct FReferenceSkeleton;
+struct FIKRigSkeleton;
+class UIKRigSolver;
+class UIKRigDefinition;
+class UIKRigEffectorGoal;
+class USkeletalMesh;
+
+UCLASS(config = Engine, hidecategories = UObject, BlueprintType)
+class IKRIGEDITOR_API UIKRigController : public UObject
+{
+	GENERATED_BODY()
+
+public:
+
+	//** use this to get a handle to a controller for the given IKRig */
+	static UIKRigController* GetIKRigController(UIKRigDefinition* InIKRigDefinition);
+	//** get the asset this controller controls */
+	UIKRigDefinition* GetAsset() const;
+
+	//** skeleton */
+	void SetSourceSkeletalMesh(USkeletalMesh* SkeletalMesh, bool bReImportBones) const;
+	USkeletalMesh* GetSourceSkeletalMesh() const;
+	void SetSkeleton(const FReferenceSkeleton& InSkeleton) const;
+	FIKRigSkeleton& GetSkeleton() const;
+	//** END skeleton */
+
+	//** solvers */
+	int32 AddSolver(TSubclassOf<UIKRigSolver> InSolverClass) const;
+	void RemoveSolver(UIKRigSolver* SolverToDelete) const;
+	bool MoveSolverInStack(int32 SolverToMoveIndex, int32 TargetSolverIndex) const;
+	UIKRigSolver* GetSolver(int32 Index) const;
+	int32 GetNumSolvers() const;
+	void SetRootBone(const FName& RootBoneName, int32 SolverIndex) const;
+	const TArray<UIKRigSolver*>& GetSolverArray() const;
+	//** END solvers */
+
+	//** goals */
+	UIKRigEffectorGoal* AddNewGoal(const FName& GoalName, const FName& BoneName) const;
+	bool RemoveGoal(const FName& GoalName) const;
+	FName RenameGoal(const FName& OldName, const FName& PotentialNewName) const;
+	bool SetGoalBone(const FName& GoalName, const FName& NewBoneName) const;
+	FName GetBoneForGoal(const FName& GoalName) const;
+	bool ConnectGoalToSolver(const UIKRigEffectorGoal& Goal, int32 SolverIndex) const;
+	bool DisconnectGoalFromSolver(const FName& GoalToRemove, int32 SolverIndex) const;
+	bool IsGoalConnectedToSolver(const FName& Goal, int32 SolverIndex) const;
+	int32 GetGoalIndex(const FName& GoalName) const;
+	FName GetGoalName(const int32& GoalIndex) const;
+	const TArray<UIKRigEffectorGoal*>& GetAllGoals() const;
+	UIKRigEffectorGoal* GetGoal(int32 GoalIndex) const;
+	UIKRigEffectorGoal* GetGoal(const FName& GoalName) const;
+	UObject* GetEffectorForGoal(const FName& GoalName, int32 SolverIndex) const;
+	FTransform GetGoalInitialTransform(const FName& GoalName) const;
+	FTransform GetGoalCurrentTransform(const FName& GoalName) const;
+	void SetGoalInitialTransform(const FName& GoalName, const FTransform& Transform) const;
+	void SetGoalCurrentTransform(const FName& GoalName, const FTransform& Transform) const;
+	static void SanitizeGoalName(FString& InOutName);
+	DECLARE_MULTICAST_DELEGATE(FGoalModified);
+	FGoalModified OnGoalModified;
+	//** END goals */
+
+	//** bone settings */
+	void AddBoneSetting(const FName& BoneName, int32 SolverIndex) const;
+	void RemoveBoneSetting(const FName& BoneName, int32 SolverIndex) const;
+	bool CanAddBoneSetting(const FName& BoneName, int32 SolverIndex) const;
+	bool CanRemoveBoneSetting(const FName& BoneName, int32 SolverIndex) const;
+	UObject* GetSettingsForBone(const FName& BoneName, int32 SolverIndex) const;
+	bool DoesBoneHaveSettings(const FName& BoneName) const;
+	//** END bone settings */
+
+	// BEGIN UObject
+	static void AddReferencedObjects(UObject* InThis, FReferenceCollector& Collector);
+	// END UObject
+
+private:
+
+	UPROPERTY(transient)
+	UIKRigDefinition* IKRigAsset = nullptr;
+
+	// todo, I'm not convinced Controller functions shouldn't just live on the asset itself...
+	static TMap<UIKRigDefinition*, UIKRigController*> AssetToControllerMap;
+
+	friend class UIKRigDefinition;
+};
diff --git a/Engine/Plugins/Animation/IKRig/Source/IKRigEditor/Public/RigEditor/IKRigDefinitionFactory.h b/Engine/Plugins/Animation/IKRig/Source/IKRigEditor/Public/RigEditor/IKRigDefinitionFactory.h
new file mode 100644
--- /dev/null
+++ b/Engine/Plugins/Animation/IKRig/Source/IKRigEditor/Public/RigEditor/IKRigDefinitionFactory.h
@@ -0,0 +1,34 @@
+// Copyright Epic Games, Inc. All Rights Reserved.
+
+#pragma once
+
+#include "CoreMinimal.h"
+#include "UObject/ObjectMacros.h"
+#include "Factories/Factory.h"
+#include "IKRigDefinitionFactory.generated.h"
+
+UCLASS(BlueprintType, hidecategories=Object)
+class UIKRigDefinitionFactory : public UFactory
+{
+	GENERATED_BODY()
+
+public:
+
+	UIKRigDefinitionFactory();
+
+	// UFactory interface
+	virtual UObject* FactoryCreateNew(
+        UClass* Class,
+        UObject* InParent,
+        FName Name,
+        EObjectFlags Flags,
+        UObject* Context,
+        FFeedbackContext* Warn) override;
+	virtual FText GetDisplayName() const override;
+	virtual uint32 GetMenuCategories() const override;
+	virtual FText GetToolTip() const override;
+	virtual FString GetDefaultNewAssetName() const override;
+	virtual bool ConfigureProperties() override;
+	virtual bool ShouldShowInNewMenu() const override;
+	// END UFactory interface
+};
diff --git a/Engine/Plugins/Animation/IKRig/Source/IKRigEditor/Public/RigEditor/IKRigEditMode.h b/Engine/Plugins/Animation/IKRig/Source/IKRigEditor/Public/RigEditor/IKRigEditMode.h
new file mode 100644
--- /dev/null
+++ b/Engine/Plugins/Animation/IKRig/Source/IKRigEditor/Public/RigEditor/IKRigEditMode.h
@@ -0,0 +1,103 @@
+// Copyright Epic Games, Inc. All Rights Reserved.
+
+#pragma once
+
+#include "CoreMinimal.h"
+
+#include "IKRigDefinition.h"
+#include "IPersonaEditMode.h"
+
+class FIKRigEditorController;
+class UIKRigEffectorGoal;
+class FIKRigEditorToolkit;
+class FIKRigPreviewScene;
+
+struct GoalGizmo
+{
+	GoalGizmo()
+	{
+		// wireframe cube lines (point pairs)
+		BoxPoints.Add(FVector(0.5f, 0.5f, 0.5f));
+		BoxPoints.Add(FVector(0.5f, -0.5f, 0.5f));
+		BoxPoints.Add(FVector(0.5f, -0.5f, 0.5f));
+		BoxPoints.Add(FVector(-0.5f, -0.5f, 0.5f));
+		BoxPoints.Add(FVector(-0.5f, -0.5f, 0.5f));
+		BoxPoints.Add(FVector(-0.5f, 0.5f, 0.5f));
+		BoxPoints.Add(FVector(-0.5f, 0.5f, 0.5f));
+		BoxPoints.Add(FVector(0.5f, 0.5f, 0.5f));
+
+		BoxPoints.Add(FVector(0.5f, 0.5f, -0.5f));
+		BoxPoints.Add(FVector(0.5f, -0.5f, -0.5f));
+		BoxPoints.Add(FVector(0.5f, -0.5f, -0.5f));
+		BoxPoints.Add(FVector(-0.5f, -0.5f, -0.5f));
+		BoxPoints.Add(FVector(-0.5f, -0.5f, -0.5f));
+		BoxPoints.Add(FVector(-0.5f, 0.5f, -0.5f));
+		BoxPoints.Add(FVector(-0.5f, 0.5f, -0.5f));
+		BoxPoints.Add(FVector(0.5f, 0.5f, -0.5f));
+
+		BoxPoints.Add(FVector(0.5f, 0.5f, 0.5f));
+		BoxPoints.Add(FVector(0.5f, 0.5f, -0.5f));
+		BoxPoints.Add(FVector(0.5f, -0.5f, 0.5f));
+		BoxPoints.Add(FVector(0.5f, -0.5f, -0.5f));
+		BoxPoints.Add(FVector(-0.5f, -0.5f, 0.5f));
+		BoxPoints.Add(FVector(-0.5f, -0.5f, -0.5f));
+		BoxPoints.Add(FVector(-0.5f, 0.5f, 0.5f));
+		BoxPoints.Add(FVector(-0.5f, 0.5f, -0.5f));
+	}
+
+	void DrawGoal(FPrimitiveDrawInterface* PDI, const UIKRigEffectorGoal* Goal, bool bIsSelected) const
+	{
+		const FLinearColor Color = bIsSelected ? FLinearColor::Green : FLinearColor::Yellow;
+		const float Thickness = Goal->GizmoThickness;
+		const float Scale = FMath::Clamp(Goal->GizmoSize, 0.1f, 1000.0f);
+		const FTransform Transform = Goal->CurrentTransform;
+		for (int32 PointIndex = 0; PointIndex < BoxPoints.Num() - 1; PointIndex += 2)
+		{
+			PDI->DrawLine(Transform.TransformPosition(BoxPoints[PointIndex] * Scale), Transform.TransformPosition(BoxPoints[PointIndex + 1] * Scale), Color, SDPG_Foreground, Thickness);
+		}
+	}
+
+	TArray<FVector> BoxPoints;
+};
+
+class FIKRigEditMode : public IPersonaEditMode
+{
+public:
+	static FName ModeName;
+	
+	FIKRigEditMode();
+
+	/** glue for all the editor parts to communicate */
+	void SetEditorController(const TSharedPtr<FIKRigEditorController> InEditorController) { EditorController = InEditorController; };
+
+	/** IPersonaEditMode interface */
+	virtual bool GetCameraTarget(FSphere& OutTarget) const override;
+	virtual class IPersonaPreviewScene& GetAnimPreviewScene() const override;
+	virtual void GetOnScreenDebugInfo(TArray<FText>& OutDebugInfo) const override;
+	/** END IPersonaEditMode interface */
+
+	/** FEdMode interface */
+	virtual void Tick(FEditorViewportClient* ViewportClient, float DeltaTime) override;
+	virtual void Render(const FSceneView* View, FViewport* Viewport, FPrimitiveDrawInterface* PDI) override;
+	virtual void DrawHUD(FEditorViewportClient* ViewportClient, FViewport* Viewport, const FSceneView* View, FCanvas* Canvas) override;
+	virtual bool IsCompatibleWith(FEditorModeID OtherModeID) const override { return true; }
+	virtual bool AllowWidgetMove() override;
+	virtual bool ShouldDrawWidget() const override;
+	virtual bool UsesTransformWidget() const override;
+	virtual bool UsesTransformWidget(UE::Widget::EWidgetMode CheckMode) const override;
+	virtual FVector GetWidgetLocation() const override;
+	virtual bool HandleClick(FEditorViewportClient* InViewportClient, HHitProxy* HitProxy, const FViewportClick& Click) override;
+	virtual bool StartTracking(FEditorViewportClient* InViewportClient, FViewport* InViewport) override;
+	virtual bool EndTracking(FEditorViewportClient* InViewportClient, FViewport* InViewport) override;
+	virtual bool InputDelta(FEditorViewportClient* InViewportClient, FViewport* InViewport, FVector& InDrag, FRotator& InRot, FVector& InScale) override;
+	virtual bool GetCustomDrawingCoordinateSystem(FMatrix& InMatrix, void* InData) override;
+	virtual bool GetCustomInputCoordinateSystem(FMatrix& InMatrix, void* InData) override;
+	/** END FEdMode interface */
+
+private:
+	/** The hosting app */
+	TWeakPtr<FIKRigEditorController> EditorController;
+
+	/** draws goals in viewport */
+	GoalGizmo GoalDrawer;
+};
diff --git a/Engine/Plugins/Animation/IKRig/Source/IKRigEditor/Public/RigEditor/IKRigEditorController.h b/Engine/Plugins/Animation/IKRig/Source/IKRigEditor/Public/RigEditor/IKRigEditorController.h
new file mode 100644
--- /dev/null
+++ b/Engine/Plugins/Animation/IKRig/Source/IKRigEditor/Public/RigEditor/IKRigEditorController.h
@@ -0,0 +1,103 @@
+// Copyright Epic Games, Inc. All Rights Reserved.
+
+#pragma once
+
+#include "IKRigSolver.h"
+#include "Templates/SubclassOf.h"
+
+#include "IKRigEditorController.generated.h"
+
+class UIKRigAnimInstance;
+class FIKRigEditorToolkit;
+class SIKRigSolverStack;
+class SIKRigSkeleton;
+class UIKRigController;
+class FSolverStackElement;
+class FIKRigTreeElement;
+class UDebugSkelMeshComponent;
+
+/** a home for cross-widget communication to synchronize state across all tabs and viewport */
+class FIKRigEditorController
+{
+public:
+
+	/** get array of the names of currently selected goals */
+	const TArray<FName>& GetSelectedGoals() const {return SelectedGoals;};
+	/** return true if a goal with the given name is selected */
+	bool IsGoalSelected(const FName& GoalName);
+	/** return number of goals that are currently selected */
+	int32 GetNumSelectedGoals();
+	/** callback when goal is selected in the viewport */
+	void HandleGoalSelectedInViewport(const FName& GoalName, bool bReplace);
+	/** callback when goal is selected in the skeleton view */
+	void HandleGoalsSelectedInTreeView(const TArray<FName>& GoalNames);
+	/** reset all goals to initial transforms */
+	void Reset() const;
+
+	/** return list of those solvers in the stack that are selected by user */
+	void GetSelectedSolvers(TArray<TSharedPtr<FSolverStackElement> >& OutSelectedSolvers);
+	/** get index of the first selected solver, return INDEX_None if nothing selected */
+	int32 GetSelectedSolverIndex();
+	/** right after importing a skeleton, we ask user what solver they want to use */
+	void PromptToAddSolverAtStartup() const;
+	/** determine if the element is connected to the selected solver */
+	bool IsElementConnectedToSolver(TSharedRef<FIKRigTreeElement> TreeElement, int32 SolverIndex);
+	/** determine if the element is connected to ANY solver */
+	bool IsElementConnectedToAnySolver(TSharedRef<FIKRigTreeElement> TreeElement);
+	
+	/** todo show BONE transform in details view */
+	void ShowDetailsForBone(const FName BoneName);
+	/** show BONE settings in details view */
+	void ShowDetailsForBoneSettings(const FName BoneName, int32 SolverIndex);
+	/** show GOAL settings in details view */
+	void ShowDetailsForGoal(const FName GoalName);
+	/** show EFFECTOR settings in details view */
+	void ShowDetailsForEffector(const FName GoalName, const int32 SolverIndex);
+	/** show SOLVER settings in details view */
+	void ShowDetailsForSolver(const int32 SolverIndex);
+	/** show nothing in details view */
+	void ShowEmptyDetails();
+	
+	/** all modifications to the data model should go through this controller */
+	UIKRigController* AssetController;
+
+	/** viewport skeletal mesh */
+	UDebugSkelMeshComponent* SkelMeshComponent;
+
+	/** viewport anim instance */
+	UIKRigAnimInstance* AnimInstance; 
+
+	/** asset properties tab */
+	TSharedPtr<class IDetailsView> DetailsView;
+
+	/** the skeleton tree view */
+	TSharedPtr<SIKRigSkeleton> SkeletonView;
+	
+	/** the solver stack view */
+	TSharedPtr<SIKRigSolverStack> SolverStackView;
+
+	/** the persona toolkit */
+	TWeakPtr<FIKRigEditorToolkit> EditorToolkit;
+
+	/** UI and viewport selection state */
+	bool bManipulatingGoals = false;
+
+private:
+
+	/** UI and viewport selection state */
+	TArray<FName> SelectedGoals;	
+};
+
+/** only used for pop-up window to selected a solver to add after importing a skeleton */
+USTRUCT()
+struct FIKRigAddFirstSolverSettings
+{
+	GENERATED_BODY()
+
+	FIKRigAddFirstSolverSettings()
+        : SolverType(nullptr)
+	{}
+
+	UPROPERTY(EditAnywhere, Category = "Add Solver")
+	TSubclassOf<UIKRigSolver> SolverType;
+};
diff --git a/Engine/Plugins/Animation/IKRig/Source/IKRigEditor/Public/RigEditor/IKRigEditorStyle.h b/Engine/Plugins/Animation/IKRig/Source/IKRigEditor/Public/RigEditor/IKRigEditorStyle.h
new file mode 100644
--- /dev/null
+++ b/Engine/Plugins/Animation/IKRig/Source/IKRigEditor/Public/RigEditor/IKRigEditorStyle.h
@@ -0,0 +1,57 @@
+// Copyright Epic Games, Inc. All Rights Reserved.
+
+#pragma once
+
+#include "EditorStyleSet.h"
+#include "Styling/SlateStyle.h"
+#include "Styling/SlateStyleRegistry.h"
+#include "Styling/SlateStyleMacros.h"
+#include "Styling/SlateTypes.h"
+
+class FIKRigEditorStyle	: public FSlateStyleSet
+{
+public:
+	
+	FIKRigEditorStyle() : FSlateStyleSet("IKRigEditorStyle")
+	{
+		SetContentRoot(FPaths::EnginePluginsDir() / TEXT("Animation/IKRig/Content"));
+		const FVector2D Icon10x10(10.0f, 10.0f);
+		const FVector2D Icon14x14(14.0f, 14.0f);
+		const FVector2D Icon16x16(16.0f, 16.0f);
+		const FVector2D Icon20x20(20.0f, 20.0f);
+		const FVector2D Icon24x24(24.0f, 24.0f);
+		const FVector2D Icon32x32(32.0f, 32.0f);
+		const FVector2D Icon40x40(40.0f, 40.0f);
+		
+		Set("IKRig.Tree.Bone", new IMAGE_BRUSH("Slate/Bone_16x", Icon16x16));
+		Set("IKRig.Tree.BoneWithSettings", new IMAGE_BRUSH("Slate/BoneWithSettings_16x", Icon16x16));
+		Set("IKRig.Tree.Goal", new IMAGE_BRUSH("Slate/Goal_16x", Icon16x16));
+		Set("IKRig.Tree.Effector", new IMAGE_BRUSH("Slate/Effector_16x", Icon16x16));
+		
+		Set("IKRig.TabIcon", new IMAGE_BRUSH("Slate/Tab_16x", Icon16x16));
+		
+		Set("IKRig.Solver", new IMAGE_BRUSH("Slate/Solver_16x", Icon16x16));
+		Set("IKRig.DragSolver", new IMAGE_BRUSH("Slate/DragSolver", FVector2D(6, 15)));
+
+		Set("IKRig.Reset", new IMAGE_BRUSH("Slate/Reset", Icon40x40));
+		Set("IKRig.Reset.Small", new IMAGE_BRUSH("Slate/Reset", Icon20x20));
+		
+		FTextBlockStyle NormalText = FEditorStyle::GetWidgetStyle<FTextBlockStyle>("SkeletonTree.NormalFont");
+		Set( "IKRig.Tree.AffectedBoneText", FTextBlockStyle(NormalText).SetColorAndOpacity( FLinearColor::Green ));
+		
+		FSlateStyleRegistry::RegisterSlateStyle(*this);
+	}
+
+	static FIKRigEditorStyle& Get()
+	{
+		static FIKRigEditorStyle Inst;
+		return Inst;
+	}
+	
+	~FIKRigEditorStyle()
+	{
+		FSlateStyleRegistry::UnRegisterSlateStyle(*this);
+	}
+};
+
+#undef IMAGE_BRUSH
\ No newline at end of file
diff --git a/Engine/Plugins/Animation/IKRig/Source/IKRigEditor/Public/RigEditor/IKRigHitProxies.h b/Engine/Plugins/Animation/IKRig/Source/IKRigEditor/Public/RigEditor/IKRigHitProxies.h
new file mode 100644
--- /dev/null
+++ b/Engine/Plugins/Animation/IKRig/Source/IKRigEditor/Public/RigEditor/IKRigHitProxies.h
@@ -0,0 +1,43 @@
+// Copyright Epic Games, Inc. All Rights Reserved.
+
+#pragma once
+
+#include "CoreMinimal.h"
+#include "HitProxies.h"
+
+// todo allow bone selection to preview transform values
+struct HIKRigEditorBoneProxy : public HHitProxy
+{
+	DECLARE_HIT_PROXY();
+
+	virtual EMouseCursor::Type GetMouseCursor()
+	{
+		return EMouseCursor::Crosshairs;
+	}
+
+	virtual bool AlwaysAllowsTranslucentPrimitives() const override
+	{
+		return true;
+	}
+};
+
+struct HIKRigEditorGoalProxy : public HHitProxy
+{
+	DECLARE_HIT_PROXY();
+
+	FName GoalName;
+	
+	HIKRigEditorGoalProxy(const FName& InGoalName)
+		: HHitProxy(HPP_World)
+		, GoalName(InGoalName) {}
+
+	virtual EMouseCursor::Type GetMouseCursor()
+	{
+		return EMouseCursor::Crosshairs;
+	}
+
+	virtual bool AlwaysAllowsTranslucentPrimitives() const override
+	{
+		return true;
+	}
+};
\ No newline at end of file
diff --git a/Engine/Plugins/Animation/IKRig/Source/IKRigEditor/Public/RigEditor/IKRigMode.h b/Engine/Plugins/Animation/IKRig/Source/IKRigEditor/Public/RigEditor/IKRigMode.h
new file mode 100644
--- /dev/null
+++ b/Engine/Plugins/Animation/IKRig/Source/IKRigEditor/Public/RigEditor/IKRigMode.h
@@ -0,0 +1,29 @@
+// Copyright Epic Games, Inc. All Rights Reserved.
+
+#pragma once
+
+#include "WorkflowOrientedApp/ApplicationMode.h"
+#include "WorkflowOrientedApp/WorkflowCentricApplication.h"
+#include "WorkflowOrientedApp/WorkflowTabManager.h"
+
+class FIKRigEditorToolkit;
+
+class FIKRigMode : public FApplicationMode
+{
+	
+public:
+	
+	FIKRigMode(TSharedRef<class FWorkflowCentricApplication> InHostingApp, TSharedRef<class IPersonaPreviewScene> InPreviewScene);
+
+	/** FApplicationMode interface */
+	virtual void RegisterTabFactories(TSharedPtr<FTabManager> InTabManager) override;
+	/** END FApplicationMode interface */
+
+protected:
+	
+	/** The hosting app */
+	TWeakPtr<FIKRigEditorToolkit> IKRigEditorPtr;
+
+	/** The tab factories we support */
+	FWorkflowAllowedTabSet TabFactories;
+};
diff --git a/Engine/Plugins/Animation/IKRig/Source/IKRigEditor/Public/RigEditor/IKRigSkeletonCommands.h b/Engine/Plugins/Animation/IKRig/Source/IKRigEditor/Public/RigEditor/IKRigSkeletonCommands.h
new file mode 100644
--- /dev/null
+++ b/Engine/Plugins/Animation/IKRig/Source/IKRigEditor/Public/RigEditor/IKRigSkeletonCommands.h
@@ -0,0 +1,40 @@
+// Copyright Epic Games, Inc. All Rights Reserved.
+
+#pragma once
+
+#include "CoreMinimal.h"
+#include "Framework/Commands/Commands.h"
+#include "IKRigEditorStyle.h"
+
+class FIKRigSkeletonCommands : public TCommands<FIKRigSkeletonCommands>
+{
+public:
+	FIKRigSkeletonCommands() : TCommands<FIKRigSkeletonCommands>
+	(
+		"IKRigSkeleton",
+		NSLOCTEXT("Contexts", "IKRigSkeleton", "IK Rig Skeleton"),
+		NAME_None,
+		FIKRigEditorStyle::Get().GetStyleSetName() // Icon Style Set
+	)
+	{}
+	
+	/** create new goal on selected bone */
+	TSharedPtr< FUICommandInfo > NewGoal;
+	/** delete selected goal */
+	TSharedPtr< FUICommandInfo > Delete;
+	/** add goal to selected solvers */
+	TSharedPtr< FUICommandInfo > ConnectGoalToSolvers;
+	/** remove goal to selected solvers */
+	TSharedPtr< FUICommandInfo > DisconnectGoalFromSolvers;
+
+	/** set root bone on selected solvers */
+	TSharedPtr< FUICommandInfo > SetRootBoneOnSolvers;
+	
+	/** add bone setting to bone */
+	TSharedPtr< FUICommandInfo > AddBoneSettings;
+	/** remove bone setting from bone */
+	TSharedPtr< FUICommandInfo > RemoveBoneSettings;
+
+	/** initialize commands */
+	virtual void RegisterCommands() override;
+};
diff --git a/Engine/Plugins/Animation/IKRig/Source/IKRigEditor/Public/RigEditor/IKRigSkeletonTabSummoner.h b/Engine/Plugins/Animation/IKRig/Source/IKRigEditor/Public/RigEditor/IKRigSkeletonTabSummoner.h
new file mode 100644
--- /dev/null
+++ b/Engine/Plugins/Animation/IKRig/Source/IKRigEditor/Public/RigEditor/IKRigSkeletonTabSummoner.h
@@ -0,0 +1,21 @@
+// Copyright Epic Games, Inc. All Rights Reserved.
+
+#pragma once
+
+#include "CoreMinimal.h"
+#include "WorkflowOrientedApp/WorkflowTabFactory.h"
+
+class FIKRigEditorToolkit;
+
+struct FIKRigSkeletonTabSummoner : public FWorkflowTabFactory
+{
+public:
+	static const FName TabID;
+	
+	FIKRigSkeletonTabSummoner(const TSharedRef<FIKRigEditorToolkit>& InIKRigEditor);
+	virtual TSharedRef<SWidget> CreateTabBody(const FWorkflowTabSpawnInfo& Info) const override;
+	virtual TSharedPtr<SToolTip> CreateTabToolTipWidget(const FWorkflowTabSpawnInfo& Info) const override;
+
+protected:
+	TWeakPtr<FIKRigEditorToolkit> IKRigEditor;
+};
diff --git a/Engine/Plugins/Animation/IKRig/Source/IKRigEditor/Public/RigEditor/IKRigSolverStackTabSummoner.h b/Engine/Plugins/Animation/IKRig/Source/IKRigEditor/Public/RigEditor/IKRigSolverStackTabSummoner.h
new file mode 100644
--- /dev/null
+++ b/Engine/Plugins/Animation/IKRig/Source/IKRigEditor/Public/RigEditor/IKRigSolverStackTabSummoner.h
@@ -0,0 +1,21 @@
+// Copyright Epic Games, Inc. All Rights Reserved.
+
+#pragma once
+
+#include "CoreMinimal.h"
+#include "WorkflowOrientedApp/WorkflowTabFactory.h"
+
+class FIKRigEditorToolkit;
+
+struct FIKRigSolverStackTabSummoner : public FWorkflowTabFactory
+{
+public:
+	static const FName TabID;
+	
+	FIKRigSolverStackTabSummoner(const TSharedRef<FIKRigEditorToolkit>& InIKRigEditor);
+	virtual TSharedRef<SWidget> CreateTabBody(const FWorkflowTabSpawnInfo& Info) const override;
+	virtual TSharedPtr<SToolTip> CreateTabToolTipWidget(const FWorkflowTabSpawnInfo& Info) const override;
+
+protected:
+	TWeakPtr<FIKRigEditorToolkit> IKRigEditor;
+};
diff --git a/Engine/Plugins/Animation/IKRig/Source/IKRigEditor/Public/RigEditor/IKRigToolkit.h b/Engine/Plugins/Animation/IKRig/Source/IKRigEditor/Public/RigEditor/IKRigToolkit.h
new file mode 100644
--- /dev/null
+++ b/Engine/Plugins/Animation/IKRig/Source/IKRigEditor/Public/RigEditor/IKRigToolkit.h
@@ -0,0 +1,103 @@
+// Copyright Epic Games, Inc. All Rights Reserved.
+
+#pragma once
+
+#include "CoreMinimal.h"
+
+#include "IHasPersonaToolkit.h"
+#include "IKRigController.h"
+#include "IKRigEditorController.h"
+#include "IKRigMode.h"
+#include "IPersonaPreviewScene.h"
+#include "PersonaAssetEditorToolkit.h"
+#include "SIKRigSkeleton.h"
+
+class UIKRigSkeletalMeshComponent;
+class FSolverStackElement;
+class SIKRigSolverStack;
+class SIKRigSkeleton;
+class FIKRigSkeletonTreeBuilder;
+class IDetailsView;
+class ISkeletonTree;
+class FGGAssetEditorToolbar;
+class UIKRigDefinition;
+class UIKRigSettings;
+class FIKRigToolbar;
+class SIKRigViewportTabBody;
+class FIKRigPreviewScene;
+
+namespace IKRigEditorModes
+{
+	extern const FName IKRigEditorMode;
+}
+
+class FIKRigEditorToolkit :
+	public FPersonaAssetEditorToolkit,
+	public IHasPersonaToolkit,
+	public FGCObject,
+	public FEditorUndoClient,
+	public FTickableEditorObject
+{
+public:
+
+	FIKRigEditorToolkit();
+	virtual ~FIKRigEditorToolkit();
+
+	void InitAssetEditor(
+		const EToolkitMode::Type Mode,
+		const TSharedPtr< IToolkitHost >& InitToolkitHost,
+		UIKRigDefinition* IKRigAsset);
+
+	/** FAssetEditorToolkit interface */
+	virtual void RegisterTabSpawners(const TSharedRef<FTabManager>& InTabManager) override;
+	virtual void UnregisterTabSpawners(const TSharedRef<FTabManager>& InTabManager) override;
+	virtual FName GetToolkitFName() const override;
+	virtual FText GetBaseToolkitName() const override;
+	virtual FText GetToolkitName() const override;
+	virtual FLinearColor GetWorldCentricTabColorScale() const override;
+	virtual FString GetWorldCentricTabPrefix() const override;
+	/** END FAssetEditorToolkit interface */
+	
+	/** FGCObject interface */
+	virtual void AddReferencedObjects(FReferenceCollector& Collector) override;
+	/** END FGCObject interface */
+
+	//** FTickableEditorObject Interface
+	virtual void Tick(float DeltaTime) override {};
+	virtual ETickableTickType GetTickableTickType() const override { return ETickableTickType::Always; }
+	virtual TStatId GetStatId() const override;
+	//~ END FTickableEditorObject Interface
+
+	/** IHasPersonaToolkit interface */
+	virtual TSharedRef<IPersonaToolkit> GetPersonaToolkit() const override { return PersonaToolkit.ToSharedRef(); }
+	/** END IHasPersonaToolkit interface */
+
+	TSharedRef<FIKRigEditorController> GetController() const {return EditorController;};
+
+private:
+
+	/** toolbar */
+	void BindCommands();
+	void ExtendToolbar();
+	void FillToolbar(FToolBarBuilder& ToolbarBuilder);
+	void HandleReset();
+	/** END toolbar */
+	
+	/** preview scene setup */
+	void HandlePreviewSceneCreated(const TSharedRef<IPersonaPreviewScene>& InPersonaPreviewScene);
+	void HandlePreviewMeshChanged(USkeletalMesh* InOldSkeletalMesh, USkeletalMesh* InNewSkeletalMesh);
+	void HandleDetailsCreated(const TSharedRef<class IDetailsView>& InDetailsView);
+	void OnFinishedChangingDetails(const FPropertyChangedEvent& PropertyChangedEvent);
+	void SetupAnimInstance(UIKRigAnimInstance* InAnimInstance);
+	/** END preview scene setup */
+	
+	/** centralized management of selection across skeleton view and viewport */
+	TSharedRef<FIKRigEditorController> EditorController;
+
+	/** Preview scene to be supplied by IHasPersonaToolkit::GetPersonaToolkit */
+	TSharedPtr<IPersonaToolkit> PersonaToolkit;
+	
+	friend FIKRigMode;
+	friend SIKRigSkeleton;
+	friend SIKRigSolverStack;
+};
diff --git a/Engine/Plugins/Animation/IKRig/Source/IKRigEditor/Public/RigEditor/SIKRigSkeleton.h b/Engine/Plugins/Animation/IKRig/Source/IKRigEditor/Public/RigEditor/SIKRigSkeleton.h
new file mode 100644
--- /dev/null
+++ b/Engine/Plugins/Animation/IKRig/Source/IKRigEditor/Public/RigEditor/SIKRigSkeleton.h
@@ -0,0 +1,233 @@
+// Copyright Epic Games, Inc. All Rights Reserved.
+
+#pragma once
+
+#include "CoreMinimal.h"
+#include "EditorUndoClient.h"
+#include "DragAndDrop/DecoratedDragDropOp.h"
+#include "Widgets/SCompoundWidget.h"
+#include "Widgets/Views/STreeView.h"
+
+#include "SIKRigSolverStack.h"
+
+#include "SIKRigSkeleton.generated.h"
+
+class FIKRigEditorController;
+class SIKRigSkeleton;
+class FIKRigEditorToolkit;
+class USkeletalMesh;
+
+enum class IKRigTreeElementType { BONE, GOAL, EFFECTOR, BONE_SETTINGS };
+
+class FIKRigTreeElement : public TSharedFromThis<FIKRigTreeElement>
+{
+public:
+	
+	FIKRigTreeElement(
+		const FName& InKey,
+		IKRigTreeElementType InType);
+
+	TSharedRef<ITableRow> MakeTreeRowWidget(
+		TSharedRef<FIKRigEditorController> InEditorController,
+		const TSharedRef<STableViewBase>& InOwnerTable,
+		TSharedRef<FIKRigTreeElement> InRigTreeElement,
+		TSharedRef<FUICommandList> InCommandList,
+		TSharedPtr<SIKRigSkeleton> InHierarchy);
+
+	FName Key;
+	IKRigTreeElementType ElementType;
+	TSharedPtr<FIKRigTreeElement> Parent;
+	TArray<TSharedPtr<FIKRigTreeElement>> Children;
+
+	/** effector meta-data */
+	FName EffectorGoalName = NAME_None;
+	int32 EffectorSolverIndex = INDEX_NONE;
+
+	/** bone setting meta-data */
+	FName BoneSettingBoneName = NAME_None;
+	int32 BoneSettingsSolverIndex = INDEX_NONE;
+
+	/** delegate for when the context menu requests a rename */
+	void RequestRename();
+	DECLARE_DELEGATE(FOnRenameRequested);
+	FOnRenameRequested OnRenameRequested;	
+};
+
+class SIKRigSkeletonItem : public STableRow<TSharedPtr<FIKRigTreeElement>>
+{
+public:
+	
+	void Construct(
+		const FArguments& InArgs,
+		TSharedRef<FIKRigEditorController> InEditorController,
+		const TSharedRef<STableViewBase>& OwnerTable,
+		TSharedRef<FIKRigTreeElement> InRigTreeElement,
+		TSharedRef<FUICommandList> InCommandList,
+		TSharedPtr<SIKRigSkeleton> InHierarchy);
+
+private:
+	
+	void OnNameCommitted(const FText& InText, ETextCommit::Type InCommitType) const;
+	FText GetName() const;
+
+	TWeakPtr<FIKRigTreeElement> WeakRigTreeElement;
+	TWeakPtr<FIKRigEditorController> EditorController;
+};
+
+class FIKRigSkeletonDragDropOp : public FDecoratedDragDropOp
+{
+public:
+	DRAG_DROP_OPERATOR_TYPE(FIKRigSkeletonDragDropOp, FDecoratedDragDropOp)
+    static TSharedRef<FIKRigSkeletonDragDropOp> New(TWeakPtr<FIKRigTreeElement> InElement);
+	virtual TSharedPtr<SWidget> GetDefaultDecorator() const override;
+	TWeakPtr<FIKRigTreeElement> Element;
+};
+
+class SIKRigSkeletonTreeView : public STreeView<TSharedPtr<FIKRigTreeElement>>
+{
+public:
+
+	virtual ~SIKRigSkeletonTreeView() {}
+
+	virtual FReply OnFocusReceived(const FGeometry& MyGeometry, const FFocusEvent& InFocusEvent) override
+	{
+		FReply Reply = STreeView<TSharedPtr<FIKRigTreeElement>>::OnFocusReceived(MyGeometry, InFocusEvent);
+		LastClickCycles = FPlatformTime::Cycles();
+		return Reply;
+	}
+
+	/** Save a snapshot of the internal map that tracks item expansion before tree reconstruction */
+	void SaveAndClearSparseItemInfos()
+	{
+		OldSparseItemInfos = SparseItemInfos;
+		ClearExpandedItems();
+	}
+
+	/** Restore the expansion infos map from the saved snapshot after tree reconstruction */
+	void RestoreSparseItemInfos(TSharedPtr<FIKRigTreeElement> ItemPtr)
+	{
+		for (const auto& Pair : OldSparseItemInfos)
+		{
+			if (Pair.Key->Key == ItemPtr->Key)
+			{
+				// the SparseItemInfos now reference the new element, but keep the same expansion state
+				SparseItemInfos.Add(ItemPtr, Pair.Value);
+				break;
+			}
+		}
+	}
+
+	/** slow double-click rename state*/
+	uint32 LastClickCycles = 0;
+	TWeakPtr<FIKRigTreeElement> LastSelected;
+
+private:
+	
+	/** A temporary snapshot of the SparseItemInfos in STreeView, used during SIKRigSkeleton::RefreshTreeView() */
+	TSparseItemMap OldSparseItemInfos;
+};
+
+USTRUCT()
+struct FIKRigSkeletonImportSettings
+{
+	GENERATED_BODY()
+
+	FIKRigSkeletonImportSettings() : Mesh(nullptr) {}
+
+	UPROPERTY(EditAnywhere, Category = "Skeleton Import")
+	USkeletalMesh* Mesh;
+};
+
+class SIKRigSkeleton : public SCompoundWidget, public FEditorUndoClient
+{
+public:
+	SLATE_BEGIN_ARGS(SIKRigSkeleton) {}
+	SLATE_END_ARGS()
+
+    ~SIKRigSkeleton() {};
+
+	void Construct(const FArguments& InArgs, TSharedRef<FIKRigEditorController> InEditorController);
+
+private:
+	/** SWidget interface */
+	virtual FReply OnKeyDown(const FGeometry& MyGeometry, const FKeyEvent& InKeyEvent);
+	/** END SWidget interface */
+	
+	/** Bind commands that this widget handles */
+	void BindCommands();
+
+	/** creating / renaming / deleting goals */
+	void HandleNewGoal();
+	bool CanAddNewGoal();
+	void HandleRenameElement() const;
+	void HandleDeleteItem();
+	bool CanDeleteItem() const;
+	/** END creating / renaming / deleting goals */
+
+	/** connecting/disconnecting goals to solvers */
+	void HandleConnectGoalToSolver();
+	void HandleDisconnectGoalFromSolver();
+	bool CanConnectGoalToSolvers() const;
+	bool CanDisconnectGoalFromSolvers() const;
+	void ConnectSelectedGoalsToSelectedSolvers(bool bConnect);
+	int32 GetNumSelectedGoalToSolverConnections(bool bCountOnlyConnected) const;
+	/** END connecting/disconnecting goals to solvers */
+
+	/** setting root bone */
+	void HandleSetRootBoneOnSolvers();
+	bool CanSetRootBoneOnSolvers();
+	/** END setting root bone */
+
+	/** per-bone settings */
+	void HandleAddBoneSettings();
+	bool CanAddBoneSettings();
+	void HandleRemoveBoneSettings();
+	bool CanRemoveBoneSettings();
+	/** END per-bone settings */
+
+	/** selection state queries */
+	void GetSelectedBones(TArray<TSharedPtr<FIKRigTreeElement>>& OutBoneItems);
+	void GetSelectedGoals(TArray<TSharedPtr<FIKRigTreeElement>>& OutSelectedGoals) const;
+	void SetSelectedGoalsFromViewport(const TArray<FName>& GoalNames);
+	/** END selection state queries */
+
+	/** skeleton import */
+	FReply OnImportSkeletonClicked();
+	void ImportSkeleton(const FAssetData& InAssetData);
+	EVisibility IsImportButtonVisible() const;
+	/** END skeleton import */
+
+	/** centralized editor controls (facilitate cross-communication between multiple UI elements)*/
+	TWeakPtr<FIKRigEditorController> EditorController;
+	
+	/** command list we bind to */
+	TSharedPtr<FUICommandList> CommandList;
+	
+	/** tree view widget */
+	TSharedPtr<SIKRigSkeletonTreeView> TreeView;
+	TArray<TSharedPtr<FIKRigTreeElement>> RootElements;
+	TArray<TSharedPtr<FIKRigTreeElement>> AllElements;
+	
+	/** tree view callbacks */
+	void RefreshTreeView(bool IsInitialSetup=false);
+	TSharedRef<ITableRow> MakeTableRowWidget(TSharedPtr<FIKRigTreeElement> InItem, const TSharedRef<STableViewBase>& OwnerTable);
+	void HandleGetChildrenForTree(TSharedPtr<FIKRigTreeElement> InItem, TArray<TSharedPtr<FIKRigTreeElement>>& OutChildren);
+	void OnSelectionChanged(TSharedPtr<FIKRigTreeElement> Selection, ESelectInfo::Type SelectInfo);
+	TSharedPtr< SWidget > CreateContextMenu();
+	void OnItemClicked(TSharedPtr<FIKRigTreeElement> InItem);
+	void OnItemDoubleClicked(TSharedPtr<FIKRigTreeElement> InItem);
+	void OnSetExpansionRecursive(TSharedPtr<FIKRigTreeElement> InItem, bool bShouldBeExpanded);
+	void SetExpansionRecursive(TSharedPtr<FIKRigTreeElement> InElement, bool bTowardsParent, bool bShouldBeExpanded);
+	void FillContextMenu(FMenuBuilder& MenuBuilder);
+	/** END tree view callbacks */
+
+	/** drag and drop */
+	FReply OnDragDetected(const FGeometry& MyGeometry, const FPointerEvent& MouseEvent);
+	TOptional<EItemDropZone> OnCanAcceptDrop(const FDragDropEvent& DragDropEvent, EItemDropZone DropZone, TSharedPtr<FIKRigTreeElement> TargetItem);
+	FReply OnAcceptDrop(const FDragDropEvent& DragDropEvent, EItemDropZone DropZone, TSharedPtr<FIKRigTreeElement> TargetItem);
+	/** END drag and drop */
+
+	friend SIKRigSkeletonItem;
+	friend FIKRigEditorController;
+	friend SIKRigSolverStack;
+};
diff --git a/Engine/Plugins/Animation/IKRig/Source/IKRigEditor/Public/RigEditor/SIKRigSolverStack.h b/Engine/Plugins/Animation/IKRig/Source/IKRigEditor/Public/RigEditor/SIKRigSolverStack.h
new file mode 100644
--- /dev/null
+++ b/Engine/Plugins/Animation/IKRig/Source/IKRigEditor/Public/RigEditor/SIKRigSolverStack.h
@@ -0,0 +1,120 @@
+// Copyright Epic Games, Inc. All Rights Reserved.
+
+#pragma once
+
+#include "CoreMinimal.h"
+#include "EditorUndoClient.h"
+#include "DragAndDrop/DecoratedDragDropOp.h"
+#include "Widgets/SCompoundWidget.h"
+#include "Widgets/Views/STableViewBase.h"
+#include "Widgets/Views/STableRow.h"
+#include "Widgets/Views/SListView.h"
+#include "Framework/Commands/UICommandList.h"
+
+class FIKRigEditorController;
+class SIKRigSolverStack;
+class FIKRigEditorToolkit;
+
+
+class FSolverStackElement
+{
+public:
+
+	TSharedRef<ITableRow> MakeListRowWidget(
+		const TSharedRef<STableViewBase>& InOwnerTable,
+        TSharedRef<FSolverStackElement> InStackElement,
+        TSharedPtr<SIKRigSolverStack> InSolverStack);
+
+	static TSharedRef<FSolverStackElement> Make(FText DisplayName, int32 SolverIndex)
+	{
+		return MakeShareable(new FSolverStackElement(DisplayName, SolverIndex));
+	}
+	
+	FText DisplayName;
+	int32 IndexInStack;
+
+private:
+	/** Hidden constructor, always use Make above */
+	FSolverStackElement(FText InDisplayName, int32 SolverIndex)
+        : DisplayName(InDisplayName), IndexInStack(SolverIndex)
+	{
+	}
+
+	/** Hidden constructor, always use Make above */
+	FSolverStackElement() {}
+};
+
+class SIKRigSolverStackItem : public STableRow<TSharedPtr<FSolverStackElement>>
+{
+public:
+	
+	void Construct(
+        const FArguments& InArgs,
+        const TSharedRef<STableViewBase>& OwnerTable,
+        TSharedRef<FSolverStackElement> InStackElement,
+        TSharedPtr<SIKRigSolverStack> InSolverStack);
+};
+
+class FIKRigSolverStackDragDropOp : public FDecoratedDragDropOp
+{
+public:
+	DRAG_DROP_OPERATOR_TYPE(FIKRigSolverStackDragDropOp, FDecoratedDragDropOp)
+    static TSharedRef<FIKRigSolverStackDragDropOp> New(TWeakPtr<FSolverStackElement> InElement);
+	virtual TSharedPtr<SWidget> GetDefaultDecorator() const override;
+	TWeakPtr<FSolverStackElement> Element;
+};
+
+
+typedef SListView< TSharedPtr<FSolverStackElement> > SSolverStackListViewType;
+
+class SIKRigSolverStack : public SCompoundWidget, public FEditorUndoClient
+{
+public:
+	SLATE_BEGIN_ARGS(SIKRigSolverStack) {}
+	SLATE_END_ARGS()
+
+    ~SIKRigSolverStack();
+
+	void Construct(const FArguments& InArgs, TSharedRef<FIKRigEditorController> InEditorController);
+
+private:
+	/** menu for adding new solver commands */
+	TSharedPtr<FUICommandList> CommandList;
+	
+	/** editor controller */
+	TWeakPtr<FIKRigEditorController> EditorController;
+	
+	/** the solver stack list view */
+	TSharedPtr<SSolverStackListViewType> ListView;
+	TArray< TSharedPtr<FSolverStackElement> > ListViewItems;
+
+	/** solver stack menu */
+	TSharedRef<SWidget> CreateAddNewMenuWidget();
+	void BuildAddNewMenu(FMenuBuilder& MenuBuilder);
+	bool IsAddSolverEnabled() const;
+	/** END solver stack menu */
+	
+	/** menu command callback for adding a new solver */
+	void AddNewSolver(UClass* Class);
+	/** when a solver is clicked on in the stack view */
+	void OnItemClicked(TSharedPtr<FSolverStackElement> InItem);
+
+	/** list view generate row callback */
+	TSharedRef<ITableRow> MakeListRowWidget(TSharedPtr<FSolverStackElement> InElement, const TSharedRef<STableViewBase>& OwnerTable);
+
+	/** call to refresh the stack view */
+	void RefreshStackView();
+
+	/** SWidget interface */
+	virtual FReply OnKeyDown(const FGeometry& MyGeometry, const FKeyEvent& InKeyEvent);
+	/** END SWidget interface */
+	
+	/** drag and drop operations */
+    FReply OnDragDetected(const FGeometry& MyGeometry, const FPointerEvent& MouseEvent);
+    TOptional<EItemDropZone> OnCanAcceptDrop(const FDragDropEvent& DragDropEvent, EItemDropZone DropZone, TSharedPtr<FSolverStackElement> TargetItem);
+    FReply OnAcceptDrop(const FDragDropEvent& DragDropEvent, EItemDropZone DropZone, TSharedPtr<FSolverStackElement> TargetItem);
+	/** END drag and drop operations */
+
+	friend SIKRigSolverStackItem;
+	friend FIKRigEditorController;
+};
