New Body Mover IK Rig solver type.

diff --git a/Engine/Plugins/Animation/IKRig/Source/IKRig/Private/Solvers/IKRig_BodyMover.cpp b/Engine/Plugins/Animation/IKRig/Source/IKRig/Private/Solvers/IKRig_BodyMover.cpp
new file mode 100644
--- /dev/null
+++ b/Engine/Plugins/Animation/IKRig/Source/IKRig/Private/Solvers/IKRig_BodyMover.cpp
@@ -0,0 +1,247 @@
+﻿// Copyright Epic Games, Inc. All Rights Reserved.
+
+#include "Solvers/IKRig_BodyMover.h"
+#include "IKRigDataTypes.h"
+#include "IKRigSkeleton.h"
+
+
+UIKRig_BodyMover::UIKRig_BodyMover()
+{
+}
+
+void UIKRig_BodyMover::Initialize(const FIKRigSkeleton& IKRigSkeleton)
+{
+	BodyBoneIndex = IKRigSkeleton.GetBoneIndexFromName(BodyBone);
+}
+
+void UIKRig_BodyMover::Solve(FIKRigSkeleton& IKRigSkeleton, const FIKRigGoalContainer& Goals)
+{
+	// no body bone specified
+	if (BodyBoneIndex == INDEX_NONE)
+	{
+		return;
+	}
+
+	// no effectors added
+	if (Effectors.IsEmpty())
+	{
+		return;
+	}
+	
+	// accumulate offsets to apply to body bone
+	check(IKRigSkeleton.RefPoseGlobal.IsValidIndex(BodyBoneIndex));
+	// the bone transform to modify (not const!)
+	FTransform& CurrentBodyTransform = IKRigSkeleton.CurrentPoseGlobal[BodyBoneIndex];
+
+	// calculate initial and current centroids
+	FVector InitialCentroid = FVector::ZeroVector;
+	FVector CurrentCentroid = FVector::ZeroVector;
+	for (UIKRig_BodyMoverEffector* Effector : Effectors)
+	{
+		FIKRigGoal Goal;
+		if (!Goals.GetGoalByName(Effector->GoalName, Goal))
+		{
+			return;
+		}
+
+		const int32 BoneIndex = IKRigSkeleton.GetBoneIndexFromName(Effector->BoneName);
+		const FTransform InitialEffector = IKRigSkeleton.RefPoseGlobal[BoneIndex];
+
+		InitialCentroid += InitialEffector.GetTranslation();
+		CurrentCentroid += Goal.FinalBlendedPosition;
+	}
+
+	// average centroids
+	const float InvNumEffectors = 1.0f / static_cast<float>(Effectors.Num());
+	InitialCentroid *= InvNumEffectors;
+	CurrentCentroid *= InvNumEffectors;
+
+	// accumulate deformation gradient to extract a rotation from
+	FVector DX = FVector::ZeroVector; // DX, DY, DZ are rows of 3x3 deformation gradient tensor
+	FVector DY = FVector::ZeroVector;
+	FVector DZ = FVector::ZeroVector;
+	for (UIKRig_BodyMoverEffector* Effector : Effectors)
+	{
+		FIKRigGoal Goal;
+		if (!Goals.GetGoalByName(Effector->GoalName, Goal))
+		{
+			return;
+		}
+
+		const int32 BoneIndex = IKRigSkeleton.GetBoneIndexFromName(Effector->BoneName);
+		const FTransform InitialEffector = IKRigSkeleton.RefPoseGlobal[BoneIndex];
+
+		//
+		// accumulate the deformation gradient tensor for all points
+		// "Meshless Deformations Based on Shape Matching"
+		// Equation 7 describes accumulation of deformation gradient from points
+		//
+		// P is normalized vector from INITIAL centroid to INITIAL point
+		// Q is normalized vector from CURRENT centroid to CURRENT point
+		FVector P = (InitialEffector.GetTranslation() - InitialCentroid).GetSafeNormal();
+		FVector Q = (Goal.FinalBlendedPosition - CurrentCentroid).GetSafeNormal();
+		// PQ^T is the outer product of P and Q which is a 3x3 matrix
+		// https://en.m.wikipedia.org/wiki/Outer_product
+		DX += FVector(P[0]*Q[0], P[0]*Q[1], P[0]*Q[2]);
+		DY += FVector(P[1]*Q[0], P[1]*Q[1], P[1]*Q[2]);
+		DZ += FVector(P[2]*Q[0], P[2]*Q[1], P[2]*Q[2]);
+	}
+
+	// extract "best fit" rotation from deformation gradient
+	FQuat RotationOffset = FQuat::Identity;
+	ExtractRotation(DX, DY, DZ, RotationOffset, 50);
+
+	// alpha blend the position offset and add it to the current bone location
+	const FVector TargetPosition = (CurrentCentroid - InitialCentroid) * PositionAlpha;
+	CurrentBodyTransform.AddToTranslation(TargetPosition);
+
+	// do per-axis alpha blend
+	FVector Euler = RotationOffset.Euler() * FVector(RotateXAlpha, RotateYAlpha, RotateZAlpha);
+	FQuat FinalRotationOffset = FQuat::MakeFromEuler(Euler);
+	// alpha blend the entire rotation offset
+	FinalRotationOffset = FQuat::FastLerp(FQuat::Identity, FinalRotationOffset, RotationAlpha).GetNormalized();
+	// add rotation offset to original rotation
+	CurrentBodyTransform.SetRotation(FinalRotationOffset * CurrentBodyTransform.GetRotation());
+
+	// do FK update of children
+	IKRigSkeleton.PropagateGlobalPoseBelowBone(BodyBoneIndex);
+}
+
+void UIKRig_BodyMover::ExtractRotation(
+	const FVector& DX,
+	const FVector& DY,
+	const FVector& DZ,
+	FQuat &Q,
+	const unsigned int MaxIter)
+{
+	// "A Robust Method to Extract the Rotational Part of Deformations" equation 7
+	// https://matthias-research.github.io/pages/publications/stablePolarDecomp.pdf
+	for (unsigned int Iter = 0; Iter < MaxIter; Iter++)
+	{
+		FMatrix R = FRotationMatrix::Make(Q);
+		FVector RCol0(R.M[0][0], R.M[0][1], R.M[0][2]);
+		FVector RCol1(R.M[1][0], R.M[1][1], R.M[1][2]);
+		FVector RCol2(R.M[2][0], R.M[2][1], R.M[2][2]);
+		FVector Omega = RCol0.Cross(DX) + RCol1.Cross(DY) + RCol2.Cross(DZ);
+		Omega *= (1.0f / fabs(RCol0.Dot(DX) + RCol1.Dot(DY) + RCol2.Dot(DZ)) + SMALL_NUMBER);
+		const float W = Omega.Size();
+		if (W < SMALL_NUMBER)
+		{
+			break;
+		}
+		Q = FQuat(FQuat((1.0 / W) * Omega, W)) * Q;
+		Q.Normalize();
+	}
+}
+
+void UIKRig_BodyMover::UpdateSolverSettings(UIKRigSolver* InSettings)
+{
+	if(UIKRig_BodyMover* Settings = Cast<UIKRig_BodyMover>(InSettings))
+	{
+		// copy solver settings
+		PositionAlpha = Settings->PositionAlpha;
+		RotationAlpha = Settings->RotationAlpha;
+		RotateXAlpha = Settings->RotateXAlpha;
+		RotateYAlpha = Settings->RotateYAlpha;
+		RotateZAlpha = Settings->RotateZAlpha;
+
+		// copy effector settings
+		for (const UIKRig_BodyMoverEffector* InEffector : Settings->Effectors)
+		{
+			for (UIKRig_BodyMoverEffector* Effector : Effectors)
+			{
+				if (Effector->GoalName == InEffector->GoalName)
+				{
+					Effector->RotationMultiplier = InEffector->RotationMultiplier;
+					break;
+				}
+			}
+		}
+	}
+}
+
+void UIKRig_BodyMover::AddGoal(const UIKRigEffectorGoal& NewGoal)
+{
+	UIKRig_BodyMoverEffector* NewEffector = NewObject<UIKRig_BodyMoverEffector>(this, UIKRig_BodyMoverEffector::StaticClass());
+	NewEffector->GoalName = NewGoal.GoalName;
+	NewEffector->BoneName = NewGoal.BoneName;
+	Effectors.Add(NewEffector);
+}
+
+void UIKRig_BodyMover::RemoveGoal(const FName& GoalName)
+{
+	// ensure goal even exists in this solver
+	const int32 GoalIndex = GetIndexOfGoal(GoalName);
+	if (GoalIndex == INDEX_NONE)
+	{
+		return;
+	}
+
+	// remove it
+	Effectors.RemoveAt(GoalIndex);
+}
+
+void UIKRig_BodyMover::RenameGoal(const FName& OldName, const FName& NewName)
+{
+	// ensure goal even exists in this solver
+	const int32 GoalIndex = GetIndexOfGoal(OldName);
+	if (GoalIndex == INDEX_NONE)
+	{
+		return;
+	}
+
+	// rename
+	Effectors[GoalIndex]->GoalName = NewName;
+}
+
+void UIKRig_BodyMover::SetGoalBone(const FName& GoalName, const FName& NewBoneName)
+{
+	// ensure goal even exists in this solver
+	const int32 GoalIndex = GetIndexOfGoal(GoalName);
+	if (GoalIndex == INDEX_NONE)
+	{
+		return;
+	}
+
+	// rename
+	Effectors[GoalIndex]->BoneName = NewBoneName;
+}
+
+bool UIKRig_BodyMover::IsGoalConnected(const FName& GoalName) const
+{
+	return GetIndexOfGoal(GoalName) != INDEX_NONE;
+}
+
+UObject* UIKRig_BodyMover::GetEffectorWithGoal(const FName& GoalName)
+{
+	const int32 GoalIndex = GetIndexOfGoal(GoalName);
+	if (GoalIndex == INDEX_NONE)
+	{
+		return nullptr;
+	}
+
+	return Effectors[GoalIndex];
+}
+
+bool UIKRig_BodyMover::IsBoneAffectedBySolver(const FName& BoneName, const FIKRigSkeleton& IKRigSkeleton) const
+{
+	return IKRigSkeleton.IsBoneInDirectLineage(BoneName, BodyBone);
+}
+
+void UIKRig_BodyMover::SetRootBone(const FName& RootBoneName)
+{
+	BodyBone = RootBoneName;
+}
+
+int32 UIKRig_BodyMover::GetIndexOfGoal(const FName& OldName) const
+{
+	for (int32 i=0; i<Effectors.Num(); ++i)
+	{
+		if (Effectors[i]->GoalName == OldName)
+		{
+			return i;
+		}
+	}
+
+	return INDEX_NONE;
+}
diff --git a/Engine/Plugins/Animation/IKRig/Source/IKRig/Public/Solvers/IKRig_BodyMover.h b/Engine/Plugins/Animation/IKRig/Source/IKRig/Public/Solvers/IKRig_BodyMover.h
new file mode 100644
--- /dev/null
+++ b/Engine/Plugins/Animation/IKRig/Source/IKRig/Public/Solvers/IKRig_BodyMover.h
@@ -0,0 +1,94 @@
+﻿// Copyright Epic Games, Inc. All Rights Reserved.
+
+#pragma once
+
+#include "IKRigDefinition.h"
+#include "IKRigSolver.h"
+
+#include "IKRig_BodyMover.generated.h"
+
+UCLASS()
+class IKRIG_API UIKRig_BodyMoverEffector : public UObject
+{
+	GENERATED_BODY()
+
+public:
+
+	UPROPERTY(VisibleAnywhere, Category = "Body Mover Effector")
+	FName GoalName;
+
+	UPROPERTY(VisibleAnywhere, Category = "Body Mover Effector")
+	FName BoneName;
+
+	/** Scale the amount this effector rotates the body. Range is 0-10. Default is 1.0. */
+	UPROPERTY(EditAnywhere, Category = "Body Mover Effector", meta = (ClampMin = "0", ClampMa = "10", UIMin = "0.0", UIMax = "10.0"))
+	float RotationMultiplier = 1.0f;
+};
+
+UCLASS(EditInlineNew)
+class IKRIG_API UIKRig_BodyMover : public UIKRigSolver
+{
+	GENERATED_BODY()
+
+public:
+	
+	UIKRig_BodyMover();
+
+	/** The target bone to move with the effectors. */
+	UPROPERTY(VisibleAnywhere, Category = "Body Mover Solver")
+	FName BodyBone;
+
+	/** Blend the translational effect of this solver on/off. Range is 0-1. Default is 1.0. */
+	UPROPERTY(EditAnywhere, Category = "Body Mover Solver", meta = (ClampMin = "0", ClampMa = "1", UIMin = "0.0", UIMax = "1.0"))
+	float PositionAlpha = 1.0f;
+
+	/** Blend the total rotational effect on/off. Range is 0-1. Default is 1.0. */
+	UPROPERTY(EditAnywhere, Category = "Body Mover Solver", meta = (ClampMin = "0", ClampMa = "1", UIMin = "0.0", UIMax = "1.0"))
+	float RotationAlpha = 1.0f;
+
+	/** Blend the X-axis rotational effect on/off. Range is 0-1. Default is 1.0. */
+	UPROPERTY(EditAnywhere, Category = "Body Mover Solver", meta = (ClampMin = "0", ClampMa = "1", UIMin = "0.0", UIMax = "1.0"))
+	float RotateXAlpha = 1.0f;
+
+	/** Blend the Y-axis rotational effect on/off. Range is 0-1. Default is 1.0. */
+	UPROPERTY(EditAnywhere, Category = "Body Mover Solver", meta = (ClampMin = "0", ClampMa = "1", UIMin = "0.0", UIMax = "1.0"))
+	float RotateYAlpha = 1.0f;
+
+	/** Blend the Z-axis rotational effect on/off. Range is 0-1. Default is 1.0. */
+	UPROPERTY(EditAnywhere, Category = "Body Mover Solver", meta = (ClampMin = "0", ClampMa = "1", UIMin = "0.0", UIMax = "1.0"))
+	float RotateZAlpha = 1.0f;
+	
+	UPROPERTY()
+	TArray<UIKRig_BodyMoverEffector*> Effectors;
+
+	/** UIKRigSolver interface */
+	virtual void Initialize(const FIKRigSkeleton& IKRigSkeleton) override;
+	virtual void Solve(FIKRigSkeleton& IKRigSkeleton, const FIKRigGoalContainer& Goals) override;
+	virtual void UpdateSolverSettings(UIKRigSolver* InSettings) override;
+	virtual bool IsBoneAffectedBySolver(const FName& BoneName, const FIKRigSkeleton& IKRigSkeleton) const override;
+	// goals
+	virtual void AddGoal(const UIKRigEffectorGoal& NewGoal) override;
+	virtual void RemoveGoal(const FName& GoalName) override;
+	virtual void RenameGoal(const FName& OldName, const FName& NewName) override;
+	virtual void SetGoalBone(const FName& GoalName, const FName& NewBoneName) override;
+	virtual bool IsGoalConnected(const FName& GoalName) const override;
+	virtual UObject* GetEffectorWithGoal(const FName& GoalName) override;
+	// root bone can be set on this solver
+	virtual bool CanSetRootBone() const override { return true; };
+	virtual void SetRootBone(const FName& RootBoneName) override;
+	/** END UIKRigSolver interface */
+
+private:
+
+	int32 GetIndexOfGoal(const FName& OldName) const;
+
+	static void ExtractRotation(
+	    const FVector& DX,
+	    const FVector& DY,
+	    const FVector& DZ,
+	    FQuat &Q,
+	    const unsigned int maxIter);
+
+	int32 BodyBoneIndex;
+};
+
