Fixing static analysis issues in IK Rig Editor.

diff --git a/Engine/Plugins/Animation/IKRig/Source/IKRig/Private/Solvers/IKRig_PBIKSolver.cpp b/Engine/Plugins/Animation/IKRig/Source/IKRig/Private/Solvers/IKRig_PBIKSolver.cpp
--- a/Engine/Plugins/Animation/IKRig/Source/IKRig/Private/Solvers/IKRig_PBIKSolver.cpp
+++ b/Engine/Plugins/Animation/IKRig/Source/IKRig/Private/Solvers/IKRig_PBIKSolver.cpp
@@ -241,7 +241,7 @@ void UIKRigPBIKSolver::AddBoneSetting(const FName& BoneName)
 
 void UIKRigPBIKSolver::RemoveBoneSetting(const FName& BoneName)
 {
-	UIKRig_PBIKBoneSettings* BoneSettingToRemove; 
+	UIKRig_PBIKBoneSettings* BoneSettingToRemove = nullptr; 
 	for (UIKRig_PBIKBoneSettings* BoneSetting : BoneSettings)
 	{
 		if (BoneSetting->Bone == BoneName)
@@ -251,7 +251,10 @@ void UIKRigPBIKSolver::RemoveBoneSetting(const FName& BoneName)
 		}
 	}
 
-	BoneSettings.Remove(BoneSettingToRemove);
+	if (BoneSettingToRemove)
+	{
+		BoneSettings.Remove(BoneSettingToRemove);
+	}
 }
 
 UObject* UIKRigPBIKSolver::GetBoneSetting(const FName& BoneName) const
diff --git a/Engine/Plugins/Animation/IKRig/Source/IKRigEditor/Private/RigEditor/SIKRigSkeleton.cpp b/Engine/Plugins/Animation/IKRig/Source/IKRigEditor/Private/RigEditor/SIKRigSkeleton.cpp
--- a/Engine/Plugins/Animation/IKRig/Source/IKRigEditor/Private/RigEditor/SIKRigSkeleton.cpp
+++ b/Engine/Plugins/Animation/IKRig/Source/IKRigEditor/Private/RigEditor/SIKRigSkeleton.cpp
@@ -71,7 +71,7 @@ void SIKRigSkeletonItem::Construct(
 	}
 		
 	// determine which icon to use for tree element
-	const FSlateBrush* Brush;
+	const FSlateBrush* Brush = FAppStyle::Get().GetBrush("SkeletonTree.Bone");
 	switch(InRigTreeElement->ElementType)
 	{
 		case IKRigTreeElementType::BONE:
