Fixing GC issue in IK Rig Editor
Adding Space options to IK Rig Goals (Additive, Component, World)

diff --git a/Engine/Plugins/Animation/IKRig/Source/IKRig/Private/AnimNodes/AnimNode_IKRetargeter.cpp b/Engine/Plugins/Animation/IKRig/Source/IKRig/Private/AnimNodes/AnimNode_IKRetargeter.cpp
--- a/Engine/Plugins/Animation/IKRig/Source/IKRig/Private/AnimNodes/AnimNode_IKRetargeter.cpp
+++ b/Engine/Plugins/Animation/IKRig/Source/IKRig/Private/AnimNodes/AnimNode_IKRetargeter.cpp
@@ -3,7 +3,6 @@
 #include "AnimNodes/AnimNode_IKRetargeter.h"
 #include "Animation/AnimInstanceProxy.h"
 
-//#pragma optimize("", off)
 
 void FAnimNode_IKRetargeter::Initialize_AnyThread(const FAnimationInitializeContext& Context)
 {
@@ -94,7 +93,7 @@ void FAnimNode_IKRetargeter::EnsureInitialized(const UAnimInstance* InAnimInstan
 	   CurrentlyUsedSourceMeshComponent.IsValid()
     && CurrentlyUsedSourceMesh.IsValid()
     && CurrentlyUsedTargetMesh.IsValid()
-	&& CurrentlyUsedRetargeter.IsValid()
+	&& CurrentlyUsedRetargeter
 	&& CurrentlyUsedSourceIKRig.IsValid()
 	&& CurrentlyUsedTargetIKRig.IsValid();
 	if (!bPartsLoaded)
@@ -141,7 +140,7 @@ void FAnimNode_IKRetargeter::InitializeRetargetData(const UAnimInstance* InAnimI
 	}
 
 	const bool bRetargeterIsValid =
-	   CurrentlyUsedRetargeter.IsValid()
+	   CurrentlyUsedRetargeter
 	&& CurrentlyUsedSourceIKRig.IsValid()
 	&& CurrentlyUsedTargetIKRig.IsValid();
 	if (!bRetargeterIsValid)
diff --git a/Engine/Plugins/Animation/IKRig/Source/IKRig/Private/AnimNodes/AnimNode_IKRig.cpp b/Engine/Plugins/Animation/IKRig/Source/IKRig/Private/AnimNodes/AnimNode_IKRig.cpp
--- a/Engine/Plugins/Animation/IKRig/Source/IKRig/Private/AnimNodes/AnimNode_IKRig.cpp
+++ b/Engine/Plugins/Animation/IKRig/Source/IKRig/Private/AnimNodes/AnimNode_IKRig.cpp
@@ -10,10 +10,11 @@
 
 
 FAnimNode_IKRig::FAnimNode_IKRig()
-	: RigDefinitionAsset (nullptr)
+	:	RigDefinitionAsset (nullptr),
 #if WITH_EDITORONLY_DATA
-	, bEnableDebugDraw(false)
+		bEnableDebugDraw(false),
 #endif
+		IKRigProcessor(nullptr)
 {
 }
 
@@ -30,12 +31,12 @@ void FAnimNode_IKRig::Evaluate_AnyThread(FPoseContext& Output)
 		Output.ResetToRefPose();
 	}
 
-	if (!RigDefinitionAsset)
+	if (!(RigDefinitionAsset && IKRigProcessor))
 	{
 		return;
 	}
 
-	if (IKRigProcessor.NeedsInitialized(RigDefinitionAsset))
+	if (IKRigProcessor->NeedsInitialized(RigDefinitionAsset))
 	{
 		return;
 	}
@@ -44,8 +45,9 @@ void FAnimNode_IKRig::Evaluate_AnyThread(FPoseContext& Output)
 	CopyInputPoseToSolver(Output.Pose);
 	// update target goal transforms
 	AssignGoalTargets();
-	// run stack of solvers, 
-	IKRigProcessor.Solve();
+	// run stack of solvers,
+	const FTransform WorldToComponent =  Output.AnimInstanceProxy->GetComponentTransform().Inverse();
+	IKRigProcessor->Solve(WorldToComponent);
 	// updates transforms with new pose
 	CopyOutputPoseToAnimGraph(Output.Pose);
 	
@@ -58,13 +60,13 @@ void FAnimNode_IKRig::CopyInputPoseToSolver(FCompactPose& InputPose)
 	// start Solve() from REFERENCE pose
 	if (bStartFromRefPose)
 	{
-		IKRigProcessor.SetInputPoseToRefPose();
+		IKRigProcessor->SetInputPoseToRefPose();
 		return;
 	}
 	
 	// start Solve() from INPUT pose
 	// copy local bone transforms into IKRigProcessor skeleton
-	FIKRigSkeleton& IKRigSkeleton = IKRigProcessor.GetSkeleton();
+	FIKRigSkeleton& IKRigSkeleton = IKRigProcessor->GetSkeleton();
 	for (FCompactPoseBoneIndex CPIndex : InputPose.ForEachBoneIndex())
 	{
 		int32* Index = CompactPoseToRigIndices.Find(CPIndex);
@@ -86,26 +88,26 @@ void FAnimNode_IKRig::AssignGoalTargets()
 	// this is used to live preview results from the IK Rig editor
 	if (bDriveWithSourceAsset)
 	{
-		IKRigProcessor.CopyAllInputsFromSourceAssetAtRuntime(RigDefinitionAsset);
+		IKRigProcessor->CopyAllInputsFromSourceAssetAtRuntime(RigDefinitionAsset);
 		return;
 	}
 	
 	// copy transforms from this anim node's goal pins from blueprint
 	for (const FIKRigGoal& Goal : Goals)
 	{
-		IKRigProcessor.SetIKGoal(Goal);
+		IKRigProcessor->SetIKGoal(Goal);
 	}
 
 	// override any goals that were manually set with goals from goal creator components (they take precedence)
 	for (const TPair<FName, FIKRigGoal>& GoalPair : GoalsFromGoalCreators)
 	{
-		IKRigProcessor.SetIKGoal(GoalPair.Value);
+		IKRigProcessor->SetIKGoal(GoalPair.Value);
 	}
 }
 
 void FAnimNode_IKRig::CopyOutputPoseToAnimGraph(FCompactPose& OutputPose)
 {
-	FIKRigSkeleton& IKRigSkeleton = IKRigProcessor.GetSkeleton();
+	FIKRigSkeleton& IKRigSkeleton = IKRigProcessor->GetSkeleton();
 	
 	// update local transforms of current IKRig pose
 	IKRigSkeleton.UpdateAllLocalTransformFromGlobal();
@@ -159,12 +161,17 @@ void FAnimNode_IKRig::Update_AnyThread(const FAnimationUpdateContext& Context)
 
 void FAnimNode_IKRig::PreUpdate(const UAnimInstance* InAnimInstance)
 {
-	if (IKRigProcessor.NeedsInitialized(RigDefinitionAsset))
+	if (!IsValid(IKRigProcessor))
+	{
+		IKRigProcessor = NewObject<UIKRigProcessor>(InAnimInstance->GetOwningComponent());	
+	}
+	
+	if (IKRigProcessor->NeedsInitialized(RigDefinitionAsset))
 	{
 		// get the retargeted local ref pose to initialize the IK with
 		const FReferenceSkeleton& RefSkeleton = InAnimInstance->CurrentSkeleton->GetReferenceSkeleton();
 		// initialize the IK Rig (will only try once on the current version of the rig asset)
-		IKRigProcessor.Initialize(RigDefinitionAsset, RefSkeleton, InAnimInstance->GetSkelMeshComponent());
+		IKRigProcessor->Initialize(RigDefinitionAsset, RefSkeleton);
 	}
 	
 	// cache list of goal creator components on the actor
@@ -194,49 +201,6 @@ void FAnimNode_IKRig::PreUpdate(const UAnimInstance* InAnimInstance)
 	}
 }
 
-bool FAnimNode_IKRig::RebuildGoalList()
-{
-	if (!RigDefinitionAsset)
-	{
-		return false;
-	}
-
-	// number of goals changed
-	const int32 NumGoalsInRig = RigDefinitionAsset->Goals.Num();
-	if (Goals.Num() != NumGoalsInRig)
-	{
-		Goals.SetNum(NumGoalsInRig);
-		for (int32 i=0; i<NumGoalsInRig; ++i)
-		{
-			Goals[i].Name = RigDefinitionAsset->Goals[i]->GoalName;
-		}
-		return true;
-	}
-
-	// potentially number of goals remains identical, but only names changed
-	bool bNameUpdated = false;
-	for (int32 i=0; i<NumGoalsInRig; ++i)
-	{
-		if (Goals[i].Name != RigDefinitionAsset->Goals[i]->GoalName)
-		{
-			Goals[i].Name = RigDefinitionAsset->Goals[i]->GoalName;
-			bNameUpdated = true;
-		}
-	}
-
-	return bNameUpdated;
-}
-
-FName FAnimNode_IKRig::GetGoalName(int32 Index) const
-{
-	if (RigDefinitionAsset && RigDefinitionAsset->Goals.IsValidIndex(Index))
-	{
-		return RigDefinitionAsset->Goals[Index]->GoalName;	
-	}
-
-	return NAME_None;
-}
-
 void FAnimNode_IKRig::CacheBones_AnyThread(const FAnimationCacheBonesContext& Context)
 {
 	FAnimNode_Base::CacheBones_AnyThread(Context);
diff --git a/Engine/Plugins/Animation/IKRig/Source/IKRig/Private/IKRigDataTypes.cpp b/Engine/Plugins/Animation/IKRig/Source/IKRig/Private/IKRigDataTypes.cpp
--- a/Engine/Plugins/Animation/IKRig/Source/IKRig/Private/IKRigDataTypes.cpp
+++ b/Engine/Plugins/Animation/IKRig/Source/IKRig/Private/IKRigDataTypes.cpp
@@ -1,7 +1,6 @@
 // Copyright Epic Games, Inc. All Rights Reserved.
 
 #include "IKRigDataTypes.h"
-#include "IKRigSolver.h"
 
 void FIKRigGoalContainer::InitializeFromGoals(const TArray<UIKRigEffectorGoal*>& InGoals)
 {
diff --git a/Engine/Plugins/Animation/IKRig/Source/IKRig/Private/IKRigProcessor.cpp b/Engine/Plugins/Animation/IKRig/Source/IKRig/Private/IKRigProcessor.cpp
--- a/Engine/Plugins/Animation/IKRig/Source/IKRig/Private/IKRigProcessor.cpp
+++ b/Engine/Plugins/Animation/IKRig/Source/IKRig/Private/IKRigProcessor.cpp
@@ -4,8 +4,7 @@
 #include "IKRigDefinition.h"
 #include "IKRigSolver.h"
 
-
-void FIKRigProcessor::Initialize(UIKRigDefinition* InRigAsset, const FReferenceSkeleton& RefSkeleton, UObject* Outer)
+void UIKRigProcessor::Initialize(UIKRigDefinition* InRigAsset, const FReferenceSkeleton& RefSkeleton)
 {
 	// we instantiate UObjects here which MUST be done on game thread...
 	check(IsInGameThread());
@@ -19,6 +18,24 @@ void FIKRigProcessor::Initialize(UIKRigDefinition* InRigAsset, const FReferenceS
 		return;
 	}
 
+	if (InRigAsset->Skeleton.BoneNames.IsEmpty())
+	{
+		UE_LOG(LogTemp, Warning, TEXT("Trying to initialize IKRigProcessor with a IKRigDefinition that has no skeleton."));
+		return;
+	}
+
+	if (InRigAsset->Solvers.IsEmpty())
+	{
+		UE_LOG(LogTemp, Warning, TEXT("Trying to initialize IKRigProcessor with a IKRigDefinition that has no solvers."));
+		return;
+	}
+
+	if (InRigAsset->Goals.IsEmpty())
+	{
+		UE_LOG(LogTemp, Warning, TEXT("Trying to initialize IKRigProcessor with a IKRigDefinition that has no goals."));
+		return;
+	}
+
 	// bail out if we've already tried initializing with this exact version of the rig asset
 	if (LastVersionTried == InRigAsset->GetAssetVersion())
 	{
@@ -85,8 +102,7 @@ void FIKRigProcessor::Initialize(UIKRigDefinition* InRigAsset, const FReferenceS
 		// new solver name
 		FString Name = IKRigSolver->GetName() + "_SolverInstance_";
 		Name.AppendInt(SolverIndex++);
-		UIKRigSolver* Solver = DuplicateObject(IKRigSolver, Outer, FName(*Name));
-		Solver->AddToRoot();
+		UIKRigSolver* Solver = DuplicateObject(IKRigSolver, this, FName(*Name));
 		Solver->Initialize(Skeleton);
 		Solvers.Add(Solver);
 	}
@@ -95,7 +111,7 @@ void FIKRigProcessor::Initialize(UIKRigDefinition* InRigAsset, const FReferenceS
 	bInitialized = true;
 }
 
-void FIKRigProcessor::SetInputPoseGlobal(const TArray<FTransform>& InGlobalBoneTransforms) 
+void UIKRigProcessor::SetInputPoseGlobal(const TArray<FTransform>& InGlobalBoneTransforms) 
 {
 	check(bInitialized);
 	check(InGlobalBoneTransforms.Num() == Skeleton.CurrentPoseGlobal.Num());
@@ -103,31 +119,31 @@ void FIKRigProcessor::SetInputPoseGlobal(const TArray<FTransform>& InGlobalBoneT
 	Skeleton.UpdateAllLocalTransformFromGlobal();
 }
 
-void FIKRigProcessor::SetInputPoseToRefPose()
+void UIKRigProcessor::SetInputPoseToRefPose()
 {
 	check(bInitialized);
 	Skeleton.CurrentPoseGlobal = Skeleton.RefPoseGlobal;
 	Skeleton.UpdateAllLocalTransformFromGlobal();
 }
 
-void FIKRigProcessor::SetIKGoal(const FIKRigGoal& InGoal)
+void UIKRigProcessor::SetIKGoal(const FIKRigGoal& InGoal)
 {
 	check(bInitialized);
 	GoalContainer.SetIKGoal(InGoal);
 }
 
-void FIKRigProcessor::SetIKGoal(const UIKRigEffectorGoal* InGoal)
+void UIKRigProcessor::SetIKGoal(const UIKRigEffectorGoal* InGoal)
 {
 	check(bInitialized);
 	GoalContainer.SetIKGoal(InGoal);
 }
 
-void FIKRigProcessor::Solve()
+void UIKRigProcessor::Solve(const FTransform& ComponentToWorld)
 {
 	check(bInitialized);
-
-	// blend goals towards input pose by alpha
-	BlendGoalsByAlpha();
+	
+	// convert goals into component space and blend towards input pose by alpha
+	ResolveFinalGoalTransforms(ComponentToWorld);
 
 	// run all the solvers
 	for (UIKRigSolver* Solver : Solvers)
@@ -142,12 +158,12 @@ void FIKRigProcessor::Solve()
 	Skeleton.NormalizeRotations(Skeleton.CurrentPoseGlobal);
 }
 
-void FIKRigProcessor::CopyOutputGlobalPoseToArray(TArray<FTransform>& OutputPoseGlobal) const
+void UIKRigProcessor::CopyOutputGlobalPoseToArray(TArray<FTransform>& OutputPoseGlobal) const
 {
 	OutputPoseGlobal = Skeleton.CurrentPoseGlobal;
 }
 
-void FIKRigProcessor::CopyAllInputsFromSourceAssetAtRuntime(UIKRigDefinition* IKRigAsset)
+void UIKRigProcessor::CopyAllInputsFromSourceAssetAtRuntime(UIKRigDefinition* IKRigAsset)
 {
 	// copy goal settings
 	for (const UIKRigEffectorGoal* AssetGoal : IKRigAsset->Goals)
@@ -164,7 +180,7 @@ void FIKRigProcessor::CopyAllInputsFromSourceAssetAtRuntime(UIKRigDefinition* IK
 	}
 }
 
-bool FIKRigProcessor::NeedsInitialized(UIKRigDefinition* IKRigAsset) const
+bool UIKRigProcessor::NeedsInitialized(UIKRigDefinition* IKRigAsset) const
 {
 	if (!bInitialized)
 	{
@@ -184,20 +200,19 @@ bool FIKRigProcessor::NeedsInitialized(UIKRigDefinition* IKRigAsset) const
 	return false;
 }
 
-const FIKRigGoalContainer& FIKRigProcessor::GetGoalContainer() const
+const FIKRigGoalContainer& UIKRigProcessor::GetGoalContainer() const
 {
 	check(bInitialized);
 	return GoalContainer;
 }
 
-FIKRigSkeleton& FIKRigProcessor::GetSkeleton()
+FIKRigSkeleton& UIKRigProcessor::GetSkeleton()
 {
 	check(bInitialized);
 	return Skeleton;
 }
 
-
-void FIKRigProcessor::BlendGoalsByAlpha()
+void UIKRigProcessor::ResolveFinalGoalTransforms(const FTransform& WorldToComponent)
 {
 	for (TPair<FName, FIKRigGoal>& GoalPair : GoalContainer.Goals)
 	{
@@ -211,15 +226,59 @@ void FIKRigProcessor::BlendGoalsByAlpha()
 		FIKRigGoal& Goal = GoalPair.Value;
 		const FGoalBone& GoalBone = GoalBones[Goal.Name];
 		const FTransform& InputPoseBoneTransform = Skeleton.CurrentPoseGlobal[GoalBone.BoneIndex];
+
+		FVector ComponentSpaceGoalPosition = Goal.Position;
+		FQuat ComponentSpaceGoalRotation = Goal.Rotation.Quaternion();
+
+		// put goal POSITION in Component Space
+		switch (Goal.PositionSpace)
+		{
+		case EIKRigGoalSpace::Additive:
+			// add position offset to bone position
+			ComponentSpaceGoalPosition = Skeleton.CurrentPoseGlobal[GoalBone.BoneIndex].GetLocation() + Goal.Position;
+			break;
+		case EIKRigGoalSpace::Component:
+			// was already supplied in Component Space
+			break;
+		case EIKRigGoalSpace::World:
+			// convert from World Space to Component Space
+			ComponentSpaceGoalPosition = WorldToComponent.TransformPosition(Goal.Position);
+			break;
+		default:
+			checkNoEntry();
+			break;
+		}
 		
+		// put goal ROTATION in Component Space
+		switch (Goal.RotationSpace)
+		{
+		case EIKRigGoalSpace::Additive:
+			// add rotation offset to bone rotation
+			ComponentSpaceGoalRotation = Goal.Rotation.Quaternion() * Skeleton.CurrentPoseGlobal[GoalBone.BoneIndex].GetRotation();
+			break;
+		case EIKRigGoalSpace::Component:
+			// was already supplied in Component Space
+			break;
+		case EIKRigGoalSpace::World:
+			// convert from World Space to Component Space
+			ComponentSpaceGoalRotation = WorldToComponent.TransformRotation(Goal.Rotation.Quaternion());
+			break;
+		default:
+			checkNoEntry();
+			break;
+		}
+
+		// blend by alpha from the input pose, to the supplied goal transform
+		// when Alpha is 0, the goal transform matches the bone transform at the input pose.
+		// when Alpha is 1, the goal transform is left fully intact
 		Goal.FinalBlendedPosition = FMath::Lerp(
             InputPoseBoneTransform.GetTranslation(),
-            Goal.Position,
+            ComponentSpaceGoalPosition,
             Goal.PositionAlpha);
 		
 		Goal.FinalBlendedRotation = FQuat::FastLerp(
             InputPoseBoneTransform.GetRotation(),
-            Goal.Rotation.Quaternion(),
+            ComponentSpaceGoalRotation,
             Goal.RotationAlpha);
 	}
 }
diff --git a/Engine/Plugins/Animation/IKRig/Source/IKRig/Private/Retargeter/IKRetargeter.cpp b/Engine/Plugins/Animation/IKRig/Source/IKRig/Private/Retargeter/IKRetargeter.cpp
--- a/Engine/Plugins/Animation/IKRig/Source/IKRig/Private/Retargeter/IKRetargeter.cpp
+++ b/Engine/Plugins/Animation/IKRig/Source/IKRig/Private/Retargeter/IKRetargeter.cpp
@@ -783,15 +783,19 @@ bool UIKRetargeter::InitializeBoneChainPairs()
 bool UIKRetargeter::InitializeIKRig(UObject* Outer, const FReferenceSkeleton& InRefSkeleton)
 {	
 	// initialize IK Rig runtime processor
-	IKRigProcessor.Initialize(TargetIKRigAsset, InRefSkeleton, Outer);
-	if (IKRigProcessor.NeedsInitialized(TargetIKRigAsset))
+	if (!IKRigProcessor)
+	{
+		IKRigProcessor = NewObject<UIKRigProcessor>(Outer);	
+	}
+	IKRigProcessor->Initialize(TargetIKRigAsset, InRefSkeleton);
+	if (IKRigProcessor->NeedsInitialized(TargetIKRigAsset))
 	{
 		return false;
 	}
 
 	// validate that all IK bone chains have an associated Goal
 	TArray<FName> GoalNames;
-	IKRigProcessor.GetGoalContainer().Goals.GenerateKeyArray(GoalNames);
+	IKRigProcessor->GetGoalContainer().Goals.GenerateKeyArray(GoalNames);
 	for (FRetargetChainPairIK& ChainPair : ChainPairsIK)
 	{
 		// does the IK rig have the IK goal this bone chain requires?
@@ -876,13 +880,13 @@ void UIKRetargeter::RunIKRetarget(
 			OutIKGoal.EndEffectorRotation,
 			1.0f,
 			1.0f);
-		IKRigProcessor.SetIKGoal(Goal);
+		IKRigProcessor->SetIKGoal(Goal);
 	}
 
 	// copy input pose to start IK solve from
-	IKRigProcessor.SetInputPoseGlobal(OutGlobalPose);
+	IKRigProcessor->SetInputPoseGlobal(OutGlobalPose);
 	// run IK solve
-	IKRigProcessor.Solve();
+	IKRigProcessor->Solve();
 	// copy results of solve
-	IKRigProcessor.CopyOutputGlobalPoseToArray(OutGlobalPose);
+	IKRigProcessor->CopyOutputGlobalPoseToArray(OutGlobalPose);
 }
diff --git a/Engine/Plugins/Animation/IKRig/Source/IKRig/Public/AnimNodes/AnimNode_IKRetargeter.h b/Engine/Plugins/Animation/IKRig/Source/IKRig/Public/AnimNodes/AnimNode_IKRetargeter.h
--- a/Engine/Plugins/Animation/IKRig/Source/IKRig/Public/AnimNodes/AnimNode_IKRetargeter.h
+++ b/Engine/Plugins/Animation/IKRig/Source/IKRig/Public/AnimNodes/AnimNode_IKRetargeter.h
@@ -58,7 +58,8 @@ private:
 	TWeakObjectPtr<USkeletalMeshComponent>	CurrentlyUsedSourceMeshComponent;
 	TWeakObjectPtr<USkeletalMesh>			CurrentlyUsedSourceMesh;
 	TWeakObjectPtr<USkeletalMesh>			CurrentlyUsedTargetMesh;
-	TWeakObjectPtr<UIKRetargeter>			CurrentlyUsedRetargeter;
+	UPROPERTY(Transient)
+	TObjectPtr<UIKRetargeter>				CurrentlyUsedRetargeter = nullptr;
 	TWeakObjectPtr<UIKRigDefinition>		CurrentlyUsedSourceIKRig;
 	TWeakObjectPtr<UIKRigDefinition>		CurrentlyUsedTargetIKRig;
 
diff --git a/Engine/Plugins/Animation/IKRig/Source/IKRig/Public/AnimNodes/AnimNode_IKRig.h b/Engine/Plugins/Animation/IKRig/Source/IKRig/Public/AnimNodes/AnimNode_IKRig.h
--- a/Engine/Plugins/Animation/IKRig/Source/IKRig/Public/AnimNodes/AnimNode_IKRig.h
+++ b/Engine/Plugins/Animation/IKRig/Source/IKRig/Public/AnimNodes/AnimNode_IKRig.h
@@ -11,8 +11,8 @@
 
 class IIKGoalCreatorInterface;
 class UIKRigDefinition;
-class UIKRigProcessor;
 
+	
 USTRUCT(BlueprintInternalUseOnly)
 struct IKRIG_API FAnimNode_IKRig : public FAnimNode_Base
 {
@@ -27,7 +27,7 @@ struct IKRIG_API FAnimNode_IKRig : public FAnimNode_Base
 	UIKRigDefinition* RigDefinitionAsset;
 
 	/** The input goal transforms used by the IK Rig solvers.*/
-	UPROPERTY(EditAnywhere, BlueprintReadWrite, EditFixedSize, Category = Goal, meta = (PinShownByDefault))
+	UPROPERTY(EditAnywhere, BlueprintReadWrite, Category = Goal, meta = (PinShownByDefault))
 	TArray<FIKRigGoal> Goals;
 
 	/** optionally ignore the input pose and start from the reference pose each solve */
@@ -47,7 +47,7 @@ private:
 
 	/** IK Rig runtime processor */
 	UPROPERTY(Transient)
-	FIKRigProcessor IKRigProcessor;
+	TObjectPtr<UIKRigProcessor> IKRigProcessor;
 
 	/** a cached list of components on the owning actor that implement the goal creator interface */
 	TArray<IIKGoalCreatorInterface*> GoalCreators;
@@ -68,14 +68,14 @@ public:
 	virtual bool HasPreUpdate() const override { return true; }
 	virtual void PreUpdate(const UAnimInstance* InAnimInstance) override;
 	// End of FAnimNode_Base interface
-
+	
 private:
+	
 	void CopyInputPoseToSolver(FCompactPose& InputPose);
 	void AssignGoalTargets();
-	void CopyOutputPoseToAnimGraph(FCompactPose& OutputPose);
-	
-	bool RebuildGoalList();
-	FName GetGoalName(int32 Index) const;
+	void CopyOutputPoseToAnimGraph(FCompactPose& OutputPose);	
 	void QueueDrawInterface(FAnimInstanceProxy* AnimProxy, const FTransform& ComponentToWorld) const;
+	
 	friend class UAnimGraphNode_IKRig;
+	friend struct FIKRigAnimInstanceProxy;
 };
diff --git a/Engine/Plugins/Animation/IKRig/Source/IKRig/Public/IKRigDataTypes.h b/Engine/Plugins/Animation/IKRig/Source/IKRig/Public/IKRigDataTypes.h
--- a/Engine/Plugins/Animation/IKRig/Source/IKRig/Public/IKRigDataTypes.h
+++ b/Engine/Plugins/Animation/IKRig/Source/IKRig/Public/IKRigDataTypes.h
@@ -8,31 +8,74 @@
 
 class UIKRigEffectorGoal;
 
+UENUM(BlueprintType)
+enum class EIKRigGoalSpace : uint8
+{
+	Component		UMETA(DisplayName = "Component Space"),
+	Additive		UMETA(DisplayName = "Additive"),
+    World			UMETA(DisplayName = "World Space"),
+};
+
+UENUM(BlueprintType)
+enum class EIKRigGoalTransformSource : uint8
+{
+	Manual			UMETA(DisplayName = "Manual Input"),
+    Bone			UMETA(DisplayName = "Bone"),
+    ActorComponent	UMETA(DisplayName = "Actor Component"),
+};
+
 USTRUCT(Blueprintable)
 struct IKRIG_API FIKRigGoal
 {
 	GENERATED_BODY()
 	
 	/** Name of the IK Goal. Must correspond to the name of a Goal in the target IKRig asset. */
-	UPROPERTY(EditAnywhere, BlueprintReadWrite, Category = FIKRigGoal)
+	UPROPERTY(EditAnywhere, BlueprintReadWrite, Category = Goal)
 	FName Name;
 
+	/** Set the source of the transform data for the Goal.
+	 * "Manual Input" uses the values provided by the blueprint node pin.
+	 * "Bone" uses the transform of the bone provided by OptionalSourceBone.
+	 * "Actor Component" uses the transform supplied by any Actor Components that implements the IIKGoalCreatorInterface*/
+	UPROPERTY(EditAnywhere, BlueprintReadWrite, Category = Goal)
+	EIKRigGoalTransformSource TransformSource;
+
+	/** When TransformSource is set to "Bone" mode, the Position and Rotation will be driven by this Bone's input transform.
+	 * When using a Bone as the transform source, the Position and Rotation Alpha values can still be set as desired.
+	 * But the PositionSpace and RotationSpace are no longer relevant and will not be used.*/
+	UPROPERTY(EditAnywhere, BlueprintReadWrite, Category = Goal, meta = (CustomWidget = "BoneName"))
+	FName OptionalSourceBone;
+
 	/** Position of the IK goal in Component Space of target actor component. */
-	UPROPERTY(EditAnywhere, BlueprintReadWrite, Category = FIKRigGoal)
+	UPROPERTY(EditAnywhere, BlueprintReadWrite, Category = Position, meta=(EditCondition="Source == EIKRigGoalSource::Manual"))
 	FVector Position;
 
 	/** Rotation of the IK goal in Component Space of target actor component. */
-	UPROPERTY(EditAnywhere, BlueprintReadWrite, Category = FIKRigGoal)
+	UPROPERTY(EditAnywhere, BlueprintReadWrite, Category = Rotation, meta=(EditCondition="Source == EIKRigGoalSource::Manual"))
 	FRotator Rotation;
 
 	/** Range 0-1. Smoothly blends the Goal position from the input pose (0.0) to the Goal position (1.0). */
-	UPROPERTY(EditAnywhere, BlueprintReadWrite, Category = FIKRigGoal, meta = (ClampMin = "0.0", ClampMax = "1.0", UIMin = "0.0", UIMax = "1.0"))
+	UPROPERTY(EditAnywhere, BlueprintReadWrite, Category = Position, meta = (ClampMin = "0.0", ClampMax = "1.0", UIMin = "0.0", UIMax = "1.0"))
 	float PositionAlpha;
 
 	/** Range 0-1. Smoothly blends the Goal rotation from the input pose (0.0) to the Goal rotation (1.0). */
-	UPROPERTY(EditAnywhere, BlueprintReadWrite, Category = FIKRigGoal, meta = (ClampMin = "0.0", ClampMax = "1.0", UIMin = "0.0", UIMax = "1.0"))
+	UPROPERTY(EditAnywhere, BlueprintReadWrite, Category = Rotation, meta = (ClampMin = "0.0", ClampMax = "1.0", UIMin = "0.0", UIMax = "1.0"))
 	float RotationAlpha;
 
+	/** The space that the goal position is in.
+	 *"Additive" treats the goal transform as an additive offset relative to the Bone at the effector.
+	 *"Component" treats the goal transform as being in the space of the Skeletal Mesh Actor Component.
+	 *"World" treats the goal transform as being in the space of the World. */
+	UPROPERTY(EditAnywhere, BlueprintReadWrite, Category = Position)
+	EIKRigGoalSpace PositionSpace;
+	
+	/** The space that the goal rotation is in.
+	*"Additive" treats the goal transform as an additive offset relative to the Bone at the effector.
+	*"Component" treats the goal transform as being in the space of the Skeletal Mesh Actor Component.
+	*"World" treats the goal transform as being in the space of the World. */
+	UPROPERTY(EditAnywhere, BlueprintReadWrite, Category = Rotation)
+	EIKRigGoalSpace RotationSpace;
+
 	UPROPERTY(Transient)
 	FVector FinalBlendedPosition;
 
@@ -42,18 +85,22 @@ struct IKRIG_API FIKRigGoal
 	FIKRigGoal()
     : Name(NAME_None),
 	Position(ForceInitToZero),
-    Rotation(FRotator::ZeroRotator),
-    PositionAlpha(0.0f),
+	Rotation(FRotator::ZeroRotator),
+	PositionAlpha(0.0f),
     RotationAlpha(0.0f),
+	PositionSpace(EIKRigGoalSpace::Component),
+	RotationSpace(EIKRigGoalSpace::Component),
 	FinalBlendedPosition(ForceInitToZero),
 	FinalBlendedRotation(FQuat::Identity){}
 	
 	FIKRigGoal(const FName& GoalName)
     : Name(GoalName),
 	Position(ForceInitToZero),
-    Rotation(FRotator::ZeroRotator),
+	Rotation(FRotator::ZeroRotator),
     PositionAlpha(0.0f),
     RotationAlpha(0.0f),
+	PositionSpace(EIKRigGoalSpace::Component),
+    RotationSpace(EIKRigGoalSpace::Component),
 	FinalBlendedPosition(ForceInitToZero),
     FinalBlendedRotation(FQuat::Identity){}
 
@@ -64,12 +111,14 @@ struct IKRIG_API FIKRigGoal
         const float PositionAlpha,
         const float RotationAlpha)
         : Name(Name),
-          Position(Position),
-          Rotation(Rotation.Rotator()),
-          PositionAlpha(PositionAlpha),
-          RotationAlpha(RotationAlpha),
-		  FinalBlendedPosition(Position),
-		  FinalBlendedRotation(Rotation){}
+        Position(Position),
+        Rotation(Rotation.Rotator()),
+        PositionAlpha(PositionAlpha),
+        RotationAlpha(RotationAlpha),
+		PositionSpace(EIKRigGoalSpace::Component),
+		RotationSpace(EIKRigGoalSpace::Component),
+		FinalBlendedPosition(Position),
+		FinalBlendedRotation(Rotation){}
 
 	FIKRigGoal(const UIKRigEffectorGoal* InGoal)
 		: Name(InGoal->GoalName),
@@ -77,6 +126,8 @@ struct IKRIG_API FIKRigGoal
 		Rotation(InGoal->CurrentTransform.Rotator()),
 		PositionAlpha(InGoal->PositionAlpha),
         RotationAlpha(InGoal->RotationAlpha),
+		PositionSpace(EIKRigGoalSpace::Component),
+        RotationSpace(EIKRigGoalSpace::Component),
         FinalBlendedPosition(ForceInitToZero),
         FinalBlendedRotation(FQuat::Identity){}
 
diff --git a/Engine/Plugins/Animation/IKRig/Source/IKRig/Public/IKRigProcessor.h b/Engine/Plugins/Animation/IKRig/Source/IKRig/Public/IKRigProcessor.h
--- a/Engine/Plugins/Animation/IKRig/Source/IKRig/Public/IKRigProcessor.h
+++ b/Engine/Plugins/Animation/IKRig/Source/IKRig/Public/IKRigProcessor.h
@@ -20,13 +20,13 @@ struct FGoalBone
 	int32 BoneIndex;
 };
 
-USTRUCT()
-struct IKRIG_API FIKRigProcessor
+UCLASS()
+class IKRIG_API UIKRigProcessor : public UObject
 {
 	GENERATED_BODY()
-
+	
 public:
-
+	
 	/** the runtime for an IKRig to convert an input pose into
 	*   a solved output pose given a set of IK Rig Goals:
 	*   
@@ -40,7 +40,7 @@ public:
 	 
 	/** setup a new processor to run the given IKRig asset
 	 *  NOTE!! this function creates new UObjects and consequently MUST be called from the main thread!!	 */
-	void Initialize(UIKRigDefinition* InRigAsset, const FReferenceSkeleton& RefSkeleton, UObject* Outer);
+	void Initialize(UIKRigDefinition* InRigAsset, const FReferenceSkeleton& RefSkeleton);
 
 	//
 	// BEGIN UPDATE SEQUENCE FUNCTIONS
@@ -60,8 +60,9 @@ public:
 	/** Set a named IK goal to go to a specific location and rotation (assumed in component space) blended by separate position/rotation alpha (0-1)*/
 	void SetIKGoal(const UIKRigEffectorGoal* Goal);
 
-	/** Run entire stack of solvers */
-	void Solve();
+	/** Run entire stack of solvers.
+	 * If any Goals were supplied in World Space, a valid WorldToComponent transform must be provided.  */
+	void Solve(const FTransform& WorldToComponent = FTransform::Identity);
 
 	/** Get the results after calling Solve() */
 	void CopyOutputGlobalPoseToArray(TArray<FTransform>& OutputPoseGlobal) const;
@@ -84,8 +85,8 @@ public:
 	
 private:
 
-	/** Update the final pos/rot of all the goals based on their alpha values. */
-	void BlendGoalsByAlpha();
+	/** Update the final pos/rot of all the goals based on their alpha values and their space settings. */
+	void ResolveFinalGoalTransforms(const FTransform& WorldToComponent);
 
 	/** the stack of solvers to run in order */
 	UPROPERTY(Transient)
diff --git a/Engine/Plugins/Animation/IKRig/Source/IKRig/Public/Retargeter/IKRetargeter.h b/Engine/Plugins/Animation/IKRig/Source/IKRig/Public/Retargeter/IKRetargeter.h
--- a/Engine/Plugins/Animation/IKRig/Source/IKRig/Public/Retargeter/IKRetargeter.h
+++ b/Engine/Plugins/Animation/IKRig/Source/IKRig/Public/Retargeter/IKRetargeter.h
@@ -2,9 +2,7 @@
 
 #pragma once
 
-#include "Engine/SkeletalMesh.h"
 #include "Containers/UnrealString.h"
-#include "Math/TransformNonVectorized.h"
 
 #include "IKRigDefinition.h"
 #include "IKRigProcessor.h"
@@ -298,7 +296,8 @@ public:
 
 private:
 
-	FIKRigProcessor IKRigProcessor;
+	UPROPERTY(Transient)
+	TObjectPtr<UIKRigProcessor> IKRigProcessor;
 	
 	TArray<FRetargetChainPairFK> ChainPairsFK;
 	
diff --git a/Engine/Plugins/Animation/IKRig/Source/IKRigDeveloper/Private/AnimGraphNode_IKRig.cpp b/Engine/Plugins/Animation/IKRig/Source/IKRigDeveloper/Private/AnimGraphNode_IKRig.cpp
--- a/Engine/Plugins/Animation/IKRig/Source/IKRigDeveloper/Private/AnimGraphNode_IKRig.cpp
+++ b/Engine/Plugins/Animation/IKRig/Source/IKRigDeveloper/Private/AnimGraphNode_IKRig.cpp
@@ -31,54 +31,4 @@ void UAnimGraphNode_IKRig::CopyNodeDataToPreviewNode(FAnimNode_Base* InPreviewNo
 	FAnimNode_IKRig* IKRigNode = static_cast<FAnimNode_IKRig*>(InPreviewNode);
 }
 
-void UAnimGraphNode_IKRig::PostEditChangeProperty(struct FPropertyChangedEvent& PropertyChangedEvent)
-{
-	const FName PropertyName = (PropertyChangedEvent.Property ? PropertyChangedEvent.Property->GetFName() : NAME_None);
-
-	// Reconstruct node to show updates to PinFriendlyNames.
-	if ((PropertyName == GET_MEMBER_NAME_STRING_CHECKED(FAnimNode_IKRig, RigDefinitionAsset)))
-	{
-		if (Node.RebuildGoalList())
-		{
-			ReconstructNode();
-		}
-	}
-}
-
-void UAnimGraphNode_IKRig::CustomizePinData(UEdGraphPin* Pin, FName SourcePropertyName, int32 ArrayIndex) const
-{
-	Super::CustomizePinData(Pin, SourcePropertyName, ArrayIndex);
-
-	// set names of GoalTransform pins to names of goals
-	const FString PinString = Pin->PinName.ToString();
-	const FString GoalPinString = GET_MEMBER_NAME_STRING_CHECKED(FAnimNode_IKRig, Goals);
-	const int32 GoalStringLen = GoalPinString.Len();
-	if (PinString.Len() > GoalStringLen && PinString.Left(GoalStringLen) == GoalPinString)
-	{
-		if (!Pin->bHidden)
-		{
-			Pin->PinFriendlyName = FText::FromName(Node.GetGoalName(ArrayIndex));
-		}
-	}
-}
-
-void UAnimGraphNode_IKRig::PostLoad()
-{
-	Super::PostLoad();
-	Node.RebuildGoalList();
-}
-
-void UAnimGraphNode_IKRig::PreloadRequiredAssets()
-{
-	if (Node.RigDefinitionAsset)
-	{
-		PreloadObject(Node.RigDefinitionAsset);
-		for (UIKRigSolver* Solver : Node.RigDefinitionAsset->Solvers)
-		{
-			PreloadObject(Solver);
-		}
-	}
-
-	Super::PreloadRequiredAssets();
-}
 #undef LOCTEXT_NAMESPACE
diff --git a/Engine/Plugins/Animation/IKRig/Source/IKRigDeveloper/Public/AnimGraphNode_IKRig.h b/Engine/Plugins/Animation/IKRig/Source/IKRigDeveloper/Public/AnimGraphNode_IKRig.h
--- a/Engine/Plugins/Animation/IKRig/Source/IKRigDeveloper/Public/AnimGraphNode_IKRig.h
+++ b/Engine/Plugins/Animation/IKRig/Source/IKRigDeveloper/Public/AnimGraphNode_IKRig.h
@@ -23,14 +23,10 @@ class IKRIGDEVELOPER_API UAnimGraphNode_IKRig : public UAnimGraphNode_Base
 public:
 	// UEdGraphNode interface
 	virtual FText GetNodeTitle(ENodeTitleType::Type TitleType) const override;
-	virtual void PostEditChangeProperty(struct FPropertyChangedEvent& PropertyChangedEvent) override;
-	virtual void PostLoad() override;
 	// End of UEdGraphNode interface
 
 	// UAnimGraphNode_Base interface
-	virtual void PreloadRequiredAssets() override;
 	virtual void CopyNodeDataToPreviewNode(FAnimNode_Base* AnimNode) override;
 	virtual void Draw(FPrimitiveDrawInterface* PDI, USkeletalMeshComponent* PreviewSkelMeshComp) const override;
-	virtual void CustomizePinData(UEdGraphPin* Pin, FName SourcePropertyName, int32 ArrayIndex) const override;
 	// End of UAnimGraphNode_Base interface
 };
diff --git a/Engine/Plugins/Animation/IKRig/Source/IKRigEditor/Private/RigEditor/IKRigAnimInstance.cpp b/Engine/Plugins/Animation/IKRig/Source/IKRigEditor/Private/RigEditor/IKRigAnimInstance.cpp
--- a/Engine/Plugins/Animation/IKRig/Source/IKRigEditor/Private/RigEditor/IKRigAnimInstance.cpp
+++ b/Engine/Plugins/Animation/IKRig/Source/IKRigEditor/Private/RigEditor/IKRigAnimInstance.cpp
@@ -17,5 +17,5 @@ void UIKRigAnimInstance::SetIKRigAsset(UIKRigDefinition* InIKRigAsset)
 
 FAnimInstanceProxy* UIKRigAnimInstance::CreateAnimInstanceProxy()
 {
-	return new FIKRigAnimInstanceProxy(this);
+	return new FIKRigAnimInstanceProxy(this, &IKRigNode);
 }
diff --git a/Engine/Plugins/Animation/IKRig/Source/IKRigEditor/Private/RigEditor/IKRigAnimInstanceProxy.cpp b/Engine/Plugins/Animation/IKRig/Source/IKRigEditor/Private/RigEditor/IKRigAnimInstanceProxy.cpp
--- a/Engine/Plugins/Animation/IKRig/Source/IKRigEditor/Private/RigEditor/IKRigAnimInstanceProxy.cpp
+++ b/Engine/Plugins/Animation/IKRig/Source/IKRigEditor/Private/RigEditor/IKRigAnimInstanceProxy.cpp
@@ -7,8 +7,9 @@ FIKRigAnimInstanceProxy::FIKRigAnimInstanceProxy()
 {
 }
 
-FIKRigAnimInstanceProxy::FIKRigAnimInstanceProxy(UAnimInstance* InAnimInstance)
-	: FAnimPreviewInstanceProxy(InAnimInstance)
+FIKRigAnimInstanceProxy::FIKRigAnimInstanceProxy(UAnimInstance* InAnimInstance, FAnimNode_IKRig* InIKRigNode)
+	: FAnimPreviewInstanceProxy(InAnimInstance),
+	IKRigNode(InIKRigNode)
 {
 }
 
@@ -19,26 +20,26 @@ FIKRigAnimInstanceProxy::~FIKRigAnimInstanceProxy()
 void FIKRigAnimInstanceProxy::Initialize(UAnimInstance* InAnimInstance)
 {
 	FAnimInstanceProxy::Initialize(InAnimInstance);
-	IKRigNode.Source.SetLinkNode(&SingleNode);
+	IKRigNode->Source.SetLinkNode(&SingleNode);
 
 	// force this instance of the IK Rig evaluation to copy setting from the source IK Rig asset
-	IKRigNode.bDriveWithSourceAsset = true; 
+	IKRigNode->bDriveWithSourceAsset = true; 
 }
 
 bool FIKRigAnimInstanceProxy::Evaluate(FPoseContext& Output)
 {
-	IKRigNode.Evaluate_AnyThread(Output);
+	IKRigNode->Evaluate_AnyThread(Output);
 	return true;
 }
 
 FAnimNode_Base* FIKRigAnimInstanceProxy::GetCustomRootNode()
 {
-	return &IKRigNode;
+	return IKRigNode;
 }
 
 void FIKRigAnimInstanceProxy::GetCustomNodes(TArray<FAnimNode_Base*>& OutNodes)
 {
-	OutNodes.Add(&IKRigNode);
+	OutNodes.Add(IKRigNode);
 }
 
 void FIKRigAnimInstanceProxy::UpdateAnimationNode(const FAnimationUpdateContext& InContext)
@@ -49,11 +50,12 @@ void FIKRigAnimInstanceProxy::UpdateAnimationNode(const FAnimationUpdateContext&
 	}
 	else
 	{
-		IKRigNode.Update_AnyThread(InContext);
+		IKRigNode->Update_AnyThread(InContext);
 	}
 }
 
 void FIKRigAnimInstanceProxy::SetIKRigAsset(UIKRigDefinition* InIKRigAsset)
 {
-	IKRigNode.RigDefinitionAsset = InIKRigAsset;
+	IKRigNode->RigDefinitionAsset = InIKRigAsset;
 }
+
diff --git a/Engine/Plugins/Animation/IKRig/Source/IKRigEditor/Private/RigEditor/IKRigToolkit.cpp b/Engine/Plugins/Animation/IKRig/Source/IKRigEditor/Private/RigEditor/IKRigToolkit.cpp
--- a/Engine/Plugins/Animation/IKRig/Source/IKRigEditor/Private/RigEditor/IKRigToolkit.cpp
+++ b/Engine/Plugins/Animation/IKRig/Source/IKRigEditor/Private/RigEditor/IKRigToolkit.cpp
@@ -163,6 +163,7 @@ FString FIKRigEditorToolkit::GetWorldCentricTabPrefix() const
 
 void FIKRigEditorToolkit::AddReferencedObjects(FReferenceCollector& Collector)
 {
+	// hold the asset we are working on
 	UIKRigDefinition* Asset = EditorController->AssetController->GetAsset();
 	Collector.AddReferencedObject(Asset);
 }
@@ -182,8 +183,8 @@ void FIKRigEditorToolkit::HandlePreviewSceneCreated(const TSharedRef<IPersonaPre
 	EditorController->SkelMeshComponent = NewObject<UDebugSkelMeshComponent>(Actor);
 
 	// setup an apply an anim instance to the skeletal mesh component
-	UIKRigAnimInstance* IKRigAnimInstance = NewObject<UIKRigAnimInstance>(EditorController->SkelMeshComponent, TEXT("IKRigAnimScriptInstance"));
-	SetupAnimInstance(IKRigAnimInstance);
+	EditorController->AnimInstance = NewObject<UIKRigAnimInstance>(EditorController->SkelMeshComponent, TEXT("IKRigAnimScriptInstance"));
+	SetupAnimInstance();
 
 	// set the skeletal mesh on the component
 	// NOTE: this must be done AFTER setting the AnimInstance so that the correct root anim node is loaded
@@ -213,8 +214,8 @@ void FIKRigEditorToolkit::HandlePreviewMeshChanged(USkeletalMesh* InOldSkeletalM
 	if (EditorSkelComp)
 	{
 		bool bWasCreated = false;
-		UIKRigAnimInstance* AnimInstance = FAnimCustomInstanceHelper::BindToSkeletalMeshComponent<UIKRigAnimInstance>(EditorSkelComp , bWasCreated);
-		SetupAnimInstance(AnimInstance);
+		EditorController->AnimInstance = FAnimCustomInstanceHelper::BindToSkeletalMeshComponent<UIKRigAnimInstance>(EditorSkelComp , bWasCreated);
+		SetupAnimInstance();
 	}
 
 	EditorController->SkelMeshComponent->SetSkeletalMesh(InNewSkeletalMesh);
@@ -238,8 +239,9 @@ void FIKRigEditorToolkit::HandleReset()
 	EditorController.Get().Reset();
 }
 
-void FIKRigEditorToolkit::SetupAnimInstance(UIKRigAnimInstance* InAnimInstance)
+void FIKRigEditorToolkit::SetupAnimInstance()
 {
+	UIKRigAnimInstance* InAnimInstance = EditorController->AnimInstance.Get();
 	InAnimInstance->SetIKRigAsset(EditorController->AssetController->GetAsset());
 	EditorController->SkelMeshComponent->PreviewInstance = InAnimInstance;
 	InAnimInstance->InitializeAnimation();
diff --git a/Engine/Plugins/Animation/IKRig/Source/IKRigEditor/Public/RigEditor/IKRigAnimInstance.h b/Engine/Plugins/Animation/IKRig/Source/IKRigEditor/Public/RigEditor/IKRigAnimInstance.h
--- a/Engine/Plugins/Animation/IKRig/Source/IKRigEditor/Public/RigEditor/IKRigAnimInstance.h
+++ b/Engine/Plugins/Animation/IKRig/Source/IKRigEditor/Public/RigEditor/IKRigAnimInstance.h
@@ -3,6 +3,8 @@
 #pragma once
 
 #include "AnimPreviewInstance.h"
+#include "IKRigProcessor.h"
+#include "AnimNodes/AnimNode_IKRig.h"
 
 #include "IKRigAnimInstance.generated.h"
 
@@ -19,4 +21,7 @@ public:
 
 protected:
 	virtual FAnimInstanceProxy* CreateAnimInstanceProxy() override;
+
+	UPROPERTY(Transient)
+	FAnimNode_IKRig IKRigNode;
 };
diff --git a/Engine/Plugins/Animation/IKRig/Source/IKRigEditor/Public/RigEditor/IKRigAnimInstanceProxy.h b/Engine/Plugins/Animation/IKRig/Source/IKRigEditor/Public/RigEditor/IKRigAnimInstanceProxy.h
--- a/Engine/Plugins/Animation/IKRig/Source/IKRigEditor/Public/RigEditor/IKRigAnimInstanceProxy.h
+++ b/Engine/Plugins/Animation/IKRig/Source/IKRigEditor/Public/RigEditor/IKRigAnimInstanceProxy.h
@@ -17,7 +17,7 @@ struct FIKRigAnimInstanceProxy : public FAnimPreviewInstanceProxy
 public:
 	
 	FIKRigAnimInstanceProxy();
-	FIKRigAnimInstanceProxy(UAnimInstance* InAnimInstance);
+	FIKRigAnimInstanceProxy(UAnimInstance* InAnimInstance, FAnimNode_IKRig* IKRigNode);
 	virtual ~FIKRigAnimInstanceProxy();
 
 	/** FAnimPreviewInstanceProxy interface */
@@ -34,6 +34,5 @@ public:
 
 	void SetIKRigAsset(UIKRigDefinition* InIKRigAsset);
 
-	UPROPERTY(Transient)
-	FAnimNode_IKRig IKRigNode;
+	FAnimNode_IKRig* IKRigNode;
 };
\ No newline at end of file
diff --git a/Engine/Plugins/Animation/IKRig/Source/IKRigEditor/Public/RigEditor/IKRigEditorController.h b/Engine/Plugins/Animation/IKRig/Source/IKRigEditor/Public/RigEditor/IKRigEditorController.h
--- a/Engine/Plugins/Animation/IKRig/Source/IKRigEditor/Public/RigEditor/IKRigEditorController.h
+++ b/Engine/Plugins/Animation/IKRig/Source/IKRigEditor/Public/RigEditor/IKRigEditorController.h
@@ -65,7 +65,8 @@ public:
 	UDebugSkelMeshComponent* SkelMeshComponent;
 
 	/** viewport anim instance */
-	UIKRigAnimInstance* AnimInstance; 
+	UPROPERTY(transient, NonTransactional)
+	TObjectPtr<class UIKRigAnimInstance> AnimInstance;
 
 	/** asset properties tab */
 	TSharedPtr<class IDetailsView> DetailsView;
diff --git a/Engine/Plugins/Animation/IKRig/Source/IKRigEditor/Public/RigEditor/IKRigToolkit.h b/Engine/Plugins/Animation/IKRig/Source/IKRigEditor/Public/RigEditor/IKRigToolkit.h
--- a/Engine/Plugins/Animation/IKRig/Source/IKRigEditor/Public/RigEditor/IKRigToolkit.h
+++ b/Engine/Plugins/Animation/IKRig/Source/IKRigEditor/Public/RigEditor/IKRigToolkit.h
@@ -88,7 +88,7 @@ private:
 	void HandlePreviewMeshChanged(USkeletalMesh* InOldSkeletalMesh, USkeletalMesh* InNewSkeletalMesh);
 	void HandleDetailsCreated(const TSharedRef<class IDetailsView>& InDetailsView);
 	void OnFinishedChangingDetails(const FPropertyChangedEvent& PropertyChangedEvent);
-	void SetupAnimInstance(UIKRigAnimInstance* InAnimInstance);
+	void SetupAnimInstance();
 	/** END preview scene setup */
 	
 	/** centralized management of selection across skeleton view and viewport */
