Fixing static analysis issue in IKRigDataTypes.h

diff --git a/Engine/Plugins/Animation/IKRig/Source/IKRig/Public/IKRigDataTypes.h b/Engine/Plugins/Animation/IKRig/Source/IKRig/Public/IKRigDataTypes.h
--- a/Engine/Plugins/Animation/IKRig/Source/IKRig/Public/IKRigDataTypes.h
+++ b/Engine/Plugins/Animation/IKRig/Source/IKRig/Public/IKRigDataTypes.h
@@ -38,7 +38,7 @@ struct IKRIG_API FIKRigGoal
 	 * "Bone" uses the transform of the bone provided by OptionalSourceBone.
 	 * "Actor Component" uses the transform supplied by any Actor Components that implements the IIKGoalCreatorInterface*/
 	UPROPERTY(EditAnywhere, BlueprintReadWrite, Category = Goal)
-	EIKRigGoalTransformSource TransformSource;
+	EIKRigGoalTransformSource TransformSource = EIKRigGoalTransformSource::Manual;
 
 	/** When TransformSource is set to "Bone" mode, the Position and Rotation will be driven by this Bone's input transform.
 	 * When using a Bone as the transform source, the Position and Rotation Alpha values can still be set as desired.
