Fixing undo/redo crash and functionality in IK Rig editor.
Adding/Removing Goals, Solvers and Bone Settings now undoable.
Reorder Solvers undoable

diff --git a/Engine/Plugins/Animation/IKRig/Source/IKRig/Private/IKRigDefinition.cpp b/Engine/Plugins/Animation/IKRig/Source/IKRig/Private/IKRigDefinition.cpp
--- a/Engine/Plugins/Animation/IKRig/Source/IKRig/Private/IKRigDefinition.cpp
+++ b/Engine/Plugins/Animation/IKRig/Source/IKRig/Private/IKRigDefinition.cpp
@@ -46,10 +46,7 @@ USkeletalMesh* UIKRigDefinition::GetPreviewMesh() const
 bool UIKRigDefinition::Modify(bool bAlwaysMarkDirty /*=true*/)
 {
 	const bool bSavedToTransactionBuffer = Super::Modify(bAlwaysMarkDirty);
-	if (bSavedToTransactionBuffer)
-	{
-		AssetVersion++; // inform any runtime/editor systems they should copy-up modifications	
-	}
+	AssetVersion++; // inform any runtime/editor systems they should copy-up modifications	
 	return bSavedToTransactionBuffer;
 }
 #endif
diff --git a/Engine/Plugins/Animation/IKRig/Source/IKRig/Public/IKRigDefinition.h b/Engine/Plugins/Animation/IKRig/Source/IKRig/Public/IKRigDefinition.h
--- a/Engine/Plugins/Animation/IKRig/Source/IKRig/Public/IKRigDefinition.h
+++ b/Engine/Plugins/Animation/IKRig/Source/IKRig/Public/IKRigDefinition.h
@@ -17,7 +17,7 @@ class IKRIG_API UIKRigEffectorGoal : public UObject
 
 public:
 	
-	UIKRigEffectorGoal() : GoalName("DefaultGoal"), BoneName(NAME_None) {}
+	UIKRigEffectorGoal() : GoalName("DefaultGoal"), BoneName(NAME_None){}
 	
 	UIKRigEffectorGoal(const FName& InGoalName, const FName& InBoneName) :  GoalName(InGoalName), BoneName(InBoneName){}
 
@@ -46,6 +46,12 @@ public:
 	float GizmoThickness = 0.7f;
 
 	bool operator==(const UIKRigEffectorGoal& Other) const { return GoalName == Other.GoalName; }
+
+	virtual void PostLoad() override
+	{
+		Super::PostLoad();
+		SetFlags(RF_Transactional);
+	}
 };
 
 USTRUCT(Blueprintable)
@@ -120,7 +126,7 @@ public:
 	UPROPERTY(VisibleAnywhere, instanced, Category = "IK")
 	TArray<UIKRigSolver*> Solvers;
 
-	/** goals, used by solvers that support them */
+	/** goals, used as effectors by solvers that support them */
 	UPROPERTY(VisibleAnywhere, Category = "IK")
 	TArray<UIKRigEffectorGoal*> Goals;
 
diff --git a/Engine/Plugins/Animation/IKRig/Source/IKRig/Public/IKRigSolver.h b/Engine/Plugins/Animation/IKRig/Source/IKRig/Public/IKRigSolver.h
--- a/Engine/Plugins/Animation/IKRig/Source/IKRig/Public/IKRigSolver.h
+++ b/Engine/Plugins/Animation/IKRig/Source/IKRig/Public/IKRigSolver.h
@@ -72,4 +72,12 @@ public:
 	/** return true if the supplied Bone is affected by this solver - this provides UI feedback for user */
 	virtual bool IsBoneAffectedBySolver(const FName& BoneName, const FIKRigSkeleton& IKRigSkeleton) const {return false;};
 	//** END ROOT BONE */
+
+	/** UObject interface */
+	void PostLoad() override
+	{
+		Super::PostLoad();
+		SetFlags(RF_Transactional); // patch old solvers to enable undo/redo
+	}; 
+	/** END UObject interface */
 };
diff --git a/Engine/Plugins/Animation/IKRig/Source/IKRigEditor/Private/RigEditor/IKRigController.cpp b/Engine/Plugins/Animation/IKRig/Source/IKRigEditor/Private/RigEditor/IKRigController.cpp
--- a/Engine/Plugins/Animation/IKRig/Source/IKRigEditor/Private/RigEditor/IKRigController.cpp
+++ b/Engine/Plugins/Animation/IKRig/Source/IKRigEditor/Private/RigEditor/IKRigController.cpp
@@ -51,6 +51,7 @@ void UIKRigController::AddBoneSetting(const FName& BoneName, int32 SolverIndex)
 	}
 
 	FScopedTransaction Transaction(LOCTEXT("AddBoneSetting_Label", "Add Bone Setting"));
+	Solver->Modify();
 	IKRigAsset->Modify();
 
 	Solver->AddBoneSetting(BoneName);
@@ -91,7 +92,8 @@ void UIKRigController::RemoveBoneSetting(const FName& BoneName, int32 SolverInde
 		return; // bone doesn't exist
 	}
 
-	FScopedTransaction Transaction(LOCTEXT("AddBoneSetting_Label", "Add Bone Setting"));
+	FScopedTransaction Transaction(LOCTEXT("RemoveBoneSetting_Label", "Remove Bone Setting"));
+	Solver->Modify();
 	IKRigAsset->Modify();
 
 	Solver->RemoveBoneSetting(BoneName);
@@ -224,7 +226,7 @@ int32 UIKRigController::AddSolver(TSubclassOf<UIKRigSolver> InIKRigSolverClass)
 	FScopedTransaction Transaction(LOCTEXT("AddSolver_Label", "Add Solver"));
 	IKRigAsset->Modify();
 
-	UIKRigSolver* NewSolver = NewObject<UIKRigSolver>(IKRigAsset, InIKRigSolverClass);
+	UIKRigSolver* NewSolver = NewObject<UIKRigSolver>(IKRigAsset, InIKRigSolverClass, NAME_None, RF_Transactional);
 	check(NewSolver);
 
 	return IKRigAsset->Solvers.Add(NewSolver);
@@ -296,9 +298,11 @@ void UIKRigController::SetRootBone(const FName& RootBoneName, int32 SolverIndex)
 	}
 
 	FScopedTransaction Transaction(LOCTEXT("SetRootBone_Label", "Set Root Bone"));
+	UIKRigSolver* Solver = IKRigAsset->Solvers[SolverIndex];
+	Solver->Modify();
 	IKRigAsset->Modify();
 
-	IKRigAsset->Solvers[SolverIndex]->SetRootBone(RootBoneName);
+	Solver->SetRootBone(RootBoneName);
 }
 
 const TArray<UIKRigSolver*>& UIKRigController::GetSolverArray() const
@@ -320,7 +324,7 @@ UIKRigEffectorGoal* UIKRigController::AddNewGoal(const FName& GoalName, const FN
 	FScopedTransaction Transaction(LOCTEXT("AddSetting_Label", "Add Bone Setting"));
 	IKRigAsset->Modify();
 
-	UIKRigEffectorGoal* NewGoal = NewObject<UIKRigEffectorGoal>(GetAsset(), UIKRigEffectorGoal::StaticClass());
+	UIKRigEffectorGoal* NewGoal = NewObject<UIKRigEffectorGoal>(GetAsset(), UIKRigEffectorGoal::StaticClass(), NAME_None, RF_Transactional);
 	NewGoal->BoneName = BoneName;
 	NewGoal->GoalName = GoalName;
 	IKRigAsset->Goals.Add(NewGoal);
@@ -347,6 +351,7 @@ bool UIKRigController::RemoveGoal(const FName& GoalName) const
 	const FName& GoalToRemove = IKRigAsset->Goals[GoalIndex]->GoalName;
 	for (UIKRigSolver* Solver : IKRigAsset->Solvers)
 	{
+		Solver->Modify();
 		Solver->RemoveGoal(GoalToRemove);
 	}
 
@@ -379,11 +384,13 @@ FName UIKRigController::RenameGoal(const FName& OldName, const FName& PotentialN
 	IKRigAsset->Modify();
 
 	// rename in core
+	IKRigAsset->Goals[GoalIndex]->Modify();
 	IKRigAsset->Goals[GoalIndex]->GoalName = NewName;
 
 	// rename in solvers
 	for (UIKRigSolver* Solver : IKRigAsset->Solvers)
 	{
+		Solver->Modify();
 		Solver->RenameGoal(OldName, NewName);
 	}
 
@@ -418,6 +425,7 @@ bool UIKRigController::SetGoalBone(const FName& GoalName, const FName& NewBoneNa
 	// update in solvers
 	for (UIKRigSolver* Solver : IKRigAsset->Solvers)
 	{
+		Solver->Modify();
 		Solver->SetGoalBone(GoalName, NewBoneName);
 	}
 
@@ -448,9 +456,11 @@ bool UIKRigController::ConnectGoalToSolver(const UIKRigEffectorGoal& Goal, int32
 	check(IKRigAsset->Solvers.IsValidIndex(SolverIndex))
 
 	FScopedTransaction Transaction(LOCTEXT("AddGoalToSolver_Label", "Add Goal"));
+	UIKRigSolver* Solver = IKRigAsset->Solvers[SolverIndex];
+	Solver->Modify();
 	IKRigAsset->Modify();
 	
-	IKRigAsset->Solvers[SolverIndex]->AddGoal(Goal);
+	Solver->AddGoal(Goal);
 	return true;
 }
 
@@ -462,9 +472,11 @@ bool UIKRigController::DisconnectGoalFromSolver(const FName& GoalToRemove, int32
 	check(IKRigAsset->Solvers.IsValidIndex(SolverIndex))
 
     FScopedTransaction Transaction(LOCTEXT("RemoveGoalFromSolver_Label", "Remove Goal"));
+	UIKRigSolver* Solver = IKRigAsset->Solvers[SolverIndex];
+	Solver->Modify();
 	IKRigAsset->Modify();
 	
-	IKRigAsset->Solvers[SolverIndex]->RemoveGoal(GoalToRemove);
+	Solver->RemoveGoal(GoalToRemove);
 	return true;
 }
 
diff --git a/Engine/Plugins/Animation/IKRig/Source/IKRigEditor/Private/RigEditor/IKRigEditorController.cpp b/Engine/Plugins/Animation/IKRig/Source/IKRigEditor/Private/RigEditor/IKRigEditorController.cpp
--- a/Engine/Plugins/Animation/IKRig/Source/IKRigEditor/Private/RigEditor/IKRigEditorController.cpp
+++ b/Engine/Plugins/Animation/IKRig/Source/IKRigEditor/Private/RigEditor/IKRigEditorController.cpp
@@ -21,6 +21,13 @@ void FIKRigEditorController::Reset() const
 	}
 }
 
+void FIKRigEditorController::RefreshAllViews() const
+{
+	SolverStackView->RefreshStackView();
+	SkeletonView->RefreshTreeView();
+	DetailsView->ForceRefresh();
+}
+
 bool FIKRigEditorController::IsGoalSelected(const FName& GoalName)
 {
 	return SelectedGoals.Contains(GoalName);
diff --git a/Engine/Plugins/Animation/IKRig/Source/IKRigEditor/Private/RigEditor/IKRigToolkit.cpp b/Engine/Plugins/Animation/IKRig/Source/IKRigEditor/Private/RigEditor/IKRigToolkit.cpp
--- a/Engine/Plugins/Animation/IKRig/Source/IKRigEditor/Private/RigEditor/IKRigToolkit.cpp
+++ b/Engine/Plugins/Animation/IKRig/Source/IKRigEditor/Private/RigEditor/IKRigToolkit.cpp
@@ -173,6 +173,16 @@ TStatId FIKRigEditorToolkit::GetStatId() const
 	RETURN_QUICK_DECLARE_CYCLE_STAT(FIKRigEditorToolkit, STATGROUP_Tickables);
 }
 
+void FIKRigEditorToolkit::PostUndo(bool bSuccess)
+{
+	EditorController->RefreshAllViews();
+}
+
+void FIKRigEditorToolkit::PostRedo(bool bSuccess)
+{
+	EditorController->RefreshAllViews();
+}
+
 void FIKRigEditorToolkit::HandlePreviewSceneCreated(const TSharedRef<IPersonaPreviewScene>& InPersonaPreviewScene)
 {
 	AAnimationEditorPreviewActor* Actor = InPersonaPreviewScene->GetWorld()->SpawnActor<AAnimationEditorPreviewActor>(AAnimationEditorPreviewActor::StaticClass(), FTransform::Identity);
@@ -204,6 +214,11 @@ void FIKRigEditorToolkit::HandlePreviewMeshChanged(USkeletalMesh* InOldSkeletalM
 	{
 		return; // first time setup
 	}
+
+	if (InOldSkeletalMesh == InNewSkeletalMesh)
+	{
+		return; // already set to this skeletal mesh
+	}
 	
 	// update asset with new skeletal mesh (will copy new skeleton data)
 	EditorController->AssetController->SetSourceSkeletalMesh(InNewSkeletalMesh, false);
diff --git a/Engine/Plugins/Animation/IKRig/Source/IKRigEditor/Public/RigEditor/IKRigEditorController.h b/Engine/Plugins/Animation/IKRig/Source/IKRigEditor/Public/RigEditor/IKRigEditorController.h
--- a/Engine/Plugins/Animation/IKRig/Source/IKRigEditor/Public/RigEditor/IKRigEditorController.h
+++ b/Engine/Plugins/Animation/IKRig/Source/IKRigEditor/Public/RigEditor/IKRigEditorController.h
@@ -33,6 +33,8 @@ public:
 	void HandleGoalsSelectedInTreeView(const TArray<FName>& GoalNames);
 	/** reset all goals to initial transforms */
 	void Reset() const;
+	/** refresh all views */
+	void RefreshAllViews() const;
 
 	/** return list of those solvers in the stack that are selected by user */
 	void GetSelectedSolvers(TArray<TSharedPtr<FSolverStackElement> >& OutSelectedSolvers);
diff --git a/Engine/Plugins/Animation/IKRig/Source/IKRigEditor/Public/RigEditor/IKRigToolkit.h b/Engine/Plugins/Animation/IKRig/Source/IKRigEditor/Public/RigEditor/IKRigToolkit.h
--- a/Engine/Plugins/Animation/IKRig/Source/IKRigEditor/Public/RigEditor/IKRigToolkit.h
+++ b/Engine/Plugins/Animation/IKRig/Source/IKRigEditor/Public/RigEditor/IKRigToolkit.h
@@ -35,7 +35,7 @@ class FIKRigEditorToolkit :
 	public FPersonaAssetEditorToolkit,
 	public IHasPersonaToolkit,
 	public FGCObject,
-	public FEditorUndoClient,
+	public FSelfRegisteringEditorUndoClient,
 	public FTickableEditorObject
 {
 public:
@@ -72,6 +72,11 @@ public:
 	virtual TSharedRef<IPersonaToolkit> GetPersonaToolkit() const override { return PersonaToolkit.ToSharedRef(); }
 	/** END IHasPersonaToolkit interface */
 
+	/** FSelfRegisteringEditorUndoClient interface */
+	virtual void PostUndo( bool bSuccess );
+	virtual void PostRedo( bool bSuccess );
+	/** END FSelfRegisteringEditorUndoClient interface */
+
 	TSharedRef<FIKRigEditorController> GetController() const {return EditorController;};
 
 private:
