Fixes IK Rig crash bug on skeleton with virtual bones.
Fixes bug that initializes IK Rig with incorrect proportions.

diff --git a/Engine/Plugins/Animation/IKRig/Source/IKRig/Private/AnimNodes/AnimNode_IKRig.cpp b/Engine/Plugins/Animation/IKRig/Source/IKRig/Private/AnimNodes/AnimNode_IKRig.cpp
--- a/Engine/Plugins/Animation/IKRig/Source/IKRig/Private/AnimNodes/AnimNode_IKRig.cpp
+++ b/Engine/Plugins/Animation/IKRig/Source/IKRig/Private/AnimNodes/AnimNode_IKRig.cpp
@@ -69,10 +69,15 @@ void FAnimNode_IKRig::CopyInputPoseToSolver(FCompactPose& InputPose)
 	FIKRigSkeleton& IKRigSkeleton = IKRigProcessor->GetSkeleton();
 	for (FCompactPoseBoneIndex CPIndex : InputPose.ForEachBoneIndex())
 	{
-		int32* Index = CompactPoseToRigIndices.Find(CPIndex);
-		if (Index)
+		if (int32* Index = CompactPoseToRigIndices.Find(CPIndex))
 		{
-			IKRigSkeleton.CurrentPoseLocal[*Index] = InputPose[CPIndex];
+			// bones that were recorded with rig indices == -1 were not in the
+			// Reference Skeleton that the IK Rig was initialized with and therefore
+			// are not considered as part of the solve.
+			if (*Index != -1)
+			{
+				IKRigSkeleton.CurrentPoseLocal[*Index] = InputPose[CPIndex];	
+			}
 		}
 	}
 	// update global pose in IK Rig
@@ -115,10 +120,16 @@ void FAnimNode_IKRig::CopyOutputPoseToAnimGraph(FCompactPose& OutputPose)
 	// copy local transforms to output pose
 	for (FCompactPoseBoneIndex CPIndex : OutputPose.ForEachBoneIndex())
 	{
-		int32* Index = CompactPoseToRigIndices.Find(CPIndex);
-		if (Index)
+		if (int32* Index = CompactPoseToRigIndices.Find(CPIndex))
 		{
-			OutputPose[CPIndex] = IKRigSkeleton.CurrentPoseLocal[*Index];
+			// bones that were recorded with rig indices == -1 were not in the
+			// Reference Skeleton that the IK Rig was initialized with and therefore
+			// are not considered as part of the solve. These transforms are left at their
+			// input pose (in local space).
+			if (*Index != -1)
+			{
+				OutputPose[CPIndex] = IKRigSkeleton.CurrentPoseLocal[*Index];	
+			}
 		}
 	}
 }
@@ -166,11 +177,10 @@ void FAnimNode_IKRig::PreUpdate(const UAnimInstance* InAnimInstance)
 		IKRigProcessor = NewObject<UIKRigProcessor>(InAnimInstance->GetOwningComponent());	
 	}
 	
+	// initialize the IK Rig (will only try once on the current version of the rig asset)
 	if (IKRigProcessor->NeedsInitialized(RigDefinitionAsset))
 	{
-		// get the retargeted local ref pose to initialize the IK with
-		const FReferenceSkeleton& RefSkeleton = InAnimInstance->CurrentSkeleton->GetReferenceSkeleton();
-		// initialize the IK Rig (will only try once on the current version of the rig asset)
+		const FReferenceSkeleton& RefSkeleton = InAnimInstance->GetSkelMeshComponent()->SkeletalMesh->GetRefSkeleton();
 		IKRigProcessor->Initialize(RigDefinitionAsset, RefSkeleton);
 	}
 	
