Fixing issue with IK Rig Body Mover using ref pose instead of input pose.
Added debug view of goals to IK Rig anim node.

diff --git a/Engine/Plugins/Animation/IKRig/Source/IKRig/Private/AnimNodes/AnimNode_IKRig.cpp b/Engine/Plugins/Animation/IKRig/Source/IKRig/Private/AnimNodes/AnimNode_IKRig.cpp
--- a/Engine/Plugins/Animation/IKRig/Source/IKRig/Private/AnimNodes/AnimNode_IKRig.cpp
+++ b/Engine/Plugins/Animation/IKRig/Source/IKRig/Private/AnimNodes/AnimNode_IKRig.cpp
@@ -50,9 +50,6 @@ void FAnimNode_IKRig::Evaluate_AnyThread(FPoseContext& Output)
 	IKRigProcessor->Solve(WorldToComponent);
 	// updates transforms with new pose
 	CopyOutputPoseToAnimGraph(Output.Pose);
-	
-	// debug drawing
-	QueueDrawInterface(Output.AnimInstanceProxy, Output.AnimInstanceProxy->GetComponentTransform());
 }
 
 void FAnimNode_IKRig::CopyInputPoseToSolver(FCompactPose& InputPose)
@@ -246,62 +243,30 @@ void FAnimNode_IKRig::CacheBones_AnyThread(const FAnimationCacheBonesContext& Co
 	}
 }
 
-void FAnimNode_IKRig::QueueDrawInterface(FAnimInstanceProxy* AnimProxy, const FTransform& ComponentToWorld) const
+void FAnimNode_IKRig::ConditionalDebugDraw(
+	FPrimitiveDrawInterface* PDI,
+	USkeletalMeshComponent* PreviewSkelMeshComp) const
 {
-	check (AnimProxy);
+#if WITH_EDITORONLY_DATA
 
-	/*
-	 *TODO implement basic drawing interface for IKRig solvers
-	for (const FControlRigDrawInstruction& Instruction : IKRigProcessor->GetDrawInterface())
+	// is anim graph setup?
+	if (!(bEnableDebugDraw && PreviewSkelMeshComp && PreviewSkelMeshComp->GetWorld()))
 	{
-		if (!Instruction.IsValid())
-		{
-			continue;
-		}
-
-		FTransform InstructionTransform = Instruction.Transform * ComponentToWorld;
-		switch (Instruction.PrimitiveType)
-		{
-		case EControlRigDrawSettings::Points:
-		{
-			ensure(false);
-			// no support for points yet, but feel free to add that to AnimProxy
-// 			for (const FVector& Point : Instruction.Positions)
-// 			{
-// 				AnimProxy->Add (InstructionTransform.TransformPosition(Point), Instruction.Color, Instruction.Thickness, SDPG_Foreground);
-// 			}
- 			break;
-		}
-		case EControlRigDrawSettings::Lines:
-		{
-			const TArray<FVector>& Points = Instruction.Positions;
-			for (int32 PointIndex = 0; PointIndex < Points.Num() - 1; PointIndex += 2)
-			{
-				AnimProxy->AnimDrawDebugLine(InstructionTransform.TransformPosition(Points[PointIndex]), InstructionTransform.TransformPosition(Points[PointIndex + 1]), Instruction.Color.ToFColor(false), false, 0.f, Instruction.Thickness);
-			}
-			break;
-		}
-		case EControlRigDrawSettings::LineStrip:
-		{
-			const TArray<FVector>& Points = Instruction.Positions;
-			for (int32 PointIndex = 0; PointIndex < Points.Num() - 1; PointIndex++)
-			{
-				AnimProxy->AnimDrawDebugLine(InstructionTransform.TransformPosition(Points[PointIndex]), InstructionTransform.TransformPosition(Points[PointIndex + 1]), Instruction.Color.ToFColor(false), false, 0.f, Instruction.Thickness);
-			}
-			break;
-		}
+		return;
+	}
 
-		case EControlRigDrawSettings::DynamicMesh:
-		{
-			ensure(false);
-			// no support for this yet
-// 			FDynamicMeshBuilder MeshBuilder(PDI->View->GetFeatureLevel());
-// 			MeshBuilder.AddVertices(Instruction.MeshVerts);
-// 			MeshBuilder.AddTriangles(Instruction.MeshIndices);
-// 			MeshBuilder.Draw(PDI, InstructionTransform.ToMatrixWithScale(), Instruction.MaterialRenderProxy, SDPG_World);
-			break;
-		}
+	// is node setup?
+	if (!(RigDefinitionAsset && IKRigProcessor && !IKRigProcessor->NeedsInitialized(RigDefinitionAsset)))
+	{
+		return;
+	}
 
-		}
-	}*/
+	const FIKRigGoalContainer& ProcessorGoals = IKRigProcessor->GetGoalContainer();
+	for (const TPair<FName, FIKRigGoal>& GoalPair : ProcessorGoals.Goals)
+	{
+		const FIKRigGoal& Goal = GoalPair.Value;
+		DrawOrientedWireBox(PDI, Goal.FinalBlendedPosition, FVector::XAxisVector, FVector::YAxisVector, FVector::ZAxisVector, FVector::One() * DebugScale, FLinearColor::Yellow, SDPG_World);
+		DrawCoordinateSystem(PDI, Goal.Position, Goal.FinalBlendedRotation.Rotator(), DebugScale, SDPG_World);
+	}
+#endif
 }
\ No newline at end of file
diff --git a/Engine/Plugins/Animation/IKRig/Source/IKRig/Private/Solvers/IKRig_BodyMover.cpp b/Engine/Plugins/Animation/IKRig/Source/IKRig/Private/Solvers/IKRig_BodyMover.cpp
--- a/Engine/Plugins/Animation/IKRig/Source/IKRig/Private/Solvers/IKRig_BodyMover.cpp
+++ b/Engine/Plugins/Animation/IKRig/Source/IKRig/Private/Solvers/IKRig_BodyMover.cpp
@@ -28,9 +28,10 @@ void UIKRig_BodyMover::Solve(FIKRigSkeleton& IKRigSkeleton, const FIKRigGoalCont
 		return;
 	}
 	
-	// accumulate offsets to apply to body bone
+	// ensure body bone exists
 	check(IKRigSkeleton.RefPoseGlobal.IsValidIndex(BodyBoneIndex));
-	// the bone transform to modify (not const!)
+
+	// the bone transform to modify
 	FTransform& CurrentBodyTransform = IKRigSkeleton.CurrentPoseGlobal[BodyBoneIndex];
 
 	// calculate initial and current centroids
@@ -45,7 +46,7 @@ void UIKRig_BodyMover::Solve(FIKRigSkeleton& IKRigSkeleton, const FIKRigGoalCont
 		}
 
 		const int32 BoneIndex = IKRigSkeleton.GetBoneIndexFromName(Effector->BoneName);
-		const FTransform InitialEffector = IKRigSkeleton.RefPoseGlobal[BoneIndex];
+		const FTransform InitialEffector = IKRigSkeleton.CurrentPoseGlobal[BoneIndex];
 
 		InitialCentroid += InitialEffector.GetTranslation();
 		CurrentCentroid += Goal.FinalBlendedPosition;
@@ -69,7 +70,7 @@ void UIKRig_BodyMover::Solve(FIKRigSkeleton& IKRigSkeleton, const FIKRigGoalCont
 		}
 
 		const int32 BoneIndex = IKRigSkeleton.GetBoneIndexFromName(Effector->BoneName);
-		const FTransform InitialEffector = IKRigSkeleton.RefPoseGlobal[BoneIndex];
+		const FTransform InitialEffector = IKRigSkeleton.CurrentPoseGlobal[BoneIndex];
 
 		//
 		// accumulate the deformation gradient tensor for all points
diff --git a/Engine/Plugins/Animation/IKRig/Source/IKRig/Public/AnimNodes/AnimNode_IKRig.h b/Engine/Plugins/Animation/IKRig/Source/IKRig/Public/AnimNodes/AnimNode_IKRig.h
--- a/Engine/Plugins/Animation/IKRig/Source/IKRig/Public/AnimNodes/AnimNode_IKRig.h
+++ b/Engine/Plugins/Animation/IKRig/Source/IKRig/Public/AnimNodes/AnimNode_IKRig.h
@@ -38,9 +38,13 @@ struct IKRIG_API FAnimNode_IKRig : public FAnimNode_Base
 	bool bDriveWithSourceAsset = false;
 
 #if WITH_EDITORONLY_DATA
-	/** Toggle drawing of axes to debug joint rotation*/
+	/** Toggle debug drawing of goals when node is selected.*/
 	UPROPERTY(EditAnywhere, Category = Solver)
 	bool bEnableDebugDraw;
+
+	/** Adjust size of debug drawing.*/
+	UPROPERTY(EditAnywhere, Category = Solver)
+	float DebugScale = 5.0f;
 #endif
 
 private:
@@ -74,7 +78,7 @@ private:
 	void CopyInputPoseToSolver(FCompactPose& InputPose);
 	void AssignGoalTargets();
 	void CopyOutputPoseToAnimGraph(FCompactPose& OutputPose);	
-	void QueueDrawInterface(FAnimInstanceProxy* AnimProxy, const FTransform& ComponentToWorld) const;
+	virtual void ConditionalDebugDraw(FPrimitiveDrawInterface* PDI, USkeletalMeshComponent* PreviewSkelMeshComp) const;
 	
 	friend class UAnimGraphNode_IKRig;
 	friend struct FIKRigAnimInstanceProxy;
diff --git a/Engine/Plugins/Animation/IKRig/Source/IKRigDeveloper/Private/AnimGraphNode_IKRig.cpp b/Engine/Plugins/Animation/IKRig/Source/IKRigDeveloper/Private/AnimGraphNode_IKRig.cpp
--- a/Engine/Plugins/Animation/IKRig/Source/IKRigDeveloper/Private/AnimGraphNode_IKRig.cpp
+++ b/Engine/Plugins/Animation/IKRig/Source/IKRigDeveloper/Private/AnimGraphNode_IKRig.cpp
@@ -12,11 +12,11 @@
 
 void UAnimGraphNode_IKRig::Draw(FPrimitiveDrawInterface* PDI, USkeletalMeshComponent* PreviewSkelMeshComp) const
 {
-	if (PreviewSkelMeshComp)
+	if(PreviewSkelMeshComp)
 	{
-		if (FAnimNode_IKRig* ActiveNode = GetActiveInstanceNode<FAnimNode_IKRig>(PreviewSkelMeshComp->GetAnimInstance()))
+		if(FAnimNode_IKRig* ActiveNode = GetActiveInstanceNode<FAnimNode_IKRig>(PreviewSkelMeshComp->GetAnimInstance()))
 		{
-			//ActiveNode->ConditionalDebugDraw(PDI, PreviewSkelMeshComp);
+			ActiveNode->ConditionalDebugDraw(PDI, PreviewSkelMeshComp);
 		}
 	}
 }
