New Retarget Editor
"Exclude Bones" added to IK Rig
Bone Chains added to IK Rig

Control Rig / FBIK fix. Moved PBIK init into the update phase of the execution. Fixes issue with rig initializing with incorrect InitialTransform values after swapping out skeletal mesh asset.
Changed "PinRotation" from bool to float for blending between input and effector rotation.
Added Excluded Bones support
Added tooltips to all exposed properties.

diff --git a/Engine/Plugins/Animation/IKRig/Source/IKRig/Private/AnimNodes/AnimNode_IKRetargeter.cpp b/Engine/Plugins/Animation/IKRig/Source/IKRig/Private/AnimNodes/AnimNode_IKRetargeter.cpp
--- a/Engine/Plugins/Animation/IKRig/Source/IKRig/Private/AnimNodes/AnimNode_IKRetargeter.cpp
+++ b/Engine/Plugins/Animation/IKRig/Source/IKRig/Private/AnimNodes/AnimNode_IKRetargeter.cpp
@@ -37,7 +37,7 @@ void FAnimNode_IKRetargeter::Evaluate_AnyThread(FPoseContext& Output)
 	}
 
 	// run the retargeter
-	CurrentlyUsedRetargeter->RunRetargeter(SourceMeshComponentSpaceBoneTransforms, bEnableIK);
+	CurrentlyUsedRetargeter->RunRetargeter(SourceMeshComponentSpaceBoneTransforms);
 
 	// copy pose back
 	FCSPose<FCompactPose> ComponentPose;
@@ -61,6 +61,11 @@ void FAnimNode_IKRetargeter::PreUpdate(const UAnimInstance* InAnimInstance)
 	}
 }
 
+UIKRetargeter* FAnimNode_IKRetargeter::GetCurrentlyUsedRetargeter() const
+{
+	return CurrentlyUsedRetargeter.Get();
+}
+
 void FAnimNode_IKRetargeter::EnsureInitialized(const UAnimInstance* InAnimInstance)
 {
 	// has user supplied a retargeter asset?
@@ -87,35 +92,13 @@ void FAnimNode_IKRetargeter::EnsureInitialized(const UAnimInstance* InAnimInstan
 		bIsInitialized = false;
 		return; // can't do anything if we don't have a source mesh
 	}
-	
-	// are all the parts already loaded?
-	const bool bPartsLoaded =
-	   CurrentlyUsedSourceMeshComponent.IsValid()
-    && CurrentlyUsedSourceMesh.IsValid()
-    && CurrentlyUsedTargetMesh.IsValid()
-	&& CurrentlyUsedRetargeter
-	&& CurrentlyUsedSourceIKRig.IsValid()
-	&& CurrentlyUsedTargetIKRig.IsValid();
-	if (!bPartsLoaded)
-	{
-		InitializeRetargetData(InAnimInstance); // nothing loaded yet, initialize
-		return;
-	}
 
-	// so parts have been loaded, but have any of the parts changed since we last initialized?
-	USkeletalMeshComponent* SourceMeshComp = SourceMeshComponent.Get();
-	USkeletalMesh* TargetMesh = InAnimInstance->GetSkelMeshComponent()->SkeletalMesh;
-	USkeletalMesh* SourceMesh = SourceMeshComp->SkeletalMesh;
-	const bool bSameParts =
-	   CurrentlyUsedSourceMeshComponent == SourceMeshComp 
-    && CurrentlyUsedTargetMesh == TargetMesh 
-    && CurrentlyUsedSourceMesh == SourceMesh 
-	&& CurrentlyUsedSourceIKRig == IKRetargeterAsset->SourceIKRigAsset
-	&& CurrentlyUsedTargetIKRig == IKRetargeterAsset->TargetIKRigAsset;
-	if (!bSameParts)
+	// has the source asset been updated? this can happen in the editor when changing source asset
+	if (!CurrentlyUsedRetargeter
+		|| (CurrentlyUsedRetargeter->AssetVersion != IKRetargeterAsset->AssetVersion))
 	{
-		InitializeRetargetData(InAnimInstance); // parts have changed, re-initialize
-		return;
+		InitializeRetargetData(InAnimInstance);
+		return;	
 	}
 }
 
@@ -123,6 +106,13 @@ void FAnimNode_IKRetargeter::InitializeRetargetData(const UAnimInstance* InAnimI
 {
 	// assume we fail until we don't
 	bIsInitialized = false;
+	
+	// don't try to reinitialize with same asset version
+	if (LastVersionTried == IKRetargeterAsset->AssetVersion)
+	{
+		return;
+	}
+	LastVersionTried = IKRetargeterAsset->AssetVersion;
 
 	// store all the components that were used to initialize
 	// if in future updates, any of this are mismatched, we have to re-initialize
@@ -130,6 +120,7 @@ void FAnimNode_IKRetargeter::InitializeRetargetData(const UAnimInstance* InAnimI
 	CurrentlyUsedSourceMesh = SourceMeshComponent->SkeletalMesh;
 	CurrentlyUsedTargetMesh = InAnimInstance->GetSkelMeshComponent()->SkeletalMesh;
 	CurrentlyUsedRetargeter = DuplicateObject(IKRetargeterAsset, const_cast<UAnimInstance*>(InAnimInstance));
+	CurrentlyUsedRetargeter->AssetVersion = IKRetargeterAsset->AssetVersion;
 	CurrentlyUsedSourceIKRig = CurrentlyUsedRetargeter->SourceIKRigAsset;
 	CurrentlyUsedTargetIKRig = CurrentlyUsedRetargeter->TargetIKRigAsset;
 
diff --git a/Engine/Plugins/Animation/IKRig/Source/IKRig/Private/IKRigDefinition.cpp b/Engine/Plugins/Animation/IKRig/Source/IKRig/Private/IKRigDefinition.cpp
--- a/Engine/Plugins/Animation/IKRig/Source/IKRig/Private/IKRigDefinition.cpp
+++ b/Engine/Plugins/Animation/IKRig/Source/IKRig/Private/IKRigDefinition.cpp
@@ -83,4 +83,14 @@ FTransform UIKRigDefinition::GetGoalInitialTransform(UIKRigEffectorGoal* Goal) c
 	return Skeleton.RefPoseGlobal[BoneIndex];
 }
 
+void UIKRigDefinition::SortRetargetChains()
+{
+	RetargetDefinition.BoneChains.Sort([this](const FBoneChain& A, const FBoneChain& B)
+	{
+		const int32 IndexA = this->Skeleton.GetBoneIndexFromName(A.StartBone);
+		const int32 IndexB = this->Skeleton.GetBoneIndexFromName(B.StartBone);
+		return IndexA < IndexB;
+	});
+}
+
 #undef LOCTEXT_NAMESPACE
diff --git a/Engine/Plugins/Animation/IKRig/Source/IKRig/Private/IKRigProcessor.cpp b/Engine/Plugins/Animation/IKRig/Source/IKRig/Private/IKRigProcessor.cpp
--- a/Engine/Plugins/Animation/IKRig/Source/IKRig/Private/IKRigProcessor.cpp
+++ b/Engine/Plugins/Animation/IKRig/Source/IKRig/Private/IKRigProcessor.cpp
@@ -12,6 +12,13 @@ void UIKRigProcessor::Initialize(UIKRigDefinition* InRigAsset, const FReferenceS
 	bInitialized = false;
 	InitializedWithIKRigAssetVersion = -1;
 
+	// bail out if we've already tried initializing with this exact version of the rig asset
+	if (LastVersionTried == InRigAsset->GetAssetVersion())
+	{
+		return; // don't keep spamming
+	}
+	LastVersionTried = InRigAsset->GetAssetVersion();
+
 	if (!InRigAsset)
 	{
 		UE_LOG(LogTemp, Warning, TEXT("Trying to initialize IKRigProcessor with a null IKRigDefinition asset."));
@@ -20,34 +27,28 @@ void UIKRigProcessor::Initialize(UIKRigDefinition* InRigAsset, const FReferenceS
 
 	if (InRigAsset->Skeleton.BoneNames.IsEmpty())
 	{
-		UE_LOG(LogTemp, Warning, TEXT("Trying to initialize IKRigProcessor with a IKRigDefinition that has no skeleton."));
+		UE_LOG(LogTemp, Warning, TEXT("Trying to initialize IKRigProcessor with a IKRigDefinition that has no skeleton: %s"), *InRigAsset->GetName());
 		return;
 	}
 
 	if (InRigAsset->Solvers.IsEmpty())
 	{
-		UE_LOG(LogTemp, Warning, TEXT("Trying to initialize IKRigProcessor with a IKRigDefinition that has no solvers."));
+		UE_LOG(LogTemp, Warning, TEXT("Trying to initialize IKRigProcessor with a IKRigDefinition that has no solvers: %s"), *InRigAsset->GetName());
 		return;
 	}
 
 	if (InRigAsset->Goals.IsEmpty())
 	{
-		UE_LOG(LogTemp, Warning, TEXT("Trying to initialize IKRigProcessor with a IKRigDefinition that has no goals."));
+		UE_LOG(LogTemp, Warning, TEXT("Trying to initialize IKRigProcessor with a IKRigDefinition that has no goals: %s"), *InRigAsset->GetName());
 		return;
 	}
 
-	// bail out if we've already tried initializing with this exact version of the rig asset
-	if (LastVersionTried == InRigAsset->GetAssetVersion())
-	{
-		return; // don't keep spamming
-	}
-	LastVersionTried = InRigAsset->GetAssetVersion();
-
 	// copy skeleton data from IKRigDefinition
 	// we use the serialized bone names and parent indices (from when asset was initialized)
 	// but we use the CURRENT ref pose from the currently running skeletal mesh (RefSkeleton)
 	Skeleton.BoneNames = InRigAsset->Skeleton.BoneNames;
 	Skeleton.ParentIndices = InRigAsset->Skeleton.ParentIndices;
+	Skeleton.ExcludedBones = InRigAsset->Skeleton.ExcludedBones;
 	const bool bSkeletonIsCompatible = Skeleton.CopyPosesFromRefSkeleton(RefSkeleton);
 	if (!bSkeletonIsCompatible)
 	{
diff --git a/Engine/Plugins/Animation/IKRig/Source/IKRig/Private/IKRigSkeleton.cpp b/Engine/Plugins/Animation/IKRig/Source/IKRig/Private/IKRigSkeleton.cpp
--- a/Engine/Plugins/Animation/IKRig/Source/IKRig/Private/IKRigSkeleton.cpp
+++ b/Engine/Plugins/Animation/IKRig/Source/IKRig/Private/IKRigSkeleton.cpp
@@ -2,7 +2,9 @@
 
 #include "IKRigSkeleton.h"
 
-void FIKRigSkeleton::Initialize(const FReferenceSkeleton& RefSkeleton)
+#include "IKRigDefinition.h"
+
+void FIKRigSkeleton::Initialize(const FReferenceSkeleton& RefSkeleton, const TArray<FName>& InExcludedBones)
 {
 	// reset all containers
 	Reset();
@@ -23,6 +25,7 @@ void FIKRigSkeleton::Reset()
 {
 	BoneNames.Reset();
 	ParentIndices.Reset();
+	ExcludedBones.Reset();
 	CurrentPoseGlobal.Reset();
 	CurrentPoseLocal.Reset();
 	RefPoseGlobal.Reset();
@@ -41,9 +44,19 @@ int32 FIKRigSkeleton::GetBoneIndexFromName(const FName InName) const
 	return INDEX_NONE;
 }
 
+FName FIKRigSkeleton::GetBoneNameFromIndex(const int32 BoneIndex) const
+{
+	if (BoneNames.IsValidIndex(BoneIndex))
+	{
+		return BoneNames[BoneIndex];
+	}
+
+	return NAME_None;
+}
+
 int32 FIKRigSkeleton::GetParentIndex(const int32 BoneIndex) const
 {
-	if (BoneIndex < 0 || BoneIndex>ParentIndices.Num() || BoneIndex == INDEX_NONE)
+	if (!BoneNames.IsValidIndex(BoneIndex))
 	{
 		return INDEX_NONE;
 	}
@@ -51,6 +64,22 @@ int32 FIKRigSkeleton::GetParentIndex(const int32 BoneIndex) const
 	return ParentIndices[BoneIndex];
 }
 
+int32 FIKRigSkeleton::GetParentIndexThatIsNotExcluded(const int32 BoneIndex) const
+{
+	// find first parent that is NOT excluded
+	int32 ParentIndex = GetParentIndex(BoneIndex);
+	while(ParentIndex != INDEX_NONE)
+	{
+		if (!IsBoneExcluded(ParentIndex))
+		{
+			break;
+		}
+		ParentIndex = GetParentIndex(ParentIndex);
+	}
+	
+	return ParentIndex;
+}
+
 bool FIKRigSkeleton::CopyPosesFromRefSkeleton(const FReferenceSkeleton& RefSkeleton)
 {
 	// get a compacted local ref pose based on the stored names
@@ -185,6 +214,16 @@ bool FIKRigSkeleton::IsBoneInDirectLineage(const FName& Child, const FName& Pote
 	return false;
 }
 
+bool FIKRigSkeleton::IsBoneExcluded(const int32 BoneIndex) const
+{
+	if (ensure(BoneNames.IsValidIndex(BoneIndex)))
+	{
+		return ExcludedBones.Contains(BoneNames[BoneIndex]);	
+	}
+	
+	return false;
+}
+
 void FIKRigSkeleton::NormalizeRotations(TArray<FTransform>& Transforms)
 {
 	for (FTransform& Transform : Transforms)
@@ -192,3 +231,51 @@ void FIKRigSkeleton::NormalizeRotations(TArray<FTransform>& Transforms)
 		Transform.NormalizeRotation();
 	}
 }
+
+void FIKRigSkeleton::GetChainsInList(const TArray<int32>& SelectedBones, TArray<FIKRigSkeletonChain>& OutChains) const
+{	
+	// no bones provided, return empty list
+	if (SelectedBones.IsEmpty())
+	{
+		return;
+	}
+
+	// get parents of all the selected bones
+	TSet<int32> SelectedParentIndices;
+	for (const int32 SelectedBone : SelectedBones)
+	{
+		SelectedParentIndices.Add(GetParentIndex(SelectedBone));
+	}
+	
+	// find all selected leaf nodes and make each one a chain
+    for (const int32 SelectedBone : SelectedBones)
+    {
+    	// is this a leaf node in the selection?
+    	if (SelectedParentIndices.Contains(SelectedBone))
+    	{
+    		continue;
+    	}
+
+    	// convert leaf nodes to chains
+    	int32 ChainStart;
+    	const int32 ChainEnd = SelectedBone;
+    	int32 ParentIndex = SelectedBone;
+    	while (true)
+    	{
+    		ChainStart = ParentIndex;
+    		ParentIndex = GetParentIndex(ParentIndex);
+    		if (ParentIndex == INDEX_NONE)
+    		{
+    			break;
+    		}
+    		if (!SelectedBones.Contains(ParentIndex))
+    		{
+    			break;
+    		}
+    	}
+
+    	const FName StartBoneName = GetBoneNameFromIndex(ChainStart);
+    	const FName EndBoneName = GetBoneNameFromIndex(ChainEnd);
+    	OutChains.Emplace(StartBoneName, EndBoneName);	
+    }
+}
diff --git a/Engine/Plugins/Animation/IKRig/Source/IKRig/Private/Retargeter/IKRetargeter.cpp b/Engine/Plugins/Animation/IKRig/Source/IKRig/Private/Retargeter/IKRetargeter.cpp
--- a/Engine/Plugins/Animation/IKRig/Source/IKRig/Private/Retargeter/IKRetargeter.cpp
+++ b/Engine/Plugins/Animation/IKRig/Source/IKRig/Private/Retargeter/IKRetargeter.cpp
@@ -1,60 +1,23 @@
 // Copyright Epic Games, Inc. All Rights Reserved.
 
 #include "Retargeter/IKRetargeter.h"
-#include "Engine/SkeletalMesh.h"
-
-/*
-TODO:
 
-Asset Testing:
-	- add chain for "upperarm_twist_01_l"
-			- will test single-bone chains AND fix twisted upperarm
-	- add finger chains to test fk AFTER ik
-	- export starfish in mannequin defailt pose
+#include "Algo/LevenshteinDistance.h"
+#include "Engine/SkeletalMesh.h"
 
-IK Rig:
-	- add "alpha" value to goals, set alpha to zero by default
-		- internally LERP effectors from input pose to goal transform by alpha
-	- port PBIK to IK Rig
+const FName UIKRetargeter::DefaultPoseName = "Default Pose";
 
-IK Retargeter:
-	- we need option to override initial poses to be different than reference pose
-	- expose root stride scale factor as parameter
+void FRetargetSkeleton::Initialize(
+	USkeletalMesh* SkeletalMesh,
+	FIKRetargetPose* RetargetPose)
+{
+	// record name for debug warnings
+	SkeletalMeshName = SkeletalMesh->GetName();
 
-FK Retarget:
-	- sort FK chains root to tip (no guarantee user sets them up this way)
-	
-IK Retarget:
-	
-	- set weight of unused goals to zero (so redundant goals don't screw up solve)
-	- set pole vectors goals on IKRig
-
-DONE:
-- create IKRigSkeleton to merge hierarchy and IKRigTransforms
-- IKRigTransforms to use only global space and provide conversion from local space
-- initialize all the things
-- encode (remove decode)
-
-
-NOTES:
-
-// WE PROBABLY DON'T NEED TO DO SEPARATE PASS OF BEFORE/AFTER IK RETARGETS BECAUSE THE
-// FK RETARGETING SHOULD ONLY AFFECT ROTATION WHICH IS INVARIANT ??
-//
-// ALWAYS RUN FK RETARGET?
-// WE MAY WANT TO ALWAYS RUN FK RETARGET BEFORE IK
-// THAT WAY THE CHAINS WILL BE CLOSE TO GOAL TARGETS!!
-// THIS IS DIFFERENT THAN MAYA BECAUSE WE HAVE FULL BODY IK
-// !!!! ACTUALLY IF WE DO FK UPDATE OF GLOBAL SPACE POSE BEFORE IK WE SHOULD BE FINE AND MUCH FASTER!!!
+	// get reference skeleton
+	const FReferenceSkeleton& RefSkeleton = SkeletalMesh->GetRefSkeleton();
 	
-*/
-
-void FRetargetSkeleton::InitializeHierarchy(
-	const USkeletalMesh* SkeletalMesh,
-	const TArray<FTransform>& InLocalRetargetPose)
-{
 	// copy names and parent indices into local storage
-	const FReferenceSkeleton& RefSkeleton = SkeletalMesh->GetRefSkeleton();
 	for (int32 BoneIndex=0; BoneIndex<RefSkeleton.GetNum(); ++BoneIndex)
 	{
 		BoneNames.Add(RefSkeleton.GetBoneName(BoneIndex));
@@ -62,21 +25,29 @@ void FRetargetSkeleton::InitializeHierarchy(
 	}
 
 	// store copy of retarget pose
-	RetargetGlobalPose.Reset(InLocalRetargetPose.Num());
-	
-	// convert local transforms to global
-	for (int32 BoneIndex=0; BoneIndex<InLocalRetargetPose.Num(); ++BoneIndex)
+	RetargetLocalPose = RefSkeleton.GetRefBonePose();
+	RetargetGlobalPose = RetargetLocalPose;
+
+	// update global retarget pose
+	UpdateGlobalTransformsBelowBone(0, RetargetLocalPose, RetargetGlobalPose);
+
+	// apply bone rotation offsets to retarget pose
+	if (RetargetPose != nullptr)
 	{
-		const int32 ParentIndex = ParentIndices[BoneIndex];
-		if (ParentIndex == INDEX_NONE)
+		for (const TTuple<FName, FQuat>& BoneRotationOffset : RetargetPose->BoneRotationOffsets)
 		{
-			RetargetGlobalPose.Add(InLocalRetargetPose[BoneIndex]); // root bone is always in local space
-			continue; 
-		}
+			const int32 BoneIndex = FindBoneIndexByName(BoneRotationOffset.Key);
+			if (BoneIndex == INDEX_NONE)
+			{
+				continue;
+			}
 
-		const FTransform& ChildLocalTransform  = InLocalRetargetPose[BoneIndex];
-		const FTransform& ParentGlobalTransform  = RetargetGlobalPose[ParentIndex];
-		RetargetGlobalPose.Add(ChildLocalTransform * ParentGlobalTransform);
+			FTransform& ChainStartTransform = RetargetGlobalPose[BoneIndex];
+			FQuat RotatedChain = BoneRotationOffset.Value * ChainStartTransform.GetRotation();
+			ChainStartTransform.SetRotation(RotatedChain);
+			UpdateLocalTransformOfSingleBone(BoneIndex, RetargetLocalPose, RetargetGlobalPose);
+			UpdateGlobalTransformsBelowBone(BoneIndex, RetargetLocalPose, RetargetGlobalPose);
+		}
 	}
 }
 
@@ -93,6 +64,66 @@ int32 FRetargetSkeleton::FindBoneIndexByName(const FName InName) const
 	return INDEX_NONE;
 }
 
+void FRetargetSkeleton::UpdateGlobalTransformsBelowBone(
+	const int32 StartBoneIndex,
+	TArray<FTransform>& InLocalPose,
+	TArray<FTransform>& OutGlobalPose) const
+{
+	check(BoneNames.IsValidIndex(StartBoneIndex));
+	check(BoneNames.Num() == InLocalPose.Num());
+	check(BoneNames.Num() == OutGlobalPose.Num());
+	
+	for (int32 BoneIndex=StartBoneIndex+1; BoneIndex<OutGlobalPose.Num(); ++BoneIndex)
+	{
+		UpdateGlobalTransformOfSingleBone(BoneIndex,InLocalPose,OutGlobalPose);
+	}
+}
+
+void FRetargetSkeleton::UpdateLocalTransformsBelowBone(
+	const int32 StartBoneIndex,
+	TArray<FTransform>& OutLocalPose,
+	TArray<FTransform>& InGlobalPose) const
+{
+	check(BoneNames.IsValidIndex(StartBoneIndex));
+	check(BoneNames.Num() == OutLocalPose.Num());
+	check(BoneNames.Num() == InGlobalPose.Num());
+	
+	for (int32 BoneIndex=StartBoneIndex+1; BoneIndex<InGlobalPose.Num(); ++BoneIndex)
+	{
+		UpdateLocalTransformOfSingleBone(BoneIndex, OutLocalPose, InGlobalPose);
+	}
+}
+
+void FRetargetSkeleton::UpdateGlobalTransformOfSingleBone(
+	const int32 BoneIndex,
+	const TArray<FTransform>& InLocalPose,
+	TArray<FTransform>& OutGlobalPose) const
+{
+	const int32 ParentIndex = ParentIndices[BoneIndex];
+	if (ParentIndex == INDEX_NONE)
+	{
+		return; // root always in global space
+	}
+	const FTransform ChildLocalTransform = InLocalPose[BoneIndex];
+	const FTransform ParentGlobalTransform = OutGlobalPose[ParentIndex];
+	OutGlobalPose[BoneIndex] = ChildLocalTransform * ParentGlobalTransform;
+}
+
+void FRetargetSkeleton::UpdateLocalTransformOfSingleBone(
+	const int32 BoneIndex,
+	TArray<FTransform>& OutLocalPose,
+	TArray<FTransform>& InGlobalPose) const
+{
+	const int32 ParentIndex = ParentIndices[BoneIndex];
+	if (ParentIndex == INDEX_NONE)
+	{
+		return; // root always in global space
+	}
+	const FTransform ChildGlobalTransform = InGlobalPose[BoneIndex];
+	const FTransform ParentGlobalTransform = InGlobalPose[ParentIndex];
+	OutLocalPose[BoneIndex] = ChildGlobalTransform.GetRelativeTransform(ParentGlobalTransform);
+}
+
 int32 FRetargetSkeleton::GetParentIndex(const int32 BoneIndex) const
 {
 	if (BoneIndex < 0 || BoneIndex>ParentIndices.Num() || BoneIndex == INDEX_NONE)
@@ -104,62 +135,38 @@ int32 FRetargetSkeleton::GetParentIndex(const int32 BoneIndex) const
 }
 
 void FTargetSkeleton::Initialize(
-	const USkeletalMesh* SkeletalMesh,
-	const TArray<FTransform>& InTargetRetargetLocalPose)
+	USkeletalMesh* SkeletalMesh, 
+	FIKRetargetPose* RetargetPose)
 {
-	InitializeHierarchy(SkeletalMesh, InTargetRetargetLocalPose);
+	FRetargetSkeleton::Initialize(SkeletalMesh, RetargetPose);
 
 	// initialize storage for output pose (the result of the retargeting)
 	OutputGlobalPose = RetargetGlobalPose;
-	// store local retarget pose so we can update partial skeleton during retarget process
-	RetargetLocalPose = InTargetRetargetLocalPose;
 
-	// make storage for per-bone "IsFK" flag (used for hierarchy updates)
-	IsBoneFK.Init(false, InTargetRetargetLocalPose.Num());
-}
-
-void FTargetSkeleton::SetBoneControlledByFKChain(const int32 BoneIndex, const bool IsFK)
-{
-	check(IsBoneFK.IsValidIndex(BoneIndex));
-	IsBoneFK[BoneIndex] = IsFK;
+	// make storage for per-bone "Is Retargeted" flag (used for hierarchy updates)
+	IsBoneRetargeted.Init(false, OutputGlobalPose.Num());
 }
 
-void FTargetSkeleton::UpdateGlobalTransformsBelowBone(const int32 StartBoneIndex)
+void FTargetSkeleton::UpdateGlobalTransformsAllNonRetargetedBones(TArray<FTransform>& InOutGlobalPose)
 {
-	check(OutputGlobalPose.IsValidIndex(StartBoneIndex));
+	check(IsBoneRetargeted.Num() == InOutGlobalPose.Num());
 	
-	for (int32 BoneIndex=StartBoneIndex+1; BoneIndex<OutputGlobalPose.Num(); ++BoneIndex)
+	for (int32 BoneIndex=0; BoneIndex<InOutGlobalPose.Num(); ++BoneIndex)
 	{
-		UpdateGlobalTransformOfSingleBone(BoneIndex);
-	}
-}
-
-void FTargetSkeleton::UpdateGlobalTransformsAllNonFKBones()
-{
-	check(IsBoneFK.Num() == OutputGlobalPose.Num());
-	
-	for (int32 BoneIndex=0; BoneIndex<OutputGlobalPose.Num(); ++BoneIndex)
-	{
-		if (!IsBoneFK[BoneIndex])
+		if (!IsBoneRetargeted[BoneIndex])
 		{
-			UpdateGlobalTransformOfSingleBone(BoneIndex);
+			UpdateGlobalTransformOfSingleBone(BoneIndex, RetargetLocalPose, InOutGlobalPose);
 		}
 	}
 }
 
-void FTargetSkeleton::UpdateGlobalTransformOfSingleBone(const int32 BoneIndex)
+void FTargetSkeleton::SetBoneIsRetargeted(const int32 BoneIndex, const bool IsRetargeted)
 {
-	const int32 ParentIndex = ParentIndices[BoneIndex];
-	if (ParentIndex == INDEX_NONE)
-	{
-		return; // root always in global space
-	}
-	const FTransform ChildLocalTransform = RetargetLocalPose[BoneIndex];
-	const FTransform ParentGlobalTransform = OutputGlobalPose[ParentIndex];
-	OutputGlobalPose[BoneIndex] = ChildLocalTransform * ParentGlobalTransform;
+	check(IsBoneRetargeted.IsValidIndex(BoneIndex));
+	IsBoneRetargeted[BoneIndex] = IsRetargeted;
 }
 
-bool FChainEncoderFK::Initialize(const TArray<int32>& BoneIndices, const TArray<FTransform>& InitialGlobalPose)
+bool FChainFK::Initialize(const TArray<int32>& BoneIndices, const TArray<FTransform>& InitialGlobalPose)
 {
 	check(!BoneIndices.IsEmpty());
 
@@ -178,7 +185,7 @@ bool FChainEncoderFK::Initialize(const TArray<int32>& BoneIndices, const TArray<
 	return CalculateBoneParameters();
 }
 
-bool FChainEncoderFK::CalculateBoneParameters()
+bool FChainFK::CalculateBoneParameters()
 {
 	// special case, a single-bone chain
 	if (InitialBoneTransforms.Num() == 1)
@@ -227,19 +234,30 @@ void FChainEncoderFK::EncodePose(
 	}
 }
 
-void FChainEncoderFK::DecodePose(
+void FChainDecoderFK::DecodePose(
 	const TArray<int32>& TargetBoneIndices,
     const FChainEncoderFK& SourceChain,
+    const FTargetSkeleton& TargetSkeleton,
     TArray<FTransform> &InOutGlobalPose)
 {
 	check(TargetBoneIndices.Num() == CurrentBoneTransforms.Num());
 	check(TargetBoneIndices.Num() == Params.Num());
-	
+
+	// Before setting this chain pose, we need to ensure that any
+	// intermediate (between chains) NON-retargeted parent bones have had their
+	// global transforms updated.
+	// 
+	// For example, if this chain is retargeting a single head bone, AND the spine was
+	// retargeted in the prior step, then the neck bones will need updating first.
+	// Otherwise the neck bones will remain at their location prior to the spine update.
+	UpdateIntermediateParents(TargetSkeleton,InOutGlobalPose);
+
+	// now retarget the pose of each bone in the chain, copying from source to target
 	for (int32 ChainIndex=0; ChainIndex<TargetBoneIndices.Num(); ++ChainIndex)
 	{
 		// get the initial and current transform of source chain at param
 		// this is the interpolated transform along the chain
-		const float Param = Params[ChainIndex];
+		float Param = Params[ChainIndex];
 		FTransform SourceCurrentTransform = GetTransformAtParam(
 			SourceChain.CurrentBoneTransforms,
 			SourceChain.Params,
@@ -249,33 +267,73 @@ void FChainEncoderFK::DecodePose(
 			SourceChain.Params,
 			Param);
 
+		const int32 BoneIndex = TargetBoneIndices[ChainIndex];
+		const int32 ParentIndex = TargetSkeleton.ParentIndices[BoneIndex];
+		
 		// the initial transform on the target chain
 		const FTransform& TargetInitialTransform = InitialBoneTransforms[ChainIndex];
 
 		// calculate output ROTATION
-		const FQuat& SourceCurrentRotation = SourceCurrentTransform.GetRotation();
-		const FQuat& SourceInitialRotation = SourceInitialTransform.GetRotation();
-		const FQuat& TargetInitialRotation = TargetInitialTransform.GetRotation();
-		const FQuat RotationDelta = TargetInitialRotation * SourceInitialRotation.Inverse();
-		const FQuat OutRotation = RotationDelta * SourceCurrentRotation;
-
-		// calculate output POSITION
-		const FVector& SourceCurrentPosition = SourceCurrentTransform.GetTranslation();
-		const FVector& SourceInitialPosition = SourceInitialTransform.GetTranslation();
-		const FVector& TargetInitialPosition = TargetInitialTransform.GetTranslation();
-		const FVector OutPosition = SourceCurrentPosition + (TargetInitialPosition - SourceInitialPosition);
+		const FQuat SourceCurrentRotation = SourceCurrentTransform.GetRotation();
+		const FQuat SourceInitialRotation = SourceInitialTransform.GetRotation();
+		const FQuat RotationDelta = SourceCurrentRotation * SourceInitialRotation.Inverse();
+		const FQuat TargetInitialRotation = TargetInitialTransform.GetRotation();
+		const FQuat OutRotation = RotationDelta * TargetInitialRotation;
+
+		// calculate output POSITION (constant for now, maybe we support retargeting translation later?)
+		const FTransform& ParentGlobalTransform = InOutGlobalPose[ParentIndex];
+		const FVector InitialLocalOffset = TargetSkeleton.RetargetLocalPose[BoneIndex].GetTranslation();
+		const FVector OutPosition = ParentGlobalTransform.TransformPosition(InitialLocalOffset);
 
 		// calculate output SCALE
-		const FVector& SourceCurrentScale = SourceCurrentTransform.GetScale3D();
-		const FVector& SourceInitialScale = SourceInitialTransform.GetScale3D();
-		const FVector& TargetInitialScale = TargetInitialTransform.GetScale3D();
+		const FVector SourceCurrentScale = SourceCurrentTransform.GetScale3D();
+		const FVector SourceInitialScale = SourceInitialTransform.GetScale3D();
+		const FVector TargetInitialScale = TargetInitialTransform.GetScale3D();
 		const FVector OutScale = SourceCurrentScale + (TargetInitialScale - SourceInitialScale);
+		
+		// apply output transform
+		FTransform OutTransform = FTransform(OutRotation, OutPosition, OutScale);
+		InOutGlobalPose[BoneIndex] = OutTransform;
+	}
+}
 
-		InOutGlobalPose[TargetBoneIndices[ChainIndex]] = FTransform(OutRotation, OutPosition, OutScale);
+void FChainDecoderFK::InitializeIntermediateParentIndices(
+	const int32 RetargetRootBoneIndex,
+	const int32 ChainRootBoneIndex,
+	const FTargetSkeleton& TargetSkeleton)
+{
+	IntermediateParentIndices.Reset();
+	int32 ParentBoneIndex = TargetSkeleton.ParentIndices[ChainRootBoneIndex];
+	while (true)
+	{
+		if (ParentBoneIndex < 0 || ParentBoneIndex == RetargetRootBoneIndex)
+		{
+			break; // reached root of skeleton
+		}
+
+		if (TargetSkeleton.IsBoneRetargeted[ParentBoneIndex])
+		{
+			break; // reached the start of another retargeted chain
+		}
+
+		IntermediateParentIndices.Add(ParentBoneIndex);
+		ParentBoneIndex = TargetSkeleton.ParentIndices[ParentBoneIndex];
+	}
+
+	Algo::Reverse(IntermediateParentIndices);
+}
+
+void FChainDecoderFK::UpdateIntermediateParents(
+	const FTargetSkeleton& TargetSkeleton,
+	TArray<FTransform>& InOutGlobalPose)
+{
+	for (const int32& ParentIndex : IntermediateParentIndices)
+	{
+		TargetSkeleton.UpdateGlobalTransformOfSingleBone(ParentIndex, TargetSkeleton.RetargetLocalPose, InOutGlobalPose);
 	}
 }
 
-FTransform FChainEncoderFK::GetTransformAtParam(
+FTransform FChainDecoderFK::GetTransformAtParam(
 	const TArray<FTransform>& Transforms,
 	const TArray<float>& InParams,
 	const float Param) const
@@ -303,7 +361,7 @@ FTransform FChainEncoderFK::GetTransformAtParam(
 		const FTransform& Prev = Transforms[ChainIndex-1];
 		const FTransform& Next = Transforms[ChainIndex];
 		const FVector Position = FMath::Lerp(Prev.GetTranslation(), Next.GetTranslation(), PercentBetweenParams);
-		const FQuat Rotation = FQuat::FastLerp(Prev.GetRotation(), Next.GetRotation(), PercentBetweenParams);
+		const FQuat Rotation = FQuat::FastLerp(Prev.GetRotation(), Next.GetRotation(), PercentBetweenParams).GetNormalized();
 		const FVector Scale = FMath::Lerp(Prev.GetScale3D(), Next.GetScale3D(), PercentBetweenParams);
 		
 		return FTransform(Rotation,Position, Scale);
@@ -347,9 +405,9 @@ bool FChainEncoderIK::Initialize(
 
 void FChainEncoderIK::EncodePose(const TArray<FTransform>& InputGlobalPose)
 {
-	FVector A = InputGlobalPose[BoneIndexA].GetTranslation();
-	FVector B = InputGlobalPose[BoneIndexB].GetTranslation();
-	FVector C = InputGlobalPose[BoneIndexC].GetTranslation();
+	const FVector A = InputGlobalPose[BoneIndexA].GetTranslation();
+	//FVector B = InputGlobalPose[BoneIndexB].GetTranslation(); TODO use for pole vector 
+	const FVector C = InputGlobalPose[BoneIndexC].GetTranslation();
 
     // get the normalized direction / length of the IK limb (how extended it is as percentage of original length)
     const FVector AC = C - A;
@@ -397,13 +455,13 @@ void FChainDecoderIK::DecodePose(
 {
 	// calculate end effector position
 	const FVector Start = OutGlobalPose[BoneIndexA].GetTranslation();
-	OutResults.EndEffectorPosition = Start + EncodedChain.EndDirectionNormalized * Length; // * LimbScale (TBD)
+	OutResults.EndEffectorPosition = Start + EncodedChain.EndDirectionNormalized * Length; // * LimbScale (todo)
 	// optionally factor in ground height
-	const bool IsGrounded = false;
-	if (IsGrounded) // TBD expose this
-	{
-		OutResults.EndEffectorPosition.Z = (EncodedChain.HeightFromGroundNormalized * Length) + EndPositionOrig.Z;
-	}
+	//const bool IsGrounded = false;
+	//if (IsGrounded) // todo expose this
+	//{
+	//	OutResults.EndEffectorPosition.Z = (EncodedChain.HeightFromGroundNormalized * Length) + EndPositionOrig.Z;
+	//}
 
 	// calculate end effector rotation
 	const FQuat Delta = EndRotationOrig * EncodedChain.EndRotationOrig.Inverse();
@@ -536,13 +594,22 @@ bool FRetargetChainPairIK::Initialize(
 	const FRetargetSkeleton& SourceSkeleton,
 	const FTargetSkeleton& TargetSkeleton)
 {
+	// validate if this chain even uses an IK Goal
+	const bool bUsingIKGoal = TargetBoneChain.IKGoalName != NAME_None && TargetBoneChain.IKGoalName != "- None -";
+	if (!bUsingIKGoal)
+	{
+		return false;
+	}
+
+	// store target IK goal
+	IKGoalName = TargetBoneChain.IKGoalName;
+
+	// initialize bone chains
 	const bool bChainInitialized = FRetargetChainPair::Initialize(SourceBoneChain, TargetBoneChain, SourceSkeleton, TargetSkeleton);
 	if (!bChainInitialized)
 	{
 		return false;
 	}
-	
-	IKGoalName = TargetBoneChain.IKGoalName;
 
 	// initialize SOURCE IK chain encoder with retarget pose
 	const bool bIKEncoderInitialized = IKEncoder.Initialize(SourceBoneIndices, SourceSkeleton.RetargetGlobalPose);
@@ -565,6 +632,20 @@ bool FRetargetChainPairIK::Initialize(
 	return true;
 }
 
+void FIKRetargetPose::AddRotationDeltaToBone(FName BoneName, FQuat RotationDelta)
+{
+	FQuat* RotOffset = BoneRotationOffsets.Find(BoneName);
+	if (RotOffset == nullptr)
+	{
+		// first time this bone has been modified in this pose
+		BoneRotationOffsets.Emplace(BoneName, RotationDelta);
+		return;
+	}
+
+	// accumulate delta rotation
+	*RotOffset = RotationDelta * (*RotOffset);
+}
+
 bool FRootEncoder::Initialize(const FName RootBoneName, const FRetargetSkeleton& SourceSkeleton)
 {
 	// validate root bone exists
@@ -655,10 +736,9 @@ void UIKRetargeter::Initialize(USkeletalMesh* SourceSkeletalMesh, USkeletalMesh*
 	check(SourceSkeletalMesh && TargetSkeletalMesh);
 	
 	// initialize skeleton data for source and target
-	const TArray<FTransform>& SourceRetargetPoseLocal = SourceSkeletalMesh->GetRefSkeleton().GetRefBonePose();
-	const TArray<FTransform>& TargetRetargetPoseLocal = TargetSkeletalMesh->GetRefSkeleton().GetRefBonePose();
-	SourceSkeleton.InitializeHierarchy(SourceSkeletalMesh, SourceRetargetPoseLocal);
-	TargetSkeleton.Initialize(TargetSkeletalMesh, TargetRetargetPoseLocal);
+	const TArray<FRetargetChainMap> DummySourceChainMapping;
+	SourceSkeleton.Initialize(SourceSkeletalMesh, nullptr);
+	TargetSkeleton.Initialize(TargetSkeletalMesh, &RetargetPoses[CurrentRetargetPose]);
 
 	// initialize roots
 	const bool bRootsInitialized = InitializeRoots();
@@ -708,7 +788,7 @@ bool UIKRetargeter::InitializeRoots()
 	}
 
 	// initialize root decoder
-	const FName TargetRootBoneName = SourceIKRigAsset->RetargetDefinition.RootBone;
+	const FName TargetRootBoneName = TargetIKRigAsset->RetargetDefinition.RootBone;
 	const bool bRootDecoderInit = RootDecoder.Initialize(TargetRootBoneName, TargetSkeleton);
 	if (!bRootDecoderInit)
 	{
@@ -727,21 +807,30 @@ bool UIKRetargeter::InitializeBoneChainPairs()
 	FRetargetDefinition& SourceRetargetDef = SourceIKRigAsset->RetargetDefinition;
 	FRetargetDefinition& TargetRetargetDef = TargetIKRigAsset->RetargetDefinition;
 
-	for (const TPair<FName, FName>& NamePair : ChainMapping)
+	for (const FRetargetChainMap& ChainSettings : ChainMapping)
 	{
+		FName SourceChain = ChainSettings.SourceChain;
+		FName TargetChain = ChainSettings.TargetChain;
+
+		// user opted to not map this to anything, we don't need to spam a warning about it
+		if (SourceChain == NAME_None)
+		{
+			continue; 
+		}
+		
 		// get source bone chain
-		FBoneChain* SourceBoneChain = SourceRetargetDef.GetBoneChainByName(NamePair.Key);
+		FBoneChain* SourceBoneChain = SourceRetargetDef.GetBoneChainByName(SourceChain);
 		if (!SourceBoneChain)
 		{
-			UE_LOG(LogTemp, Warning, TEXT("IK Retargeter missing source bone chain: %s"), *NamePair.Key.ToString());
+			UE_LOG(LogTemp, Warning, TEXT("IK Retargeter missing source bone chain: %s"), *SourceChain.ToString());
 			continue;
 		}
 
 		// get target bone chain
-		FBoneChain* TargetBoneChain = TargetRetargetDef.GetBoneChainByName(NamePair.Value);
+		FBoneChain* TargetBoneChain = TargetRetargetDef.GetBoneChainByName(TargetChain);
 		if (!TargetBoneChain)
 		{
-			UE_LOG(LogTemp, Warning, TEXT("IK Retargeter missing target bone chain: %s"), *NamePair.Value.ToString());
+			UE_LOG(LogTemp, Warning, TEXT("IK Retargeter missing target bone chain: %s"), *TargetChain.ToString());
 			continue;
 		}
 
@@ -754,28 +843,34 @@ bool UIKRetargeter::InitializeBoneChainPairs()
 		}
 		
 		// load IK chain
-		if (TargetBoneChain->bUseIK)
+		FRetargetChainPairIK ChainPairIK;
+		const bool bChainPairIKValid = ChainPairIK.Initialize(*SourceBoneChain, *TargetBoneChain, SourceSkeleton, TargetSkeleton);
+		if (bChainPairIKValid)
 		{
-			FRetargetChainPairIK ChainPairIK;
-			const bool bChainPairIKValid = ChainPairIK.Initialize(*SourceBoneChain, *TargetBoneChain, SourceSkeleton, TargetSkeleton);
-			if (bChainPairIKValid)
-			{
-				ChainPairsIK.Add(ChainPairIK);
-			}
+			ChainPairsIK.Add(ChainPairIK);
 		}
 	}
 
-	// record which bones in the target skeleton are controlled by FK
+	// record which bones in the target skeleton are being retargeted
 	for (const FRetargetChainPairFK& FKChainPair : ChainPairsFK)
 	{
 		for (const int32 BoneIndex : FKChainPair.TargetBoneIndices)
 		{
-			TargetSkeleton.SetBoneControlledByFKChain(BoneIndex, true);
+			TargetSkeleton.SetBoneIsRetargeted(BoneIndex, true);
 		}
 	}
 
+	// record intermediate bones (non-retargeted bones located BETWEEN FK chains on the target skeleton)
+	for (FRetargetChainPairFK& FKChainPair : ChainPairsFK)
+	{
+		FKChainPair.FKDecoder.InitializeIntermediateParentIndices(
+			RootDecoder.BoneIndex,
+			FKChainPair.TargetBoneIndices[0],
+			TargetSkeleton);
+	}
+
 	// root is updated before IK as well
-	TargetSkeleton.SetBoneControlledByFKChain(RootDecoder.BoneIndex, true);
+	TargetSkeleton.SetBoneIsRetargeted(RootDecoder.BoneIndex, true);
 
 	// return true if at least 1 pair of bone chains were initialized
 	return !(ChainPairsIK.IsEmpty() && ChainPairsFK.IsEmpty());
@@ -811,30 +906,43 @@ bool UIKRetargeter::InitializeIKRig(UObject* Outer, const FReferenceSkeleton& In
 	return true;
 }
 
-void UIKRetargeter::RunRetargeter(const TArray<FTransform>& InSourceGlobalPose, bool bEnableIK)
+void UIKRetargeter::RunRetargeter(const TArray<FTransform>& InSourceGlobalPose)
 {
 	check(bIsLoadedAndValid);
 
-	// ROOT retargeting
-	RunRootRetarget(InSourceGlobalPose, TargetSkeleton.OutputGlobalPose);
-
-	// update global transforms below root
-	TargetSkeleton.UpdateGlobalTransformsBelowBone(RootDecoder.BoneIndex);
+	// start from retarget pose
+	TargetSkeleton.OutputGlobalPose = TargetSkeleton.RetargetGlobalPose;
+	
+	// in edit mode we just want to see the edited reference pose, not actually run the retargeting
+	// as long as the retargeter is reinitialized after every modification to the limb rotation offsets,
+	// then the TargetSkeleton.RetargetGlobalPose will contain the updated retarget pose.
+	if (bEditReferencePoseMode)
+	{
+		return; 
+	}
 
+	// ROOT retargeting
+	if (bRetargetRoot)
+	{
+		RunRootRetarget(InSourceGlobalPose, TargetSkeleton.OutputGlobalPose);
+		// update global transforms below root
+		TargetSkeleton.UpdateGlobalTransformsBelowBone(RootDecoder.BoneIndex, TargetSkeleton.RetargetLocalPose, TargetSkeleton.OutputGlobalPose);
+	}	
 	// FK CHAIN retargeting
-	RunFKRetarget(InSourceGlobalPose,TargetSkeleton.OutputGlobalPose, true);
-
-	// update all the bones that are not controlled by FK chains or root
-	TargetSkeleton.UpdateGlobalTransformsAllNonFKBones();
+	if (bRetargetFK)
+	{
+		RunFKRetarget(InSourceGlobalPose, TargetSkeleton.OutputGlobalPose);
+		// update all the bones that are not controlled by FK chains or root
+		TargetSkeleton.UpdateGlobalTransformsAllNonRetargetedBones(TargetSkeleton.OutputGlobalPose);
+	}
 	
 	// IK CHAIN retargeting
-	if (bEnableIK)
+	if (bRetargetIK)
 	{
 		RunIKRetarget(InSourceGlobalPose, TargetSkeleton.OutputGlobalPose);
 	}	
 }
 
-
 void UIKRetargeter::RunRootRetarget(
 	const TArray<FTransform>& InGlobalTransforms,
     TArray<FTransform>& OutGlobalTransforms)
@@ -846,14 +954,13 @@ void UIKRetargeter::RunRootRetarget(
 
 void UIKRetargeter::RunFKRetarget(
 	const TArray<FTransform>& InGlobalTransforms,
-    TArray<FTransform>& OutGlobalTransforms,
-    const bool bPreIK)
+    TArray<FTransform>& OutGlobalTransforms)
 {
 	// spin through chains and encode/decode them all using the input pose
     for (FRetargetChainPairFK& ChainPair : ChainPairsFK)
     {
     	ChainPair.FKEncoder.EncodePose(ChainPair.SourceBoneIndices, InGlobalTransforms);
-    	ChainPair.FKDecoder.DecodePose(ChainPair.TargetBoneIndices, ChainPair.FKEncoder,OutGlobalTransforms);
+    	ChainPair.FKDecoder.DecodePose(ChainPair.TargetBoneIndices, ChainPair.FKEncoder, TargetSkeleton,OutGlobalTransforms);
     }
 }
 
@@ -891,3 +998,361 @@ void UIKRetargeter::RunIKRetarget(
 	// copy results of solve
 	IKRigProcessor->CopyOutputGlobalPoseToArray(OutGlobalPose);
 }
+
+#if WITH_EDITOR
+
+FName UIKRetargeter::GetSourceRootBone()
+{
+	return SourceIKRigAsset ? SourceIKRigAsset->RetargetDefinition.RootBone : FName("None");
+}
+
+FName UIKRetargeter::GetTargetRootBone()
+{
+	return TargetIKRigAsset ? TargetIKRigAsset->RetargetDefinition.RootBone : FName("None");
+}
+
+void UIKRetargeter::GetTargetChainNames(TArray<FName>& OutNames) const
+{
+	if (TargetIKRigAsset)
+	{
+		for (const FBoneChain& Chain : TargetIKRigAsset->RetargetDefinition.BoneChains)
+		{
+			OutNames.Add(Chain.ChainName);
+		}
+	}
+}
+
+void UIKRetargeter::GetSourceChainNames(TArray<FName>& OutNames) const
+{
+	if (SourceIKRigAsset)
+	{
+		for (const FBoneChain& Chain : SourceIKRigAsset->RetargetDefinition.BoneChains)
+		{
+			OutNames.Add(Chain.ChainName);
+		}
+	}
+}
+
+void UIKRetargeter::CleanChainMapping()
+{
+	if (!TargetIKRigAsset)
+	{
+		// don't clean chain mappings, in case user is replacing with IK Rig asset that has some valid mappings
+		return;
+	}
+	
+	TArray<FName> TargetChainNames;
+	GetTargetChainNames(TargetChainNames);
+
+	// remove all target chains that are no longer in the target IK rig asset
+	TArray<FName> TargetChainsToRemove;
+	for (FRetargetChainMap& ChainMap : ChainMapping)
+	{
+		if (!TargetChainNames.Contains(ChainMap.TargetChain))
+		{
+			TargetChainsToRemove.Add(ChainMap.TargetChain);
+		}
+	}
+	for (FName TargetChainToRemove : TargetChainsToRemove)
+	{
+		ChainMapping.RemoveAll([&TargetChainToRemove](FRetargetChainMap& Element)
+		{
+			return Element.TargetChain == TargetChainToRemove;
+		});
+	}
+
+	// add a mapping for each chain that is in the target IK rig (if it doesn't have one already)
+	for (FName TargetChainName : TargetChainNames)
+	{
+		const bool HasChain = ChainMapping.ContainsByPredicate([&TargetChainName](FRetargetChainMap& Element)
+		{
+			return Element.TargetChain == TargetChainName;
+		});
+		
+		if (!HasChain)
+		{
+			ChainMapping.Add(FRetargetChainMap(TargetChainName));
+		}
+	}
+
+	TArray<FName> SourceChainNames;
+	GetSourceChainNames(SourceChainNames);
+	
+	// reset any sources that are no longer present to "None"
+	for (FRetargetChainMap& ChainMap : ChainMapping)
+	{
+		if (!SourceChainNames.Contains(ChainMap.SourceChain))
+		{
+			ChainMap.SourceChain = NAME_None;
+		}
+	}
+
+	// enforce the same chain order as the target IK rig
+	ChainMapping.Sort([this](const FRetargetChainMap& A, const FRetargetChainMap& B)
+	{
+		const TArray<FBoneChain>& BoneChains = TargetIKRigAsset->RetargetDefinition.BoneChains;
+		
+		const int32 IndexA = BoneChains.IndexOfByPredicate([&A](const FBoneChain& Chain)
+		{
+			return A.TargetChain == Chain.ChainName;
+		});
+
+		const int32 IndexB = BoneChains.IndexOfByPredicate([&B](const FBoneChain& Chain)
+		{
+			return B.TargetChain == Chain.ChainName;
+		});
+ 
+		return IndexA < IndexB;
+	});
+
+	// force update with latest mapping
+	Modify();
+}
+
+void UIKRetargeter::CleanPoseList()
+{
+	// enforce the existence of a default pose
+	const bool HasDefaultPose = RetargetPoses.Contains(DefaultPoseName);
+	if (!HasDefaultPose)
+	{
+		RetargetPoses.Emplace(DefaultPoseName);
+	}
+	
+	// use default pose unless set to something else
+	if (CurrentRetargetPose == NAME_None)
+	{
+		CurrentRetargetPose = DefaultPoseName;
+	}
+
+	// remove all bone offsets that are no longer part of the target skeleton
+	if (TargetIKRigAsset)
+	{
+		const TArray<FName> AllowedBoneNames = TargetIKRigAsset->Skeleton.BoneNames;
+		for (TTuple<FName, FIKRetargetPose>& Pose : RetargetPoses)
+		{
+			// find bone offsets no longer in target skeleton
+			TArray<FName> BonesToRemove;
+			for (TTuple<FName, FQuat>& BoneOffset : Pose.Value.BoneRotationOffsets)
+			{
+				if (!AllowedBoneNames.Contains(BoneOffset.Key))
+				{
+					BonesToRemove.Add(BoneOffset.Key);
+				}
+			}
+			
+			// remove bone offsets
+			for (const FName& BoneToRemove : BonesToRemove)
+			{
+				Pose.Value.BoneRotationOffsets.Remove(BoneToRemove);
+			}
+		}
+	}
+
+	Modify();
+}
+
+void UIKRetargeter::AutoMapChains()
+{
+	TArray<FName> SourceChainNames;
+	GetSourceChainNames(SourceChainNames);
+	
+	// auto-map any chains that have no value using a fuzzy string search
+	for (FRetargetChainMap& ChainMap : ChainMapping)
+	{
+		if (ChainMap.SourceChain != NAME_None)
+		{
+			continue; // already set by user
+		}
+
+		// find "best match" automatically as a convenience for the user
+		FString TargetNameLowerCase = ChainMap.TargetChain.ToString().ToLower();
+		float HighestScore = 0.2f;
+		int32 HighestScoreIndex = -1;
+		for (int32 ChainIndex=0; ChainIndex<SourceChainNames.Num(); ++ChainIndex)
+		{
+			FString SourceNameLowerCase = SourceChainNames[ChainIndex].ToString().ToLower();
+			float WorstCase = TargetNameLowerCase.Len() + SourceNameLowerCase.Len();
+			WorstCase = WorstCase < 1.0f ? 1.0f : WorstCase;
+			const float Score = 1.0f - (Algo::LevenshteinDistance(TargetNameLowerCase, SourceNameLowerCase) / WorstCase);
+			if (Score > HighestScore)
+			{
+				HighestScore = Score;
+				HighestScoreIndex = ChainIndex;
+			}
+		}
+
+		// apply source if any decent matches were found
+		if (SourceChainNames.IsValidIndex(HighestScoreIndex))
+		{
+			ChainMap.SourceChain = SourceChainNames[HighestScoreIndex];
+		}
+	}
+
+	// force update with latest mapping
+	Modify();
+}
+
+USkeleton* UIKRetargeter::GetSourceSkeletonAsset() const
+{
+	if (!SourceIKRigAsset)
+	{
+		return nullptr;
+	}
+
+	if (!SourceIKRigAsset->PreviewSkeletalMesh)
+	{
+		return nullptr;
+	}
+
+	return SourceIKRigAsset->PreviewSkeletalMesh->GetSkeleton();
+}
+
+void UIKRetargeter::AddRetargetPose(FName NewPoseName)
+{
+	if (RetargetPoses.Contains(NewPoseName))
+	{
+		return;
+	}
+
+	RetargetPoses.Add(NewPoseName);
+	CurrentRetargetPose = NewPoseName;
+
+	Modify();
+}
+
+void UIKRetargeter::RemoveRetargetPose(FName PoseToRemove)
+{
+	if (PoseToRemove == DefaultPoseName)
+	{
+		return; // cannot remove default pose
+	}
+
+	if (!RetargetPoses.Contains(PoseToRemove))
+	{
+		return; // cannot remove pose that doesn't exist
+	}
+
+	RetargetPoses.Remove(PoseToRemove);
+
+	// did we remove the currently used pose?
+	if (CurrentRetargetPose == PoseToRemove)
+	{
+		CurrentRetargetPose = UIKRetargeter::DefaultPoseName;
+	}
+
+	Modify();
+}
+
+void UIKRetargeter::ResetRetargetPose(FName PoseToReset)
+{
+	if (!RetargetPoses.Contains(PoseToReset))
+	{
+		return; // cannot reset pose that doesn't exist
+	}
+
+	RetargetPoses[PoseToReset].BoneRotationOffsets.Reset();
+	RetargetPoses[PoseToReset].RootTranslationOffset = FVector::ZeroVector;
+	RetargetPoses[PoseToReset].RootRotationOffset = FQuat::Identity;
+	
+	Modify();
+}
+
+void UIKRetargeter::SetCurrentRetargetPose(FName CurrentPose)
+{
+	check(RetargetPoses.Contains(CurrentPose));
+	CurrentRetargetPose = CurrentPose;
+	Modify();
+}
+
+void UIKRetargeter::AddRotationOffsetToRetargetPoseBone(FName BoneName, FQuat RotationOffset)
+{
+	RetargetPoses[CurrentRetargetPose].AddRotationDeltaToBone(BoneName, RotationOffset);
+	Modify();
+}
+
+bool UIKRetargeter::Modify(bool bAlwaysMarkDirty)
+{
+	const bool bSavedToTransactionBuffer = Super::Modify(bAlwaysMarkDirty);
+	AssetVersion++; // inform any runtime/editor systems they should copy-up modifications	
+	return bSavedToTransactionBuffer;
+}
+
+void UIKRetargeter::PostEditChangeProperty(FPropertyChangedEvent& PropertyChangedEvent)
+{
+	if (PropertyChangedEvent.GetPropertyName() == "TargetIKRigAsset"
+		|| PropertyChangedEvent.GetPropertyName() == "SourceIKRigAsset")
+	{
+		AutoMapChains();
+		CleanChainMapping();
+	}
+}
+
+void UIKRetargeter::PostLoad()
+{
+	Super::PostLoad();
+	CleanChainMapping();
+	CleanPoseList();
+}
+#endif
+
+FTransform UIKRetargeter::GetTargetBoneRetargetPoseGlobalTransform(const FName& TargetBoneName) const
+{
+	check(bIsLoadedAndValid);
+
+	// makes sure is actually in skeleton
+	const int32 BoneIndex = TargetSkeleton.FindBoneIndexByName(TargetBoneName);
+	if (BoneIndex == INDEX_NONE)
+	{
+		return FTransform::Identity;
+	}
+
+	// get the current retarget pose
+	return TargetSkeleton.RetargetGlobalPose[BoneIndex];
+}
+
+void UIKRetargeter::GetTargetBoneStartAndEnd(const FName& TargetBoneName, FVector& OutStart, FVector& OutEnd) const
+{
+	check(bIsLoadedAndValid);
+
+	// find the bone
+	const int32 BoneIndex = TargetSkeleton.FindBoneIndexByName(TargetBoneName);
+	if (BoneIndex == INDEX_NONE)
+	{
+		return;
+	}
+
+	// find the end bone of the bone chain
+	// todo, return list of lines to each child
+	int32 FirstChildIndex = INDEX_NONE;
+	for (int32 I=0; I<TargetSkeleton.ParentIndices.Num(); ++I)
+	{
+		const int32 ParentIndex = TargetSkeleton.ParentIndices[I];
+		if (ParentIndex == BoneIndex)
+		{
+			FirstChildIndex = I;
+			break;
+		}
+	}
+
+	if (FirstChildIndex == INDEX_NONE)
+	{
+		return;
+	}
+
+	// get the origin of the bone chain
+	OutStart = TargetSkeleton.RetargetGlobalPose[BoneIndex].GetTranslation();
+	OutEnd = TargetSkeleton.RetargetGlobalPose[FirstChildIndex].GetTranslation();
+}
+
+FRetargetChainMap* UIKRetargeter::GetChainMap(const FName& TargetChainName)
+{
+	for (FRetargetChainMap& ChainMap : ChainMapping)
+	{
+		if (ChainMap.TargetChain == TargetChainName)
+		{
+			return &ChainMap;
+		}
+	}
+
+	return nullptr;
+}
diff --git a/Engine/Plugins/Animation/IKRig/Source/IKRig/Private/Solvers/IKRig_PBIKSolver.cpp b/Engine/Plugins/Animation/IKRig/Source/IKRig/Private/Solvers/IKRig_PBIKSolver.cpp
--- a/Engine/Plugins/Animation/IKRig/Source/IKRig/Private/Solvers/IKRig_PBIKSolver.cpp
+++ b/Engine/Plugins/Animation/IKRig/Source/IKRig/Private/Solvers/IKRig_PBIKSolver.cpp
@@ -34,7 +34,7 @@ void UIKRigPBIKSolver::Initialize(const FIKRigSkeleton& InSkeleton)
 		const FName& Name = InSkeleton.BoneNames[BoneIndex];
 
 		// get the parent bone solver index
-		const int32 ParentIndex = InSkeleton.GetParentIndex(BoneIndex);
+		const int32 ParentIndex = InSkeleton.GetParentIndexThatIsNotExcluded(BoneIndex);
 		const FTransform OrigTransform = InSkeleton.RefPoseGlobal[BoneIndex];
 		const FVector InOrigPosition = OrigTransform.GetLocation();
 		const FQuat InOrigRotation = OrigTransform.GetRotation();
@@ -98,7 +98,7 @@ void UIKRigPBIKSolver::Solve(FIKRigSkeleton& IKRigSkeleton, const FIKRigGoalCont
 			OffsetAlpha,
 			Effector->StrengthAlpha,
 			Effector->PullChainAlpha,
-			Effector->bPinRotation);
+			Effector->PinRotation);
 	}
 
 	// update settings
diff --git a/Engine/Plugins/Animation/IKRig/Source/IKRig/Public/AnimNodes/AnimNode_IKRetargeter.h b/Engine/Plugins/Animation/IKRig/Source/IKRig/Public/AnimNodes/AnimNode_IKRetargeter.h
--- a/Engine/Plugins/Animation/IKRig/Source/IKRig/Public/AnimNodes/AnimNode_IKRetargeter.h
+++ b/Engine/Plugins/Animation/IKRig/Source/IKRig/Public/AnimNodes/AnimNode_IKRetargeter.h
@@ -20,18 +20,10 @@ struct IKRIG_API FAnimNode_IKRetargeter : public FAnimNode_Base
 	/* If SourceMeshComponent is not valid, and if this is true, it will look for attached parent as a source */
 	UPROPERTY(EditAnywhere, BlueprintReadWrite, Category=Settings, meta = (NeverAsPin))
 	bool bUseAttachedParent = true;
-
-	/** Map of chain names to per-chain retarget settings (can be modified at runtime).*/
-	UPROPERTY(BlueprintReadWrite, transient, Category=Settings, meta=(PinShownByDefault))
-	TMap<FName, FIKRetargetChainSettings> ChainSettings;
 	
 	/** Retarget asset to use. Must define a Source and Target IK Rig compatible with the SourceMeshComponent and current anim instance.*/
 	UPROPERTY(EditAnywhere, Category = Settings)
 	UIKRetargeter* IKRetargeterAsset = nullptr;
-
-	/** When false, IK is not applied as part of retargeter. Useful for debugging limb issues suspected to be caused by IK.*/
-	UPROPERTY(EditAnywhere, Category = Settings)
-	bool bEnableIK = true;
 	
 	// FAnimNode_Base interface
 	virtual void Initialize_AnyThread(const FAnimationInitializeContext& Context) override;
@@ -43,6 +35,8 @@ struct IKRIG_API FAnimNode_IKRetargeter : public FAnimNode_Base
 	virtual void PreUpdate(const UAnimInstance* InAnimInstance) override;
 	// End of FAnimNode_Base interface
 
+	UIKRetargeter* GetCurrentlyUsedRetargeter() const;
+
 private:
 	
 	void EnsureInitialized(const UAnimInstance* InAnimInstance);
@@ -53,6 +47,11 @@ private:
 	// indicates that all prerequisites are met and node is ready to operate
 	UPROPERTY(Transient)
 	bool bIsInitialized = false;
+	/** which version of the IK Retargeter asset was this instance last initialized with?
+	* this allows the IK Retargeter asset to undergo modifications at runtime via the editor and trigger
+	* re-initialization by incrementing this transient asset version.
+	 */
+	int32 LastVersionTried = -2;
 	
 	// source mesh references, cached during init so that we can compare and see if it has changed
 	TWeakObjectPtr<USkeletalMeshComponent>	CurrentlyUsedSourceMeshComponent;
diff --git a/Engine/Plugins/Animation/IKRig/Source/IKRig/Public/IKRigDefinition.h b/Engine/Plugins/Animation/IKRig/Source/IKRig/Public/IKRigDefinition.h
--- a/Engine/Plugins/Animation/IKRig/Source/IKRig/Public/IKRigDefinition.h
+++ b/Engine/Plugins/Animation/IKRig/Source/IKRig/Public/IKRigDefinition.h
@@ -57,27 +57,27 @@ public:
 USTRUCT(Blueprintable)
 struct IKRIG_API FBoneChain
 {
-	FBoneChain()
-		: bUseIK(false)
-	{
-	}
-
 	GENERATED_BODY()
 
-	UPROPERTY(EditAnywhere, Category = BoneChain)
-	FName ChainName;
+	FBoneChain() = default;
+	
+	FBoneChain(FName InName, FName InStartBone, FName InEndBone)
+	: ChainName(InName),
+	StartBone(InStartBone),
+	EndBone(InEndBone),
+	IKGoalName(NAME_None){}
 
 	UPROPERTY(EditAnywhere, Category = BoneChain)
-	FName StartBone;
+	FName ChainName = NAME_None;
 
 	UPROPERTY(EditAnywhere, Category = BoneChain)
-	FName EndBone;
+	FName StartBone = NAME_None;
 
-	UPROPERTY(EditAnywhere, Category = IK)
-	bool bUseIK;
+	UPROPERTY(EditAnywhere, Category = BoneChain)
+	FName EndBone = NAME_None;
 	
 	UPROPERTY(EditAnywhere, Category = IK)
-	FName IKGoalName;
+	FName IKGoalName = NAME_None;
 
 	/*
 	UPROPERTY(EditAnywhere, Category = IK)
@@ -157,6 +157,8 @@ private:
 
 	void ResetGoalTransforms() const;
 	FTransform GetGoalInitialTransform(UIKRigEffectorGoal* Goal) const;
+
+	void SortRetargetChains();
 	
 	friend class UIKRigController;
 };
diff --git a/Engine/Plugins/Animation/IKRig/Source/IKRig/Public/IKRigProcessor.h b/Engine/Plugins/Animation/IKRig/Source/IKRig/Public/IKRigProcessor.h
--- a/Engine/Plugins/Animation/IKRig/Source/IKRig/Public/IKRigProcessor.h
+++ b/Engine/Plugins/Animation/IKRig/Source/IKRig/Public/IKRigProcessor.h
@@ -106,5 +106,5 @@ private:
 	/** which version of the IKRig asset was this instance last initialized with?
 	 * this allows the IKRig asset to undergo modifications at runtime via the editor*/
 	int32 InitializedWithIKRigAssetVersion = -1;
-	int32 LastVersionTried = -1;
+	int32 LastVersionTried = -2;
 };
\ No newline at end of file
diff --git a/Engine/Plugins/Animation/IKRig/Source/IKRig/Public/IKRigSkeleton.h b/Engine/Plugins/Animation/IKRig/Source/IKRig/Public/IKRigSkeleton.h
--- a/Engine/Plugins/Animation/IKRig/Source/IKRig/Public/IKRigSkeleton.h
+++ b/Engine/Plugins/Animation/IKRig/Source/IKRig/Public/IKRigSkeleton.h
@@ -7,34 +7,61 @@
 
 #include "IKRigSkeleton.generated.h"
 
+class UIKRigDefinition;
+
+USTRUCT()
+struct IKRIG_API FIKRigSkeletonChain
+{
+	GENERATED_BODY()
+	
+	FIKRigSkeletonChain(){}
+	FIKRigSkeletonChain(const FName& StartBone, const FName& EndBone): StartBone(StartBone), EndBone(EndBone){}
+
+	FName StartBone;
+	FName EndBone;
+};
+
 USTRUCT()
 struct IKRIG_API FIKRigSkeleton
 {
 	GENERATED_BODY()
 
+	/** Names of bones. Used to match hierarchy with runtime skeleton. */
 	UPROPERTY(VisibleAnywhere, Category = Skeleton)
 	TArray<FName> BoneNames;
-	
+
+	/** Same length as BoneNames, stores index of parent for each bone. */
 	UPROPERTY(VisibleAnywhere, Category = Skeleton)
 	TArray<int32> ParentIndices;
-	
+
+	/** Sparse array of bones that are to be excluded from any solvers (parented around, treated as FK children). */
+	UPROPERTY(VisibleAnywhere, Category = Skeleton)
+	TArray<FName> ExcludedBones;
+
+	/** The current GLOBAL space pose of each bone. */
 	UPROPERTY(VisibleAnywhere, Category = Skeleton)
 	TArray<FTransform> CurrentPoseGlobal;
-	
+
+	/** The current LOCAL space pose of each bone. */
 	UPROPERTY(VisibleAnywhere, Category = Skeleton)
 	TArray<FTransform> CurrentPoseLocal;
-	
+
+	/** The initial/reference GLOBAL space pose of each bone. */
 	UPROPERTY(VisibleAnywhere, Category = Skeleton)
 	TArray<FTransform> RefPoseGlobal;
 
-	void Initialize(const FReferenceSkeleton& RefSkeleton);
+	void Initialize(const FReferenceSkeleton& RefSkeleton, const TArray<FName>& ExcludedBones);
 
 	void Reset();
 	
 	int32 GetBoneIndexFromName(const FName InName) const;
+
+	FName GetBoneNameFromIndex(const int32 BoneIndex) const;
 	
 	int32 GetParentIndex(const int32 BoneIndex) const;
 
+	int32 GetParentIndexThatIsNotExcluded(const int32 BoneIndex) const;
+
 	bool CopyPosesFromRefSkeleton(const FReferenceSkeleton& RefSkeleton);
 
 	static void ConvertLocalPoseToGlobal(
@@ -54,6 +81,9 @@ struct IKRIG_API FIKRigSkeleton
 
 	bool IsBoneInDirectLineage(const FName& Child, const FName& PotentialParent) const;
 
+	bool IsBoneExcluded(const int32 BoneIndex) const;
+
 	static void NormalizeRotations(TArray<FTransform>& Transforms);
-};
 
+	void GetChainsInList(const TArray<int32>& SelectedBones, TArray<FIKRigSkeletonChain>& OutChains) const;
+};
diff --git a/Engine/Plugins/Animation/IKRig/Source/IKRig/Public/Retargeter/IKRetargeter.h b/Engine/Plugins/Animation/IKRig/Source/IKRig/Public/Retargeter/IKRetargeter.h
--- a/Engine/Plugins/Animation/IKRig/Source/IKRig/Public/Retargeter/IKRetargeter.h
+++ b/Engine/Plugins/Animation/IKRig/Source/IKRig/Public/Retargeter/IKRetargeter.h
@@ -9,23 +9,24 @@
 
 #include "IKRetargeter.generated.h"
 
+struct FIKRetargetPose;
+
 USTRUCT(Blueprintable)
-struct IKRIG_API FIKRetargetChainSettings
+struct IKRIG_API FRetargetChainMap
 {
 	GENERATED_BODY()
 
-	/** Range -1 to 1. Default 0. Rotates chain away or towards the midline (positive values will raise arms upwards). */
-	UPROPERTY(EditAnywhere, BlueprintReadWrite, Category = Offsets, meta = (ClampMin = "-1.0", ClampMax = "1.0", UIMin = "-1.0", UIMax = "1.0"))
-	float AbductionOffset = 0.0f;
-
-	/** Range -1 to 1. Default 0. Rotates chain forward (+) and backward (-) relative to facing direction. */
-	UPROPERTY(EditAnywhere, BlueprintReadWrite, Category = Offsets, meta = (ClampMin = "-1.0", ClampMax = "1.0", UIMin = "-1.0", UIMax = "1.0"))
-	float SwingOffset = 0.0f;
+	FRetargetChainMap(){}
+	
+	FRetargetChainMap(const FName& TargetChain) : TargetChain(TargetChain){}
+	
+	/** The name of the chain on the Source IK Rig asset to copy animation from. */
+	UPROPERTY(EditAnywhere, BlueprintReadWrite, Category = Offsets)
+	FName SourceChain = NAME_None;
 
-	/** Range -1 to 1. Default 0. Rotates the chain along an axis from the start to the end of the chain.
-	*  At -1 the chain is rotated laterally 90 degrees. At +1 the chain is rotated medially 90 degrees. */
-	UPROPERTY(EditAnywhere, BlueprintReadWrite, Category = Offsets, meta = (ClampMin = "-1.0", ClampMax = "1.0", UIMin = "-1.0", UIMax = "1.0"))
-	float TwistOffset = 0.0f;
+	/** The name of the chain on the Target IK Rig asset to copy animation from. */
+	UPROPERTY(EditAnywhere, BlueprintReadWrite, Category = Offsets)
+	FName TargetChain = NAME_None;
 
 	/** Range -1 to 2. Default 0. Brings IK effector closer (-) or further (+) from origin of chain.
 	*  At -1 the end is on top of the origin. At +2 the end is fully extended twice the length of the chain. */
@@ -34,54 +35,81 @@ struct IKRIG_API FIKRetargetChainSettings
 
 	/** Range 0 to 1. Default 0. Allow the chain to stretch by translating to reach the IK goal locations.
 	*  At 0 the chain will not stretch at all. At 1 the chain will be allowed to stretch double it's full length to reach IK. */
-	UPROPERTY(EditAnywhere, BlueprintReadWrite, Category = Stretch, meta = (ClampMin = "0.0", ClampMax = "1.0", UIMin = "0.0", UIMax = "1.0"))
-	float StretchTolerance = 0.0f;
+	//UPROPERTY(EditAnywhere, BlueprintReadWrite, Category = Stretch, meta = (ClampMin = "0.0", ClampMax = "1.0", UIMin = "0.0", UIMax = "1.0"))
+	//float StretchTolerance = 0.0f;
 
 	/** Range 0 to 1. Default is 1.0. Blend IK effector at the end of this chain towards the original position
 	 *  on the source skeleton (0.0) or the position on the retargeted target skeleton (1.0). */
-	UPROPERTY(EditAnywhere, BlueprintReadWrite, Category = IkMode, meta = (ClampMin = "0.0", ClampMax = "1.0", UIMin = "0.0", UIMax = "1.0"))
-	float IkToSourceOrTarget = 1.0f;
+	//UPROPERTY(EditAnywhere, BlueprintReadWrite, Category = IkMode, meta = (ClampMin = "0.0", ClampMax = "1.0", UIMin = "0.0", UIMax = "1.0"))
+	//float IkToSourceOrTarget = 1.0f;
 	
 	/** When true, the source IK position is calculated relative to a source bone and applied relative to a target bone. */
-	UPROPERTY(EditAnywhere, BlueprintReadWrite, Category = IkMode)
-	bool bIkRelativeToSourceBone = false;
+	//UPROPERTY(EditAnywhere, BlueprintReadWrite, Category = IkMode)
+	//bool bIkRelativeToSourceBone = false;
 	
 	/** A bone in the SOURCE skeleton that the IK location is relative to. */
-	UPROPERTY(EditAnywhere, BlueprintReadWrite, Category = IkMode, meta = (EditCondition="bIkRelativeToSourceBone"))
-	FName IkRelativeSourceBone;
+	//UPROPERTY(EditAnywhere, BlueprintReadWrite, Category = IkMode, meta = (EditCondition="bIkRelativeToSourceBone"))
+	//FName IkRelativeSourceBone;
 	
 	/** A bone in the TARGET skeleton that the IK location is relative to.
 	 * This is usually the same bone as the source skeleton, but may have a different name in the target skeleton. */
-	UPROPERTY(EditAnywhere, BlueprintReadWrite, Category = IkMode, meta = (EditCondition="bIkRelativeToSourceBone"))
-	FName IkRelativeTargetBone;	
+	//UPROPERTY(EditAnywhere, BlueprintReadWrite, Category = IkMode, meta = (EditCondition="bIkRelativeToSourceBone"))
+	//FName IkRelativeTargetBone;	
 };
 
 struct FRetargetSkeleton
 {
 	TArray<FName> BoneNames;
+	
 	TArray<int32> ParentIndices;
+	
+	TArray<FTransform> RetargetLocalPose;
+	
 	TArray<FTransform> RetargetGlobalPose;
+	
 	FString SkeletalMeshName;
 
-	void InitializeHierarchy(const USkeletalMesh* SkeletalMesh, const TArray<FTransform>& InLocalRetargetPose);
+	void Initialize(
+		USkeletalMesh* SkeletalMesh, 
+		FIKRetargetPose* RetargetPose);
+
 	int32 FindBoneIndexByName(const FName InName) const;
+
 	int32 GetParentIndex(const int32 BoneIndex) const;
+
+	void UpdateGlobalTransformsBelowBone(
+		const int32 StartBoneIndex,
+		TArray<FTransform>& InOutLocalPose,
+		TArray<FTransform>& InOutGlobalPose) const;
+
+	void UpdateLocalTransformsBelowBone(
+		const int32 StartBoneIndex,
+		TArray<FTransform>& InOutLocalPose,
+		TArray<FTransform>& InOutGlobalPose) const;
+	
+	void UpdateGlobalTransformOfSingleBone(
+		const int32 BoneIndex,
+		const TArray<FTransform>& InLocalPose,
+		TArray<FTransform>& OutGlobalPose) const;
+	
+	void UpdateLocalTransformOfSingleBone(
+		const int32 BoneIndex,
+		TArray<FTransform>& OutLocalPose,
+		TArray<FTransform>& InGlobalPose) const;
 };
 
 struct FTargetSkeleton : public FRetargetSkeleton
 {
 	TArray<FTransform> OutputGlobalPose;
-	TArray<FTransform> RetargetLocalPose;
-	FString SkeletalMeshName;
-	TArray<bool> IsBoneFK;
+	TArray<bool> IsBoneRetargeted;
 
-	void Initialize(const USkeletalMesh* SkeletalMesh, const TArray<FTransform>& InTargetRetargetLocalPose);
-	void SetBoneControlledByFKChain(const int32 BoneIndex, const bool IsFK);
-	void UpdateGlobalTransformsBelowBone(const int32 StartBoneIndex);
-	void UpdateGlobalTransformsAllNonFKBones();
+	void Initialize(
+		USkeletalMesh* SkeletalMesh, 
+		FIKRetargetPose* RetargetPose);
 
-private:
-	void UpdateGlobalTransformOfSingleBone(const int32 BoneIndex);
+	void SetBoneIsRetargeted(const int32 BoneIndex, const bool IsRetargeted);
+
+	void UpdateGlobalTransformsAllNonRetargetedBones(TArray<FTransform>& InOutGlobalPose);
 };
 
 struct FEncodedRoot
@@ -122,7 +150,7 @@ struct FRootDecoder
         const float StrideScale) const;
 };
 
-struct FChainEncoderFK
+struct FChainFK
 {
 	TArray<FTransform> InitialBoneTransforms;
 
@@ -133,30 +161,56 @@ struct FChainEncoderFK
 	bool Initialize(
 		const TArray<int32>& BoneIndices,
 		const TArray<FTransform> &InitialGlobalPose);
+
+private:
 	
-	bool CalculateBoneParameters();
+	bool CalculateBoneParameters();	
+};
 
+struct FChainEncoderFK : public FChainFK
+{
 	void EncodePose(
 		const TArray<int32>& SourceBoneIndices,
 		const TArray<FTransform> &InputGlobalPose);
+};
 
+struct FChainDecoderFK : public FChainFK
+{
+	void InitializeIntermediateParentIndices(
+		const int32 RetargetRootBoneIndex,
+		const int32 ChainRootBoneIndex,
+		const FTargetSkeleton& TargetSkeleton);
+	
 	void DecodePose(
 		const TArray<int32>& TargetBoneIndices,
 		const FChainEncoderFK& SourceChain,
+		const FTargetSkeleton& TargetSkeleton,
 		TArray<FTransform> &InOutGlobalPose);
 
+private:
+
 	FTransform GetTransformAtParam(
 		const TArray<FTransform>& Transforms,
 		const TArray<float>& InParams,
 		const float Param) const;
+	
+	void UpdateIntermediateParents(
+		const FTargetSkeleton& TargetSkeleton,
+		TArray<FTransform> &InOutGlobalPose);
+
+	TArray<int32> IntermediateParentIndices;
 };
 
 struct FEncodedIKChain
 {
 	FVector EndDirectionNormalized;
+	
 	FQuat EndRotation;
+	
 	FQuat EndRotationOrig;
+	
 	float HeightFromGroundNormalized;
+	
 	FVector PoleVectorDirection;
 };
 
@@ -240,7 +294,7 @@ struct FRetargetChainPairFK : FRetargetChainPair
 {
 	FChainEncoderFK FKEncoder;
 
-	FChainEncoderFK FKDecoder;
+	FChainDecoderFK FKDecoder;
 	
 	virtual bool Initialize(
         const FBoneChain& SourceBoneChain,
@@ -266,6 +320,25 @@ struct FRetargetChainPairIK : FRetargetChainPair
         const FTargetSkeleton& TargetSkeleton) override;
 };
 
+USTRUCT()
+struct IKRIG_API FIKRetargetPose
+{
+	GENERATED_BODY()
+	
+	FIKRetargetPose(){}
+	
+	UPROPERTY(EditAnywhere, Category = RetargetPose)
+	FVector RootTranslationOffset;
+	
+	UPROPERTY(EditAnywhere, Category = RetargetPose)
+	FQuat RootRotationOffset;
+	
+	UPROPERTY(EditAnywhere, Category = RetargetPose)
+	TMap<FName, FQuat> BoneRotationOffsets;
+
+	void AddRotationDeltaToBone(FName BoneName, FQuat RotationDelta);
+};
+
 UCLASS(Blueprintable)
 class IKRIG_API UIKRetargeter : public UObject
 {
@@ -273,26 +346,101 @@ class IKRIG_API UIKRetargeter : public UObject
 	
 public:
 
-	UPROPERTY(EditAnywhere, Category = Rigs)
+	/** The rig to copy animation FROM.*/
+	UPROPERTY(VisibleAnywhere, Category = Rigs)
 	UIKRigDefinition *SourceIKRigAsset;
-
+	
+	/** The rig to copy animation TO.*/
 	UPROPERTY(EditAnywhere, Category = Rigs)
 	UIKRigDefinition *TargetIKRigAsset;
 
-	UPROPERTY(EditAnywhere, Category = Mapping)
-	TMap<FName, FName> ChainMapping;
+	/** Mapping of chains to copy animation between source and target rigs.*/
+	UPROPERTY()
+	TArray<FRetargetChainMap> ChainMapping;
+
+	/** The set of retarget poses available as options for retargeting.*/
+	UPROPERTY(EditAnywhere, Category = RetargetPose)
+	FName CurrentRetargetPose = NAME_None;
+	static const FName DefaultPoseName;
+	
+	/** The set of retarget poses available as options for retargeting.*/
+	UPROPERTY(EditAnywhere, Category = RetargetPose)
+	TMap<FName, FIKRetargetPose> RetargetPoses;
+
+	/** When false, translational motion of skeleton root is not copied. Useful for debugging.*/
+	UPROPERTY(EditAnywhere, Category = Settings)
+	bool bRetargetRoot = true;
+	
+	/** When false, limbs are not copied via FK. Useful for debugging limb issues suspected to be caused by FK pose.*/
+	UPROPERTY(EditAnywhere, Category = Settings)
+	bool bRetargetFK = true;
+	
+	/** When false, IK is not applied as part of retargeter. Useful for debugging limb issues suspected to be caused by IK.*/
+	UPROPERTY(EditAnywhere, Category = Settings)
+	bool bRetargetIK = true;
+
+	UPROPERTY(EditAnywhere, Category = DebugPreview, meta = (UIMin = "-500.0", UIMax = "500.0"))
+	float TargetActorOffset = 150.0f;
+
+	UPROPERTY(EditAnywhere, Category = DebugPreview, meta = (UIMin = "0.01", UIMax = "2.0"))
+	float TargetActorScale = 1.0f;
+
+	UPROPERTY(EditAnywhere, Category = DebugPreview, meta = (ClampMin = "0.0", UIMin = "0.01", UIMax = "10.0"))
+	float BoneGizmoSize = 8.0f;
+	
+	UPROPERTY(EditAnywhere, Category = DebugPreview, meta = (ClampMin = "0.0", UIMin = "0.01", UIMax = "10.0"))
+	float BoneGizmoThickness = 1.0f;
 
 	FRetargetSkeleton SourceSkeleton;
 	
 	FTargetSkeleton TargetSkeleton;
 	
-	bool bIsLoadedAndValid;
+	bool bIsLoadedAndValid = false;
+
+	UPROPERTY()
+	bool bEditReferencePoseMode = false;
+
+	UPROPERTY()
+	int32 AssetVersion = 0;
 
 	UIKRetargeter();
 	
 	void Initialize(USkeletalMesh *SourceSkeleton, USkeletalMesh *TargetSkeleton, UObject* Outer);
 	
-	void RunRetargeter(const TArray<FTransform>& InSourceGlobalPose, bool bEnableIK);
+	void RunRetargeter(const TArray<FTransform>& InSourceGlobalPose);
+
+	FTransform GetTargetBoneRetargetPoseGlobalTransform(const FName& ChainName) const;
+
+	void GetTargetBoneStartAndEnd(const FName& TargetBoneName, FVector& OutStart, FVector& OutEnd) const;
+
+	FRetargetChainMap* GetChainMap(const FName& TargetChainName);
+
+#if WITH_EDITOR
+	// UObject
+	virtual bool Modify(bool bAlwaysMarkDirty = true) override;
+	virtual void PostEditChangeProperty( struct FPropertyChangedEvent& PropertyChangedEvent) override;
+	virtual void PostLoad() override;
+	// END UObject
+
+	FName GetSourceRootBone();
+	FName GetTargetRootBone();
+
+	void GetTargetChainNames(TArray<FName>& OutNames) const;
+	void GetSourceChainNames(TArray<FName>& OutNames) const;
+
+	void CleanChainMapping();
+	void CleanPoseList();
+	void AutoMapChains();
+
+	USkeleton* GetSourceSkeletonAsset() const;
+
+	void AddRetargetPose(FName NewPoseName);
+	void RemoveRetargetPose(FName PoseToRemove);
+	void ResetRetargetPose(FName PoseToReset);
+	void SetCurrentRetargetPose(FName CurrentPose);
+	void AddRotationOffsetToRetargetPoseBone(FName BoneName, FQuat RotationOffset);
+	
+#endif
 
 private:
 
@@ -312,17 +460,10 @@ private:
 	bool InitializeBoneChainPairs();
 	
 	bool InitializeIKRig(UObject* Outer, const FReferenceSkeleton& InRefSkeleton);
-
-	void RunRootRetarget(
-		const TArray<FTransform>& InGlobalTransforms,
-		TArray<FTransform>& OutGlobalTransforms);
 	
-	void RunFKRetarget(
-		const TArray<FTransform>& InGlobalTransforms,
-		TArray<FTransform>& OutGlobalTransforms,
-		const bool bPreIK);
+	void RunRootRetarget(const TArray<FTransform>& InGlobalTransforms, TArray<FTransform>& OutGlobalTransforms);
+	
+	void RunFKRetarget(const TArray<FTransform>& InGlobalTransforms, TArray<FTransform>& OutGlobalTransforms);
 
-	void RunIKRetarget(
-		const TArray<FTransform>& InGlobalPose,
-		TArray<FTransform>& OutGlobalPose);
+	void RunIKRetarget(const TArray<FTransform>& InGlobalPose, TArray<FTransform>& OutGlobalPose);
 };
diff --git a/Engine/Plugins/Animation/IKRig/Source/IKRig/Public/Solvers/IKRig_PBIKSolver.h b/Engine/Plugins/Animation/IKRig/Source/IKRig/Public/Solvers/IKRig_PBIKSolver.h
--- a/Engine/Plugins/Animation/IKRig/Source/IKRig/Public/Solvers/IKRig_PBIKSolver.h
+++ b/Engine/Plugins/Animation/IKRig/Source/IKRig/Public/Solvers/IKRig_PBIKSolver.h
@@ -28,8 +28,8 @@ public:
 	UPROPERTY(EditAnywhere, Category = "Full Body IK Effector", meta = (ClampMin = "0", ClampMax = "1", UIMin = "0.0", UIMax = "1.0"))
 	float PullChainAlpha = 1.0f;
 
-	UPROPERTY(EditAnywhere, Category = "Full Body IK Effector")
-	bool bPinRotation = true;
+	UPROPERTY(EditAnywhere, Category = "Full Body IK Effector", meta = (ClampMin = "0", ClampMax = "1", UIMin = "0.0", UIMax = "1.0"))
+	float PinRotation = 1.0f;
 
 	UPROPERTY(Transient)
 	int32 IndexInSolver = -1;
@@ -38,7 +38,7 @@ public:
 	{
 		StrengthAlpha = Other->StrengthAlpha;
 		PullChainAlpha = Other->PullChainAlpha;
-		bPinRotation = Other->bPinRotation;
+		PinRotation = Other->PinRotation;
 	}
 };
 
diff --git a/Engine/Plugins/Animation/IKRig/Source/IKRigEditor/Private/IKRigEditor.cpp b/Engine/Plugins/Animation/IKRig/Source/IKRigEditor/Private/IKRigEditor.cpp
--- a/Engine/Plugins/Animation/IKRig/Source/IKRigEditor/Private/IKRigEditor.cpp
+++ b/Engine/Plugins/Animation/IKRig/Source/IKRigEditor/Private/IKRigEditor.cpp
@@ -10,6 +10,8 @@
 #include "IKRigDefinition.h"
 #include "RigEditor/AssetTypeActions_IKRigDefinition.h"
 #include "RetargetEditor/AssetTypeActions_IKRetargeter.h"
+#include "RetargetEditor/IKRetargetCommands.h"
+#include "RetargetEditor/IKRetargetEditMode.h"
 #include "RigEditor/IKRigCommands.h"
 #include "RigEditor/IKRigEditMode.h"
 #include "RigEditor/IKRigSkeletonCommands.h"
@@ -22,27 +24,31 @@ DEFINE_LOG_CATEGORY_STATIC(LogIKRigEditor, Log, All);
 
 void FIKRigEditor::StartupModule()
 {
+	// register commands
 	FIKRigCommands::Register();
 	FIKRigSkeletonCommands::Register();
+	FIKRetargetCommands::Register();
 	
-	// register IKRigDefinition asset type
+	// register custom asset type actions
 	IKRigDefinitionAssetAction = MakeShareable(new FAssetTypeActions_IKRigDefinition);
 	FModuleManager::LoadModuleChecked<FAssetToolsModule>("AssetTools").Get().RegisterAssetTypeActions(IKRigDefinitionAssetAction.ToSharedRef());
-
-	// register IKRetargeter asset type
+	//
 	IKRetargeterAssetAction = MakeShareable(new FAssetTypeActions_IKRetargeter);
 	FModuleManager::LoadModuleChecked<FAssetToolsModule>("AssetTools").Get().RegisterAssetTypeActions(IKRetargeterAssetAction.ToSharedRef());
 
-	// register custom editor mode
+	// register custom editor modes
 	FEditorModeRegistry::Get().RegisterMode<FIKRigEditMode>(FIKRigEditMode::ModeName, LOCTEXT("IKRigEditMode", "IKRig"), FSlateIcon(), false);
+	FEditorModeRegistry::Get().RegisterMode<FIKRetargetEditMode>(FIKRetargetEditMode::ModeName, LOCTEXT("IKRetargetEditMode", "IKRetarget"), FSlateIcon(), false);
 }
 
 void FIKRigEditor::ShutdownModule()
 {
 	FIKRigCommands::Unregister();
 	FIKRigSkeletonCommands::Unregister();
+	FIKRetargetCommands::Unregister();
 	
 	FEditorModeRegistry::Get().UnregisterMode(FIKRigEditMode::ModeName);
+	FEditorModeRegistry::Get().UnregisterMode(FIKRetargetEditMode::ModeName);
 
 	FPropertyEditorModule& PropertyModule = FModuleManager::GetModuleChecked<FPropertyEditorModule>("PropertyEditor");
 	PropertyModule.UnregisterCustomPropertyTypeLayout("IKRigEffector");
diff --git a/Engine/Plugins/Animation/IKRig/Source/IKRigEditor/Private/RetargetEditor/AssetTypeActions_IKRetargeter.cpp b/Engine/Plugins/Animation/IKRig/Source/IKRigEditor/Private/RetargetEditor/AssetTypeActions_IKRetargeter.cpp
--- a/Engine/Plugins/Animation/IKRig/Source/IKRigEditor/Private/RetargetEditor/AssetTypeActions_IKRetargeter.cpp
+++ b/Engine/Plugins/Animation/IKRig/Source/IKRigEditor/Private/RetargetEditor/AssetTypeActions_IKRetargeter.cpp
@@ -1,7 +1,9 @@
 // Copyright Epic Games, Inc. All Rights Reserved.
 
 #include "RetargetEditor/AssetTypeActions_IKRetargeter.h"
+#include "RetargetEditor/IKRetargetEditor.h"
 #include "Retargeter/IKRetargeter.h"
+#include "ThumbnailRendering/SceneThumbnailInfo.h"
 
 #define LOCTEXT_NAMESPACE "AssetTypeActions"
 
@@ -15,4 +17,26 @@ void FAssetTypeActions_IKRetargeter::GetActions(const TArray<UObject*>& InObject
 	FAssetTypeActions_Base::GetActions(InObjects, Section);
 }
 
+void FAssetTypeActions_IKRetargeter::OpenAssetEditor(
+	const TArray<UObject*>& InObjects,
+	TSharedPtr<IToolkitHost> EditWithinLevelEditor)
+{
+	const EToolkitMode::Type Mode = EditWithinLevelEditor.IsValid() ? EToolkitMode::WorldCentric : EToolkitMode::Standalone;
+    
+	for (auto ObjIt = InObjects.CreateConstIterator(); ObjIt; ++ObjIt)
+	{
+		if (UIKRetargeter* Asset = Cast<UIKRetargeter>(*ObjIt))
+		{
+			TSharedRef<FIKRetargetEditor> NewEditor(new FIKRetargetEditor());
+			NewEditor->InitAssetEditor(Mode, EditWithinLevelEditor, Asset);
+		}
+	}
+}
+
+UThumbnailInfo* FAssetTypeActions_IKRetargeter::GetThumbnailInfo(UObject* Asset) const
+{
+	UIKRetargeter* IKRetargeter = CastChecked<UIKRetargeter>(Asset);
+	return NewObject<USceneThumbnailInfo>(IKRetargeter, NAME_None, RF_Transactional);
+}
+
 #undef LOCTEXT_NAMESPACE
diff --git a/Engine/Plugins/Animation/IKRig/Source/IKRigEditor/Private/RetargetEditor/IKRetargetAnimInstance.cpp b/Engine/Plugins/Animation/IKRig/Source/IKRigEditor/Private/RetargetEditor/IKRetargetAnimInstance.cpp
new file mode 100644
--- /dev/null
+++ b/Engine/Plugins/Animation/IKRig/Source/IKRigEditor/Private/RetargetEditor/IKRetargetAnimInstance.cpp
@@ -0,0 +1,28 @@
+// Copyright Epic Games, Inc. All Rights Reserved.
+
+#include "RetargetEditor/IKRetargetAnimInstance.h"
+#include "RetargetEditor/IKRetargetAnimInstanceProxy.h"
+
+UIKRetargetAnimInstance::UIKRetargetAnimInstance(const FObjectInitializer& ObjectInitializer)
+	: Super(ObjectInitializer)
+{
+	bUseMultiThreadedAnimationUpdate = false;
+}
+
+void UIKRetargetAnimInstance::SetRetargetAssetAndSourceComponent(
+	UIKRetargeter* InAsset,
+	TWeakObjectPtr<USkeletalMeshComponent> InSourceMeshComponent)
+{
+	FIKRetargetAnimInstanceProxy& Proxy = GetProxyOnGameThread<FIKRetargetAnimInstanceProxy>();
+	Proxy.SetRetargetAssetAndSourceComponent(InAsset, InSourceMeshComponent);
+}
+
+UIKRetargeter* UIKRetargetAnimInstance::GetCurrentlyUsedRetargeter() const
+{
+	return IKRetargeterNode.GetCurrentlyUsedRetargeter();
+}
+
+FAnimInstanceProxy* UIKRetargetAnimInstance::CreateAnimInstanceProxy()
+{
+	return new FIKRetargetAnimInstanceProxy(this, &IKRetargeterNode);
+}
diff --git a/Engine/Plugins/Animation/IKRig/Source/IKRigEditor/Private/RetargetEditor/IKRetargetAnimInstanceProxy.cpp b/Engine/Plugins/Animation/IKRig/Source/IKRigEditor/Private/RetargetEditor/IKRetargetAnimInstanceProxy.cpp
new file mode 100644
--- /dev/null
+++ b/Engine/Plugins/Animation/IKRig/Source/IKRigEditor/Private/RetargetEditor/IKRetargetAnimInstanceProxy.cpp
@@ -0,0 +1,54 @@
+// Copyright Epic Games, Inc. All Rights Reserved.
+
+#include "RetargetEditor/IKRetargetAnimInstanceProxy.h"
+#include "RetargetEditor/IKRetargetAnimInstance.h"
+
+
+FIKRetargetAnimInstanceProxy::FIKRetargetAnimInstanceProxy(UAnimInstance* InAnimInstance, FAnimNode_IKRetargeter* InNode)
+	: FAnimPreviewInstanceProxy(InAnimInstance),
+	IKRetargetNode(InNode)
+{
+}
+
+void FIKRetargetAnimInstanceProxy::Initialize(UAnimInstance* InAnimInstance)
+{
+	FAnimInstanceProxy::Initialize(InAnimInstance);
+}
+
+bool FIKRetargetAnimInstanceProxy::Evaluate(FPoseContext& Output)
+{
+	IKRetargetNode->Evaluate_AnyThread(Output);
+	return true;
+}
+
+FAnimNode_Base* FIKRetargetAnimInstanceProxy::GetCustomRootNode()
+{
+	return IKRetargetNode;
+}
+
+void FIKRetargetAnimInstanceProxy::GetCustomNodes(TArray<FAnimNode_Base*>& OutNodes)
+{
+	OutNodes.Add(IKRetargetNode);
+}
+
+void FIKRetargetAnimInstanceProxy::UpdateAnimationNode(const FAnimationUpdateContext& InContext)
+{
+	if (CurrentAsset != nullptr)
+	{
+		FAnimPreviewInstanceProxy::UpdateAnimationNode(InContext);
+	}
+	else
+	{
+		IKRetargetNode->Update_AnyThread(InContext);
+	}
+}
+
+void FIKRetargetAnimInstanceProxy::SetRetargetAssetAndSourceComponent(
+	UIKRetargeter* InIKRetargetAsset,
+	TWeakObjectPtr<USkeletalMeshComponent> InSourceMeshComponent) const
+{
+	IKRetargetNode->IKRetargeterAsset = InIKRetargetAsset;
+	IKRetargetNode->bUseAttachedParent = false;
+	IKRetargetNode->SourceMeshComponent = InSourceMeshComponent;
+}
+
diff --git a/Engine/Plugins/Animation/IKRig/Source/IKRigEditor/Private/RetargetEditor/IKRetargetChainTabSummoner.cpp b/Engine/Plugins/Animation/IKRig/Source/IKRigEditor/Private/RetargetEditor/IKRetargetChainTabSummoner.cpp
new file mode 100644
--- /dev/null
+++ b/Engine/Plugins/Animation/IKRig/Source/IKRigEditor/Private/RetargetEditor/IKRetargetChainTabSummoner.cpp
@@ -0,0 +1,36 @@
+﻿// Copyright Epic Games, Inc. All Rights Reserved.
+#include "RetargetEditor/IKRetargetChainTabSummoner.h"
+
+#include "IDocumentation.h"
+#include "RigEditor/IKRigEditorStyle.h"
+#include "RetargetEditor/IKRetargetEditor.h"
+#include "RetargetEditor/SIKRetargetChainMapList.h"
+
+#define LOCTEXT_NAMESPACE "RetargetChainTabSummoner"
+
+const FName FIKRetargetChainTabSummoner::TabID(TEXT("ChainMapping"));
+
+FIKRetargetChainTabSummoner::FIKRetargetChainTabSummoner(const TSharedRef<FIKRetargetEditor>& InRetargetEditor)
+	: FWorkflowTabFactory(TabID, InRetargetEditor),
+	IKRetargetEditor(InRetargetEditor)
+{
+	bIsSingleton = true; // only allow a single instance of this tab
+	
+	TabLabel = LOCTEXT("IKRetargetChainTabLabel", "Chain Mapping");
+	TabIcon = FSlateIcon(FIKRigEditorStyle::Get().GetStyleSetName(), "IKRig.TabIcon");
+
+	ViewMenuDescription = LOCTEXT("IKRetargetChain_ViewMenu_Desc", "Chain Mapping");
+	ViewMenuTooltip = LOCTEXT("IKRetargetChain_ViewMenu_ToolTip", "Show the Chain Mapping Tab");
+}
+
+TSharedPtr<SToolTip> FIKRetargetChainTabSummoner::CreateTabToolTipWidget(const FWorkflowTabSpawnInfo& Info) const
+{
+	return  IDocumentation::Get()->CreateToolTip(LOCTEXT("IKRetargetChainTooltip", "Map the bone chains between the source and target IK rigs."), NULL, TEXT("Shared/Editors/Persona"), TEXT("IKRetargetChain_Window"));
+}
+
+TSharedRef<SWidget> FIKRetargetChainTabSummoner::CreateTabBody(const FWorkflowTabSpawnInfo& Info) const
+{
+	return SNew(SIKRetargetChainMapList, IKRetargetEditor.Pin()->GetController());
+}
+
+#undef LOCTEXT_NAMESPACE 
diff --git a/Engine/Plugins/Animation/IKRig/Source/IKRigEditor/Private/RetargetEditor/IKRetargetCommands.cpp b/Engine/Plugins/Animation/IKRig/Source/IKRigEditor/Private/RetargetEditor/IKRetargetCommands.cpp
new file mode 100644
--- /dev/null
+++ b/Engine/Plugins/Animation/IKRig/Source/IKRigEditor/Private/RetargetEditor/IKRetargetCommands.cpp
@@ -0,0 +1,16 @@
+// Copyright Epic Games, Inc. All Rights Reserved.
+
+#include "RetargetEditor/IKRetargetCommands.h"
+
+#define LOCTEXT_NAMESPACE "IKRetargetCommands"
+
+void FIKRetargetCommands::RegisterCommands()
+{
+	UI_COMMAND(EditRetargetPose, "Edit Pose", "Edit retarget reference pose of target mesh. Usually best to match the source mesh.", EUserInterfaceActionType::ToggleButton, FInputChord());
+	UI_COMMAND(NewRetargetPose, "New Pose", "Create a new retarget pose for the target mesh.", EUserInterfaceActionType::Button, FInputChord());
+	UI_COMMAND(DeleteRetargetPose, "Delete Pose", "Delete a retarget pose.", EUserInterfaceActionType::Button, FInputChord());
+	UI_COMMAND(ResetRetargetPose, "Reset Pose", "Resets a retarget pose to the mesh reference pose.", EUserInterfaceActionType::Button, FInputChord());
+	UI_COMMAND(ExportAnimation, "Export Animation", "Export selected animation(s) on the target mesh. Creates new Animation Sequence assets.", EUserInterfaceActionType::Button, FInputChord());
+}
+
+#undef LOCTEXT_NAMESPACE
diff --git a/Engine/Plugins/Animation/IKRig/Source/IKRigEditor/Private/RetargetEditor/IKRetargetEditMode.cpp b/Engine/Plugins/Animation/IKRig/Source/IKRigEditor/Private/RetargetEditor/IKRetargetEditMode.cpp
new file mode 100644
--- /dev/null
+++ b/Engine/Plugins/Animation/IKRig/Source/IKRigEditor/Private/RetargetEditor/IKRetargetEditMode.cpp
@@ -0,0 +1,275 @@
+// Copyright Epic Games, Inc. All Rights Reserved.
+
+#include "RetargetEditor/IKRetargetEditMode.h"
+
+#include "EditorViewportClient.h"
+#include "AssetEditorModeManager.h"
+#include "EditorScriptingHelpers.h"
+#include "IPersonaPreviewScene.h"
+#include "Animation/DebugSkelMeshComponent.h"
+#include "Materials/MaterialInstanceDynamic.h"
+#include "RetargetEditor/IKRetargetHitProxies.h"
+#include "RetargetEditor/IKRetargetEditor.h"
+
+FName FIKRetargetEditMode::ModeName("IKRetargetAssetEditMode");
+
+bool FIKRetargetEditMode::GetCameraTarget(FSphere& OutTarget) const
+{
+	// target skeletal mesh
+	if (const TSharedPtr<FIKRetargetEditorController> Controller = EditorController.Pin())
+	{
+		if (Controller->SourceSkelMeshComponent)
+		{
+			OutTarget = Controller->SourceSkelMeshComponent->Bounds.GetSphere();
+			return true;
+		}
+	}
+	
+	return false;
+}
+
+IPersonaPreviewScene& FIKRetargetEditMode::GetAnimPreviewScene() const
+{
+	return *static_cast<IPersonaPreviewScene*>(static_cast<FAssetEditorModeManager*>(Owner)->GetPreviewScene());
+}
+
+void FIKRetargetEditMode::GetOnScreenDebugInfo(TArray<FText>& OutDebugInfo) const
+{
+	// todo: provide warnings
+}
+
+void FIKRetargetEditMode::Render(const FSceneView* View, FViewport* Viewport, FPrimitiveDrawInterface* PDI)
+{
+	FEdMode::Render(View, Viewport, PDI);
+
+	const TSharedPtr<FIKRetargetEditorController> Controller = EditorController.Pin();
+	if (!Controller.IsValid())
+	{
+		return;
+	}
+	
+	UIKRetargeter* Retargeter = Controller->Asset;
+	for (const FName& BoneName : Retargeter->TargetIKRigAsset->Skeleton.BoneNames)
+	{
+		// is this chain currently selected?
+		const bool bIsSelected = IsBoneSelected(BoneName);
+
+		// get the location of the bone on the currently initialized target skeletal mesh
+		const FTransform BoneTransform = Controller->GetTargetBoneTransform(BoneName);
+
+		FVector Start;
+		FVector End;
+		Controller->GetTargetBoneStartAndEnd(BoneName, Start, End);
+
+		// draw the chain rotation gizmo
+		PDI->SetHitProxy(new HIKRetargetEditorBoneProxy(BoneName));
+		DrawBoneGizmo(
+			PDI,
+			BoneTransform,
+			Start,
+			End,
+			Retargeter->BoneGizmoSize,
+			Retargeter->BoneGizmoThickness,
+			bIsSelected);
+		PDI->SetHitProxy(NULL);
+	}
+}
+
+bool FIKRetargetEditMode::AllowWidgetMove()
+{
+	return !SelectedBones.IsEmpty();
+}
+
+bool FIKRetargetEditMode::ShouldDrawWidget() const
+{
+	return !SelectedBones.IsEmpty();
+}
+
+bool FIKRetargetEditMode::UsesTransformWidget() const
+{
+	return !SelectedBones.IsEmpty();
+}
+
+bool FIKRetargetEditMode::UsesTransformWidget(UE::Widget::EWidgetMode CheckMode) const
+{
+	return !SelectedBones.IsEmpty() && CheckMode == UE::Widget::EWidgetMode::WM_Rotate;
+}
+
+FVector FIKRetargetEditMode::GetWidgetLocation() const
+{
+	if (SelectedBones.IsEmpty())
+	{
+		return FVector::ZeroVector; 
+	}
+
+	const TSharedPtr<FIKRetargetEditorController> Controller = EditorController.Pin();
+	if (!Controller.IsValid())
+	{
+		return FVector::ZeroVector; 
+	}
+	
+	return Controller->GetTargetBoneTransform(SelectedBones.Last()).GetTranslation();
+}
+
+bool FIKRetargetEditMode::HandleClick(FEditorViewportClient* InViewportClient, HHitProxy* HitProxy, const FViewportClick& Click)
+{
+	// check for selections
+	if (Click.GetKey() == EKeys::LeftMouseButton)
+	{
+		// selected goal
+		if (HitProxy && HitProxy->IsA(HIKRetargetEditorBoneProxy::StaticGetType()))
+		{
+			HIKRetargetEditorBoneProxy* BoneProxy = static_cast<HIKRetargetEditorBoneProxy*>(HitProxy);
+			const bool bReplaceSelection = !(InViewportClient->IsCtrlPressed() || InViewportClient->IsShiftPressed());
+			HandleBoneSelectedInViewport(BoneProxy->BoneName, bReplaceSelection);
+			return true;
+		}
+		else
+		{
+			// clicking in empty space clears selected goals
+			const bool bReplaceSelection = true;
+			HandleBoneSelectedInViewport(NAME_None, bReplaceSelection);
+		}
+	}
+	
+	return false;
+}
+
+bool FIKRetargetEditMode::StartTracking(FEditorViewportClient* InViewportClient, FViewport* InViewport)
+{
+	if (SelectedBones.IsEmpty())
+	{
+		return false; // no goals selected to manipulate
+	}
+
+	const EAxisList::Type CurrentAxis = InViewportClient->GetCurrentWidgetAxis();
+	if (CurrentAxis == EAxisList::None)
+	{
+		return false; // not manipulating a required axis
+	}
+	
+	bRotatingBones = true;
+	return true;
+}
+
+bool FIKRetargetEditMode::EndTracking(FEditorViewportClient* InViewportClient, FViewport* InViewport)
+{
+	if (!bRotatingBones)
+	{
+		return false; // not handled
+	}
+
+	bRotatingBones = false;
+	return true;
+}
+
+bool FIKRetargetEditMode::InputDelta(FEditorViewportClient* InViewportClient, FViewport* InViewport, FVector& InDrag, FRotator& InRot, FVector& InScale)
+{
+	if (!bRotatingBones)
+	{
+		return false; // not handled
+	}
+	
+	if(InViewportClient->GetWidgetMode() != UE::Widget::WM_Rotate)
+	{
+		return false;
+	}
+
+	const TSharedPtr<FIKRetargetEditorController> Controller = EditorController.Pin();
+	if (!Controller.IsValid())
+	{
+		return false; 
+	}
+	
+	// rotate chains
+	for (const FName& BoneName : SelectedBones)
+	{
+		Controller->Asset->AddRotationOffsetToRetargetPoseBone(BoneName, InRot.Quaternion());
+	}
+	
+	return true;
+}
+
+bool FIKRetargetEditMode::GetCustomDrawingCoordinateSystem(FMatrix& InMatrix, void* InData)
+{
+	if (SelectedBones.IsEmpty())
+	{
+		return false; // nothing selected to manipulate
+	}
+
+	const TSharedPtr<FIKRetargetEditorController> Controller = EditorController.Pin();
+	if (!Controller.IsValid())
+	{
+		return false; 
+	}
+
+	const FTransform CurrentTransform = Controller->GetTargetBoneTransform(SelectedBones[0]);
+	InMatrix = CurrentTransform.ToMatrixNoScale().RemoveTranslation();
+	return true;
+}
+
+bool FIKRetargetEditMode::GetCustomInputCoordinateSystem(FMatrix& InMatrix, void* InData)
+{
+	return GetCustomDrawingCoordinateSystem(InMatrix, InData);
+}
+
+void FIKRetargetEditMode::DrawBoneGizmo(
+	FPrimitiveDrawInterface* PDI,
+	const FTransform& BoneTransform,
+	const FVector& Start,
+	const FVector& End,
+	float Size,
+	float Thickness,
+	bool bIsSelected) const
+{
+	const FLinearColor Color = bIsSelected ? FLinearColor::Green : FLinearColor::Yellow;
+	Size = FMath::Clamp(Size, 0.1f, 1000.0f);
+	Thickness = FMath::Clamp(Thickness, 0.1f, 1000.0f);
+
+	// draw a thick line along length of the limb
+	PDI->DrawLine(Start, End, Color, SDPG_Foreground, Thickness * 2.0f);
+	// draw a coordinate frame at the tip of the limb
+	DrawCoordinateSystem(PDI, End, BoneTransform.GetRotation().Rotator(), Size, SDPG_Foreground, Thickness);
+}
+
+void FIKRetargetEditMode::Tick(FEditorViewportClient* ViewportClient, float DeltaTime)
+{
+	FEdMode::Tick(ViewportClient, DeltaTime);
+}
+
+void FIKRetargetEditMode::DrawHUD(FEditorViewportClient* ViewportClient, FViewport* Viewport, const FSceneView* View, FCanvas* Canvas)
+{
+	FEdMode::DrawHUD(ViewportClient, Viewport, View, Canvas);
+}
+
+bool FIKRetargetEditMode::IsBoneSelected(const FName& BoneName) const
+{
+	return SelectedBones.Contains(BoneName);
+}
+
+void FIKRetargetEditMode::HandleBoneSelectedInViewport(const FName& BoneName, bool bReplace)
+{
+	if (bReplace)
+	{
+		SelectedBones.Reset();
+		if (BoneName != NAME_None)
+		{
+			SelectedBones.Add(BoneName);
+		}
+	}
+	else
+	{
+		if (BoneName != NAME_None)
+		{
+			const bool bAlreadySelected = SelectedBones.Contains(BoneName);
+			if (bAlreadySelected)
+			{
+				SelectedBones.Remove(BoneName);	
+			}
+			else
+			{
+				SelectedBones.Add(BoneName);
+			}
+		}	
+	}
+}
diff --git a/Engine/Plugins/Animation/IKRig/Source/IKRigEditor/Private/RetargetEditor/IKRetargetEditor.cpp b/Engine/Plugins/Animation/IKRig/Source/IKRigEditor/Private/RetargetEditor/IKRetargetEditor.cpp
new file mode 100644
--- /dev/null
+++ b/Engine/Plugins/Animation/IKRig/Source/IKRigEditor/Private/RetargetEditor/IKRetargetEditor.cpp
@@ -0,0 +1,537 @@
+// Copyright Epic Games, Inc. All Rights Reserved.
+
+#include "RetargetEditor/IKRetargetEditor.h"
+
+#include "AnimationEditorPreviewActor.h"
+#include "EditorModeManager.h"
+#include "GameFramework/WorldSettings.h"
+#include "Modules/ModuleManager.h"
+#include "PersonaModule.h"
+#include "IPersonaToolkit.h"
+#include "IAssetFamily.h"
+#include "ISkeletonEditorModule.h"
+#include "Preferences/PersonaOptions.h"
+#include "AnimCustomInstanceHelper.h"
+#include "Animation/DebugSkelMeshComponent.h"
+#include "Framework/MultiBox/MultiBoxBuilder.h"
+
+#include "Retargeter/IKRetargeter.h"
+#include "RetargetEditor/IKRetargetAnimInstance.h"
+#include "RetargetEditor/IKRetargetCommands.h"
+#include "RetargetEditor/IKRetargetEditMode.h"
+#include "RetargetEditor/IKRetargetMode.h"
+#include "RetargetEditor/IKRetargetEditorController.h"
+
+#define LOCTEXT_NAMESPACE "IKRetargeterEditor"
+
+const FName IKRetargetEditorModes::IKRetargetEditorMode("IKRetargetEditorMode");
+const FName IKRetargetEditorAppName = FName(TEXT("IKRetargetEditorApp"));
+
+FIKRetargetEditor::FIKRetargetEditor()
+	: EditorController(MakeShared<FIKRetargetEditorController>())
+{
+}
+
+FIKRetargetEditor::~FIKRetargetEditor()
+{
+}
+
+void FIKRetargetEditor::InitAssetEditor(
+	const EToolkitMode::Type Mode,
+    const TSharedPtr<IToolkitHost>& InitToolkitHost, 
+    UIKRetargeter* InAsset)
+{
+	EditorController->Editor = SharedThis(this);
+	EditorController->Asset = InAsset;
+
+	BindCommands();
+	
+	FPersonaToolkitArgs PersonaToolkitArgs;
+	PersonaToolkitArgs.OnPreviewSceneCreated = FOnPreviewSceneCreated::FDelegate::CreateSP(this, &FIKRetargetEditor::HandlePreviewSceneCreated);
+	
+	FPersonaModule& PersonaModule = FModuleManager::LoadModuleChecked<FPersonaModule>("Persona");
+	EditorController->PersonaToolkit = PersonaModule.CreatePersonaToolkit(InAsset, PersonaToolkitArgs, InAsset->GetSourceSkeletonAsset());
+	
+	// when/if preview mesh is changed, we need to reinitialize the anim instance
+    EditorController->PersonaToolkit->GetPreviewScene()->RegisterOnPreviewMeshChanged(FOnPreviewMeshChanged::CreateSP(this, &FIKRetargetEditor::HandlePreviewMeshChanged));
+
+	TSharedRef<IAssetFamily> AssetFamily = PersonaModule.CreatePersonaAssetFamily(InAsset);
+	AssetFamily->RecordAssetOpened(FAssetData(InAsset));
+
+	const bool bCreateDefaultStandaloneMenu = true;
+	const bool bCreateDefaultToolbar = true;
+	FAssetEditorToolkit::InitAssetEditor(
+		Mode, 
+		InitToolkitHost, 
+		IKRetargetEditorAppName, 
+		FTabManager::FLayout::NullLayout, 
+		bCreateDefaultStandaloneMenu, 
+		bCreateDefaultToolbar, 
+		InAsset);
+
+	AddApplicationMode(
+		IKRetargetEditorModes::IKRetargetEditorMode,
+		MakeShareable(new FIKRetargetMode(SharedThis(this), EditorController->PersonaToolkit->GetPreviewScene())));
+
+	SetCurrentMode(IKRetargetEditorModes::IKRetargetEditorMode);
+
+	GetEditorModeManager().SetDefaultMode(FPersonaEditModes::SkeletonSelection);
+	GetEditorModeManager().ActivateMode(FIKRetargetEditMode::ModeName);
+	FIKRetargetEditMode* EditMode = GetEditorModeManager().GetActiveModeTyped<FIKRetargetEditMode>(FIKRetargetEditMode::ModeName);
+	EditMode->SetEditorController(EditorController);
+	GetEditorModeManager().DeactivateMode(FIKRetargetEditMode::ModeName);
+
+	ExtendToolbar();
+	RegenerateMenusAndToolbars();
+}
+
+void FIKRetargetEditor::RegisterTabSpawners(const TSharedRef<FTabManager>& InTabManager)
+{
+	WorkspaceMenuCategory = InTabManager->AddLocalWorkspaceMenuCategory(LOCTEXT("WorkspaceMenu_IKRigEditor", "IK Rig Editor"));
+	TSharedRef<FWorkspaceItem> WorkspaceMenuCategoryRef = WorkspaceMenuCategory.ToSharedRef();
+
+	FAssetEditorToolkit::RegisterTabSpawners(InTabManager);
+}
+
+void FIKRetargetEditor::UnregisterTabSpawners(const TSharedRef<FTabManager>& InTabManager)
+{
+	FAssetEditorToolkit::UnregisterTabSpawners(InTabManager);
+}
+
+void FIKRetargetEditor::BindCommands()
+{
+	const FIKRetargetCommands& Commands = FIKRetargetCommands::Get();
+
+	ToolkitCommands->MapAction(
+        Commands.EditRetargetPose,
+        FExecuteAction::CreateSP(this, &FIKRetargetEditor::HandleEditPose),
+        FCanExecuteAction::CreateSP(this, &FIKRetargetEditor::CanEditPose),
+        FIsActionChecked::CreateSP(this,  &FIKRetargetEditor::IsEditingPose),
+		EUIActionRepeatMode::RepeatDisabled);
+
+	ToolkitCommands->MapAction(
+		Commands.NewRetargetPose,
+		FExecuteAction::CreateSP(this, &FIKRetargetEditor::HandleNewPose),
+		FCanExecuteAction(),
+		FCanExecuteAction(),
+		EUIActionRepeatMode::RepeatDisabled);
+
+	ToolkitCommands->MapAction(
+		Commands.DeleteRetargetPose,
+		FExecuteAction::CreateSP(this, &FIKRetargetEditor::HandleDeletePose),
+		FCanExecuteAction::CreateSP(this, &FIKRetargetEditor::CanDeletePose),
+		FCanExecuteAction(),
+		EUIActionRepeatMode::RepeatDisabled);
+
+	ToolkitCommands->MapAction(
+		Commands.ResetRetargetPose,
+		FExecuteAction::CreateSP(this, &FIKRetargetEditor::HandleResetPose),
+		FCanExecuteAction(),
+		EUIActionRepeatMode::RepeatDisabled);
+		
+	ToolkitCommands->MapAction(
+		Commands.ExportAnimation,
+		FExecuteAction::CreateSP(this, &FIKRetargetEditor::ExportAnimation),
+		FCanExecuteAction(),
+		EUIActionRepeatMode::RepeatDisabled);
+}
+
+void FIKRetargetEditor::ExtendToolbar()
+{
+	TSharedPtr<FExtender> ToolbarExtender = MakeShareable(new FExtender);
+
+	AddToolbarExtender(ToolbarExtender);
+
+	ToolbarExtender->AddToolBarExtension(
+        "Asset",
+        EExtensionHook::After,
+        GetToolkitCommands(),
+        FToolBarExtensionDelegate::CreateSP(this, &FIKRetargetEditor::FillToolbar)
+    );
+}
+
+void FIKRetargetEditor::FillToolbar(FToolBarBuilder& ToolbarBuilder)
+{
+	ToolbarBuilder.BeginSection("Retarget Pose");
+	{
+		PoseNames.Reset();
+		for (TTuple<FName, FIKRetargetPose>& Pose : EditorController->Asset->RetargetPoses)
+		{
+			PoseNames.Add(MakeShareable(new FName(Pose.Key)));
+		}
+
+		TSharedRef<SWidget> PoseListWidget = SNew(SHorizontalBox)
+
+		+ SHorizontalBox::Slot()
+		.AutoWidth()
+		.HAlign(HAlign_Left)
+		.VAlign(VAlign_Center)
+		.Padding(3.0f, 1.0f)
+		[
+			SNew(STextBlock)
+			.Text(LOCTEXT("IKRetargetPoseTitleLabel", "Current Retarget Pose: "))
+		]
+
+		+ SHorizontalBox::Slot()
+		.AutoWidth()
+		.HAlign(HAlign_Left)
+		.VAlign(VAlign_Center)
+		.Padding(3.0f, 1.0f)
+		[
+			SNew(SComboBox<TSharedPtr<FName>>)
+			.OptionsSource(&PoseNames)
+			.OnGenerateWidget_Lambda([](TSharedPtr<FName> InItem)
+			{
+				return SNew(STextBlock).Text(FText::FromName(*InItem.Get()));
+			})
+			.OnSelectionChanged(this, &FIKRetargetEditor::OnPoseSelected)
+			[
+				SNew(STextBlock).Text(this, &FIKRetargetEditor::GetCurrentPoseName)
+			]
+		];
+		ToolbarBuilder.AddWidget(PoseListWidget);
+		
+		ToolbarBuilder.AddToolBarButton(
+			FIKRetargetCommands::Get().EditRetargetPose,
+			NAME_None,
+			TAttribute<FText>(),
+			TAttribute<FText>(),
+			FSlateIcon(FAppStyle::Get().GetStyleSetName(),"Icons.Edit"));
+
+		ToolbarBuilder.AddToolBarButton(
+			FIKRetargetCommands::Get().NewRetargetPose,
+			NAME_None,
+			TAttribute<FText>(),
+			TAttribute<FText>(),
+			FSlateIcon(FAppStyle::Get().GetStyleSetName(),"Icons.Plus"));
+
+		ToolbarBuilder.AddToolBarButton(
+			FIKRetargetCommands::Get().DeleteRetargetPose,
+			NAME_None,
+			TAttribute<FText>(),
+			TAttribute<FText>(),
+			FSlateIcon(FAppStyle::Get().GetStyleSetName(),"Icons.Delete"));
+
+		ToolbarBuilder.AddToolBarButton(
+			FIKRetargetCommands::Get().ResetRetargetPose,
+			NAME_None,
+			TAttribute<FText>(),
+			TAttribute<FText>(),
+			FSlateIcon(FAppStyle::Get().GetStyleSetName(),"Icons.Refresh"));
+	}
+	ToolbarBuilder.EndSection();
+
+	ToolbarBuilder.BeginSection("Export");
+	{
+		ToolbarBuilder.AddToolBarButton(
+			FIKRetargetCommands::Get().ExportAnimation,
+			NAME_None,
+			TAttribute<FText>(),
+			TAttribute<FText>(),
+			FSlateIcon(FAppStyle::Get().GetStyleSetName(),"Icons.Save"));
+	}
+	ToolbarBuilder.EndSection();	
+}
+
+FName FIKRetargetEditor::GetToolkitFName() const
+{
+	return FName("IKRetargetEditor");
+}
+
+FText FIKRetargetEditor::GetBaseToolkitName() const
+{
+	return LOCTEXT("IKRetargetEditorAppLabel", "IK Retarget Editor");
+}
+
+FText FIKRetargetEditor::GetToolkitName() const
+{
+	const bool bDirtyState = EditorController->Asset->GetOutermost()->IsDirty();
+
+	FFormatNamedArguments Args;
+	Args.Add(TEXT("AssetName"), FText::FromString(EditorController->Asset->GetName()));
+	Args.Add(TEXT("DirtyState"), bDirtyState ? FText::FromString(TEXT("*")) : FText::GetEmpty());
+	return FText::Format(LOCTEXT("RetargetEditorToolkitName", "{AssetName}{DirtyState}"), Args);
+}
+
+FLinearColor FIKRetargetEditor::GetWorldCentricTabColorScale() const
+{
+	return FLinearColor::White;
+}
+
+FString FIKRetargetEditor::GetWorldCentricTabPrefix() const
+{
+	return TEXT("IKRetargetEditor");
+}
+
+void FIKRetargetEditor::AddReferencedObjects(FReferenceCollector& Collector)
+{
+	// hold the asset we are working on
+	Collector.AddReferencedObject(EditorController->Asset);
+}
+
+void FIKRetargetEditor::Tick(float DeltaTime)
+{
+	// apply offset to the target component
+	if (EditorController->TargetSkelMeshComponent)
+	{
+		const float TargetOffset = EditorController->Asset->TargetActorOffset;
+		EditorController->TargetSkelMeshComponent->SetRelativeLocation(FVector(TargetOffset,0,0));
+
+		const float TargetScale = EditorController->Asset->TargetActorScale;
+		EditorController->TargetSkelMeshComponent->SetRelativeScale3D(FVector(TargetScale,TargetScale,TargetScale));
+	}
+}
+
+TStatId FIKRetargetEditor::GetStatId() const
+{
+	RETURN_QUICK_DECLARE_CYCLE_STAT(FIKRetargetEditor, STATGROUP_Tickables);
+}
+
+void FIKRetargetEditor::HandlePreviewSceneCreated(const TSharedRef<IPersonaPreviewScene>& InPersonaPreviewScene)
+{
+	AAnimationEditorPreviewActor* Actor = InPersonaPreviewScene->GetWorld()->SpawnActor<AAnimationEditorPreviewActor>(AAnimationEditorPreviewActor::StaticClass(), FTransform::Identity);
+	Actor->SetFlags(RF_Transient);
+	InPersonaPreviewScene->SetActor(Actor);
+	
+	// create the skeletal mesh components
+	EditorController->SourceSkelMeshComponent = NewObject<UDebugSkelMeshComponent>(Actor);
+	EditorController->TargetSkelMeshComponent = NewObject<UDebugSkelMeshComponent>(Actor);
+
+	// setup an apply an anim instance to the skeletal mesh component
+	EditorController->SourceAnimInstance = NewObject<UAnimPreviewInstance>(EditorController->SourceSkelMeshComponent, TEXT("IKRetargetSourceAnimScriptInstance"));
+	EditorController->TargetAnimInstance = NewObject<UIKRetargetAnimInstance>(EditorController->TargetSkelMeshComponent, TEXT("IKRetargetTargetAnimScriptInstance"));
+	SetupAnimInstance();
+	
+	// set the source and target skeletal meshes on the component
+	// NOTE: this must be done AFTER setting the AnimInstance so that the correct root anim node is loaded
+	USkeletalMesh* SourceMesh = EditorController->GetSourceSkeletalMesh();
+	USkeletalMesh* TargetMesh = EditorController->GetTargetSkeletalMesh();
+	EditorController->SourceSkelMeshComponent->SetSkeletalMesh(SourceMesh);
+	EditorController->TargetSkelMeshComponent->SetSkeletalMesh(TargetMesh);
+
+	// apply mesh to the preview scene
+	InPersonaPreviewScene->SetPreviewMeshComponent(EditorController->SourceSkelMeshComponent);
+	InPersonaPreviewScene->SetAdditionalMeshesSelectable(false);
+	InPersonaPreviewScene->SetPreviewMesh(SourceMesh);
+	InPersonaPreviewScene->AddComponent(EditorController->SourceSkelMeshComponent, FTransform::Identity);
+	InPersonaPreviewScene->AddComponent(EditorController->TargetSkelMeshComponent, FTransform::Identity);
+}
+
+void FIKRetargetEditor::SetupAnimInstance()
+{
+	// connect the retarget asset and the source component ot the target anim instance
+	EditorController->TargetAnimInstance->SetRetargetAssetAndSourceComponent(EditorController->Asset, EditorController->SourceSkelMeshComponent);
+
+	EditorController->SourceSkelMeshComponent->PreviewInstance = EditorController->SourceAnimInstance.Get();
+	EditorController->TargetSkelMeshComponent->PreviewInstance = EditorController->TargetAnimInstance.Get();
+
+	EditorController->SourceAnimInstance->InitializeAnimation();
+	EditorController->TargetAnimInstance->InitializeAnimation();
+}
+
+void FIKRetargetEditor::HandleOpenNewAsset(UObject* InNewAsset)
+{
+	if (UAnimationAsset* NewAnimationAsset = Cast<UAnimationAsset>(InNewAsset))
+	{
+		EditorController->SourceAnimInstance->SetAnimationAsset(NewAnimationAsset);
+		PreviousAsset = NewAnimationAsset;
+	}
+}
+
+void FIKRetargetEditor::HandleAnimationSequenceBrowserCreated(
+	const TSharedRef<IAnimationSequenceBrowser>& InSequenceBrowser)
+{
+	EditorController->SequenceBrowser = InSequenceBrowser;
+}
+
+FText FIKRetargetEditor::GetCurrentPoseName() const
+{
+	return FText::FromName(EditorController->Asset->CurrentRetargetPose);
+}
+
+void FIKRetargetEditor::OnPoseSelected(TSharedPtr<FName> InPosePose, ESelectInfo::Type SelectInfo)
+{
+	EditorController->Asset->SetCurrentRetargetPose(*InPosePose.Get());
+}
+
+void FIKRetargetEditor::ExportAnimation() const
+{
+}
+
+void FIKRetargetEditor::HandleSourceOrTargetIKRigAssetChanged()
+{
+	// set the source and target skeletal meshes on the component
+	// NOTE: this must be done AFTER setting the AnimInstance so that the correct root anim node is loaded
+	USkeletalMesh* SourceMesh = EditorController->GetSourceSkeletalMesh();
+	USkeletalMesh* TargetMesh = EditorController->GetTargetSkeletalMesh();
+	EditorController->SourceSkelMeshComponent->SetSkeletalMesh(SourceMesh);
+	EditorController->TargetSkelMeshComponent->SetSkeletalMesh(TargetMesh);
+	
+	// apply mesh to the preview scene
+	TSharedRef<IPersonaPreviewScene> PreviewScene = GetPersonaToolkit()->GetPreviewScene();
+	if (PreviewScene->GetPreviewMesh() != SourceMesh)
+	{
+		PreviewScene->SetPreviewMeshComponent(EditorController->SourceSkelMeshComponent);
+		PreviewScene->SetPreviewMesh(SourceMesh);
+	}
+	
+	SetupAnimInstance();
+	
+	EditorController->RefreshAllViews();
+}
+
+void FIKRetargetEditor::HandlePreviewMeshChanged(USkeletalMesh* InOldSkeletalMesh, USkeletalMesh* InNewSkeletalMesh)
+{
+	if (!InOldSkeletalMesh)
+	{
+		return; // first time setup
+	}
+	
+	HandleSourceOrTargetIKRigAssetChanged();
+}
+
+void FIKRetargetEditor::HandleDetailsCreated(const TSharedRef<class IDetailsView>& InDetailsView)
+{
+	EditorController->DetailsView = InDetailsView;
+	EditorController->DetailsView->OnFinishedChangingProperties().AddSP(this, &FIKRetargetEditor::OnFinishedChangingDetails);
+	EditorController->DetailsView->SetObject(EditorController->Asset);
+}
+
+void FIKRetargetEditor::OnFinishedChangingDetails(
+    const FPropertyChangedEvent& PropertyChangedEvent)
+{
+	if (PropertyChangedEvent.GetPropertyName() == "TargetIKRigAsset"
+		|| PropertyChangedEvent.GetPropertyName() == "SourceIKRigAsset")
+	{
+		HandleSourceOrTargetIKRigAssetChanged();
+	}
+}
+
+void FIKRetargetEditor::HandleEditPose() const
+{
+	UIKRetargeter* Asset = EditorController.Get().Asset;
+	Asset->bEditReferencePoseMode = !EditorController.Get().Asset->bEditReferencePoseMode;
+	if (Asset->bEditReferencePoseMode)
+	{
+		GetEditorModeManager().ActivateMode(FIKRetargetEditMode::ModeName);
+		EditorController->SourceSkelMeshComponent->ShowReferencePose(true);
+	}
+	else
+	{
+		GetEditorModeManager().DeactivateMode(FIKRetargetEditMode::ModeName);
+		if (PreviousAsset)
+		{
+			EditorController->SourceSkelMeshComponent->EnablePreview(true, PreviousAsset);
+		}
+	}
+
+	Asset->Modify();
+}
+
+bool FIKRetargetEditor::CanEditPose() const
+{
+	UIKRetargeter* CurrentRetargeter = EditorController->GetCurrentlyRunningRetargeter();
+	if (!CurrentRetargeter)
+	{
+		return false;
+	}
+
+	return CurrentRetargeter->bIsLoadedAndValid;
+}
+
+bool FIKRetargetEditor::IsEditingPose() const
+{
+	return EditorController.Get().Asset->bEditReferencePoseMode;
+}
+
+void FIKRetargetEditor::HandleNewPose()
+{
+	SAssignNew(NewPoseWindow, SWindow)
+	.Title(LOCTEXT("NewRetargetPoseOptions", "Create New Retarget Pose"))
+	.ClientSize(FVector2D(250, 80))
+	.HasCloseButton(true)
+	.SupportsMinimize(false) .SupportsMaximize(false)
+	[
+		SNew(SBorder)
+		.BorderImage( FEditorStyle::GetBrush("Menu.Background") )
+		[
+			SNew(SVerticalBox)
+
+			+ SVerticalBox::Slot()
+			.Padding(4)
+			.AutoHeight()
+			[
+				SAssignNew(NewPoseEditableText, SEditableTextBox)
+				.Text(LOCTEXT("NewRetargetPoseName", "NewPose"))
+			]
+
+			+ SVerticalBox::Slot()
+			.Padding(4)
+			.AutoHeight()
+			[
+				SNew(SHorizontalBox)
+				+ SHorizontalBox::Slot()
+				.HAlign(HAlign_Center)
+				[
+					SNew(SButton)
+					.ButtonStyle(FAppStyle::Get(), "Button")
+					.TextStyle( FAppStyle::Get(), "DialogButtonText" )
+					.HAlign(HAlign_Center)
+					.VAlign(VAlign_Center)
+					.Text(LOCTEXT("OkButtonLabel", "Ok"))
+					.OnClicked(this, &FIKRetargetEditor::CreateNewPose)
+				]
+				
+				+ SHorizontalBox::Slot()
+				.HAlign(HAlign_Center)
+				[
+					SNew(SButton)
+					.ButtonStyle(FAppStyle::Get(), "Button")
+					.TextStyle( FAppStyle::Get(), "DialogButtonText" )
+					.HAlign(HAlign_Center)
+					.VAlign(VAlign_Center)
+					.Text(LOCTEXT("CancelButtonLabel", "Cancel"))
+					.OnClicked_Lambda( [=]()
+					{
+						NewPoseWindow->RequestDestroyWindow();
+						return FReply::Handled();
+					})
+				]
+			]	
+		]
+	];
+
+	GEditor->EditorAddModalWindow(NewPoseWindow.ToSharedRef());
+	NewPoseWindow.Reset();
+}
+
+FReply FIKRetargetEditor::CreateNewPose()
+{
+	const FName NewPoseName = FName(NewPoseEditableText.Get()->GetText().ToString());
+	EditorController->Asset->AddRetargetPose(NewPoseName);
+	NewPoseWindow->RequestDestroyWindow();
+	RegenerateMenusAndToolbars();
+	return FReply::Handled();
+}
+
+void FIKRetargetEditor::HandleDeletePose()
+{
+	const FName CurrentPose = EditorController->Asset->CurrentRetargetPose;
+	EditorController->Asset->RemoveRetargetPose(CurrentPose);
+	RegenerateMenusAndToolbars();
+}
+
+bool FIKRetargetEditor::CanDeletePose() const
+{	
+	// cannot delete default pose
+	return EditorController->Asset->CurrentRetargetPose != UIKRetargeter::DefaultPoseName;
+}
+
+void FIKRetargetEditor::HandleResetPose()
+{
+	const FName CurrentPose = EditorController->Asset->CurrentRetargetPose;
+	EditorController->Asset->ResetRetargetPose(CurrentPose);
+}
+
+
+#undef LOCTEXT_NAMESPACE
diff --git a/Engine/Plugins/Animation/IKRig/Source/IKRigEditor/Private/RetargetEditor/IKRetargetEditorController.cpp b/Engine/Plugins/Animation/IKRig/Source/IKRigEditor/Private/RetargetEditor/IKRetargetEditorController.cpp
new file mode 100644
--- /dev/null
+++ b/Engine/Plugins/Animation/IKRig/Source/IKRigEditor/Private/RetargetEditor/IKRetargetEditorController.cpp
@@ -0,0 +1,91 @@
+// Copyright Epic Games, Inc. All Rights Reserved.
+
+#include "RetargetEditor/IKRetargetEditorController.h"
+
+#include "RetargetEditor/IKRetargetAnimInstance.h"
+#include "RetargetEditor/SIKRetargetChainMapList.h"
+#include "Retargeter/IKRetargeter.h"
+
+#define LOCTEXT_NAMESPACE "IKRetargetEditorController"
+
+USkeletalMesh* FIKRetargetEditorController::GetSourceSkeletalMesh() const
+{
+	if (!(Asset && Asset->SourceIKRigAsset))
+	{
+		return nullptr;
+	}
+
+	return Asset->SourceIKRigAsset->PreviewSkeletalMesh;
+}
+
+USkeletalMesh* FIKRetargetEditorController::GetTargetSkeletalMesh() const
+{
+	if (!(Asset && Asset->TargetIKRigAsset))
+	{
+		return nullptr;
+	}
+
+	return Asset->TargetIKRigAsset->PreviewSkeletalMesh;
+}
+
+FTransform FIKRetargetEditorController::GetTargetBoneTransform(const FName& BoneName) const
+{
+	UIKRetargetAnimInstance* AnimInstance = TargetAnimInstance.Get();
+	check(AnimInstance);
+
+	UIKRetargeter* CurrentRetargeter = AnimInstance->GetCurrentlyUsedRetargeter();
+	check(CurrentRetargeter)
+
+	// get transform of root of chain
+	FTransform BoneTransform = CurrentRetargeter->GetTargetBoneRetargetPoseGlobalTransform(BoneName);
+
+	BoneTransform.ScaleTranslation(Asset->TargetActorScale);
+
+	// add the target translation offset
+	const FVector TargetOffset(Asset->TargetActorOffset, 0.f, 0.f);
+	BoneTransform.AddToTranslation(TargetOffset);
+
+	return BoneTransform;
+}
+
+void FIKRetargetEditorController::GetTargetBoneStartAndEnd(const FName& BoneName, FVector& Start, FVector& End) const
+{
+	UIKRetargeter* CurrentRetargeter = GetCurrentlyRunningRetargeter();
+	check(CurrentRetargeter)
+
+	// get line of chain
+	CurrentRetargeter->GetTargetBoneStartAndEnd(BoneName, Start, End);
+
+	// add the target translation offset
+	const FVector TargetOffset(Asset->TargetActorOffset, 0.f, 0.f);
+
+	Start *= Asset->TargetActorScale;
+	End *= Asset->TargetActorScale;
+	
+	Start += TargetOffset;
+	End += TargetOffset;
+}
+
+UIKRetargeter* FIKRetargetEditorController::GetCurrentlyRunningRetargeter() const
+{
+	UIKRetargetAnimInstance* AnimInstance = TargetAnimInstance.Get();
+	if(!AnimInstance)
+	{
+		return nullptr;
+	}
+
+	UIKRetargeter* CurrentRetargeter = AnimInstance->GetCurrentlyUsedRetargeter();
+	if (!CurrentRetargeter)
+	{
+		return nullptr;
+	}
+
+	return CurrentRetargeter;
+}
+
+void FIKRetargetEditorController::RefreshAllViews()
+{
+	ChainsView.Get()->RefreshView();
+}
+
+#undef LOCTEXT_NAMESPACE
diff --git a/Engine/Plugins/Animation/IKRig/Source/IKRigEditor/Private/RetargetEditor/IKRetargetFactory.cpp b/Engine/Plugins/Animation/IKRig/Source/IKRigEditor/Private/RetargetEditor/IKRetargetFactory.cpp
new file mode 100644
--- /dev/null
+++ b/Engine/Plugins/Animation/IKRig/Source/IKRigEditor/Private/RetargetEditor/IKRetargetFactory.cpp
@@ -0,0 +1,112 @@
+// Copyright Epic Games, Inc. All Rights Reserved.
+
+#include "RetargetEditor/IKRetargetFactory.h"
+#include "IKRig/Public/Retargeter/IKRetargeter.h"
+#include "AssetTypeCategories.h"
+#include "ContentBrowserModule.h"
+#include "IContentBrowserSingleton.h"
+
+#define LOCTEXT_NAMESPACE "IKRetargeterFactory"
+
+
+UIKRetargetFactory::UIKRetargetFactory()
+{
+	bCreateNew = true;
+	bEditAfterNew = true;
+	SupportedClass = UIKRetargeter::StaticClass();
+}
+
+UObject* UIKRetargetFactory::FactoryCreateNew(
+	UClass* Class,
+	UObject* InParent,
+	FName Name,
+	EObjectFlags Flags,
+	UObject* Context,
+	FFeedbackContext* Warn)
+{
+	if (!SourceIKRig)
+	{
+		UE_LOG(LogTemp, Warning, TEXT("Unable to create IK Retargter. No source IK Rig asset supplied."));
+		return nullptr;
+	}
+
+	if (!SourceIKRig->GetPreviewMesh())
+	{
+		UE_LOG(LogTemp, Warning, TEXT("Unable to create IK Retargter. Source IK Rig does not have an associated skeletal mesh."));
+		return nullptr;
+	}
+	
+	UIKRetargeter* Retargeter = NewObject<UIKRetargeter>(InParent, Class, Name, Flags);
+	Retargeter->SourceIKRigAsset = SourceIKRig;
+	return Retargeter;
+}
+
+bool UIKRetargetFactory::ShouldShowInNewMenu() const
+{
+	return true;
+}
+
+void UIKRetargetFactory::OnTargetIKRigSelected(const FAssetData& SelectedAsset)
+{
+	SourceIKRig = Cast<UIKRigDefinition>(SelectedAsset.GetAsset());
+	PickerWindow->RequestDestroyWindow();
+}
+
+bool UIKRetargetFactory::ConfigureProperties()
+{
+	// Null the parent class so we can check for selection later
+	SourceIKRig = nullptr;
+
+	// Load the content browser module to display an asset picker
+	FContentBrowserModule& ContentBrowserModule = FModuleManager::LoadModuleChecked<FContentBrowserModule>("ContentBrowser");
+
+	FAssetPickerConfig AssetPickerConfig;
+
+	/** The asset picker will only show skeletal meshes */
+	AssetPickerConfig.Filter.ClassNames.Add(UIKRigDefinition::StaticClass()->GetFName());
+	AssetPickerConfig.Filter.bRecursiveClasses = true;
+
+	/** The delegate that fires when an asset was selected */
+	AssetPickerConfig.OnAssetSelected = FOnAssetSelected::CreateUObject(this, &UIKRetargetFactory::OnTargetIKRigSelected);
+
+	/** The default view mode should be a list view */
+	AssetPickerConfig.InitialAssetViewType = EAssetViewType::List;
+
+	PickerWindow = SNew(SWindow)
+	.Title(LOCTEXT("CreateRetargeterOptions", "Pick IK Rig To Copy Animation From"))
+	.ClientSize(FVector2D(500, 600))
+	.SupportsMinimize(false) .SupportsMaximize(false)
+	[
+		SNew(SBorder)
+		.BorderImage( FEditorStyle::GetBrush("Menu.Background") )
+		[
+			ContentBrowserModule.Get().CreateAssetPicker(AssetPickerConfig)
+		]
+	];
+
+	GEditor->EditorAddModalWindow(PickerWindow.ToSharedRef());
+	PickerWindow.Reset();
+
+	return SourceIKRig != nullptr;
+}
+
+FText UIKRetargetFactory::GetDisplayName() const
+{
+	return LOCTEXT("IKRetargeter_DisplayName", "IK Retargeter");
+}
+
+uint32 UIKRetargetFactory::GetMenuCategories() const
+{
+	return EAssetTypeCategories::Animation;
+}
+
+FText UIKRetargetFactory::GetToolTip() const
+{
+	return LOCTEXT("IKRetargeter_Tooltip", "Defines a pair of Source/Target Retarget Rigs and the mapping between them.");
+}
+
+FString UIKRetargetFactory::GetDefaultNewAssetName() const
+{
+	return FString(TEXT("NewIKRetargeter"));
+}
+#undef LOCTEXT_NAMESPACE
diff --git a/Engine/Plugins/Animation/IKRig/Source/IKRigEditor/Private/RetargetEditor/IKRetargetHitProxies.cpp b/Engine/Plugins/Animation/IKRig/Source/IKRigEditor/Private/RetargetEditor/IKRetargetHitProxies.cpp
new file mode 100644
--- /dev/null
+++ b/Engine/Plugins/Animation/IKRig/Source/IKRigEditor/Private/RetargetEditor/IKRetargetHitProxies.cpp
@@ -0,0 +1,5 @@
+// Copyright Epic Games, Inc. All Rights Reserved.
+
+#include "RetargetEditor/IKRetargetHitProxies.h"
+
+IMPLEMENT_HIT_PROXY(HIKRetargetEditorBoneProxy, HHitProxy);
\ No newline at end of file
diff --git a/Engine/Plugins/Animation/IKRig/Source/IKRigEditor/Private/RetargetEditor/IKRetargetMode.cpp b/Engine/Plugins/Animation/IKRig/Source/IKRigEditor/Private/RetargetEditor/IKRetargetMode.cpp
new file mode 100644
--- /dev/null
+++ b/Engine/Plugins/Animation/IKRig/Source/IKRigEditor/Private/RetargetEditor/IKRetargetMode.cpp
@@ -0,0 +1,98 @@
+// Copyright Epic Games, Inc. All Rights Reserved.
+
+#include "RetargetEditor/IKRetargetMode.h"
+
+#include "IPersonaPreviewScene.h"
+#include "PersonaModule.h"
+#include "Modules/ModuleManager.h"
+#include "PersonaTabs.h"
+
+#include "RetargetEditor/IKRetargetChainTabSummoner.h"
+#include "RetargetEditor/IKRetargetEditor.h"
+
+#define LOCTEXT_NAMESPACE "IKRetargetMode"
+
+
+FIKRetargetMode::FIKRetargetMode(
+	TSharedRef<FWorkflowCentricApplication> InHostingApp,  
+	TSharedRef<IPersonaPreviewScene> InPreviewScene)
+	: FApplicationMode(IKRetargetEditorModes::IKRetargetEditorMode)
+{
+	IKRetargetEditorPtr = StaticCastSharedRef<FIKRetargetEditor>(InHostingApp);
+	TSharedRef<FIKRetargetEditor> IKRetargetEditor = StaticCastSharedRef<FIKRetargetEditor>(InHostingApp);
+
+	FPersonaViewportArgs ViewportArgs(InPreviewScene);
+	ViewportArgs.bAlwaysShowTransformToolbar = true;
+	ViewportArgs.bShowStats = false;
+	ViewportArgs.bShowTurnTable = false;
+	ViewportArgs.ContextName = TEXT("IKRetargetEditor.Viewport");
+
+	// register Persona tabs
+	FPersonaModule& PersonaModule = FModuleManager::LoadModuleChecked<FPersonaModule>("Persona");
+	TabFactories.RegisterFactory(PersonaModule.CreatePersonaViewportTabFactory(InHostingApp, ViewportArgs));
+	TabFactories.RegisterFactory(PersonaModule.CreateAdvancedPreviewSceneTabFactory(InHostingApp, InPreviewScene));
+	TabFactories.RegisterFactory(PersonaModule.CreateDetailsTabFactory(InHostingApp, FOnDetailsCreated::CreateSP(&IKRetargetEditor.Get(), &FIKRetargetEditor::HandleDetailsCreated)));
+	TabFactories.RegisterFactory(PersonaModule.CreateAnimationAssetBrowserTabFactory(InHostingApp, IKRetargetEditor->GetPersonaToolkit(), FOnOpenNewAsset::CreateSP(&IKRetargetEditor.Get(), &FIKRetargetEditor::HandleOpenNewAsset), FOnAnimationSequenceBrowserCreated::CreateSP(&IKRetargetEditor.Get(), &FIKRetargetEditor::HandleAnimationSequenceBrowserCreated), true));
+
+	// register custom tabs
+	TabFactories.RegisterFactory(MakeShared<FIKRetargetChainTabSummoner>(IKRetargetEditor));
+
+	// create tab layout
+	TabLayout = FTabManager::NewLayout("Standalone_IKRetargetEditor_Layout_v1.005")
+		->AddArea
+		(
+			FTabManager::NewPrimaryArea()
+			->SetOrientation(Orient_Vertical)
+			->Split
+			(
+				FTabManager::NewSplitter()
+				->SetSizeCoefficient(0.9f)
+				->SetOrientation(Orient_Horizontal)
+				->Split
+				(
+					FTabManager::NewStack()
+					->SetSizeCoefficient(0.6f)
+					->SetHideTabWell(true)
+					->AddTab(FPersonaTabs::PreviewViewportID, ETabState::OpenedTab)
+				)
+				->Split
+				(
+					FTabManager::NewSplitter()
+					->SetSizeCoefficient(0.9f)
+					->SetOrientation(Orient_Vertical)
+					->Split
+					(
+					FTabManager::NewStack()
+						->SetSizeCoefficient(0.6f)
+						->AddTab(FPersonaTabs::DetailsID, ETabState::OpenedTab)
+						->AddTab(FPersonaTabs::AdvancedPreviewSceneSettingsID, ETabState::OpenedTab)
+						->SetForegroundTab(FPersonaTabs::DetailsID)
+					)
+					->Split
+					(
+					FTabManager::NewStack()
+						->SetSizeCoefficient(0.6f)
+						->AddTab(FIKRetargetChainTabSummoner::TabID, ETabState::OpenedTab)
+						->AddTab(FPersonaTabs::AssetBrowserID, ETabState::OpenedTab)
+						->SetForegroundTab(FIKRetargetChainTabSummoner::TabID)
+					)
+					
+				)
+			)
+		);
+
+	PersonaModule.OnRegisterTabs().Broadcast(TabFactories, InHostingApp);
+	LayoutExtender = MakeShared<FLayoutExtender>();
+	PersonaModule.OnRegisterLayoutExtensions().Broadcast(*LayoutExtender.Get());
+	TabLayout->ProcessExtensions(*LayoutExtender.Get());
+}
+
+void FIKRetargetMode::RegisterTabFactories(TSharedPtr<FTabManager> InTabManager)
+{
+	TSharedPtr<FIKRetargetEditor> IKRigEditor = IKRetargetEditorPtr.Pin();
+	IKRigEditor->RegisterTabSpawners(InTabManager.ToSharedRef());
+	IKRigEditor->PushTabFactories(TabFactories);
+	FApplicationMode::RegisterTabFactories(InTabManager);
+}
+
+#undef LOCTEXT_NAMESPACE
diff --git a/Engine/Plugins/Animation/IKRig/Source/IKRigEditor/Private/RetargetEditor/IKRetargeterFactory.cpp b/Engine/Plugins/Animation/IKRig/Source/IKRigEditor/Private/RetargetEditor/IKRetargeterFactory.cpp
deleted file mode 100644
diff --git a/Engine/Plugins/Animation/IKRig/Source/IKRigEditor/Private/RetargetEditor/SIKRetargetChainMapList.cpp b/Engine/Plugins/Animation/IKRig/Source/IKRigEditor/Private/RetargetEditor/SIKRetargetChainMapList.cpp
new file mode 100644
--- /dev/null
+++ b/Engine/Plugins/Animation/IKRig/Source/IKRigEditor/Private/RetargetEditor/SIKRetargetChainMapList.cpp
@@ -0,0 +1,325 @@
+﻿// Copyright Epic Games, Inc. All Rights Reserved.
+
+#include "RetargetEditor/SIKRetargetChainMapList.h"
+
+
+#include "Framework/Commands/UICommandList.h"
+#include "Framework/MultiBox/MultiBoxBuilder.h"
+#include "IPersonaToolkit.h"
+#include "SKismetInspector.h"
+#include "Dialogs/Dialogs.h"
+#include "Widgets/Text/SInlineEditableTextBlock.h"
+#include "SEditorHeaderButton.h"
+#include "Retargeter/IKRetargeter.h"
+#include "RigEditor/IKRigEditorStyle.h"
+
+#define LOCTEXT_NAMESPACE "SIKRigRetargetChains"
+
+static const FName ColumnId_TargetChainLabel( "Target Bone Chain" );
+static const FName ColumnId_SourceChainLabel( "Source Bone Chain" );
+
+TSharedRef<ITableRow> FRetargetChainMapElement::MakeListRowWidget(
+	const TSharedRef<STableViewBase>& InOwnerTable,
+	TSharedRef<FRetargetChainMapElement> InChainElement,
+	TSharedPtr<SIKRetargetChainMapList> InChainList)
+{
+	return SNew(SIKRetargetChainMapRow, InOwnerTable, InChainElement, InChainList);
+}
+
+void SIKRetargetChainMapRow::Construct(
+	const FArguments& InArgs,
+	const TSharedRef<STableViewBase>& InOwnerTableView,
+	TSharedRef<FRetargetChainMapElement> InChainElement,
+	TSharedPtr<SIKRetargetChainMapList> InChainList)
+{
+	ChainMapElement = InChainElement;
+	ChainMapList = InChainList;
+
+	// generate list of source chains
+	// NOTE: cannot just use FName because "None" is considered a null entry and removed from ComboBox.
+	SourceChainOptions.Reset();
+	SourceChainOptions.Add(MakeShareable(new FString(TEXT("None"))));
+	UIKRigDefinition* SourceIKRig = ChainMapList.Pin()->EditorController.Pin()->Asset->SourceIKRigAsset;
+	if (SourceIKRig)
+	{
+		for (const FBoneChain& BoneChain : SourceIKRig->RetargetDefinition.BoneChains)
+		{
+			SourceChainOptions.Add(MakeShareable(new FString(BoneChain.ChainName.ToString())));
+		}
+	}
+
+	SMultiColumnTableRow< FRetargetChainMapElementPtr >::Construct( FSuperRowType::FArguments(), InOwnerTableView );
+}
+
+TSharedRef< SWidget > SIKRetargetChainMapRow::GenerateWidgetForColumn( const FName& ColumnName )
+{
+	if (ColumnName == ColumnId_TargetChainLabel)
+	{
+		TSharedRef<SWidget> NewWidget =
+		SNew(SHorizontalBox)
+		+ SHorizontalBox::Slot()
+		.AutoWidth()
+		.HAlign(HAlign_Left)
+		.VAlign(VAlign_Center)
+		.Padding(3.0f, 1.0f)
+		[
+			SNew(STextBlock)
+			.Text(FText::FromName(ChainMapElement.Pin()->TargetChainName))
+			.Font(FEditorStyle::GetFontStyle(TEXT("BoldFont")))
+		];
+		return NewWidget;
+	}
+	else
+	{
+		TSharedRef<SWidget> NewWidget =
+		SNew(SHorizontalBox)
+		+ SHorizontalBox::Slot()
+		.AutoWidth()
+		.HAlign(HAlign_Left)
+		.VAlign(VAlign_Center)
+		.Padding(3.0f, 1.0f)
+		[
+			SNew(SComboBox<TSharedPtr<FString>>)
+			.OptionsSource(&SourceChainOptions)
+			.OnGenerateWidget_Lambda([](TSharedPtr<FString> InItem)
+			{
+				return SNew(STextBlock).Text(FText::FromString(*InItem.Get()));
+			})
+			.OnSelectionChanged(this, &SIKRetargetChainMapRow::OnSourceChainComboSelectionChanged)
+			[
+				SNew(STextBlock)
+				.Text(this, &SIKRetargetChainMapRow::GetSourceChainName)
+			]
+		];
+		return NewWidget;
+	}
+}
+
+void SIKRetargetChainMapRow::OnSourceChainComboSelectionChanged(TSharedPtr<FString> InName, ESelectInfo::Type SelectInfo)
+{
+	const TSharedPtr<FIKRetargetEditorController> Controller = ChainMapList.Pin()->EditorController.Pin();
+	if (!Controller.IsValid())
+	{
+		return; 
+	}
+	
+	const FName TargetChainName = ChainMapElement.Pin()->TargetChainName;
+	FRetargetChainMap& ChainMap = *Controller->Asset->GetChainMap(TargetChainName);
+	ChainMap.SourceChain = FName(*InName.Get());
+	Controller->Asset->Modify();
+}
+
+FText SIKRetargetChainMapRow::GetSourceChainName() const
+{
+	const TSharedPtr<FIKRetargetEditorController> Controller = ChainMapList.Pin()->EditorController.Pin();
+	if (!Controller.IsValid())
+	{
+		return FText::FromName(NAME_None); 
+	}
+	
+	const FName TargetChainName = ChainMapElement.Pin()->TargetChainName;
+	const FRetargetChainMap& Mapping = *Controller->Asset->GetChainMap(TargetChainName);
+	return FText::FromName(Mapping.SourceChain);
+}
+
+SIKRetargetChainMapList::~SIKRetargetChainMapList()
+{
+	
+}
+
+void SIKRetargetChainMapList::Construct(
+	const FArguments& InArgs,
+	TSharedRef<FIKRetargetEditorController> InEditorController)
+{
+	EditorController = InEditorController;
+	EditorController.Pin()->ChainsView = SharedThis(this);
+
+	ChildSlot
+    [
+        SNew(SVerticalBox)
+        
+        +SVerticalBox::Slot()
+        .AutoHeight()
+        .VAlign(VAlign_Top)
+        [
+			SNew(SHorizontalBox)
+			+ SHorizontalBox::Slot()
+			.AutoWidth()
+			.HAlign(HAlign_Left)
+			.VAlign(VAlign_Center)
+			[
+				SNew(STextBlock)
+				.Text(LOCTEXT("TargetRootLabel", "Target Root: "))
+				.TextStyle(FEditorStyle::Get(), "NormalText")
+			]
+
+			+ SHorizontalBox::Slot()
+			.AutoWidth()
+			.HAlign(HAlign_Left)
+			.VAlign(VAlign_Center)
+			[
+				SNew(STextBlock)
+				.Text(this, &SIKRetargetChainMapList::GetSourceRootBone)
+				.IsEnabled(false)
+			]
+        ]
+
+		+SVerticalBox::Slot()
+		.AutoHeight()
+		.VAlign(VAlign_Top)
+		[
+			SNew(SHorizontalBox)
+			+ SHorizontalBox::Slot()
+			.AutoWidth()
+			.HAlign(HAlign_Left)
+			.VAlign(VAlign_Center)
+			[
+				SNew(STextBlock)
+				.Text(LOCTEXT("SourceRootLabel", "Source Root: "))
+				.TextStyle(FEditorStyle::Get(), "NormalText")
+			]
+
+			+ SHorizontalBox::Slot()
+			.AutoWidth()
+			.HAlign(HAlign_Left)
+			.VAlign(VAlign_Center)
+			[
+				SNew(STextBlock)
+				.Text(this, &SIKRetargetChainMapList::GetTargetRootBone)
+				.IsEnabled(false)
+			]
+		]
+
+		+ SVerticalBox::Slot()
+		.AutoHeight()
+		.VAlign(VAlign_Top)
+		.HAlign(HAlign_Left)
+		[
+			SNew(SEditorHeaderButton)
+			.Visibility(this, &SIKRetargetChainMapList::IsAutoMapButtonVisible)
+			.Icon(FAppStyle::Get().GetBrush("Icons.Refresh"))
+			.Text(LOCTEXT("AutoMapButtonLabel", "Auto-Map Chains"))
+			.ToolTipText(LOCTEXT("AutoMapButtonToolTip", "Automatically assign source chains based on fuzzy string match"))
+			.OnClicked(this, &SIKRetargetChainMapList::OnAutoMapButtonClicked)
+		]
+		
+        +SVerticalBox::Slot()
+        [
+			SAssignNew(ListView, SRetargetChainMapListViewType )
+			.SelectionMode(ESelectionMode::Multi)
+			.IsEnabled(this, &SIKRetargetChainMapList::IsChainMapEnabled)
+			.ListItemsSource( &ListViewItems )
+			.OnGenerateRow( this, &SIKRetargetChainMapList::MakeListRowWidget )
+			.OnMouseButtonClick(this, &SIKRetargetChainMapList::OnItemClicked)
+			.ItemHeight( 22.0f )
+			.HeaderRow
+			(
+				SNew( SHeaderRow )
+				+ SHeaderRow::Column( ColumnId_TargetChainLabel )
+				.DefaultLabel( LOCTEXT( "TargetColumnLabel", "Target Chain" ) )
+				.FixedWidth(150.f)
+
+				+ SHeaderRow::Column( ColumnId_SourceChainLabel )
+				.DefaultLabel( LOCTEXT( "SourceColumnLabel", "Source Chain" ) )
+			)
+        ]
+    ];
+
+	RefreshView();
+}
+
+FText SIKRetargetChainMapList::GetSourceRootBone() const
+{
+	const TSharedPtr<FIKRetargetEditorController> Controller = EditorController.Pin();
+	if (!Controller.IsValid())
+	{
+		return FText::FromName(NAME_None); 
+	}
+	
+	return FText::FromName(Controller->Asset->GetSourceRootBone());
+}
+
+FText SIKRetargetChainMapList::GetTargetRootBone() const
+{
+	const TSharedPtr<FIKRetargetEditorController> Controller = EditorController.Pin();
+	if (!Controller.IsValid())
+	{
+		return FText::FromName(NAME_None); 
+	}
+	
+	return FText::FromName(Controller->Asset->GetTargetRootBone());
+}
+
+bool SIKRetargetChainMapList::IsChainMapEnabled() const
+{
+	const TSharedPtr<FIKRetargetEditorController> Controller = EditorController.Pin();
+	if (!Controller.IsValid())
+	{
+		return false; 
+	}
+	
+	if (Controller->Asset->TargetIKRigAsset)
+	{
+		return Controller->Asset->TargetIKRigAsset->RetargetDefinition.BoneChains.Num() > 0;
+	}
+
+	return false;
+}
+
+void SIKRetargetChainMapList::RefreshView()
+{
+	const TSharedPtr<FIKRetargetEditorController> Controller = EditorController.Pin();
+	if (!Controller.IsValid())
+	{
+		return; 
+	}
+	
+	// get latest chain mapping
+	Controller->Asset->CleanChainMapping();
+	
+	// refresh list of chains
+	ListViewItems.Reset();
+	const TArray<FRetargetChainMap>& ChainMappings = Controller->Asset->ChainMapping;
+	for (const FRetargetChainMap& ChainMap : ChainMappings)
+	{
+		TSharedPtr<FRetargetChainMapElement> ChainItem = FRetargetChainMapElement::Make(ChainMap.TargetChain);
+		ListViewItems.Add(ChainItem);
+	}
+
+	ListView->RequestListRefresh();
+}
+
+TSharedRef<ITableRow> SIKRetargetChainMapList::MakeListRowWidget(
+	TSharedPtr<FRetargetChainMapElement> InElement,
+    const TSharedRef<STableViewBase>& OwnerTable)
+{
+	return InElement->MakeListRowWidget(
+		OwnerTable,
+		InElement.ToSharedRef(),
+		SharedThis(this));
+}
+
+void SIKRetargetChainMapList::OnItemClicked(TSharedPtr<FRetargetChainMapElement> InItem)
+{
+	// TODO highlight chain in skeleton view
+}
+
+EVisibility SIKRetargetChainMapList::IsAutoMapButtonVisible() const
+{
+	return IsChainMapEnabled() ? EVisibility::Visible : EVisibility::Hidden;
+}
+
+FReply SIKRetargetChainMapList::OnAutoMapButtonClicked() const
+{
+	const TSharedPtr<FIKRetargetEditorController> Controller = EditorController.Pin();
+	if (!Controller.IsValid())
+	{
+		return FReply::Unhandled();
+	}
+	
+	Controller->Asset->CleanChainMapping();
+	Controller->Asset->AutoMapChains();
+	return FReply::Handled();
+}
+
+#undef LOCTEXT_NAMESPACE
diff --git a/Engine/Plugins/Animation/IKRig/Source/IKRigEditor/Private/RigEditor/IKRigController.cpp b/Engine/Plugins/Animation/IKRig/Source/IKRigEditor/Private/RigEditor/IKRigController.cpp
--- a/Engine/Plugins/Animation/IKRig/Source/IKRigEditor/Private/RigEditor/IKRigController.cpp
+++ b/Engine/Plugins/Animation/IKRig/Source/IKRigEditor/Private/RigEditor/IKRigController.cpp
@@ -154,6 +154,107 @@ bool UIKRigController::DoesBoneHaveSettings(const FName& BoneName) const
 	return false;
 }
 
+bool UIKRigController::AddRetargetChain(const FName& ChainName, const FName& StartBone, const FName& EndBone) const
+{
+	if (IKRigAsset->RetargetDefinition.GetBoneChainByName(ChainName))
+	{
+		return false; // bone chain already exists with that name
+	}
+
+	FScopedTransaction Transaction(LOCTEXT("AddRetargetChain_Label", "Add Retarget Chain"));
+	IKRigAsset->Modify();
+	
+	IKRigAsset->RetargetDefinition.BoneChains.Emplace(ChainName, StartBone, EndBone);
+	IKRigAsset->SortRetargetChains();
+	return true;
+}
+
+bool UIKRigController::RemoveRetargetChain(const FName& ChainName) const
+{
+	FScopedTransaction Transaction(LOCTEXT("RemoveRetargetChain_Label", "Remove Retarget Chain"));
+	IKRigAsset->Modify();
+	
+	auto Pred = [ChainName](const FBoneChain& Element)
+	{
+		if (Element.ChainName == ChainName)
+		{
+			return true;
+		}
+
+		return false;
+	};
+	
+	if (IKRigAsset->RetargetDefinition.BoneChains.RemoveAll(Pred) > 0)
+	{
+		IKRigAsset->SortRetargetChains();
+		return true;
+	}
+	
+	return false;
+}
+
+bool UIKRigController::SetRetargetChainStartBone(const FName& ChainName, const FName& StartBoneName) const
+{
+	if (FBoneChain* BoneChain = IKRigAsset->RetargetDefinition.GetBoneChainByName(ChainName))
+	{
+		FScopedTransaction Transaction(LOCTEXT("SetRetargetChainStartBone_Label", "Set Retarget Chain Start Bone"));
+		IKRigAsset->Modify();
+	
+		BoneChain->StartBone = StartBoneName;
+		IKRigAsset->SortRetargetChains();
+		return true;
+	}
+
+	return false; // no bone chain with that name
+}
+
+bool UIKRigController::SetRetargetChainEndBone(const FName& ChainName, const FName& EndBoneName) const
+{
+	if (FBoneChain* BoneChain = IKRigAsset->RetargetDefinition.GetBoneChainByName(ChainName))
+	{
+		FScopedTransaction Transaction(LOCTEXT("SetRetargetChainEndBone_Label", "Set Retarget Chain End Bone"));
+		IKRigAsset->Modify();
+		
+		BoneChain->EndBone = EndBoneName;
+		return true;
+	}
+
+	return false; // no bone chain with that name
+}
+
+FBoneChain* UIKRigController::GetRetargetChain(int32 Index) const
+{
+	if (ensure(IKRigAsset) && IKRigAsset->RetargetDefinition.BoneChains.IsValidIndex(Index))
+	{
+		return &IKRigAsset->RetargetDefinition.BoneChains[Index];
+	}
+	
+	return nullptr;
+}
+
+int32 UIKRigController::GetNumRetargetChains() const
+{
+	check(IKRigAsset)
+	return IKRigAsset->RetargetDefinition.BoneChains.Num();
+}
+
+void UIKRigController::SetRetargetRoot(const FName& RootBoneName)
+{
+	check(IKRigAsset)
+
+	FScopedTransaction Transaction(LOCTEXT("SetRetargetRootBone_Label", "Set Retarget Root Bone"));
+	IKRigAsset->Modify();
+	
+	IKRigAsset->RetargetDefinition.RootBone = RootBoneName;
+}
+
+FName UIKRigController::GetRetargetRoot()
+{
+	check(IKRigAsset)
+
+	return IKRigAsset->RetargetDefinition.RootBone;
+}
+
 void UIKRigController::AddReferencedObjects(UObject* InThis, FReferenceCollector& Collector)
 {
 	// since static member, we just add only for default object
@@ -190,7 +291,8 @@ void UIKRigController::SetSourceSkeletalMesh(USkeletalMesh* SkeletalMesh, bool b
 	if (SkeletalMesh)
 	{
 		SetSkeleton(SkeletalMesh->GetRefSkeleton());		
-	}else
+	}
+	else
 	{
 		IKRigAsset->Skeleton.Reset();
 	}
@@ -201,7 +303,7 @@ void UIKRigController::SetSkeleton(const FReferenceSkeleton& InSkeleton) const
 	FScopedTransaction Transaction(LOCTEXT("SetSkeleton_Label", "Set Skeleton"));
 	IKRigAsset->Modify();
 
-	IKRigAsset->Skeleton.Initialize(InSkeleton);
+	IKRigAsset->Skeleton.Initialize(InSkeleton, IKRigAsset->Skeleton.ExcludedBones);
 }
 
 FIKRigSkeleton& UIKRigController::GetSkeleton() const
@@ -209,6 +311,32 @@ FIKRigSkeleton& UIKRigController::GetSkeleton() const
 	return IKRigAsset->Skeleton;
 }
 
+void UIKRigController::SetBoneExcluded(const FName& BoneName, const bool bExclude) const
+{
+	const bool bIsExcluded = IKRigAsset->Skeleton.ExcludedBones.Contains(BoneName);
+	if (bIsExcluded == bExclude)
+	{
+		return; // already in the requested state of exclusion
+	}
+
+	FScopedTransaction Transaction(LOCTEXT("SetBoneExcluded_Label", "Set Bone Excluded"));
+	IKRigAsset->Modify();
+
+	if (bExclude)
+	{
+		IKRigAsset->Skeleton.ExcludedBones.Add(BoneName);
+	}
+	else
+	{
+		IKRigAsset->Skeleton.ExcludedBones.Remove(BoneName);
+	}
+}
+
+bool UIKRigController::GetBoneExcluded(const FName& BoneName) const
+{
+	return IKRigAsset->Skeleton.ExcludedBones.Contains(BoneName);
+}
+
 USkeletalMesh* UIKRigController::GetSourceSkeletalMesh() const
 {
 	return IKRigAsset->PreviewSkeletalMesh.Get();
diff --git a/Engine/Plugins/Animation/IKRig/Source/IKRigEditor/Private/RigEditor/IKRigEditMode.cpp b/Engine/Plugins/Animation/IKRig/Source/IKRigEditor/Private/RigEditor/IKRigEditMode.cpp
--- a/Engine/Plugins/Animation/IKRig/Source/IKRigEditor/Private/RigEditor/IKRigEditMode.cpp
+++ b/Engine/Plugins/Animation/IKRig/Source/IKRigEditor/Private/RigEditor/IKRigEditMode.cpp
@@ -18,14 +18,20 @@ FIKRigEditMode::FIKRigEditMode()
 
 bool FIKRigEditMode::GetCameraTarget(FSphere& OutTarget) const
 {
+	const TSharedPtr<FIKRigEditorController> Controller = EditorController.Pin();
+	if (!Controller.IsValid())
+	{
+		return false; 
+	}
+	
 	// target union of selected goals
-	const TArray<FName>& OutGoalNames = EditorController.Pin()->GetSelectedGoals();
+	const TArray<FName>& OutGoalNames = Controller->GetSelectedGoals();
 	if (!OutGoalNames.IsEmpty())
 	{
 		TArray<FVector> GoalPoints;
 		for (const auto& GoalName : OutGoalNames)
 		{
-			FTransform GoalTransform = EditorController.Pin()->AssetController->GetGoalInitialTransform(GoalName);
+			FTransform GoalTransform = Controller->AssetController->GetGoalInitialTransform(GoalName);
 			GoalPoints.Add(GoalTransform.GetLocation());
 		}
 
@@ -35,9 +41,9 @@ bool FIKRigEditMode::GetCameraTarget(FSphere& OutTarget) const
 	}
 
 	// target skeletal mesh
-	if (EditorController.Pin()->SkelMeshComponent)
+	if (Controller->SkelMeshComponent)
 	{
-		OutTarget = EditorController.Pin()->SkelMeshComponent->Bounds.GetSphere();
+		OutTarget = Controller->SkelMeshComponent->Bounds.GetSphere();
 		return true;
 	}
 	
@@ -58,11 +64,17 @@ void FIKRigEditMode::Render(const FSceneView* View, FViewport* Viewport, FPrimit
 {
 	FEdMode::Render(View, Viewport, PDI);
 
-	UIKRigController* Controller = EditorController.Pin()->AssetController;
-	TArray<UIKRigEffectorGoal*> Goals = Controller->GetAllGoals();
+	const TSharedPtr<FIKRigEditorController> Controller = EditorController.Pin();
+	if (!Controller.IsValid())
+	{
+		return; 
+	}
+	
+	UIKRigController* AssetController = Controller->AssetController;
+	TArray<UIKRigEffectorGoal*> Goals = AssetController->GetAllGoals();
 	for (const UIKRigEffectorGoal* Goal : Goals)
 	{
-		const bool bIsSelected = EditorController.Pin()->IsGoalSelected(Goal->GoalName);
+		const bool bIsSelected = Controller->IsGoalSelected(Goal->GoalName);
 		PDI->SetHitProxy(new HIKRigEditorGoalProxy(Goal->GoalName));
 		GoalDrawer.DrawGoal(PDI, Goal, bIsSelected);
 		PDI->SetHitProxy(NULL);
@@ -71,37 +83,73 @@ void FIKRigEditMode::Render(const FSceneView* View, FViewport* Viewport, FPrimit
 
 bool FIKRigEditMode::AllowWidgetMove()
 {
-	return EditorController.Pin()->GetNumSelectedGoals() > 0;
+	const TSharedPtr<FIKRigEditorController> Controller = EditorController.Pin();
+	if (!Controller.IsValid())
+	{
+		return false; 
+	}
+	
+	return Controller->GetNumSelectedGoals() > 0;
 }
 
 bool FIKRigEditMode::ShouldDrawWidget() const
 {
-	return EditorController.Pin()->GetNumSelectedGoals() > 0;
+	const TSharedPtr<FIKRigEditorController> Controller = EditorController.Pin();
+	if (!Controller.IsValid())
+	{
+		return false; 
+	}
+	
+	return Controller->GetNumSelectedGoals() > 0;
 }
 
 bool FIKRigEditMode::UsesTransformWidget() const
 {
-	return EditorController.Pin()->GetNumSelectedGoals() > 0;
+	const TSharedPtr<FIKRigEditorController> Controller = EditorController.Pin();
+	if (!Controller.IsValid())
+	{
+		return false; 
+	}
+	
+	return Controller->GetNumSelectedGoals() > 0;
 }
 
 bool FIKRigEditMode::UsesTransformWidget(UE::Widget::EWidgetMode CheckMode) const
 {
-	return EditorController.Pin()->GetNumSelectedGoals() > 0;
+	const TSharedPtr<FIKRigEditorController> Controller = EditorController.Pin();
+	if (!Controller.IsValid())
+	{
+		return false; 
+	}
+	
+	return Controller->GetNumSelectedGoals() > 0;
 }
 
 FVector FIKRigEditMode::GetWidgetLocation() const
 {
-	const TArray<FName>& OutGoalNames = EditorController.Pin()->GetSelectedGoals();
+	const TSharedPtr<FIKRigEditorController> Controller = EditorController.Pin();
+	if (!Controller.IsValid())
+	{
+		return FVector::ZeroVector; 
+	}
+	
+	const TArray<FName>& OutGoalNames = Controller->GetSelectedGoals();
 	if (OutGoalNames.IsEmpty())
 	{
 		return FVector::ZeroVector; 
 	}
 
-	return EditorController.Pin()->AssetController->GetGoalCurrentTransform(OutGoalNames.Last()).GetTranslation();
+	return Controller->AssetController->GetGoalCurrentTransform(OutGoalNames.Last()).GetTranslation();
 }
 
 bool FIKRigEditMode::HandleClick(FEditorViewportClient* InViewportClient, HHitProxy* HitProxy, const FViewportClick& Click)
 {
+	const TSharedPtr<FIKRigEditorController> Controller = EditorController.Pin();
+	if (!Controller.IsValid())
+	{
+		return false; 
+	}
+	
 	// check for selections
 	if (Click.GetKey() == EKeys::LeftMouseButton)
 	{
@@ -110,13 +158,14 @@ bool FIKRigEditMode::HandleClick(FEditorViewportClient* InViewportClient, HHitPr
 		{
 			HIKRigEditorGoalProxy* GoalProxy = static_cast<HIKRigEditorGoalProxy*>(HitProxy);
 			const bool bReplaceSelection = !(InViewportClient->IsCtrlPressed() || InViewportClient->IsShiftPressed());
-			EditorController.Pin()->HandleGoalSelectedInViewport(GoalProxy->GoalName, bReplaceSelection);
+			Controller->HandleGoalSelectedInViewport(GoalProxy->GoalName, bReplaceSelection);
 			return true;
-		}else
+		}
+		else
 		{
 			// clicking in empty space clears selected goals
 			const bool bReplaceSelection = true;
-			EditorController.Pin()->HandleGoalSelectedInViewport(NAME_None, bReplaceSelection);
+			Controller->HandleGoalSelectedInViewport(NAME_None, bReplaceSelection);
 		}
 	}
 	
@@ -125,7 +174,13 @@ bool FIKRigEditMode::HandleClick(FEditorViewportClient* InViewportClient, HHitPr
 
 bool FIKRigEditMode::StartTracking(FEditorViewportClient* InViewportClient, FViewport* InViewport)
 {
-	const TArray<FName>& SelectedGoalNames = EditorController.Pin()->GetSelectedGoals();
+	const TSharedPtr<FIKRigEditorController> Controller = EditorController.Pin();
+	if (!Controller.IsValid())
+	{
+		return false; 
+	}
+	
+	const TArray<FName>& SelectedGoalNames = Controller->GetSelectedGoals();
 	if (SelectedGoalNames.IsEmpty())
 	{
 		return false; // no goals selected to manipulate
@@ -137,53 +192,64 @@ bool FIKRigEditMode::StartTracking(FEditorViewportClient* InViewportClient, FVie
 		return false; // not manipulating a required axis
 	}
 	
-	EditorController.Pin()->bManipulatingGoals = true;
+	Controller->bManipulatingGoals = true;
 	return true;
 }
 
 bool FIKRigEditMode::EndTracking(FEditorViewportClient* InViewportClient, FViewport* InViewport)
 {
-	if (!EditorController.Pin()->bManipulatingGoals)
+	const TSharedPtr<FIKRigEditorController> Controller = EditorController.Pin();
+	if (!Controller.IsValid())
+	{
+		return false; 
+	}
+	
+	if (!Controller->bManipulatingGoals)
 	{
 		return false; // not handled
 	}
 
-	EditorController.Pin()->bManipulatingGoals = false;
+	Controller->bManipulatingGoals = false;
 	return true;
 }
 
 bool FIKRigEditMode::InputDelta(FEditorViewportClient* InViewportClient, FViewport* InViewport, FVector& InDrag, FRotator& InRot, FVector& InScale)
 {
-	if (!EditorController.Pin()->bManipulatingGoals)
+	const TSharedPtr<FIKRigEditorController> Controller = EditorController.Pin();
+	if (!Controller.IsValid())
+	{
+		return false; 
+	}
+	
+	if (!Controller->bManipulatingGoals)
 	{
 		return false; // not handled
 	}
 
+	const TArray<FName>& SelectedGoalNames = Controller->GetSelectedGoals();
+	UIKRigController* AssetController = Controller->AssetController;
+
 	// translate goals
 	if(InViewportClient->GetWidgetMode() == UE::Widget::WM_Translate)
 	{
-		const TArray<FName>& SelectedGoalNames = EditorController.Pin()->GetSelectedGoals();
-		UIKRigController* Controller = EditorController.Pin()->AssetController;
 		for (const FName& GoalName : SelectedGoalNames)
 		{
-			FTransform CurrentTransform = Controller->GetGoal(GoalName)->CurrentTransform;
+			FTransform CurrentTransform = AssetController->GetGoal(GoalName)->CurrentTransform;
 			CurrentTransform.AddToTranslation(InDrag);
-			Controller->SetGoalCurrentTransform(GoalName, CurrentTransform);
+			AssetController->SetGoalCurrentTransform(GoalName, CurrentTransform);
 		}
 	}
 
 	// rotate goals
 	if(InViewportClient->GetWidgetMode() == UE::Widget::WM_Rotate)
 	{
-		const TArray<FName>& SelectedGoalNames = EditorController.Pin()->GetSelectedGoals();
-		UIKRigController* Controller = EditorController.Pin()->AssetController;
 		for (const FName& GoalName : SelectedGoalNames)
 		{
-			FTransform CurrentTransform = Controller->GetGoal(GoalName)->CurrentTransform;
+			FTransform CurrentTransform = AssetController->GetGoal(GoalName)->CurrentTransform;
 			FQuat CurrentRotation = CurrentTransform.GetRotation();
 			CurrentRotation = (InRot.Quaternion() * CurrentRotation);
 			CurrentTransform.SetRotation(CurrentRotation);
-			Controller->SetGoalCurrentTransform(GoalName, CurrentTransform);
+			AssetController->SetGoalCurrentTransform(GoalName, CurrentTransform);
 		}
 	}
 	
@@ -192,13 +258,19 @@ bool FIKRigEditMode::InputDelta(FEditorViewportClient* InViewportClient, FViewpo
 
 bool FIKRigEditMode::GetCustomDrawingCoordinateSystem(FMatrix& InMatrix, void* InData)
 {
-	const TArray<FName>& SelectedGoals = EditorController.Pin()->GetSelectedGoals();
+	const TSharedPtr<FIKRigEditorController> Controller = EditorController.Pin();
+	if (!Controller.IsValid())
+	{
+		return false;
+	}
+	
+	const TArray<FName>& SelectedGoals = Controller->GetSelectedGoals();
 	if (SelectedGoals.IsEmpty())
 	{
 		return false; // nothing selected to manipulate
 	}
 
-	if (UIKRigEffectorGoal* Goal = EditorController.Pin()->AssetController->GetGoal(SelectedGoals[0]))
+	if (UIKRigEffectorGoal* Goal = Controller->AssetController->GetGoal(SelectedGoals[0]))
 	{
 		InMatrix = Goal->CurrentTransform.ToMatrixNoScale().RemoveTranslation();
 		return true;
diff --git a/Engine/Plugins/Animation/IKRig/Source/IKRigEditor/Private/RigEditor/IKRigEditorController.cpp b/Engine/Plugins/Animation/IKRig/Source/IKRigEditor/Private/RigEditor/IKRigEditorController.cpp
--- a/Engine/Plugins/Animation/IKRig/Source/IKRigEditor/Private/RigEditor/IKRigEditorController.cpp
+++ b/Engine/Plugins/Animation/IKRig/Source/IKRigEditor/Private/RigEditor/IKRigEditorController.cpp
@@ -26,6 +26,7 @@ void FIKRigEditorController::RefreshAllViews() const
 	SolverStackView->RefreshStackView();
 	SkeletonView->RefreshTreeView();
 	DetailsView->ForceRefresh();
+	RetargetingView->RefreshView();
 }
 
 bool FIKRigEditorController::IsGoalSelected(const FName& GoalName)
@@ -47,7 +48,8 @@ void FIKRigEditorController::HandleGoalSelectedInViewport(const FName& GoalName,
 		{
 			SelectedGoals.Add(GoalName);
 		}
-	}else
+	}
+	else
 	{
 		if (GoalName != NAME_None)
 		{
@@ -55,7 +57,8 @@ void FIKRigEditorController::HandleGoalSelectedInViewport(const FName& GoalName,
 			if (bAlreadySelected)
 			{
 				SelectedGoals.Remove(GoalName);	
-			}else
+			}
+			else
 			{
 				SelectedGoals.Add(GoalName);
 			}
@@ -164,6 +167,17 @@ bool FIKRigEditorController::IsElementConnectedToAnySolver(TSharedRef<FIKRigTree
 	return false;
 }
 
+bool FIKRigEditorController::IsElementExcludedBone(TSharedRef<FIKRigTreeElement> TreeElement)
+{
+	if (TreeElement->ElementType != IKRigTreeElementType::BONE)
+	{
+		return false;
+	}
+	
+	// is this bone excluded?
+	return AssetController->GetBoneExcluded(TreeElement->Key);
+}
+
 void FIKRigEditorController::ShowDetailsForBone(const FName BoneName)
 {
 	ShowEmptyDetails();
@@ -207,4 +221,26 @@ void FIKRigEditorController::ShowEmptyDetails()
 	DetailsView->SetObject(nullptr);
 }
 
+void FIKRigEditorController::AddNewRetargetChain(const FName ChainName, const FName StartBone, const FName EndBone)
+{
+	FIKRigRetargetChainSettings Settings(ChainName, StartBone, EndBone);
+	const TSharedPtr<FStructOnScope> StructToDisplay = MakeShareable(new FStructOnScope(FIKRigRetargetChainSettings::StaticStruct(), (uint8*)&Settings));
+	TSharedRef<SKismetInspector> KismetInspector = SNew(SKismetInspector);
+	KismetInspector->ShowSingleStruct(StructToDisplay);
+
+	SGenericDialogWidget::FArguments DialogArguments;
+	DialogArguments.OnOkPressed_Lambda([&Settings, this] ()
+	{
+		// add the retarget chain
+		AssetController->AddRetargetChain(Settings.ChainName, Settings.StartBone, Settings.EndBone);
+		RefreshAllViews();
+	});
+
+	SGenericDialogWidget::OpenDialog(
+		LOCTEXT("SIKRigRetargetChains", "Add New Retarget Chain"),
+		KismetInspector,
+		DialogArguments,
+		true);
+}
+
 #undef LOCTEXT_NAMESPACE
diff --git a/Engine/Plugins/Animation/IKRig/Source/IKRigEditor/Private/RigEditor/IKRigMode.cpp b/Engine/Plugins/Animation/IKRig/Source/IKRigEditor/Private/RigEditor/IKRigMode.cpp
--- a/Engine/Plugins/Animation/IKRig/Source/IKRigEditor/Private/RigEditor/IKRigMode.cpp
+++ b/Engine/Plugins/Animation/IKRig/Source/IKRigEditor/Private/RigEditor/IKRigMode.cpp
@@ -9,11 +9,10 @@
 #include "PersonaTabs.h"
 #include "RigEditor/IKRigSkeletonTabSummoner.h"
 #include "RigEditor/IKRigSolverStackTabSummoner.h"
+#include "RigEditor/IKRigRetargetChainTabSummoner.h"
 
 #define LOCTEXT_NAMESPACE "IKRigMode"
 
-static const FName IKRigSolverStackName("IKRigSolverStack");
-
 FIKRigMode::FIKRigMode(
 	TSharedRef<FWorkflowCentricApplication> InHostingApp,  
 	TSharedRef<IPersonaPreviewScene> InPreviewScene)
@@ -37,9 +36,10 @@ FIKRigMode::FIKRigMode(
 	// register custom tabs
 	TabFactories.RegisterFactory(MakeShared<FIKRigSkeletonTabSummoner>(IKRigEditor));
 	TabFactories.RegisterFactory(MakeShared<FIKRigSolverStackTabSummoner>(IKRigEditor));
+	TabFactories.RegisterFactory(MakeShared<FIKRigRetargetChainTabSummoner>(IKRigEditor));
 
 	// create tab layout
-	TabLayout = FTabManager::NewLayout("Standalone_IKRigEditor_Layout_v1.118")
+	TabLayout = FTabManager::NewLayout("Standalone_IKRigEditor_Layout_v1.120")
 		->AddArea
 		(
 			FTabManager::NewPrimaryArea()
@@ -64,7 +64,7 @@ FIKRigMode::FIKRigMode(
 					(
 					    FTabManager::NewStack()
 					    ->SetSizeCoefficient(0.4f)
-					    ->AddTab(IKRigSolverStackName, ETabState::OpenedTab)
+					    ->AddTab(FIKRigSolverStackTabSummoner::TabID, ETabState::OpenedTab)
 					)
 				)
 				->Split
@@ -76,11 +76,24 @@ FIKRigMode::FIKRigMode(
 				)
 				->Split
 				(
-					FTabManager::NewStack()
-                    ->SetSizeCoefficient(0.6f)
-                    ->AddTab(FPersonaTabs::DetailsID, ETabState::OpenedTab)
-                    ->AddTab(FPersonaTabs::AdvancedPreviewSceneSettingsID, ETabState::OpenedTab)
-                    ->SetForegroundTab(FPersonaTabs::DetailsID)
+					FTabManager::NewSplitter()
+					->SetSizeCoefficient(0.2f)
+                    ->SetOrientation(Orient_Vertical)
+                    ->Split
+                    (
+						FTabManager::NewStack()
+						->SetSizeCoefficient(0.6f)
+						->AddTab(FPersonaTabs::DetailsID, ETabState::OpenedTab)
+						->AddTab(FPersonaTabs::AdvancedPreviewSceneSettingsID, ETabState::OpenedTab)
+						->SetForegroundTab(FPersonaTabs::DetailsID)
+                    )
+                    ->Split
+                    (
+						FTabManager::NewStack()
+						->SetSizeCoefficient(0.6f)
+						->AddTab(FIKRigRetargetChainTabSummoner::TabID, ETabState::OpenedTab)
+						->SetForegroundTab(FIKRigRetargetChainTabSummoner::TabID)
+                    )
 				)
 			)
 		);
diff --git a/Engine/Plugins/Animation/IKRig/Source/IKRigEditor/Private/RigEditor/IKRigRetargetChainTabSummoner.cpp b/Engine/Plugins/Animation/IKRig/Source/IKRigEditor/Private/RigEditor/IKRigRetargetChainTabSummoner.cpp
new file mode 100644
--- /dev/null
+++ b/Engine/Plugins/Animation/IKRig/Source/IKRigEditor/Private/RigEditor/IKRigRetargetChainTabSummoner.cpp
@@ -0,0 +1,37 @@
+// Copyright Epic Games, Inc. All Rights Reserved.
+
+#include "RigEditor/IKRigRetargetChainTabSummoner.h"
+
+#include "IDocumentation.h"
+#include "RigEditor/IKRigEditorStyle.h"
+#include "RigEditor/IKRigToolkit.h"
+#include "RigEditor/SIKRigSolverStack.h"
+
+#define LOCTEXT_NAMESPACE "RetargetChainTabSummoner"
+
+const FName FIKRigRetargetChainTabSummoner::TabID(TEXT("RetargetingChains"));
+
+FIKRigRetargetChainTabSummoner::FIKRigRetargetChainTabSummoner(const TSharedRef<FIKRigEditorToolkit>& InIKRigEditor)
+	: FWorkflowTabFactory(TabID, InIKRigEditor)
+	, IKRigEditor(InIKRigEditor)
+{
+	bIsSingleton = true; // only allow a single instance of this tab
+	
+	TabLabel = LOCTEXT("IKRigRetargetChainTabLabel", "IK Retargeting");
+	TabIcon = FSlateIcon(FIKRigEditorStyle::Get().GetStyleSetName(), "IKRig.TabIcon");
+
+	ViewMenuDescription = LOCTEXT("IKRigRetargetChain_ViewMenu_Desc", "IK Rig Retarget Chains");
+	ViewMenuTooltip = LOCTEXT("IKRigRetargetChain_ViewMenu_ToolTip", "Show the IK Rig Retarget Chains Tab");
+}
+
+TSharedPtr<SToolTip> FIKRigRetargetChainTabSummoner::CreateTabToolTipWidget(const FWorkflowTabSpawnInfo& Info) const
+{
+	return  IDocumentation::Get()->CreateToolTip(LOCTEXT("IKRigRetargetChainTooltip", "A list of labeled bone chains for retargeting animation between different skeletons. Used by the IK Retargeter asset."), NULL, TEXT("Shared/Editors/Persona"), TEXT("IKRigRetargetChain_Window"));
+}
+
+TSharedRef<SWidget> FIKRigRetargetChainTabSummoner::CreateTabBody(const FWorkflowTabSpawnInfo& Info) const
+{
+	return SNew(SIKRigRetargetChainList, IKRigEditor.Pin()->GetController());
+}
+
+#undef LOCTEXT_NAMESPACE 
diff --git a/Engine/Plugins/Animation/IKRig/Source/IKRigEditor/Private/RigEditor/IKRigSkeletonCommands.cpp b/Engine/Plugins/Animation/IKRig/Source/IKRigEditor/Private/RigEditor/IKRigSkeletonCommands.cpp
--- a/Engine/Plugins/Animation/IKRig/Source/IKRigEditor/Private/RigEditor/IKRigSkeletonCommands.cpp
+++ b/Engine/Plugins/Animation/IKRig/Source/IKRigEditor/Private/RigEditor/IKRigSkeletonCommands.cpp
@@ -13,6 +13,10 @@ void FIKRigSkeletonCommands::RegisterCommands()
 	UI_COMMAND(SetRootBoneOnSolvers, "Set Root Bone on Selected Solvers", "Set the Root Bone setting on the selected solvers.", EUserInterfaceActionType::Button, FInputChord());
 	UI_COMMAND(AddBoneSettings, "Add Settings to Selected Bone", "Apply settings to the selected bone for all selected solvers (defined per-solver for limits, stiffness etc).", EUserInterfaceActionType::Button, FInputChord());
 	UI_COMMAND(RemoveBoneSettings, "Remove Settings on Selected Bone", "Remove all settings on the selected bone in all selected solvers.", EUserInterfaceActionType::Button, FInputChord());
+	UI_COMMAND(ExcludeBone, "Exclude Selected Bone From Solve", "Ignore bone in all solvers.", EUserInterfaceActionType::Button, FInputChord());
+	UI_COMMAND(IncludeBone, "Include Selected Bone In Solve", "Include bone in all solvers.", EUserInterfaceActionType::Button, FInputChord());
+	UI_COMMAND(NewRetargetChain, "New Retarget Chain from Selected Bones", "Create a new retarget bone chain from the selected bones.", EUserInterfaceActionType::Button, FInputChord());
+	UI_COMMAND(SetRetargetRoot, "Set Retargeting Root", "Set the Root Bone used for retargeting. Usually 'Pelvis'.", EUserInterfaceActionType::Button, FInputChord());
 }
 
 #undef LOCTEXT_NAMESPACE
diff --git a/Engine/Plugins/Animation/IKRig/Source/IKRigEditor/Private/RigEditor/IKRigToolkit.cpp b/Engine/Plugins/Animation/IKRig/Source/IKRigEditor/Private/RigEditor/IKRigToolkit.cpp
--- a/Engine/Plugins/Animation/IKRig/Source/IKRigEditor/Private/RigEditor/IKRigToolkit.cpp
+++ b/Engine/Plugins/Animation/IKRig/Source/IKRigEditor/Private/RigEditor/IKRigToolkit.cpp
@@ -87,7 +87,7 @@ void FIKRigEditorToolkit::InitAssetEditor(
 void FIKRigEditorToolkit::RegisterTabSpawners(const TSharedRef<FTabManager>& InTabManager)
 {
 	WorkspaceMenuCategory = InTabManager->AddLocalWorkspaceMenuCategory(LOCTEXT("WorkspaceMenu_IKRigEditor", "IK Rig Editor"));
-	auto WorkspaceMenuCategoryRef = WorkspaceMenuCategory.ToSharedRef();
+	TSharedRef<FWorkspaceItem> WorkspaceMenuCategoryRef = WorkspaceMenuCategory.ToSharedRef();
 
 	FAssetEditorToolkit::RegisterTabSpawners(InTabManager);
 }
diff --git a/Engine/Plugins/Animation/IKRig/Source/IKRigEditor/Private/RigEditor/SIKRigRetargetChainList.cpp b/Engine/Plugins/Animation/IKRig/Source/IKRigEditor/Private/RigEditor/SIKRigRetargetChainList.cpp
new file mode 100644
--- /dev/null
+++ b/Engine/Plugins/Animation/IKRig/Source/IKRigEditor/Private/RigEditor/SIKRigRetargetChainList.cpp
@@ -0,0 +1,372 @@
+// Copyright Epic Games, Inc. All Rights Reserved.
+
+#include "RigEditor/SIKRigRetargetChainList.h"
+
+#include "RigEditor/IKRigEditorStyle.h"
+#include "RigEditor/IKRigToolkit.h"
+#include "RigEditor/IKRigEditorController.h"
+
+#include "Framework/Commands/UICommandList.h"
+#include "Framework/MultiBox/MultiBoxBuilder.h"
+#include "IPersonaToolkit.h"
+#include "SKismetInspector.h"
+#include "Dialogs/Dialogs.h"
+#include "Widgets/Text/SInlineEditableTextBlock.h"
+#include "SEditorHeaderButton.h"
+
+#define LOCTEXT_NAMESPACE "SIKRigRetargetChains"
+
+static const FName ColumnId_ChainNameLabel( "Chain Name" );
+static const FName ColumnId_ChainRootLabel( "Root Bone" );
+static const FName ColumnId_ChainTipLabel( "Tip Bone" );
+static const FName ColumnId_IKGoalLabel( "IK Goal" );
+
+TSharedRef<ITableRow> FRetargetChainElement::MakeListRowWidget(
+	const TSharedRef<STableViewBase>& InOwnerTable,
+	TSharedRef<FRetargetChainElement> InChainElement,
+	TSharedPtr<SIKRigRetargetChainList> InChainList)
+{
+	return SNew(SIKRigRetargetChainRow, InOwnerTable, InChainElement, InChainList);
+}
+
+void SIKRigRetargetChainRow::Construct(
+	const FArguments& InArgs,
+	const TSharedRef<STableViewBase>& InOwnerTableView,
+	TSharedRef<FRetargetChainElement> InChainElement,
+	TSharedPtr<SIKRigRetargetChainList> InChainList)
+{
+	ChainElement = InChainElement;
+	ChainList = InChainList;
+
+	// generate list of goals
+	// NOTE: cannot just use literal "None" because Combobox considers that a "null" entry and will discard it from the list.
+	GoalOptions.Add(MakeShareable(new FString("None")));
+	for (UIKRigEffectorGoal* Goal : InChainList->EditorController.Pin()->AssetController->GetAsset()->Goals)
+	{
+		GoalOptions.Add(MakeShareable(new FString(Goal->GoalName.ToString())));
+	}
+
+	SMultiColumnTableRow< TSharedPtr<FRetargetChainElement> >::Construct( FSuperRowType::FArguments(), InOwnerTableView );
+}
+
+TSharedRef<SWidget> SIKRigRetargetChainRow::GenerateWidgetForColumn(const FName& ColumnName)
+{
+	if (ColumnName == ColumnId_ChainNameLabel)
+	{
+		TSharedRef<SWidget> ChainWidget =
+		SNew(SHorizontalBox)
+		+ SHorizontalBox::Slot()
+		.AutoWidth()
+		.HAlign(HAlign_Left)
+		.VAlign(VAlign_Center)
+		.Padding(3.0f, 1.0f)
+		[
+			SNew(SEditableTextBox)
+			.Text(FText::FromName(ChainElement.Pin()->Chain->ChainName))
+			.Font(FEditorStyle::GetFontStyle(TEXT("BoldFont")))
+			.OnTextCommitted(this, &SIKRigRetargetChainRow::OnRenameChain)
+		];
+		return ChainWidget;
+	}
+
+	if (ColumnName == ColumnId_ChainRootLabel)
+	{
+		TSharedRef<SWidget> RootWidget =
+		SNew(SHorizontalBox)
+		+ SHorizontalBox::Slot()
+		.AutoWidth()
+		.HAlign(HAlign_Left)
+		.VAlign(VAlign_Center)
+		.Padding(3.0f, 1.0f)
+		[
+			SNew(SComboBox<FName>)
+			.OptionsSource(&ChainList.Pin()->EditorController.Pin()->AssetController->GetSkeleton().BoneNames)
+			.OnGenerateWidget(this, &SIKRigRetargetChainRow::MakeBoneComboEntryWidget)
+			.OnSelectionChanged(this, &SIKRigRetargetChainRow::OnStartBoneComboSelectionChanged)
+			[
+				SNew(STextBlock)
+				.Text(this, &SIKRigRetargetChainRow::GetStartBoneName)
+			]
+		];
+		return RootWidget;
+	}
+
+	if (ColumnName == ColumnId_ChainTipLabel)
+	{
+		TSharedRef<SWidget> TipWidget =
+		SNew(SHorizontalBox)
+		+ SHorizontalBox::Slot()
+		.AutoWidth()
+		.HAlign(HAlign_Left)
+		.VAlign(VAlign_Center)
+		.Padding(3.0f, 1.0f)
+		[
+			SNew(SComboBox<FName>)
+			.OptionsSource(&ChainList.Pin()->EditorController.Pin()->AssetController->GetSkeleton().BoneNames)
+			.OnGenerateWidget(this, &SIKRigRetargetChainRow::MakeBoneComboEntryWidget)
+			.OnSelectionChanged(this, &SIKRigRetargetChainRow::OnEndBoneComboSelectionChanged)
+			[
+				SNew(STextBlock)
+				.Text(this, &SIKRigRetargetChainRow::GetEndBoneName)
+			]
+		];
+		return TipWidget;
+	}
+	else
+	{
+		// ColumnId_IKGoalLabel
+		
+		TSharedRef<SWidget> GoalWidget =
+		SNew(SHorizontalBox)
+		+ SHorizontalBox::Slot()
+		.AutoWidth()
+		.HAlign(HAlign_Left)
+		.VAlign(VAlign_Center)
+		.Padding(3.0f, 1.0f)
+		[
+			SNew(SComboBox<TSharedPtr<FString>>)
+			.OptionsSource(&GoalOptions)
+			.OnGenerateWidget(this, &SIKRigRetargetChainRow::MakeGoalComboEntryWidget)
+			.OnSelectionChanged(this, &SIKRigRetargetChainRow::OnGoalComboSelectionChanged)
+			[
+				SNew(STextBlock)
+				.Text(this, &SIKRigRetargetChainRow::GetGoalName)
+			]
+		];
+		return GoalWidget;
+	}
+}
+
+TSharedRef<SWidget> SIKRigRetargetChainRow::MakeBoneComboEntryWidget(FName InItem) const
+{
+	return SNew(STextBlock).Text(FText::FromName(InItem));
+}
+
+TSharedRef<SWidget> SIKRigRetargetChainRow::MakeGoalComboEntryWidget(TSharedPtr<FString> InItem) const
+{
+	return SNew(STextBlock).Text(FText::FromString(*InItem.Get()));
+}
+
+void SIKRigRetargetChainRow::OnStartBoneComboSelectionChanged(FName InName, ESelectInfo::Type SelectInfo)
+{
+	const TSharedPtr<FIKRigEditorController> Controller = ChainList.Pin()->EditorController.Pin();
+	if (!Controller.IsValid())
+	{
+		return; 
+	}
+	
+	Controller->AssetController->SetRetargetChainStartBone(ChainElement.Pin()->Chain->ChainName, InName);
+	ChainList.Pin()->RefreshView();
+}
+
+FText SIKRigRetargetChainRow::GetStartBoneName() const
+{
+	return FText::FromName(ChainElement.Pin()->Chain->StartBone);
+}
+
+void SIKRigRetargetChainRow::OnEndBoneComboSelectionChanged(FName InName, ESelectInfo::Type SelectInfo)
+{
+	const TSharedPtr<FIKRigEditorController> Controller = ChainList.Pin()->EditorController.Pin();
+	if (!Controller.IsValid())
+	{
+		return; 
+	}
+	
+	Controller->AssetController->SetRetargetChainEndBone(ChainElement.Pin()->Chain->ChainName, InName);
+	ChainList.Pin()->RefreshView();
+}
+
+void SIKRigRetargetChainRow::OnGoalComboSelectionChanged(TSharedPtr<FString> InGoalName, ESelectInfo::Type SelectInfo)
+{
+	ChainElement.Pin()->Chain->IKGoalName = FName(*InGoalName.Get());
+	ChainList.Pin()->RefreshView();
+}
+
+FText SIKRigRetargetChainRow::GetEndBoneName() const
+{
+	return FText::FromName(ChainElement.Pin()->Chain->EndBone);
+}
+
+FText SIKRigRetargetChainRow::GetGoalName() const
+{
+	return FText::FromName(ChainElement.Pin()->Chain->IKGoalName);
+}
+
+void SIKRigRetargetChainRow::OnRenameChain(const FText& InText, ETextCommit::Type) const
+{
+	ChainElement.Pin()->Chain->ChainName = FName(*InText.ToString());
+	ChainList.Pin()->RefreshView();
+}
+
+SIKRigRetargetChainList::~SIKRigRetargetChainList()
+{
+	
+}
+
+void SIKRigRetargetChainList::Construct(const FArguments& InArgs, TSharedRef<FIKRigEditorController> InEditorController)
+{
+	EditorController = InEditorController;
+	EditorController.Pin()->RetargetingView = SharedThis(this);
+	
+	CommandList = MakeShared<FUICommandList>();
+
+	ChildSlot
+    [
+        SNew(SVerticalBox)
+        +SVerticalBox::Slot()
+        .AutoHeight()
+        .VAlign(VAlign_Top)
+        [
+			SNew(SHorizontalBox)
+			+ SHorizontalBox::Slot()
+			.AutoWidth()
+			.HAlign(HAlign_Left)
+			.VAlign(VAlign_Center)
+			.Padding(3.0f, 0.0f)
+			[
+				SNew(STextBlock)
+				.Text(LOCTEXT("RetargetRootLabel", "Retarget Root:"))
+				.TextStyle(FEditorStyle::Get(), "NormalText")
+			]
+				
+			+ SHorizontalBox::Slot()
+			.AutoWidth()
+			.HAlign(HAlign_Left)
+			.VAlign(VAlign_Center)
+			.Padding(3.0f, 0.0f)
+			[
+				SAssignNew(RetargetRootTextBox, SEditableTextBox)
+				.Text(FText::FromName(InEditorController->AssetController->GetRetargetRoot()))
+				.Font(FEditorStyle::GetFontStyle(TEXT("BoldFont")))
+				.IsReadOnly(true)
+			]
+        ]
+
+        +SVerticalBox::Slot()
+		[
+			SAssignNew(ListView, SRetargetChainListViewType )
+			.SelectionMode(ESelectionMode::Single)
+			.IsEnabled(this, &SIKRigRetargetChainList::IsAddChainEnabled)
+			.ListItemsSource( &ListViewItems )
+			.OnGenerateRow( this, &SIKRigRetargetChainList::MakeListRowWidget )
+			.OnMouseButtonClick(this, &SIKRigRetargetChainList::OnItemClicked)
+			.ItemHeight( 22.0f )
+			.HeaderRow
+			(
+				SNew( SHeaderRow )
+				+ SHeaderRow::Column( ColumnId_ChainNameLabel )
+				.DefaultLabel( LOCTEXT( "ChainNameColumnLabel", "Chain Name" ) )
+
+				+ SHeaderRow::Column( ColumnId_ChainRootLabel )
+				.DefaultLabel( LOCTEXT( "ChainRootColumnLabel", "Root Bone" ) )
+
+				+ SHeaderRow::Column( ColumnId_ChainTipLabel )
+				.DefaultLabel( LOCTEXT( "ChainTipColumnLabel", "Tip Bone" ) )
+
+				+ SHeaderRow::Column( ColumnId_IKGoalLabel )
+				.DefaultLabel( LOCTEXT( "IKGoalColumnLabel", "IK Goal" ) )
+			)
+		]
+    ];
+
+	RefreshView();
+}
+
+bool SIKRigRetargetChainList::IsAddChainEnabled() const
+{
+	const TSharedPtr<FIKRigEditorController> Controller = EditorController.Pin();
+	if (!Controller.IsValid())
+	{
+		return false;
+	}
+	
+	if (UIKRigController* AssetController = Controller->AssetController)
+	{
+		if (AssetController->GetSkeleton().BoneNames.Num() > 0)
+		{
+			return true;
+		}
+	}
+	
+	return false;
+}
+
+void SIKRigRetargetChainList::RefreshView()
+{
+	const TSharedPtr<FIKRigEditorController> Controller = EditorController.Pin();
+	if (!Controller.IsValid())
+	{
+		return;
+	}
+	
+	// refresh retarget root
+	RetargetRootTextBox.Get()->SetText(FText::FromName(Controller->AssetController->GetRetargetRoot()));
+
+	// refresh list of chains
+	ListViewItems.Reset();
+	UIKRigController* AssetController = Controller->AssetController;
+	const int32 NumChains = AssetController->GetNumRetargetChains();
+	for (int32 i=0; i<NumChains; ++i)
+	{
+		FBoneChain* Chain = AssetController->GetRetargetChain(i);
+		const FText DisplayName = FText::FromName(Chain->ChainName);
+		TSharedPtr<FRetargetChainElement> ChainItem = FRetargetChainElement::Make(Chain);
+		ListViewItems.Add(ChainItem);
+	}
+
+	// select first item if none others selected
+	if (ListViewItems.Num() > 0 && ListView->GetNumItemsSelected() == 0)
+	{
+		ListView->SetSelection(ListViewItems[0]);
+	}
+
+	ListView->RequestListRefresh();
+}
+
+TSharedRef<ITableRow> SIKRigRetargetChainList::MakeListRowWidget(
+	TSharedPtr<FRetargetChainElement> InElement,
+    const TSharedRef<STableViewBase>& OwnerTable)
+{
+	return InElement->MakeListRowWidget(
+		OwnerTable,
+		InElement.ToSharedRef(),
+		SharedThis(this));
+}
+
+void SIKRigRetargetChainList::OnItemClicked(TSharedPtr<FRetargetChainElement> InItem)
+{
+	// TODO highlight chain in skeleton view
+}
+
+FReply SIKRigRetargetChainList::OnKeyDown(const FGeometry& MyGeometry, const FKeyEvent& InKeyEvent)
+{
+	const FKey Key = InKeyEvent.GetKey();
+
+	// handle deleting selected chain
+	if (Key == EKeys::Delete)
+	{
+		TArray<TSharedPtr<FRetargetChainElement>> SelectedItems = ListView->GetSelectedItems();
+		if (SelectedItems.IsEmpty())
+		{
+			return FReply::Unhandled();
+		}
+		
+		const TSharedPtr<FIKRigEditorController> Controller = EditorController.Pin();
+		if (!Controller.IsValid())
+		{
+			return FReply::Unhandled();
+		}
+
+		UIKRigController* AssetController = Controller->AssetController;
+		const FName ChainToRemove = SelectedItems[0]->Chain->ChainName;
+		AssetController->RemoveRetargetChain(ChainToRemove);
+
+		RefreshView();
+		
+		return FReply::Handled();
+	}
+	
+	return FReply::Unhandled();
+}
+
+#undef LOCTEXT_NAMESPACE
diff --git a/Engine/Plugins/Animation/IKRig/Source/IKRigEditor/Private/RigEditor/SIKRigSkeleton.cpp b/Engine/Plugins/Animation/IKRig/Source/IKRigEditor/Private/RigEditor/SIKRigSkeleton.cpp
--- a/Engine/Plugins/Animation/IKRig/Source/IKRigEditor/Private/RigEditor/SIKRigSkeleton.cpp
+++ b/Engine/Plugins/Animation/IKRig/Source/IKRigEditor/Private/RigEditor/SIKRigSkeleton.cpp
@@ -53,32 +53,43 @@ void SIKRigSkeletonItem::Construct(
 
 	// is this element affected by the selected solver?
 	bool bIsConnectedToSelectedSolver;
-	const int32 SelectedSolver = EditorController.Pin()->GetSelectedSolverIndex();
+	const int32 SelectedSolver = InEditorController->GetSelectedSolverIndex();
 	if (SelectedSolver == INDEX_NONE)
 	{
-		bIsConnectedToSelectedSolver = EditorController.Pin()->IsElementConnectedToAnySolver(InRigTreeElement);
-	}else
+		bIsConnectedToSelectedSolver = InEditorController->IsElementConnectedToAnySolver(InRigTreeElement);
+	}
+	else
 	{
-		bIsConnectedToSelectedSolver = EditorController.Pin()->IsElementConnectedToSolver(InRigTreeElement, SelectedSolver);
+		bIsConnectedToSelectedSolver = InEditorController->IsElementConnectedToSolver(InRigTreeElement, SelectedSolver);
 	}
-	
+
 	// determine text style
-	FSlateFontInfo TextStyle = FEditorStyle::GetWidgetStyle<FTextBlockStyle>("SkeletonTree.NormalFont").Font;
-	// highlight elements connected to the selected solver
+	FTextBlockStyle NormalText = FIKRigEditorStyle::Get().GetWidgetStyle<FTextBlockStyle>("IKRig.Tree.NormalText");
+	FTextBlockStyle ItalicText = FIKRigEditorStyle::Get().GetWidgetStyle<FTextBlockStyle>("IKRig.Tree.ItalicText");
+	FSlateFontInfo TextFont;
+	FSlateColor TextColor;
 	if (bIsConnectedToSelectedSolver)
 	{
-		TextStyle = FIKRigEditorStyle::Get().GetWidgetStyle<FTextBlockStyle>("IKRig.Tree.AffectedBoneText").Font;
+		// elements connected to the selected solver are green
+		TextFont = ItalicText.Font;
+		TextColor = NormalText.ColorAndOpacity;
 	}
-		
+	else
+	{
+		TextFont = NormalText.Font;
+		TextColor = FLinearColor(0.2f, 0.2f, 0.2f, 0.5f);
+	}
+
 	// determine which icon to use for tree element
 	const FSlateBrush* Brush = FAppStyle::Get().GetBrush("SkeletonTree.Bone");
 	switch(InRigTreeElement->ElementType)
 	{
 		case IKRigTreeElementType::BONE:
-			if (bIsConnectedToSelectedSolver)
+			if (!InEditorController->IsElementExcludedBone(InRigTreeElement))
 			{
 				Brush = FAppStyle::Get().GetBrush("SkeletonTree.Bone");
-			}else
+			}
+			else
 			{
 				Brush = FAppStyle::Get().GetBrush("SkeletonTree.BoneNonWeighted");
 			}
@@ -125,9 +136,24 @@ void SIKRigSkeletonItem::Construct(
         [
 			SNew(STextBlock)
 			.Text(this, &SIKRigSkeletonItem::GetName)
-			.Font(TextStyle)
+			.Font(TextFont)
+			.ColorAndOpacity(TextColor)
         ];
-	}else
+
+		if (InEditorController->AssetController->GetRetargetRoot() == InRigTreeElement->Key)
+		{	
+			HorizontalBox->AddSlot()
+			.AutoWidth()
+			.HAlign(HAlign_Left)
+			[
+				SNew(STextBlock)
+				.Text(LOCTEXT("RetargetRootLabel", " (Retarget Root)"))
+				.Font(ItalicText.Font)
+				.ColorAndOpacity(FLinearColor(0.9f, 0.9f, 0.9f, 0.5f))
+			];
+		}
+	}
+	else
 	{
 		TSharedPtr< SInlineEditableTextBlock > InlineWidget;
 		HorizontalBox->AddSlot()
@@ -151,16 +177,22 @@ void SIKRigSkeletonItem::OnNameCommitted(const FText& InText, ETextCommit::Type
 	{
 		return; // make sure user actually intends to commit a name change
 	}
+
+	const TSharedPtr<FIKRigEditorController> Controller = EditorController.Pin();
+	if (!Controller.IsValid())
+	{
+		return; 
+	}
 	
 	const FName OldName = WeakRigTreeElement.Pin()->Key;
 	const FName PotentialNewName = FName(InText.ToString());
-	const FName NewName = EditorController.Pin()->AssetController->RenameGoal(OldName, PotentialNewName);
+	const FName NewName = Controller->AssetController->RenameGoal(OldName, PotentialNewName);
 	if (NewName != NAME_None)
 	{
 		WeakRigTreeElement.Pin()->Key = NewName;
 	}
 
-	EditorController.Pin()->SkeletonView->RefreshTreeView();
+	Controller->SkeletonView->RefreshTreeView();
 }
 
 FText SIKRigSkeletonItem::GetName() const
@@ -224,7 +256,7 @@ void SIKRigSkeleton::Construct(const FArguments& InArgs, TSharedRef<FIKRigEditor
                         .Icon(FAppStyle::Get().GetBrush("Icons.Plus"))
                         .Text(LOCTEXT("ImportSkeletonLabel", "Import Skeleton"))
                         .ToolTipText(LOCTEXT("ImportSkeletonToolTip", "Import a skeletal mesh to apply IK to."))
-                        .OnClicked(FOnClicked::CreateSP(this, &SIKRigSkeleton::OnImportSkeletonClicked))
+                        .OnClicked(this, &SIKRigSkeleton::OnImportSkeletonClicked)
                     ]
                 ]
             ]
@@ -264,7 +296,7 @@ void SIKRigSkeleton::SetSelectedGoalsFromViewport(const TArray<FName>& GoalNames
 		return;
 	}
 	
-	for (const auto& Item : AllElements)
+	for (const TSharedPtr<FIKRigTreeElement>& Item : AllElements)
 	{
 		if (GoalNames.Contains(Item->Key))
 		{
@@ -273,6 +305,32 @@ void SIKRigSkeleton::SetSelectedGoalsFromViewport(const TArray<FName>& GoalNames
 	}
 }
 
+void SIKRigSkeleton::GetSelectedBoneChains(TArray<FIKRigSkeletonChain>& OutChains)
+{
+	const TSharedPtr<FIKRigEditorController> Controller = EditorController.Pin();
+	if (!Controller.IsValid())
+	{
+		return; 
+	}
+
+	// get selected bones
+	TArray<TSharedPtr<FIKRigTreeElement>> SelectedBoneItems;
+	GetSelectedBones(SelectedBoneItems);
+
+	FIKRigSkeleton& Skeleton = Controller->AssetController->GetSkeleton();
+
+	// get selected bone indices
+	TArray<int32> SelectedBones;
+	for (const TSharedPtr<FIKRigTreeElement>& BoneItem : SelectedBoneItems)
+	{
+		const FName BoneName = BoneItem.Get()->Key;
+		const int32 BoneIndex = Skeleton.GetBoneIndexFromName(BoneName);
+		SelectedBones.Add(BoneIndex);
+	}
+	
+	return Skeleton.GetChainsInList(SelectedBones, OutChains);
+}
+
 void SIKRigSkeleton::BindCommands()
 {
 	const FIKRigSkeletonCommands& Commands = FIKRigSkeletonCommands::Get();
@@ -304,6 +362,22 @@ void SIKRigSkeleton::BindCommands()
 	CommandList->MapAction(Commands.RemoveBoneSettings,
         FExecuteAction::CreateSP(this, &SIKRigSkeleton::HandleRemoveBoneSettings),
         FCanExecuteAction::CreateSP(this, &SIKRigSkeleton::CanRemoveBoneSettings));
+
+	CommandList->MapAction(Commands.ExcludeBone,
+		FExecuteAction::CreateSP(this, &SIKRigSkeleton::HandleExcludeBone),
+		FCanExecuteAction::CreateSP(this, &SIKRigSkeleton::CanExcludeBone));
+
+	CommandList->MapAction(Commands.IncludeBone,
+		FExecuteAction::CreateSP(this, &SIKRigSkeleton::HandleIncludeBone),
+		FCanExecuteAction::CreateSP(this, &SIKRigSkeleton::CanIncludeBone));
+
+	CommandList->MapAction(Commands.NewRetargetChain,
+		FExecuteAction::CreateSP(this, &SIKRigSkeleton::HandleNewRetargetChain),
+		FCanExecuteAction::CreateSP(this, &SIKRigSkeleton::CanAddNewRetargetChain));
+
+	CommandList->MapAction(Commands.SetRetargetRoot,
+		FExecuteAction::CreateSP(this, &SIKRigSkeleton::HandleSetRetargetRoot),
+		FCanExecuteAction::CreateSP(this, &SIKRigSkeleton::CanSetRetargetRoot));
 }
 
 void SIKRigSkeleton::FillContextMenu(FMenuBuilder& MenuBuilder)
@@ -316,11 +390,6 @@ void SIKRigSkeleton::FillContextMenu(FMenuBuilder& MenuBuilder)
 		return;
 	}
 
-	if (SelectedItems.Num() > 1)
-	{
-		return; // TODO create limb definitions from multiple selected bones
-	}
-
 	MenuBuilder.BeginSection("AddRemoveGoals", LOCTEXT("AddRemoveGoalOperations", "Goals"));
 	MenuBuilder.AddMenuEntry(Actions.NewGoal);
 	MenuBuilder.AddMenuEntry(Actions.Delete);
@@ -334,53 +403,69 @@ void SIKRigSkeleton::FillContextMenu(FMenuBuilder& MenuBuilder)
 	MenuBuilder.BeginSection("BoneSettings", LOCTEXT("BoneSettingsOperations", "Bone Settings"));
 	MenuBuilder.AddMenuEntry(Actions.AddBoneSettings);
 	MenuBuilder.AddMenuEntry(Actions.RemoveBoneSettings);
+	MenuBuilder.AddMenuEntry(Actions.SetRootBoneOnSolvers);
 	MenuBuilder.EndSection();
 
-	MenuBuilder.BeginSection("RootBone", LOCTEXT("RootBoneOperations", "Root Bone"));
-	MenuBuilder.AddMenuEntry(Actions.SetRootBoneOnSolvers);
+	MenuBuilder.BeginSection("IncludeExclude", LOCTEXT("IncludeExcludeOperations", "Exclude Bones"));
+	MenuBuilder.AddMenuEntry(Actions.ExcludeBone);
+	MenuBuilder.AddMenuEntry(Actions.IncludeBone);
+	MenuBuilder.EndSection();
+
+	MenuBuilder.BeginSection("Retargeting", LOCTEXT("RetargetingOperations", "Retargeting"));
+	MenuBuilder.AddMenuEntry(Actions.NewRetargetChain);
+	MenuBuilder.AddMenuEntry(Actions.SetRetargetRoot);
 	MenuBuilder.EndSection();
 }
 
 void SIKRigSkeleton::HandleNewGoal()
 {
-	const TArray<TSharedPtr<FIKRigTreeElement>> SelectedItems = TreeView.Get()->GetSelectedItems();
-	if (SelectedItems.IsEmpty())
+	const TSharedPtr<FIKRigEditorController> Controller = EditorController.Pin();
+	if (!Controller.IsValid())
 	{
-		return; // have to select something
+		return; 
 	}
-
-	TSharedPtr<FIKRigTreeElement> SelectedItem = SelectedItems[0];
-	if (SelectedItem->ElementType != IKRigTreeElementType::BONE)
+	
+	// create goal for each selected bone
+	FName LastCreatedGoalName = NAME_None;
+	const TArray<TSharedPtr<FIKRigTreeElement>> SelectedItems = TreeView.Get()->GetSelectedItems();
+	for (const TSharedPtr<FIKRigTreeElement>& Item : SelectedItems)
 	{
-		return; // can only add goals to bones
-	}
-
-	UIKRigController* Controller = EditorController.Pin()->AssetController;
+		if (Item->ElementType != IKRigTreeElementType::BONE)
+		{
+			continue; // can only add goals to bones
+		}
 
-	// build default name for the new goal
-	const FName BoneName = SelectedItem->Key;
-	const FName NewGoalName = FName(BoneName.ToString() + "_Goal");
+		UIKRigController* AssetController = Controller->AssetController;
 
-	// create a new goal
-	UIKRigEffectorGoal* NewGoal = Controller->AddNewGoal(NewGoalName, BoneName);
+		// build default name for the new goal
+		const FName BoneName = Item->Key;
+		const FName NewGoalName = FName(BoneName.ToString() + "_Goal");
 
-	if (NewGoal)
-	{
-		// get selected solvers
-		TArray<TSharedPtr<FSolverStackElement>> SelectedSolvers;
-		EditorController.Pin()->GetSelectedSolvers(SelectedSolvers);
-		// connect the new goal to all the selected solvers
-		for (const TSharedPtr<FSolverStackElement>& SolverElement : SelectedSolvers)
+		// create a new goal
+		UIKRigEffectorGoal* NewGoal = AssetController->AddNewGoal(NewGoalName, BoneName);
+		if (NewGoal)
 		{
-			Controller->ConnectGoalToSolver(*NewGoal, SolverElement->IndexInStack);	
+			// get selected solvers
+			TArray<TSharedPtr<FSolverStackElement>> SelectedSolvers;
+			Controller->GetSelectedSolvers(SelectedSolvers);
+			// connect the new goal to all the selected solvers
+			for (const TSharedPtr<FSolverStackElement>& SolverElement : SelectedSolvers)
+			{
+				AssetController->ConnectGoalToSolver(*NewGoal, SolverElement->IndexInStack);	
+			}
+
+			LastCreatedGoalName = NewGoalName;
 		}
 	}
-	
-	// create new goal tree element and update view
-	RefreshTreeView();
 
-	// show goal in details view
-	EditorController.Pin()->ShowDetailsForGoal(NewGoalName);
+	// were any goals created?
+	if (LastCreatedGoalName != NAME_None)
+	{
+		// show last created goal in details view
+		Controller->ShowDetailsForGoal(LastCreatedGoalName);
+		// update all views
+		Controller->RefreshAllViews();
+	}
 }
 
 bool SIKRigSkeleton::CanAddNewGoal()
@@ -393,7 +478,7 @@ bool SIKRigSkeleton::CanAddNewGoal()
 	}
 
 	// can only add goals to selected bones
-	for (const auto&Item : SelectedItems)
+	for (const TSharedPtr<FIKRigTreeElement>& Item : SelectedItems)
 	{
 		if (Item->ElementType != IKRigTreeElementType::BONE)
 		{
@@ -406,22 +491,30 @@ bool SIKRigSkeleton::CanAddNewGoal()
 
 void SIKRigSkeleton::HandleDeleteItem()
 {
+	const TSharedPtr<FIKRigEditorController> Controller = EditorController.Pin();
+	if (!Controller.IsValid())
+	{
+		return; 
+	}
+	
 	TArray<TSharedPtr<FIKRigTreeElement>> SelectedItems = TreeView->GetSelectedItems();
 	TSharedPtr<FIKRigTreeElement> ParentOfDeletedGoal;
-	for (const auto&Item : SelectedItems)
+	for (const TSharedPtr<FIKRigTreeElement>& Item : SelectedItems)
 	{
 		if (Item->ElementType == IKRigTreeElementType::GOAL)
 		{
-			EditorController.Pin()->AssetController->RemoveGoal(Item->Key);
+			Controller->AssetController->RemoveGoal(Item->Key);
 			
-		}else if (Item->ElementType == IKRigTreeElementType::EFFECTOR)
+		}
+		else if (Item->ElementType == IKRigTreeElementType::EFFECTOR)
 		{
-			EditorController.Pin()->AssetController->DisconnectGoalFromSolver(Item->EffectorGoalName, Item->EffectorSolverIndex);
+			Controller->AssetController->DisconnectGoalFromSolver(Item->EffectorGoalName, Item->EffectorSolverIndex);
 		}
 	}
 
-	EditorController.Pin()->ShowEmptyDetails();
-	RefreshTreeView();
+	Controller->ShowEmptyDetails();
+	// update all views
+	Controller->RefreshAllViews();
 }
 bool SIKRigSkeleton::CanDeleteItem() const
 {
@@ -433,7 +526,7 @@ bool SIKRigSkeleton::CanDeleteItem() const
 	}
 
 	// are all selected items goals or effectors?
-	for (const auto&Item : SelectedItems)
+	for (const TSharedPtr<FIKRigTreeElement>& Item : SelectedItems)
 	{
 		const bool bIsGoal = Item->ElementType == IKRigTreeElementType::GOAL;
 		const bool bIsEffector = Item->ElementType == IKRigTreeElementType::EFFECTOR;
@@ -474,26 +567,33 @@ bool SIKRigSkeleton::CanDisconnectGoalFromSolvers() const
 
 void SIKRigSkeleton::ConnectSelectedGoalsToSelectedSolvers(bool bConnect)
 {
+	const TSharedPtr<FIKRigEditorController> Controller = EditorController.Pin();
+	if (!Controller.IsValid())
+	{
+		return; 
+	}
+	
 	TArray<TSharedPtr<FIKRigTreeElement>> SelectedGoals;
 	GetSelectedGoals(SelectedGoals);
 	TArray<TSharedPtr<FSolverStackElement>> SelectedSolvers;
-	EditorController.Pin()->GetSelectedSolvers(SelectedSolvers);
+	Controller->GetSelectedSolvers(SelectedSolvers);
 
-	UIKRigController* Controller = EditorController.Pin()->AssetController;
+	UIKRigController* AssetController = Controller->AssetController;
 	for (const TSharedPtr<FIKRigTreeElement>& GoalElement : SelectedGoals)
 	{
 		const FName GoalName = GoalElement->Key;
-		const int32 GoalIndex = Controller->GetGoalIndex(GoalName);
+		const int32 GoalIndex = AssetController->GetGoalIndex(GoalName);
 		check(GoalIndex != INDEX_NONE);
-		const UIKRigEffectorGoal& EffectorGoal = *Controller->GetGoal(GoalIndex);
+		const UIKRigEffectorGoal& EffectorGoal = *AssetController->GetGoal(GoalIndex);
 		for (const TSharedPtr<FSolverStackElement>& SolverElement : SelectedSolvers)
 		{
 			if (bConnect)
 			{
-				Controller->ConnectGoalToSolver(EffectorGoal, SolverElement->IndexInStack);	
-			}else
+				AssetController->ConnectGoalToSolver(EffectorGoal, SolverElement->IndexInStack);	
+			}
+			else
 			{
-				Controller->DisconnectGoalFromSolver(EffectorGoal.GoalName, SolverElement->IndexInStack);	
+				AssetController->DisconnectGoalFromSolver(EffectorGoal.GoalName, SolverElement->IndexInStack);	
 			}
 		}
 	}
@@ -504,17 +604,23 @@ void SIKRigSkeleton::ConnectSelectedGoalsToSelectedSolvers(bool bConnect)
 
 int32 SIKRigSkeleton::GetNumSelectedGoalToSolverConnections(bool bCountOnlyConnected) const
 {
+	const TSharedPtr<FIKRigEditorController> Controller = EditorController.Pin();
+	if (!Controller.IsValid())
+	{
+		return 0;
+	}
+	
 	TArray<TSharedPtr<FIKRigTreeElement>> SelectedGoals;
 	GetSelectedGoals(SelectedGoals);
 	TArray<TSharedPtr<FSolverStackElement>> SelectedSolvers;
-	EditorController.Pin()->GetSelectedSolvers(SelectedSolvers);
+	Controller->GetSelectedSolvers(SelectedSolvers);
 
 	int32 NumMatched = 0;
 	for (const TSharedPtr<FIKRigTreeElement>& Goal : SelectedGoals)
 	{
 		for (const TSharedPtr<FSolverStackElement>& Solver : SelectedSolvers)
 		{
-			const bool bIsConnected = EditorController.Pin()->AssetController->IsGoalConnectedToSolver(Goal->Key, Solver->IndexInStack);
+			const bool bIsConnected = Controller->AssetController->IsGoalConnectedToSolver(Goal->Key, Solver->IndexInStack);
 			if (bIsConnected == bCountOnlyConnected)
 			{
 				++NumMatched;
@@ -527,24 +633,30 @@ int32 SIKRigSkeleton::GetNumSelectedGoalToSolverConnections(bool bCountOnlyConne
 
 void SIKRigSkeleton::HandleSetRootBoneOnSolvers()
 {
+	const TSharedPtr<FIKRigEditorController> Controller = EditorController.Pin();
+	if (!Controller.IsValid())
+	{
+		return;
+	}
+
     // get name of selected root bone
     TArray<TSharedPtr<FIKRigTreeElement>> SelectedBones;
 	GetSelectedBones(SelectedBones);
 	const FName RootBoneName = SelectedBones[0]->Key;
 
 	// apply to all selected solvers (ignored on solvers that don't accept a root bone)
-	UIKRigController* Controller = EditorController.Pin()->AssetController;
+	UIKRigController* AssetController = Controller->AssetController;
 	TArray<TSharedPtr<FSolverStackElement>> SelectedSolvers;
-	EditorController.Pin()->GetSelectedSolvers(SelectedSolvers);
+	Controller->GetSelectedSolvers(SelectedSolvers);
 	int32 SolverToShow = 0;
-	for (const auto& Solver : SelectedSolvers)
+	for (const TSharedPtr<FSolverStackElement>& Solver : SelectedSolvers)
 	{
-		Controller->SetRootBone(RootBoneName, Solver->IndexInStack);
+		AssetController->SetRootBone(RootBoneName, Solver->IndexInStack);
 		SolverToShow = Solver->IndexInStack;
 	}
 
 	// show solver that had it's root bone updated
-	EditorController.Pin()->ShowDetailsForSolver(SolverToShow);
+	Controller->ShowDetailsForSolver(SolverToShow);
 	
 	// show new icon when bone has settings applied
 	RefreshTreeView();
@@ -552,6 +664,12 @@ void SIKRigSkeleton::HandleSetRootBoneOnSolvers()
 
 bool SIKRigSkeleton::CanSetRootBoneOnSolvers()
 {
+	const TSharedPtr<FIKRigEditorController> Controller = EditorController.Pin();
+	if (!Controller.IsValid())
+	{
+		return false;
+	}
+	
 	// must have at least 1 bone selected
 	TArray<TSharedPtr<FIKRigTreeElement>> SelectedBones;
 	GetSelectedBones(SelectedBones);
@@ -561,12 +679,12 @@ bool SIKRigSkeleton::CanSetRootBoneOnSolvers()
 	}
 
 	// must have at least 1 solver selected that accepts root bones
-	UIKRigController* Controller = EditorController.Pin()->AssetController;
+	UIKRigController* AssetController = Controller->AssetController;
 	TArray<TSharedPtr<FSolverStackElement>> SelectedSolvers;
-	EditorController.Pin()->GetSelectedSolvers(SelectedSolvers);
-	for (const auto& Solver : SelectedSolvers)
+	Controller->GetSelectedSolvers(SelectedSolvers);
+	for (const TSharedPtr<FSolverStackElement>& Solver : SelectedSolvers)
 	{
-		if (Controller->GetSolver(Solver->IndexInStack)->CanSetRootBone())
+		if (AssetController->GetSolver(Solver->IndexInStack)->CanSetRootBone())
 		{
 			return true;
 		}
@@ -577,27 +695,33 @@ bool SIKRigSkeleton::CanSetRootBoneOnSolvers()
 
 void SIKRigSkeleton::HandleAddBoneSettings()
 {
+	const TSharedPtr<FIKRigEditorController> Controller = EditorController.Pin();
+	if (!Controller.IsValid())
+	{
+		return;
+	}
+	
 	// get selected bones
 	TArray<TSharedPtr<FIKRigTreeElement>> SelectedBones;
 	GetSelectedBones(SelectedBones);
 	
 	// add settings for bone on all selected solvers (ignored if already present)
-	UIKRigController* Controller = EditorController.Pin()->AssetController;
+	UIKRigController* AssetController = Controller->AssetController;
 	TArray<TSharedPtr<FSolverStackElement>> SelectedSolvers;
-	EditorController.Pin()->GetSelectedSolvers(SelectedSolvers);
+	Controller->GetSelectedSolvers(SelectedSolvers);
 	FName BoneNameForSettings;
 	int32 SolverIndex = INDEX_NONE;
 	for (const TSharedPtr<FIKRigTreeElement>& BoneItem : SelectedBones)
 	{
 		for (const TSharedPtr<FSolverStackElement>& Solver : SelectedSolvers)
 		{
-			Controller->AddBoneSetting(BoneItem->Key, Solver->IndexInStack);
+			AssetController->AddBoneSetting(BoneItem->Key, Solver->IndexInStack);
 			BoneNameForSettings = BoneItem->Key;
 			SolverIndex = Solver->IndexInStack;
         }
 	}
 
-	EditorController.Pin()->ShowDetailsForBoneSettings(BoneNameForSettings, SolverIndex);
+	Controller->ShowDetailsForBoneSettings(BoneNameForSettings, SolverIndex);
 
 	// show new icon when bone has settings applied
 	RefreshTreeView();
@@ -605,6 +729,12 @@ void SIKRigSkeleton::HandleAddBoneSettings()
 
 bool SIKRigSkeleton::CanAddBoneSettings()
 {
+	const TSharedPtr<FIKRigEditorController> Controller = EditorController.Pin();
+	if (!Controller.IsValid())
+	{
+		return false;
+	}
+	
     // must have at least 1 bone selected
     TArray<TSharedPtr<FIKRigTreeElement>> SelectedBones;
 	GetSelectedBones(SelectedBones);
@@ -614,14 +744,14 @@ bool SIKRigSkeleton::CanAddBoneSettings()
 	}
 
 	// must have at least 1 solver selected that does not already have a bone setting for the selected bones
-	UIKRigController* Controller = EditorController.Pin()->AssetController;
+	UIKRigController* AssetController = Controller->AssetController;
 	TArray<TSharedPtr<FSolverStackElement>> SelectedSolvers;
-	EditorController.Pin()->GetSelectedSolvers(SelectedSolvers);
+	Controller->GetSelectedSolvers(SelectedSolvers);
 	for (const TSharedPtr<FIKRigTreeElement>& BoneItem : SelectedBones)
 	{
 		for (const TSharedPtr<FSolverStackElement>& Solver : SelectedSolvers)
 		{
-			if (Controller->CanAddBoneSetting(BoneItem->Key, Solver->IndexInStack))
+			if (AssetController->CanAddBoneSetting(BoneItem->Key, Solver->IndexInStack))
 			{
 				return true;
 			}
@@ -633,32 +763,44 @@ bool SIKRigSkeleton::CanAddBoneSettings()
 
 void SIKRigSkeleton::HandleRemoveBoneSettings()
 {
+	const TSharedPtr<FIKRigEditorController> Controller = EditorController.Pin();
+	if (!Controller.IsValid())
+	{
+		return;
+	}
+	
 	// get selected bones
 	TArray<TSharedPtr<FIKRigTreeElement>> SelectedBones;
 	GetSelectedBones(SelectedBones);
 	
 	// add settings for bone on all selected solvers (ignored if already present)
-	UIKRigController* Controller = EditorController.Pin()->AssetController;
+	UIKRigController* AssetController = Controller->AssetController;
 	TArray<TSharedPtr<FSolverStackElement>> SelectedSolvers;
-	EditorController.Pin()->GetSelectedSolvers(SelectedSolvers);
+	Controller->GetSelectedSolvers(SelectedSolvers);
 	FName BoneToShowInDetailsView;
 	for (const TSharedPtr<FIKRigTreeElement>& BoneItem : SelectedBones)
 	{
 		for (const TSharedPtr<FSolverStackElement>& Solver : SelectedSolvers)
 		{
-			Controller->RemoveBoneSetting(BoneItem->Key, Solver->IndexInStack);
+			AssetController->RemoveBoneSetting(BoneItem->Key, Solver->IndexInStack);
 			BoneToShowInDetailsView = BoneItem->Key;
 		}
 	}
 
-	EditorController.Pin()->ShowDetailsForBone(BoneToShowInDetailsView);
-
+	Controller->ShowDetailsForBone(BoneToShowInDetailsView);
+	
 	// show new icon when bone has settings applied
 	RefreshTreeView();
 }
 
 bool SIKRigSkeleton::CanRemoveBoneSettings()
 {
+	const TSharedPtr<FIKRigEditorController> Controller = EditorController.Pin();
+	if (!Controller.IsValid())
+	{
+		return false;
+	}
+	
     // must have at least 1 bone selected
     TArray<TSharedPtr<FIKRigTreeElement>> SelectedBones;
 	GetSelectedBones(SelectedBones);
@@ -668,14 +810,14 @@ bool SIKRigSkeleton::CanRemoveBoneSettings()
 	}
 
 	// must have at least 1 solver selected that has a bone setting for 1 of the selected bones
-	UIKRigController* Controller = EditorController.Pin()->AssetController;
+	UIKRigController* AssetController = Controller->AssetController;
 	TArray<TSharedPtr<FSolverStackElement>> SelectedSolvers;
-	EditorController.Pin()->GetSelectedSolvers(SelectedSolvers);
+	Controller->GetSelectedSolvers(SelectedSolvers);
 	for (const TSharedPtr<FIKRigTreeElement>& BoneItem : SelectedBones)
 	{
 		for (const TSharedPtr<FSolverStackElement>& Solver : SelectedSolvers)
 		{
-			if (Controller->CanRemoveBoneSetting(BoneItem->Key, Solver->IndexInStack))
+			if (AssetController->CanRemoveBoneSetting(BoneItem->Key, Solver->IndexInStack))
 			{
 				return true;
 			}
@@ -685,10 +827,151 @@ bool SIKRigSkeleton::CanRemoveBoneSettings()
 	return false;
 }
 
+void SIKRigSkeleton::HandleExcludeBone()
+{
+	const TSharedPtr<FIKRigEditorController> Controller = EditorController.Pin();
+	if (!Controller.IsValid())
+	{
+		return;
+	}
+	
+	// exclude selected bones
+	TArray<TSharedPtr<FIKRigTreeElement>> SelectedBones;
+	GetSelectedBones(SelectedBones);
+	for (const TSharedPtr<FIKRigTreeElement>& BoneItem : SelectedBones)
+	{
+		Controller->AssetController->SetBoneExcluded(BoneItem->Key, true);
+	}
+
+	// show greyed out bone name after being excluded
+	RefreshTreeView();
+}
+
+bool SIKRigSkeleton::CanExcludeBone()
+{
+	const TSharedPtr<FIKRigEditorController> Controller = EditorController.Pin();
+	if (!Controller.IsValid())
+	{
+		return false;
+	}
+	
+	// must have at least 1 bone selected that is INCLUDED
+	TArray<TSharedPtr<FIKRigTreeElement>> SelectedBones;
+	GetSelectedBones(SelectedBones);
+	for (const TSharedPtr<FIKRigTreeElement>& BoneItem : SelectedBones)
+	{
+		if (!Controller->AssetController->GetBoneExcluded(BoneItem->Key))
+		{
+			return true;
+		}
+	}
+
+	return false;
+}
+
+void SIKRigSkeleton::HandleIncludeBone()
+{
+	const TSharedPtr<FIKRigEditorController> Controller = EditorController.Pin();
+	if (!Controller.IsValid())
+	{
+		return;
+	}
+	
+	// exclude selected bones
+	TArray<TSharedPtr<FIKRigTreeElement>> SelectedBones;
+	GetSelectedBones(SelectedBones);
+	for (const TSharedPtr<FIKRigTreeElement>& BoneItem : SelectedBones)
+	{
+		Controller->AssetController->SetBoneExcluded(BoneItem->Key, false);
+	}
+
+	// show normal bone name after being included
+	RefreshTreeView();
+}
+
+bool SIKRigSkeleton::CanIncludeBone()
+{
+	const TSharedPtr<FIKRigEditorController> Controller = EditorController.Pin();
+	if (!Controller.IsValid())
+	{
+		return false;
+	}
+	
+	// must have at least 1 bone selected that is EXCLUDED
+	TArray<TSharedPtr<FIKRigTreeElement>> SelectedBones;
+	GetSelectedBones(SelectedBones);
+	for (const TSharedPtr<FIKRigTreeElement>& BoneItem : SelectedBones)
+	{
+		if (Controller->AssetController->GetBoneExcluded(BoneItem->Key))
+		{
+			return true;
+		}
+	}
+
+	return false;
+}
+
+void SIKRigSkeleton::HandleNewRetargetChain()
+{
+	const TSharedPtr<FIKRigEditorController> Controller = EditorController.Pin();
+	if (!Controller.IsValid())
+	{
+		return;
+	}
+	
+	TArray<FIKRigSkeletonChain> BoneChains;
+	GetSelectedBoneChains(BoneChains);
+	for (const FIKRigSkeletonChain& BoneChain : BoneChains)
+	{
+		Controller->AddNewRetargetChain(BoneChain.StartBone, BoneChain.StartBone, BoneChain.EndBone);
+	}
+	
+	Controller->RefreshAllViews();
+}
+
+bool SIKRigSkeleton::CanAddNewRetargetChain()
+{
+	TArray<TSharedPtr<FIKRigTreeElement>> SelectedBones;
+	GetSelectedBones(SelectedBones);
+	return !SelectedBones.IsEmpty();
+}
+
+void SIKRigSkeleton::HandleSetRetargetRoot()
+{
+	const TSharedPtr<FIKRigEditorController> Controller = EditorController.Pin();
+	if (!Controller.IsValid())
+	{
+		return;
+	}
+	
+	// get selected bones
+	TArray<TSharedPtr<FIKRigTreeElement>> SelectedBones;
+	GetSelectedBones(SelectedBones);
+
+	// must have at least 1 bone selected
+	if (SelectedBones.IsEmpty())
+	{
+		return;
+	}
+
+	// set the first selected bone as the retarget root
+	Controller->AssetController->SetRetargetRoot(SelectedBones[0]->Key);
+
+	// show root bone after being set
+	Controller->RefreshAllViews();
+}
+
+bool SIKRigSkeleton::CanSetRetargetRoot()
+{
+	TArray<TSharedPtr<FIKRigTreeElement>> SelectedBones;
+	GetSelectedBones(SelectedBones);
+	return !SelectedBones.IsEmpty();
+}
+
 void SIKRigSkeleton::GetSelectedBones(TArray<TSharedPtr<FIKRigTreeElement>>& OutBoneItems)
 {
 	const TArray<TSharedPtr<FIKRigTreeElement>> SelectedItems = TreeView->GetSelectedItems();
-	for (const auto& Item : SelectedItems)
+	for (const TSharedPtr<FIKRigTreeElement>& Item : SelectedItems)
 	{
 		if (Item->ElementType == IKRigTreeElementType::BONE)
 		{
@@ -701,7 +984,7 @@ void SIKRigSkeleton::GetSelectedGoals(TArray<TSharedPtr<FIKRigTreeElement>>& Out
 {
 	OutSelectedGoals.Reset();
 	TArray<TSharedPtr<FIKRigTreeElement>> SelectedItems = TreeView->GetSelectedItems();
-	for (const auto& Item : SelectedItems)
+	for (const TSharedPtr<FIKRigTreeElement>& Item : SelectedItems)
 	{
 		if (Item->ElementType == IKRigTreeElementType::GOAL)
 		{
@@ -724,6 +1007,12 @@ void SIKRigSkeleton::HandleRenameElement() const
 
 FReply SIKRigSkeleton::OnImportSkeletonClicked()
 {
+	const TSharedPtr<FIKRigEditorController> Controller = EditorController.Pin();
+	if (!Controller.IsValid())
+	{
+		return FReply::Unhandled();
+	}
+	
 	FIKRigSkeletonImportSettings Settings;
 	const TSharedPtr<FStructOnScope> StructToDisplay = MakeShareable(new FStructOnScope(FIKRigSkeletonImportSettings::StaticStruct(), (uint8*)&Settings));
 	TSharedRef<SKismetInspector> KismetInspector = SNew(SKismetInspector);
@@ -734,7 +1023,7 @@ FReply SIKRigSkeleton::OnImportSkeletonClicked()
 	if (Settings.Mesh != nullptr)
 	{
 		ImportSkeleton(FAssetData(Settings.Mesh));
-		EditorController.Pin()->PromptToAddSolverAtStartup();
+		Controller->PromptToAddSolverAtStartup();
 	}
 	
 	return FReply::Handled();
@@ -742,9 +1031,15 @@ FReply SIKRigSkeleton::OnImportSkeletonClicked()
 
 EVisibility SIKRigSkeleton::IsImportButtonVisible() const
 {
-	if (UIKRigController* Controller = EditorController.Pin()->AssetController)
+	const TSharedPtr<FIKRigEditorController> Controller = EditorController.Pin();
+	if (!Controller.IsValid())
+	{
+		return EVisibility::Hidden;
+	}
+	
+	if (UIKRigController* AssetController = Controller->AssetController)
 	{
-		if (Controller->GetSkeleton().BoneNames.Num() > 0)
+		if (AssetController->GetSkeleton().BoneNames.Num() > 0)
 		{
 			return EVisibility::Collapsed;
 		}
@@ -755,6 +1050,12 @@ EVisibility SIKRigSkeleton::IsImportButtonVisible() const
 
 void SIKRigSkeleton::RefreshTreeView(bool IsInitialSetup)
 {
+	const TSharedPtr<FIKRigEditorController> Controller = EditorController.Pin();
+	if (!Controller.IsValid())
+	{
+		return;
+	}
+	
 	// save expansion state
 	TreeView->SaveAndClearSparseItemInfos();
 
@@ -763,8 +1064,8 @@ void SIKRigSkeleton::RefreshTreeView(bool IsInitialSetup)
 	AllElements.Reset();
 
 	// validate we have a skeleton to load
-	UIKRigController* Controller = EditorController.Pin()->AssetController;
-	FIKRigSkeleton& Skeleton = Controller->GetSkeleton();
+	UIKRigController* AssetController = Controller->AssetController;
+	FIKRigSkeleton& Skeleton = AssetController->GetSkeleton();
 	if (Skeleton.BoneNames.IsEmpty())
 	{
 		TreeView->RequestTreeRefresh();
@@ -786,7 +1087,8 @@ void SIKRigSkeleton::RefreshTreeView(bool IsInitialSetup)
 		{
 			// store root element
 			RootElements.Add(AllElements[BoneIndex]);
-		}else
+		}
+		else
 		{
 			// store pointer to child on parent
 			AllElements[ParentIndex]->Children.Add(AllElements[BoneIndex]);
@@ -796,9 +1098,9 @@ void SIKRigSkeleton::RefreshTreeView(bool IsInitialSetup)
 	}
 
 	// create all bone settings elements
-	for (int32 SolverIndex=0; SolverIndex<Controller->GetNumSolvers(); ++SolverIndex)
+	for (int32 SolverIndex=0; SolverIndex<AssetController->GetNumSolvers(); ++SolverIndex)
 	{
-		UIKRigSolver* Solver = Controller->GetSolver(SolverIndex);
+		UIKRigSolver* Solver = AssetController->GetSolver(SolverIndex);
 		if (!Solver->UsesBoneSettings())
 		{
 			continue;
@@ -820,7 +1122,7 @@ void SIKRigSkeleton::RefreshTreeView(bool IsInitialSetup)
     }
 	
 	// create all goal and effector elements
-	TArray<UIKRigEffectorGoal*> Goals = Controller->GetAllGoals();
+	TArray<UIKRigEffectorGoal*> Goals = AssetController->GetAllGoals();
 	for (const UIKRigEffectorGoal* Goal : Goals)
 	{
 		// make new element for goal
@@ -834,12 +1136,12 @@ void SIKRigSkeleton::RefreshTreeView(bool IsInitialSetup)
 		GoalItem->Parent = AllElements[BoneIndex];
 
 		// add all effectors connected to this goal
-		for (int32 SolverIndex=0; SolverIndex<Controller->GetNumSolvers(); ++SolverIndex)
+		for (int32 SolverIndex=0; SolverIndex<AssetController->GetNumSolvers(); ++SolverIndex)
 		{
-			if (UObject* Effector = Controller->GetEffectorForGoal(Goal->GoalName, SolverIndex))
+			if (UObject* Effector = AssetController->GetEffectorForGoal(Goal->GoalName, SolverIndex))
 			{
 				// make new element for effector
-				const UIKRigSolver* Solver = Controller->GetSolver(SolverIndex);
+				const UIKRigSolver* Solver = AssetController->GetSolver(SolverIndex);
 				const FName DisplayName = FName("Effector for: " + Solver->GetName());
 				TSharedPtr<FIKRigTreeElement> EffectorItem = MakeShared<FIKRigTreeElement>(DisplayName, IKRigTreeElementType::EFFECTOR);
 				EffectorItem->EffectorSolverIndex = SolverIndex;
@@ -885,17 +1187,23 @@ void SIKRigSkeleton::HandleGetChildrenForTree(
 
 void SIKRigSkeleton::OnSelectionChanged(TSharedPtr<FIKRigTreeElement> Selection, ESelectInfo::Type SelectInfo)
 {
+	const TSharedPtr<FIKRigEditorController> Controller = EditorController.Pin();
+	if (!Controller.IsValid())
+	{
+		return;
+	}
+	
 	// gate any selection changes NOT made by user clicking mouse
 	if (SelectInfo == ESelectInfo::OnMouseClick)
 	{
 		TArray<TSharedPtr<FIKRigTreeElement>> SelectedGoals;
 		GetSelectedGoals(SelectedGoals);
 		TArray<FName> SelectedGoalNames;
-		for (const auto& Goal : SelectedGoals)
+		for (const TSharedPtr<FIKRigTreeElement>& Goal : SelectedGoals)
 		{
 			SelectedGoalNames.Add(Goal->Key);
 		}
-		EditorController.Pin()->HandleGoalsSelectedInTreeView(SelectedGoalNames);
+		Controller->HandleGoalsSelectedInTreeView(SelectedGoalNames);
 	}
 }
 
@@ -909,19 +1217,28 @@ TSharedPtr<SWidget> SIKRigSkeleton::CreateContextMenu()
 
 void SIKRigSkeleton::OnItemClicked(TSharedPtr<FIKRigTreeElement> InItem)
 {
+	const TSharedPtr<FIKRigEditorController> Controller = EditorController.Pin();
+	if (!Controller.IsValid())
+	{
+		return;
+	}
+	
 	// update details view
 	if (InItem->ElementType == IKRigTreeElementType::BONE)
 	{
-		EditorController.Pin()->ShowDetailsForBone(InItem->Key);
-	}else if (InItem->ElementType == IKRigTreeElementType::GOAL)
+		Controller->ShowDetailsForBone(InItem->Key);
+	}
+	else if (InItem->ElementType == IKRigTreeElementType::GOAL)
 	{
-		EditorController.Pin()->ShowDetailsForGoal(InItem->Key);
-	}else if (InItem->ElementType == IKRigTreeElementType::EFFECTOR)
+		Controller->ShowDetailsForGoal(InItem->Key);
+	}
+	else if (InItem->ElementType == IKRigTreeElementType::EFFECTOR)
 	{
-		EditorController.Pin()->ShowDetailsForEffector(InItem->EffectorGoalName, InItem->EffectorSolverIndex);
-	}else if (InItem->ElementType == IKRigTreeElementType::BONE_SETTINGS)
+		Controller->ShowDetailsForEffector(InItem->EffectorGoalName, InItem->EffectorSolverIndex);
+	}
+	else if (InItem->ElementType == IKRigTreeElementType::BONE_SETTINGS)
 	{
-		EditorController.Pin()->ShowDetailsForBoneSettings(InItem->BoneSettingBoneName, InItem->BoneSettingsSolverIndex);
+		Controller->ShowDetailsForBoneSettings(InItem->BoneSettingBoneName, InItem->BoneSettingsSolverIndex);
 	}
 
 	// to rename an item, you have to select it first, then click on it again within a time limit (slow double click)
@@ -989,13 +1306,19 @@ void SIKRigSkeleton::ImportSkeleton(const FAssetData& InAssetData)
 		return;
 	}
 	
+	const TSharedPtr<FIKRigEditorController> Controller = EditorController.Pin();
+	if (!Controller.IsValid())
+	{
+		return;
+	}
+	
 	// load the skeletal mesh in the viewport
-	TSharedRef<IPersonaToolkit> Persona = EditorController.Pin()->EditorToolkit.Pin()->GetPersonaToolkit();
+	TSharedRef<IPersonaToolkit> Persona = Controller->EditorToolkit.Pin()->GetPersonaToolkit();
 	Persona->SetPreviewMesh(Mesh, true);
 	Persona->GetPreviewScene()->FocusViews();
 
 	// import the skeleton data into the IKRig asset
-	EditorController.Pin()->AssetController->SetSourceSkeletalMesh(Mesh, true);
+	Controller->AssetController->SetSourceSkeletalMesh(Mesh, true);
 
 	// update the bone hierarchy view
 	RefreshTreeView(true);
@@ -1053,10 +1376,16 @@ FReply SIKRigSkeleton::OnAcceptDrop(
 	{
 		return FReply::Unhandled();
 	}
+	
+	const TSharedPtr<FIKRigEditorController> Controller = EditorController.Pin();
+	if (!Controller.IsValid())
+	{
+		return FReply::Handled();
+	}
 
 	const FIKRigTreeElement& DraggedElement = *DragDropOp.Get()->Element.Pin().Get();
-	UIKRigController* Controller = EditorController.Pin()->AssetController;
-	const bool bWasReparented = Controller->SetGoalBone(DraggedElement.Key, TargetItem.Get()->Key);
+	UIKRigController* AssetController = Controller->AssetController;
+	const bool bWasReparented = AssetController->SetGoalBone(DraggedElement.Key, TargetItem.Get()->Key);
 	if (bWasReparented)
 	{
 		RefreshTreeView();
@@ -1069,28 +1398,31 @@ FReply SIKRigSkeleton::OnKeyDown(const FGeometry& MyGeometry, const FKeyEvent& I
 {
 	const FKey Key = InKeyEvent.GetKey();
 
+	const TSharedPtr<FIKRigEditorController> Controller = EditorController.Pin();
+	if (!Controller.IsValid())
+	{
+		return FReply::Handled();
+	}
+
 	// handle deleting selected goals
 	if (Key == EKeys::Delete)
 	{
 		TArray<TSharedPtr<FIKRigTreeElement>> SelectedItems = TreeView->GetSelectedItems();
-		for (const auto& SelectedItem : SelectedItems)
+		for (const TSharedPtr<FIKRigTreeElement>& SelectedItem : SelectedItems)
 		{
 			if (SelectedItem->ElementType == IKRigTreeElementType::GOAL)
 			{
-				UIKRigController* Controller = EditorController.Pin()->AssetController;
-				Controller->RemoveGoal(SelectedItem->Key);
+				Controller->AssetController->RemoveGoal(SelectedItem->Key);
 			}
 
 			if (SelectedItem->ElementType == IKRigTreeElementType::EFFECTOR)
 			{
-				UIKRigController* Controller = EditorController.Pin()->AssetController;
-				Controller->DisconnectGoalFromSolver(SelectedItem->EffectorGoalName, SelectedItem->EffectorSolverIndex);
+				Controller->AssetController->DisconnectGoalFromSolver(SelectedItem->EffectorGoalName, SelectedItem->EffectorSolverIndex);
 			}
 			
 			if (SelectedItem->ElementType == IKRigTreeElementType::BONE_SETTINGS)
 			{
-				UIKRigController* Controller = EditorController.Pin()->AssetController;
-				Controller->RemoveBoneSetting(SelectedItem->BoneSettingBoneName, SelectedItem->BoneSettingsSolverIndex);
+				Controller->AssetController->RemoveBoneSetting(SelectedItem->BoneSettingBoneName, SelectedItem->BoneSettingsSolverIndex);
 			}
 		}
 
diff --git a/Engine/Plugins/Animation/IKRig/Source/IKRigEditor/Private/RigEditor/SIKRigSolverStack.cpp b/Engine/Plugins/Animation/IKRig/Source/IKRigEditor/Private/RigEditor/SIKRigSolverStack.cpp
--- a/Engine/Plugins/Animation/IKRig/Source/IKRigEditor/Private/RigEditor/SIKRigSolverStack.cpp
+++ b/Engine/Plugins/Animation/IKRig/Source/IKRigEditor/Private/RigEditor/SIKRigSolverStack.cpp
@@ -175,9 +175,15 @@ void SIKRigSolverStack::BuildAddNewMenu(FMenuBuilder& MenuBuilder)
 
 bool SIKRigSolverStack::IsAddSolverEnabled() const
 {
-	if (UIKRigController* Controller = EditorController.Pin()->AssetController)
+	const TSharedPtr<FIKRigEditorController> Controller = EditorController.Pin();
+	if (!Controller.IsValid())
 	{
-		if (Controller->GetSkeleton().BoneNames.Num() > 0)
+		return false; 
+	}
+	
+	if (UIKRigController* AssetController = Controller->AssetController)
+	{
+		if (AssetController->GetSkeleton().BoneNames.Num() > 0)
 		{
 			return true;
 		}
@@ -188,24 +194,36 @@ bool SIKRigSolverStack::IsAddSolverEnabled() const
 
 void SIKRigSolverStack::AddNewSolver(UClass* Class)
 {
-	if (UIKRigController* Controller = EditorController.Pin()->AssetController)
+	const TSharedPtr<FIKRigEditorController> Controller = EditorController.Pin();
+	if (!Controller.IsValid())
+	{
+		return; 
+	}
+	
+	if (UIKRigController* AssetController = Controller->AssetController)
 	{
 		// add the solver
-		const int32 NewSolverIndex = Controller->AddSolver(Class);
+		const int32 NewSolverIndex = AssetController->AddSolver(Class);
 	}
 
 	RefreshStackView();
-	EditorController.Pin()->SkeletonView->RefreshTreeView(); // update solver indices in effector items
+	Controller->SkeletonView->RefreshTreeView(); // update solver indices in effector items
 }
 
 void SIKRigSolverStack::RefreshStackView()
 {
+	const TSharedPtr<FIKRigEditorController> Controller = EditorController.Pin();
+	if (!Controller.IsValid())
+	{
+		return; 
+	}
+	
 	ListViewItems.Reset();
-	UIKRigController* Controller = EditorController.Pin()->AssetController;
-	const int32 NumSolvers = Controller->GetNumSolvers();
+	UIKRigController* AssetController = Controller->AssetController;
+	const int32 NumSolvers = AssetController->GetNumSolvers();
 	for (int32 i=0; i<NumSolvers; ++i)
 	{
-		UIKRigSolver* Solver = Controller->GetSolver(i);
+		UIKRigSolver* Solver = AssetController->GetSolver(i);
 		const FText DisplayName = Solver ? FText::FromString(Solver->GetName()) : FText::FromString("Unknown Solver");
 		TSharedPtr<FSolverStackElement> SolverItem = FSolverStackElement::Make(DisplayName, i);
 		ListViewItems.Add(SolverItem);
@@ -252,8 +270,14 @@ FReply SIKRigSolverStack::OnDragDetected(
 
 void SIKRigSolverStack::OnItemClicked(TSharedPtr<FSolverStackElement> InItem)
 {
-	EditorController.Pin()->ShowDetailsForSolver(InItem.Get()->IndexInStack);
-	EditorController.Pin()->SkeletonView->RefreshTreeView(); // update filter showing which bones are affected
+	const TSharedPtr<FIKRigEditorController> Controller = EditorController.Pin();
+	if (!Controller.IsValid())
+	{
+		return; 
+	}
+	
+	Controller->ShowDetailsForSolver(InItem.Get()->IndexInStack);
+	Controller->SkeletonView->RefreshTreeView(); // update filter showing which bones are affected
 }
 
 TOptional<EItemDropZone> SIKRigSolverStack::OnCanAcceptDrop(
@@ -282,14 +306,20 @@ FReply SIKRigSolverStack::OnAcceptDrop(
 	{
 		return FReply::Unhandled();
 	}
+	
+	const TSharedPtr<FIKRigEditorController> Controller = EditorController.Pin();
+	if (!Controller.IsValid())
+	{
+		return FReply::Handled();
+	}
 
 	const FSolverStackElement& DraggedElement = *DragDropOp.Get()->Element.Pin().Get();
-	UIKRigController* Controller = EditorController.Pin()->AssetController;
-	const bool bWasReparented = Controller->MoveSolverInStack(DraggedElement.IndexInStack, TargetItem.Get()->IndexInStack);
+	UIKRigController* AssetController = Controller->AssetController;
+	const bool bWasReparented = AssetController->MoveSolverInStack(DraggedElement.IndexInStack, TargetItem.Get()->IndexInStack);
 	if (bWasReparented)
 	{
 		RefreshStackView();
-		EditorController.Pin()->SkeletonView->RefreshTreeView(); // update solver indices in effector items
+		Controller->SkeletonView->RefreshTreeView(); // update solver indices in effector items
 	}
 	
 	return FReply::Handled();
@@ -298,6 +328,12 @@ FReply SIKRigSolverStack::OnAcceptDrop(
 FReply SIKRigSolverStack::OnKeyDown(const FGeometry& MyGeometry, const FKeyEvent& InKeyEvent)
 {
 	const FKey Key = InKeyEvent.GetKey();
+	
+	const TSharedPtr<FIKRigEditorController> Controller = EditorController.Pin();
+	if (!Controller.IsValid())
+	{
+		return FReply::Handled();
+	}
 
 	// handle deleting selected solver
 	if (Key == EKeys::Delete)
@@ -308,12 +344,12 @@ FReply SIKRigSolverStack::OnKeyDown(const FGeometry& MyGeometry, const FKeyEvent
 			return FReply::Unhandled();
 		}
 
-		UIKRigController* Controller = EditorController.Pin()->AssetController;
-		UIKRigSolver* SolverToRemove = Controller->GetSolver(SelectedItems[0]->IndexInStack);
-		Controller->RemoveSolver(SolverToRemove);
+		UIKRigController* AssetController = Controller->AssetController;
+		UIKRigSolver* SolverToRemove = AssetController->GetSolver(SelectedItems[0]->IndexInStack);
+		AssetController->RemoveSolver(SolverToRemove);
 
 		RefreshStackView();
-		EditorController.Pin()->SkeletonView->RefreshTreeView(); // update solver indices in effector items
+		Controller->SkeletonView->RefreshTreeView(); // update solver indices in effector items
 		
 		return FReply::Handled();
 	}
diff --git a/Engine/Plugins/Animation/IKRig/Source/IKRigEditor/Public/RetargetEditor/AssetTypeActions_IKRetargeter.h b/Engine/Plugins/Animation/IKRig/Source/IKRigEditor/Public/RetargetEditor/AssetTypeActions_IKRetargeter.h
--- a/Engine/Plugins/Animation/IKRig/Source/IKRigEditor/Public/RetargetEditor/AssetTypeActions_IKRetargeter.h
+++ b/Engine/Plugins/Animation/IKRig/Source/IKRigEditor/Public/RetargetEditor/AssetTypeActions_IKRetargeter.h
@@ -14,7 +14,7 @@ public:
 	virtual void GetActions(const TArray<UObject*>& InObjects, struct FToolMenuSection& Section) override;
 	virtual FColor GetTypeColor() const override { return FColor(255, 192, 128); }
 	virtual UClass* GetSupportedClass() const override;
-	//virtual void OpenAssetEditor(const TArray<UObject*>& InObjects, TSharedPtr<class IToolkitHost> EditWithinLevelEditor = TSharedPtr<IToolkitHost>()) override;
+	virtual void OpenAssetEditor(const TArray<UObject*>& InObjects, TSharedPtr<class IToolkitHost> EditWithinLevelEditor = TSharedPtr<IToolkitHost>()) override;
 	virtual uint32 GetCategories() override { return EAssetTypeCategories::Animation; }
-	//virtual UThumbnailInfo* GetThumbnailInfo(UObject* Asset) const override;
+	virtual UThumbnailInfo* GetThumbnailInfo(UObject* Asset) const override;
 };
diff --git a/Engine/Plugins/Animation/IKRig/Source/IKRigEditor/Public/RetargetEditor/IKRetargetAnimInstance.h b/Engine/Plugins/Animation/IKRig/Source/IKRigEditor/Public/RetargetEditor/IKRetargetAnimInstance.h
new file mode 100644
--- /dev/null
+++ b/Engine/Plugins/Animation/IKRig/Source/IKRigEditor/Public/RetargetEditor/IKRetargetAnimInstance.h
@@ -0,0 +1,30 @@
+// Copyright Epic Games, Inc. All Rights Reserved.
+
+#pragma once
+
+#include "AnimPreviewInstance.h"
+#include "AnimNodes/AnimNode_IKRetargeter.h"
+
+#include "IKRetargetAnimInstance.generated.h"
+
+class UIKRetargeter;
+
+UCLASS(transient, NotBlueprintable)
+class UIKRetargetAnimInstance : public UAnimPreviewInstance
+{
+	GENERATED_UCLASS_BODY()
+
+public:
+
+	void SetRetargetAssetAndSourceComponent(
+		UIKRetargeter* InIKRetargetAsset,
+		TWeakObjectPtr<USkeletalMeshComponent> InSourceMeshComponent);
+
+	UIKRetargeter* GetCurrentlyUsedRetargeter() const;
+
+protected:
+	virtual FAnimInstanceProxy* CreateAnimInstanceProxy() override;
+
+	UPROPERTY(Transient)
+	FAnimNode_IKRetargeter IKRetargeterNode;
+};
diff --git a/Engine/Plugins/Animation/IKRig/Source/IKRigEditor/Public/RetargetEditor/IKRetargetAnimInstanceProxy.h b/Engine/Plugins/Animation/IKRig/Source/IKRigEditor/Public/RetargetEditor/IKRetargetAnimInstanceProxy.h
new file mode 100644
--- /dev/null
+++ b/Engine/Plugins/Animation/IKRig/Source/IKRigEditor/Public/RetargetEditor/IKRetargetAnimInstanceProxy.h
@@ -0,0 +1,40 @@
+// Copyright Epic Games, Inc. All Rights Reserved.
+
+#pragma once
+
+#include "AnimPreviewInstance.h"
+#include "Animation/AnimNode_LinkedInputPose.h"
+#include "AnimNodes/AnimNode_IKRetargeter.h"
+
+#include "IKRetargetAnimInstanceProxy.generated.h"
+
+/** Proxy override for this UAnimInstance-derived class */
+USTRUCT()
+struct FIKRetargetAnimInstanceProxy : public FAnimPreviewInstanceProxy
+{
+	GENERATED_BODY()
+
+public:
+	
+	FIKRetargetAnimInstanceProxy() = default;
+	FIKRetargetAnimInstanceProxy(UAnimInstance* InAnimInstance, FAnimNode_IKRetargeter* IKRetargeterNode);
+	virtual ~FIKRetargetAnimInstanceProxy() override = default;
+
+	/** FAnimPreviewInstanceProxy interface */
+	virtual void Initialize(UAnimInstance* InAnimInstance) override;
+	virtual bool Evaluate(FPoseContext& Output) override;
+	virtual void UpdateAnimationNode(const FAnimationUpdateContext& InContext) override;
+	/** END FAnimPreviewInstanceProxy interface */
+
+	/** FAnimInstanceProxy interface */
+	/** Called when the anim instance is being initialized. These nodes can be provided */
+	virtual FAnimNode_Base* GetCustomRootNode() override;
+	virtual void GetCustomNodes(TArray<FAnimNode_Base*>& OutNodes) override;
+	/** END FAnimInstanceProxy interface */
+
+	void SetRetargetAssetAndSourceComponent(
+		UIKRetargeter* InIKRetargetAsset,
+		TWeakObjectPtr<USkeletalMeshComponent> InSourceMeshComponent) const;
+
+	FAnimNode_IKRetargeter* IKRetargetNode;
+};
\ No newline at end of file
diff --git a/Engine/Plugins/Animation/IKRig/Source/IKRigEditor/Public/RetargetEditor/IKRetargetChainTabSummoner.h b/Engine/Plugins/Animation/IKRig/Source/IKRigEditor/Public/RetargetEditor/IKRetargetChainTabSummoner.h
new file mode 100644
--- /dev/null
+++ b/Engine/Plugins/Animation/IKRig/Source/IKRigEditor/Public/RetargetEditor/IKRetargetChainTabSummoner.h
@@ -0,0 +1,21 @@
+﻿// Copyright Epic Games, Inc. All Rights Reserved.
+
+#pragma once
+
+#include "CoreMinimal.h"
+#include "WorkflowOrientedApp/WorkflowTabFactory.h"
+
+class FIKRetargetEditor;
+
+struct FIKRetargetChainTabSummoner : public FWorkflowTabFactory
+{
+public:
+	static const FName TabID;
+	
+	FIKRetargetChainTabSummoner(const TSharedRef<FIKRetargetEditor>& InIKRetargetEditor);
+	virtual TSharedRef<SWidget> CreateTabBody(const FWorkflowTabSpawnInfo& Info) const override;
+	virtual TSharedPtr<SToolTip> CreateTabToolTipWidget(const FWorkflowTabSpawnInfo& Info) const override;
+
+protected:
+	TWeakPtr<FIKRetargetEditor> IKRetargetEditor;
+};
diff --git a/Engine/Plugins/Animation/IKRig/Source/IKRigEditor/Public/RetargetEditor/IKRetargetCommands.h b/Engine/Plugins/Animation/IKRig/Source/IKRigEditor/Public/RetargetEditor/IKRetargetCommands.h
new file mode 100644
--- /dev/null
+++ b/Engine/Plugins/Animation/IKRig/Source/IKRigEditor/Public/RetargetEditor/IKRetargetCommands.h
@@ -0,0 +1,38 @@
+// Copyright Epic Games, Inc. All Rights Reserved.
+
+#pragma once
+
+#include "CoreMinimal.h"
+#include "Framework/Commands/Commands.h"
+#include "IKRetargetEditorStyle.h"
+
+class FIKRetargetCommands : public TCommands<FIKRetargetCommands>
+{
+public:
+	FIKRetargetCommands() : TCommands<FIKRetargetCommands>
+	(
+		"IKRetarget",
+		NSLOCTEXT("Contexts", "IKRetarget", "IK Retarget"),
+		NAME_None,
+		FIKRetargetEditorStyle::Get().GetStyleSetName() // Icon Style Set
+	)
+	{}
+	
+	/** edit reference pose */
+	TSharedPtr< FUICommandInfo > EditRetargetPose;
+
+	/** new reference pose */
+	TSharedPtr< FUICommandInfo > NewRetargetPose;
+
+	/** delete reference pose */
+	TSharedPtr< FUICommandInfo > DeleteRetargetPose;
+
+	/** reset reference pose */
+	TSharedPtr< FUICommandInfo > ResetRetargetPose;
+	
+	/** export animation */
+	TSharedPtr< FUICommandInfo > ExportAnimation;
+
+	/** initialize commands */
+	virtual void RegisterCommands() override;
+};
diff --git a/Engine/Plugins/Animation/IKRig/Source/IKRigEditor/Public/RetargetEditor/IKRetargetEditMode.h b/Engine/Plugins/Animation/IKRig/Source/IKRigEditor/Public/RetargetEditor/IKRetargetEditMode.h
new file mode 100644
--- /dev/null
+++ b/Engine/Plugins/Animation/IKRig/Source/IKRigEditor/Public/RetargetEditor/IKRetargetEditMode.h
@@ -0,0 +1,68 @@
+// Copyright Epic Games, Inc. All Rights Reserved.
+
+#pragma once
+
+#include "CoreMinimal.h"
+
+#include "Retargeter/IKRetargeter.h"
+#include "IPersonaEditMode.h"
+
+class FIKRetargetEditorController;
+class FIKRetargetEditor;
+class FIKRetargetPreviewScene;
+
+class FIKRetargetEditMode : public IPersonaEditMode
+{
+public:
+	static FName ModeName;
+	
+	FIKRetargetEditMode() = default;
+
+	/** glue for all the editor parts to communicate */
+	void SetEditorController(const TSharedPtr<FIKRetargetEditorController> InEditorController) { EditorController = InEditorController; };
+
+	/** IPersonaEditMode interface */
+	virtual bool GetCameraTarget(FSphere& OutTarget) const override;
+	virtual class IPersonaPreviewScene& GetAnimPreviewScene() const override;
+	virtual void GetOnScreenDebugInfo(TArray<FText>& OutDebugInfo) const override;
+	/** END IPersonaEditMode interface */
+
+	/** FEdMode interface */
+	virtual void Tick(FEditorViewportClient* ViewportClient, float DeltaTime) override;
+	virtual void Render(const FSceneView* View, FViewport* Viewport, FPrimitiveDrawInterface* PDI) override;
+	virtual void DrawHUD(FEditorViewportClient* ViewportClient, FViewport* Viewport, const FSceneView* View, FCanvas* Canvas) override;
+	virtual bool IsCompatibleWith(FEditorModeID OtherModeID) const override { return true; }
+	virtual bool AllowWidgetMove() override;
+	virtual bool ShouldDrawWidget() const override;
+	virtual bool UsesTransformWidget() const override;
+	virtual bool UsesTransformWidget(UE::Widget::EWidgetMode CheckMode) const override;
+	virtual FVector GetWidgetLocation() const override;
+	virtual bool HandleClick(FEditorViewportClient* InViewportClient, HHitProxy* HitProxy, const FViewportClick& Click) override;
+	virtual bool StartTracking(FEditorViewportClient* InViewportClient, FViewport* InViewport) override;
+	virtual bool EndTracking(FEditorViewportClient* InViewportClient, FViewport* InViewport) override;
+	virtual bool InputDelta(FEditorViewportClient* InViewportClient, FViewport* InViewport, FVector& InDrag, FRotator& InRot, FVector& InScale) override;
+	virtual bool GetCustomDrawingCoordinateSystem(FMatrix& InMatrix, void* InData) override;
+	virtual bool GetCustomInputCoordinateSystem(FMatrix& InMatrix, void* InData) override;
+	/** END FEdMode interface */
+
+	bool IsBoneSelected(const FName& BoneName) const;
+
+	void HandleBoneSelectedInViewport(const FName& BoneName, bool bReplace);
+
+private:
+	/** The hosting app */
+	TWeakPtr<FIKRetargetEditorController> EditorController;
+
+	void DrawBoneGizmo(
+		FPrimitiveDrawInterface* PDI,
+		const FTransform& BoneTransform,
+		const FVector& Start,
+		const FVector& End,
+		float Size,
+		float Thickness,
+		bool bIsSelected) const;
+
+	/** chain viewport selection state */
+	TArray<FName> SelectedBones;
+	bool bRotatingBones = false;
+};
diff --git a/Engine/Plugins/Animation/IKRig/Source/IKRigEditor/Public/RetargetEditor/IKRetargetEditor.h b/Engine/Plugins/Animation/IKRig/Source/IKRigEditor/Public/RetargetEditor/IKRetargetEditor.h
new file mode 100644
--- /dev/null
+++ b/Engine/Plugins/Animation/IKRig/Source/IKRigEditor/Public/RetargetEditor/IKRetargetEditor.h
@@ -0,0 +1,124 @@
+// Copyright Epic Games, Inc. All Rights Reserved.
+
+#pragma once
+
+#include "CoreMinimal.h"
+
+#include "IHasPersonaToolkit.h"
+#include "IKRetargetEditorController.h"
+#include "IKRetargetMode.h"
+#include "IPersonaPreviewScene.h"
+#include "PersonaAssetEditorToolkit.h"
+
+class IAnimationSequenceBrowser;
+class UIKRetargetSkeletalMeshComponent;
+class IDetailsView;
+class FGGAssetEditorToolbar;
+class UIKRetargeter;
+class FIKRetargetToolbar;
+class SIKRetargetViewportTabBody;
+class FIKRetargetPreviewScene;
+struct FIKRetargetPose;
+
+namespace IKRetargetEditorModes
+{
+	extern const FName IKRetargetEditorMode;
+}
+
+class FIKRetargetEditor :
+	public FPersonaAssetEditorToolkit,
+	public IHasPersonaToolkit,
+	public FGCObject,
+	public FEditorUndoClient,
+	public FTickableEditorObject
+{
+public:
+
+	FIKRetargetEditor();
+	virtual ~FIKRetargetEditor() override;
+
+	void InitAssetEditor(
+		const EToolkitMode::Type Mode,
+		const TSharedPtr< IToolkitHost >& InitToolkitHost,
+		UIKRetargeter* Asset);
+
+	/** FAssetEditorToolkit interface */
+	virtual void RegisterTabSpawners(const TSharedRef<FTabManager>& InTabManager) override;
+	virtual void UnregisterTabSpawners(const TSharedRef<FTabManager>& InTabManager) override;
+	virtual FName GetToolkitFName() const override;
+	virtual FText GetBaseToolkitName() const override;
+	virtual FText GetToolkitName() const override;
+	virtual FLinearColor GetWorldCentricTabColorScale() const override;
+	virtual FString GetWorldCentricTabPrefix() const override;
+	/** END FAssetEditorToolkit interface */
+	
+	/** FGCObject interface */
+	virtual void AddReferencedObjects(FReferenceCollector& Collector) override;
+	/** END FGCObject interface */
+
+	//** FTickableEditorObject Interface
+	virtual void Tick(float DeltaTime) override;
+	virtual ETickableTickType GetTickableTickType() const override { return ETickableTickType::Always; }
+	virtual TStatId GetStatId() const override;
+	//~ END FTickableEditorObject Interface
+
+	/** IHasPersonaToolkit interface */
+	virtual TSharedRef<IPersonaToolkit> GetPersonaToolkit() const override { return EditorController->PersonaToolkit.ToSharedRef(); }
+	/** END IHasPersonaToolkit interface */
+
+	TSharedRef<FIKRetargetEditorController> GetController() const {return EditorController;};
+
+private:
+
+	/** toolbar */
+	void BindCommands();
+	void ExtendToolbar();
+	void FillToolbar(FToolBarBuilder& ToolbarBuilder);
+	/** END toolbar */
+	
+	/** preview scene setup */
+	void HandlePreviewSceneCreated(const TSharedRef<IPersonaPreviewScene>& InPersonaPreviewScene);
+	void HandlePreviewMeshChanged(USkeletalMesh* InOldSkeletalMesh, USkeletalMesh* InNewSkeletalMesh);
+	void HandleDetailsCreated(const TSharedRef<class IDetailsView>& InDetailsView);
+	void OnFinishedChangingDetails(const FPropertyChangedEvent& PropertyChangedEvent);
+	void SetupAnimInstance();
+	/** END preview scene setup */
+
+	/** animation asset browser */
+	void HandleOpenNewAsset(UObject* InNewAsset);
+	void HandleAnimationSequenceBrowserCreated(const TSharedRef<IAnimationSequenceBrowser>& InSequenceBrowser);
+	/* END animation asset browser */
+
+	/** edit retarget pose */
+	void HandleEditPose() const;
+	bool CanEditPose() const;
+	bool IsEditingPose() const;
+	UAnimationAsset* PreviousAsset = nullptr;
+	/** END edit retarget pose*/
+
+	/** new/delete retarget pose */
+	void HandleNewPose();
+	FReply CreateNewPose();
+	void HandleDeletePose();
+	bool CanDeletePose() const;
+	void HandleResetPose();
+	TSharedPtr<SWindow> NewPoseWindow;
+	TSharedPtr<SEditableTextBox> NewPoseEditableText;
+	/** END new/delete retarget pose */
+
+	TArray<TSharedPtr<FName>> PoseNames;
+	FText GetCurrentPoseName() const;
+	void OnPoseSelected(TSharedPtr<FName> InPoseName, ESelectInfo::Type SelectInfo);
+	/* END edit reference pose */
+
+	/* export animation*/
+	void ExportAnimation() const;
+	/* END edit reference pose */
+
+	void HandleSourceOrTargetIKRigAssetChanged();
+	
+	/** centralized management of across all views */
+	TSharedRef<FIKRetargetEditorController> EditorController;
+	
+	friend FIKRetargetMode;
+};
diff --git a/Engine/Plugins/Animation/IKRig/Source/IKRigEditor/Public/RetargetEditor/IKRetargetEditorController.h b/Engine/Plugins/Animation/IKRig/Source/IKRigEditor/Public/RetargetEditor/IKRetargetEditorController.h
new file mode 100644
--- /dev/null
+++ b/Engine/Plugins/Animation/IKRig/Source/IKRigEditor/Public/RetargetEditor/IKRetargetEditorController.h
@@ -0,0 +1,57 @@
+// Copyright Epic Games, Inc. All Rights Reserved.
+#pragma once
+#include "IPersonaToolkit.h"
+
+
+class SIKRetargetChainMapList;
+class UIKRetargetAnimInstance;
+class FIKRetargetEditor;
+class UIKRetargeter;
+class UDebugSkelMeshComponent;
+
+/** a home for cross-widget communication to synchronize state across all tabs and viewport */
+class FIKRetargetEditorController
+{
+public:
+	
+	/** the data model */
+	UIKRetargeter* Asset;
+
+	/** viewport skeletal mesh */
+	UDebugSkelMeshComponent* SourceSkelMeshComponent;
+	UDebugSkelMeshComponent* TargetSkelMeshComponent;
+
+	/** viewport anim instance */
+	UPROPERTY(transient, NonTransactional)
+	TObjectPtr<class UAnimPreviewInstance> SourceAnimInstance;
+	UPROPERTY(transient, NonTransactional)
+	TObjectPtr<class UIKRetargetAnimInstance> TargetAnimInstance;
+
+	/** asset properties tab */
+	TSharedPtr<class IDetailsView> DetailsView;
+
+	/** chain list view */
+	TSharedPtr<SIKRetargetChainMapList> ChainsView;
+
+	/** Preview scene to be supplied by IHasPersonaToolkit::GetPersonaToolkit */
+	TSharedPtr<IPersonaToolkit> PersonaToolkit;
+
+	/** Sequence Browser **/
+	TWeakPtr<class IAnimationSequenceBrowser> SequenceBrowser;
+
+	/** the persona toolkit */
+	TWeakPtr<FIKRetargetEditor> Editor;
+
+	/** get the skeletal meshes we are copying animation between */
+	USkeletalMesh* GetSourceSkeletalMesh() const;
+	USkeletalMesh* GetTargetSkeletalMesh() const;
+
+	/** get current chain pose */
+	FTransform GetTargetBoneTransform(const FName& BoneName) const;
+	void GetTargetBoneStartAndEnd(const FName& BoneName, FVector& Start, FVector& End) const;
+
+	/** get the retargeter that is running in the viewport (which is a duplicate of the source asset) */
+	UIKRetargeter* GetCurrentlyRunningRetargeter() const;
+
+	void RefreshAllViews();
+};
\ No newline at end of file
diff --git a/Engine/Plugins/Animation/IKRig/Source/IKRigEditor/Public/RetargetEditor/IKRetargetEditorStyle.h b/Engine/Plugins/Animation/IKRig/Source/IKRigEditor/Public/RetargetEditor/IKRetargetEditorStyle.h
new file mode 100644
--- /dev/null
+++ b/Engine/Plugins/Animation/IKRig/Source/IKRigEditor/Public/RetargetEditor/IKRetargetEditorStyle.h
@@ -0,0 +1,43 @@
+// Copyright Epic Games, Inc. All Rights Reserved.
+
+#pragma once
+
+#include "EditorStyleSet.h"
+#include "Styling/SlateStyle.h"
+#include "Styling/SlateStyleRegistry.h"
+#include "Styling/SlateStyleMacros.h"
+#include "Styling/SlateTypes.h"
+
+class FIKRetargetEditorStyle	: public FSlateStyleSet
+{
+public:
+	
+	FIKRetargetEditorStyle() : FSlateStyleSet("IKRetargetEditorStyle")
+	{
+		SetContentRoot(FPaths::EnginePluginsDir() / TEXT("Animation/IKRig/Content"));
+		const FVector2D Icon10x10(10.0f, 10.0f);
+		const FVector2D Icon14x14(14.0f, 14.0f);
+		const FVector2D Icon16x16(16.0f, 16.0f);
+		const FVector2D Icon20x20(20.0f, 20.0f);
+		const FVector2D Icon24x24(24.0f, 24.0f);
+		const FVector2D Icon32x32(32.0f, 32.0f);
+		const FVector2D Icon40x40(40.0f, 40.0f);
+
+		Set("IKRetarget.Tree.Bone", new IMAGE_BRUSH("Slate/Bone_16x", Icon16x16));
+		
+		FSlateStyleRegistry::RegisterSlateStyle(*this);
+	}
+
+	static FIKRetargetEditorStyle& Get()
+	{
+		static FIKRetargetEditorStyle Inst;
+		return Inst;
+	}
+	
+	~FIKRetargetEditorStyle()
+	{
+		FSlateStyleRegistry::UnRegisterSlateStyle(*this);
+	}
+};
+
+#undef IMAGE_BRUSH
\ No newline at end of file
diff --git a/Engine/Plugins/Animation/IKRig/Source/IKRigEditor/Public/RetargetEditor/IKRetargetFactory.h b/Engine/Plugins/Animation/IKRig/Source/IKRigEditor/Public/RetargetEditor/IKRetargetFactory.h
new file mode 100644
--- /dev/null
+++ b/Engine/Plugins/Animation/IKRig/Source/IKRigEditor/Public/RetargetEditor/IKRetargetFactory.h
@@ -0,0 +1,37 @@
+// Copyright Epic Games, Inc. All Rights Reserved.
+
+#pragma once
+
+#include "CoreMinimal.h"
+#include "UObject/ObjectMacros.h"
+#include "Factories/Factory.h"
+#include "IKRetargetFactory.generated.h"
+
+
+UCLASS(hidecategories=Object, MinimalAPI)
+class UIKRetargetFactory : public UFactory
+{
+	GENERATED_BODY()
+
+	UPROPERTY()
+	TObjectPtr<class UIKRigDefinition>	SourceIKRig;
+
+public:
+
+	UIKRetargetFactory();
+
+	// UFactory Interface
+	virtual FText GetDisplayName() const override;
+	virtual uint32 GetMenuCategories() const override;
+	virtual FText GetToolTip() const override;
+	virtual FString GetDefaultNewAssetName() const override;
+	virtual UObject* FactoryCreateNew(UClass* Class, UObject* InParent, FName Name, EObjectFlags Flags, UObject* Context, FFeedbackContext* Warn) override;
+	virtual bool ConfigureProperties() override;
+	virtual bool ShouldShowInNewMenu() const override;
+
+private:
+	
+	void OnTargetIKRigSelected(const FAssetData& SelectedAsset);
+
+	TSharedPtr<SWindow> PickerWindow;
+};
diff --git a/Engine/Plugins/Animation/IKRig/Source/IKRigEditor/Public/RetargetEditor/IKRetargetHitProxies.h b/Engine/Plugins/Animation/IKRig/Source/IKRigEditor/Public/RetargetEditor/IKRetargetHitProxies.h
new file mode 100644
--- /dev/null
+++ b/Engine/Plugins/Animation/IKRig/Source/IKRigEditor/Public/RetargetEditor/IKRetargetHitProxies.h
@@ -0,0 +1,28 @@
+// Copyright Epic Games, Inc. All Rights Reserved.
+
+#pragma once
+
+#include "CoreMinimal.h"
+#include "HitProxies.h"
+
+// allow limb selection to edit retarget pose
+struct HIKRetargetEditorBoneProxy : public HHitProxy
+{
+	DECLARE_HIT_PROXY();
+
+	FName BoneName;
+	
+	HIKRetargetEditorBoneProxy(const FName& InBoneName)
+		: HHitProxy(HPP_World)
+		, BoneName(InBoneName) {}
+
+	virtual EMouseCursor::Type GetMouseCursor()
+	{
+		return EMouseCursor::Crosshairs;
+	}
+
+	virtual bool AlwaysAllowsTranslucentPrimitives() const override
+	{
+		return true;
+	}
+};
diff --git a/Engine/Plugins/Animation/IKRig/Source/IKRigEditor/Public/RetargetEditor/IKRetargetMode.h b/Engine/Plugins/Animation/IKRig/Source/IKRigEditor/Public/RetargetEditor/IKRetargetMode.h
new file mode 100644
--- /dev/null
+++ b/Engine/Plugins/Animation/IKRig/Source/IKRigEditor/Public/RetargetEditor/IKRetargetMode.h
@@ -0,0 +1,29 @@
+// Copyright Epic Games, Inc. All Rights Reserved.
+
+#pragma once
+
+#include "WorkflowOrientedApp/ApplicationMode.h"
+#include "WorkflowOrientedApp/WorkflowCentricApplication.h"
+#include "WorkflowOrientedApp/WorkflowTabManager.h"
+
+class FIKRetargetEditor;
+
+class FIKRetargetMode : public FApplicationMode
+{
+	
+public:
+	
+	FIKRetargetMode(TSharedRef<class FWorkflowCentricApplication> InHostingApp, TSharedRef<class IPersonaPreviewScene> InPreviewScene);
+
+	/** FApplicationMode interface */
+	virtual void RegisterTabFactories(TSharedPtr<FTabManager> InTabManager) override;
+	/** END FApplicationMode interface */
+
+protected:
+	
+	/** The hosting app */
+	TWeakPtr<FIKRetargetEditor> IKRetargetEditorPtr;
+
+	/** The tab factories we support */
+	FWorkflowAllowedTabSet TabFactories;
+};
diff --git a/Engine/Plugins/Animation/IKRig/Source/IKRigEditor/Public/RetargetEditor/IKRetargeterFactory.h b/Engine/Plugins/Animation/IKRig/Source/IKRigEditor/Public/RetargetEditor/IKRetargeterFactory.h
deleted file mode 100644
diff --git a/Engine/Plugins/Animation/IKRig/Source/IKRigEditor/Public/RetargetEditor/SIKRetargetChainMapList.h b/Engine/Plugins/Animation/IKRig/Source/IKRigEditor/Public/RetargetEditor/SIKRetargetChainMapList.h
new file mode 100644
--- /dev/null
+++ b/Engine/Plugins/Animation/IKRig/Source/IKRigEditor/Public/RetargetEditor/SIKRetargetChainMapList.h
@@ -0,0 +1,119 @@
+﻿// Copyright Epic Games, Inc. All Rights Reserved.
+
+#pragma once
+
+#include "CoreMinimal.h"
+#include "EditorUndoClient.h"
+#include "IKRetargetEditorController.h"
+#include "DragAndDrop/DecoratedDragDropOp.h"
+#include "Widgets/SCompoundWidget.h"
+#include "Widgets/Views/STableViewBase.h"
+#include "Widgets/Views/STableRow.h"
+#include "Widgets/Views/SListView.h"
+#include "Framework/Commands/UICommandList.h"
+
+class FIKRigEditorController;
+class SIKRetargetChainMapList;
+class FIKRetargetEditor;
+struct FBoneChain;
+
+class FRetargetChainMapElement
+{
+public:
+
+	TSharedRef<ITableRow> MakeListRowWidget(
+		const TSharedRef<STableViewBase>& InOwnerTable,
+        TSharedRef<FRetargetChainMapElement> InStackElement,
+        TSharedPtr<SIKRetargetChainMapList> InChainList);
+
+	static TSharedRef<FRetargetChainMapElement> Make(FName InTargetChain)
+	{
+		return MakeShareable(new FRetargetChainMapElement(InTargetChain));
+	}
+
+	FName TargetChainName;
+
+private:
+	
+	/** Hidden constructor, always use Make above */
+	FRetargetChainMapElement(FName InTargetChain) : TargetChainName(InTargetChain) {}
+
+	/** Hidden constructor, always use Make above */
+	FRetargetChainMapElement() {}
+};
+
+typedef TSharedPtr< FRetargetChainMapElement > FRetargetChainMapElementPtr;
+class SIKRetargetChainMapRow : public SMultiColumnTableRow< FRetargetChainMapElementPtr >
+{
+public:
+	
+	void Construct(
+        const FArguments& InArgs,
+        const TSharedRef<STableViewBase>& InOwnerTableView,
+        TSharedRef<FRetargetChainMapElement> InChainElement,
+        TSharedPtr<SIKRetargetChainMapList> InChainList);
+
+	/** Overridden from SMultiColumnTableRow.  Generates a widget for this column of the table row. */
+	virtual TSharedRef<SWidget> GenerateWidgetForColumn(const FName& ColumnName) override;
+
+	void OnSourceChainComboSelectionChanged(TSharedPtr<FString> InName, ESelectInfo::Type SelectInfo);
+
+	FText GetSourceChainName() const;
+	
+private:
+
+	TArray<TSharedPtr<FString>> SourceChainOptions;
+
+	TWeakPtr<FRetargetChainMapElement> ChainMapElement;
+	
+	TWeakPtr<SIKRetargetChainMapList> ChainMapList;
+
+	friend SIKRetargetChainMapList;
+};
+
+typedef SListView< TSharedPtr<FRetargetChainMapElement> > SRetargetChainMapListViewType;
+
+class SIKRetargetChainMapList : public SCompoundWidget, public FEditorUndoClient
+{
+public:
+	SLATE_BEGIN_ARGS(SIKRetargetChainMapList) {}
+	SLATE_END_ARGS()
+
+    virtual ~SIKRetargetChainMapList() override;
+
+	void Construct(const FArguments& InArgs, TSharedRef<FIKRetargetEditorController> InEditorController);
+
+private:
+	
+	/** menu for adding new solver commands */
+	TSharedPtr<FUICommandList> CommandList;
+	
+	/** editor controller */
+	TWeakPtr<FIKRetargetEditorController> EditorController;
+
+	/** list view */
+	TSharedPtr<SRetargetChainMapListViewType> ListView;
+	TArray< TSharedPtr<FRetargetChainMapElement> > ListViewItems;
+	/** END list view */
+
+	/** callbacks */
+	FText GetSourceRootBone() const;
+	FText GetTargetRootBone() const;
+	bool IsChainMapEnabled() const;
+	/** when a chain is clicked on in the table view */
+	void OnItemClicked(TSharedPtr<FRetargetChainMapElement> InItem);
+
+	/** auto-map chain button*/
+	EVisibility IsAutoMapButtonVisible() const;
+	FReply OnAutoMapButtonClicked() const;
+	/** END auto-map chain button*/
+
+	/** list view generate row callback */
+	TSharedRef<ITableRow> MakeListRowWidget(TSharedPtr<FRetargetChainMapElement> InElement, const TSharedRef<STableViewBase>& OwnerTable);
+
+	/** call to refresh the list view */
+	void RefreshView();
+
+	friend SIKRetargetChainMapRow;
+	friend FIKRetargetEditorController;
+};
diff --git a/Engine/Plugins/Animation/IKRig/Source/IKRigEditor/Public/RigEditor/IKRigController.h b/Engine/Plugins/Animation/IKRig/Source/IKRigEditor/Public/RigEditor/IKRigController.h
--- a/Engine/Plugins/Animation/IKRig/Source/IKRigEditor/Public/RigEditor/IKRigController.h
+++ b/Engine/Plugins/Animation/IKRig/Source/IKRigEditor/Public/RigEditor/IKRigController.h
@@ -13,6 +13,7 @@ class UIKRigSolver;
 class UIKRigDefinition;
 class UIKRigEffectorGoal;
 class USkeletalMesh;
+struct FBoneChain;
 
 UCLASS(config = Engine, hidecategories = UObject, BlueprintType)
 class IKRIGEDITOR_API UIKRigController : public UObject
@@ -31,6 +32,8 @@ public:
 	USkeletalMesh* GetSourceSkeletalMesh() const;
 	void SetSkeleton(const FReferenceSkeleton& InSkeleton) const;
 	FIKRigSkeleton& GetSkeleton() const;
+	void SetBoneExcluded(const FName& BoneName, const bool bExclude) const;
+	bool GetBoneExcluded(const FName& BoneName) const;
 	//** END skeleton */
 
 	//** solvers */
@@ -76,6 +79,17 @@ public:
 	bool DoesBoneHaveSettings(const FName& BoneName) const;
 	//** END bone settings */
 
+	//** retargeting and retarget chains */
+	bool AddRetargetChain(const FName& ChainName, const FName& StartBone, const FName& EndBone) const;
+	bool RemoveRetargetChain(const FName& ChainName) const;
+	bool SetRetargetChainStartBone(const FName& ChainName, const FName& StartBoneName) const;
+	bool SetRetargetChainEndBone(const FName& ChainName, const FName& EndBoneName) const;
+	FBoneChain* GetRetargetChain(int32 Index) const;
+	int32 GetNumRetargetChains() const;
+	void SetRetargetRoot(const FName& RootBoneName);
+	FName GetRetargetRoot();
+	//** END retarget chains */
+	
 	// BEGIN UObject
 	static void AddReferencedObjects(UObject* InThis, FReferenceCollector& Collector);
 	// END UObject
diff --git a/Engine/Plugins/Animation/IKRig/Source/IKRigEditor/Public/RigEditor/IKRigEditorController.h b/Engine/Plugins/Animation/IKRig/Source/IKRigEditor/Public/RigEditor/IKRigEditorController.h
--- a/Engine/Plugins/Animation/IKRig/Source/IKRigEditor/Public/RigEditor/IKRigEditorController.h
+++ b/Engine/Plugins/Animation/IKRig/Source/IKRigEditor/Public/RigEditor/IKRigEditorController.h
@@ -3,6 +3,7 @@
 #pragma once
 
 #include "IKRigSolver.h"
+#include "SIKRigRetargetChainList.h"
 #include "Templates/SubclassOf.h"
 
 #include "IKRigEditorController.generated.h"
@@ -46,6 +47,8 @@ public:
 	bool IsElementConnectedToSolver(TSharedRef<FIKRigTreeElement> TreeElement, int32 SolverIndex);
 	/** determine if the element is connected to ANY solver */
 	bool IsElementConnectedToAnySolver(TSharedRef<FIKRigTreeElement> TreeElement);
+	/** determine if the element is an excluded bone*/
+	bool IsElementExcludedBone(TSharedRef<FIKRigTreeElement> TreeElement);
 	
 	/** todo show BONE transform in details view */
 	void ShowDetailsForBone(const FName BoneName);
@@ -59,6 +62,9 @@ public:
 	void ShowDetailsForSolver(const int32 SolverIndex);
 	/** show nothing in details view */
 	void ShowEmptyDetails();
+
+	/** create a new retarget chain */
+	void AddNewRetargetChain(const FName ChainName, const FName StartBone, const FName EndBone);
 	
 	/** all modifications to the data model should go through this controller */
 	UIKRigController* AssetController;
@@ -79,6 +85,9 @@ public:
 	/** the solver stack view */
 	TSharedPtr<SIKRigSolverStack> SolverStackView;
 
+	/** the solver stack view */
+	TSharedPtr<SIKRigRetargetChainList> RetargetingView;
+
 	/** the persona toolkit */
 	TWeakPtr<FIKRigEditorToolkit> EditorToolkit;
 
diff --git a/Engine/Plugins/Animation/IKRig/Source/IKRigEditor/Public/RigEditor/IKRigEditorStyle.h b/Engine/Plugins/Animation/IKRig/Source/IKRigEditor/Public/RigEditor/IKRigEditorStyle.h
--- a/Engine/Plugins/Animation/IKRig/Source/IKRigEditor/Public/RigEditor/IKRigEditorStyle.h
+++ b/Engine/Plugins/Animation/IKRig/Source/IKRigEditor/Public/RigEditor/IKRigEditorStyle.h
@@ -22,22 +22,25 @@ public:
 		const FVector2D Icon24x24(24.0f, 24.0f);
 		const FVector2D Icon32x32(32.0f, 32.0f);
 		const FVector2D Icon40x40(40.0f, 40.0f);
+
+#define IKRIG_IMAGE_BRUSH(RelativePath, ...) FSlateImageBrush(RootToContentDir(RelativePath, TEXT(".png")), __VA_ARGS__)
+		Set("IKRig.Tree.Bone", new IKRIG_IMAGE_BRUSH("Slate/Bone_16x", Icon16x16));
+		Set("IKRig.Tree.BoneWithSettings", new IKRIG_IMAGE_BRUSH("Slate/BoneWithSettings_16x", Icon16x16));
+		Set("IKRig.Tree.Goal", new IKRIG_IMAGE_BRUSH("Slate/Goal_16x", Icon16x16));
+		Set("IKRig.Tree.Effector", new IKRIG_IMAGE_BRUSH("Slate/Effector_16x", Icon16x16));
 		
-		Set("IKRig.Tree.Bone", new IMAGE_BRUSH("Slate/Bone_16x", Icon16x16));
-		Set("IKRig.Tree.BoneWithSettings", new IMAGE_BRUSH("Slate/BoneWithSettings_16x", Icon16x16));
-		Set("IKRig.Tree.Goal", new IMAGE_BRUSH("Slate/Goal_16x", Icon16x16));
-		Set("IKRig.Tree.Effector", new IMAGE_BRUSH("Slate/Effector_16x", Icon16x16));
-		
-		Set("IKRig.TabIcon", new IMAGE_BRUSH("Slate/Tab_16x", Icon16x16));
+		Set("IKRig.TabIcon", new IKRIG_IMAGE_BRUSH("Slate/Tab_16x", Icon16x16));
 		
-		Set("IKRig.Solver", new IMAGE_BRUSH("Slate/Solver_16x", Icon16x16));
-		Set("IKRig.DragSolver", new IMAGE_BRUSH("Slate/DragSolver", FVector2D(6, 15)));
+		Set("IKRig.Solver", new IKRIG_IMAGE_BRUSH("Slate/Solver_16x", Icon16x16));
+		Set("IKRig.DragSolver", new IKRIG_IMAGE_BRUSH("Slate/DragSolver", FVector2D(6, 15)));
 
-		Set("IKRig.Reset", new IMAGE_BRUSH("Slate/Reset", Icon40x40));
-		Set("IKRig.Reset.Small", new IMAGE_BRUSH("Slate/Reset", Icon20x20));
+		Set("IKRig.Reset", new IKRIG_IMAGE_BRUSH("Slate/Reset", Icon40x40));
+		Set("IKRig.Reset.Small", new IKRIG_IMAGE_BRUSH("Slate/Reset", Icon20x20));
+#undef IKRIG_IMAGE_BRUSH
 		
 		FTextBlockStyle NormalText = FEditorStyle::GetWidgetStyle<FTextBlockStyle>("SkeletonTree.NormalFont");
-		Set( "IKRig.Tree.AffectedBoneText", FTextBlockStyle(NormalText).SetColorAndOpacity( FLinearColor::Green ));
+		Set( "IKRig.Tree.NormalText", FTextBlockStyle(NormalText));
+		Set( "IKRig.Tree.ItalicText", FTextBlockStyle(NormalText).SetFont(DEFAULT_FONT("Italic", 10)));
 		
 		FSlateStyleRegistry::RegisterSlateStyle(*this);
 	}
@@ -52,6 +55,4 @@ public:
 	{
 		FSlateStyleRegistry::UnRegisterSlateStyle(*this);
 	}
-};
-
-#undef IMAGE_BRUSH
\ No newline at end of file
+};
\ No newline at end of file
diff --git a/Engine/Plugins/Animation/IKRig/Source/IKRigEditor/Public/RigEditor/IKRigRetargetChainTabSummoner.h b/Engine/Plugins/Animation/IKRig/Source/IKRigEditor/Public/RigEditor/IKRigRetargetChainTabSummoner.h
new file mode 100644
--- /dev/null
+++ b/Engine/Plugins/Animation/IKRig/Source/IKRigEditor/Public/RigEditor/IKRigRetargetChainTabSummoner.h
@@ -0,0 +1,21 @@
+// Copyright Epic Games, Inc. All Rights Reserved.
+
+#pragma once
+
+#include "CoreMinimal.h"
+#include "WorkflowOrientedApp/WorkflowTabFactory.h"
+
+class FIKRigEditorToolkit;
+
+struct FIKRigRetargetChainTabSummoner : public FWorkflowTabFactory
+{
+public:
+	static const FName TabID;
+	
+	FIKRigRetargetChainTabSummoner(const TSharedRef<FIKRigEditorToolkit>& InIKRigEditor);
+	virtual TSharedRef<SWidget> CreateTabBody(const FWorkflowTabSpawnInfo& Info) const override;
+	virtual TSharedPtr<SToolTip> CreateTabToolTipWidget(const FWorkflowTabSpawnInfo& Info) const override;
+
+protected:
+	TWeakPtr<FIKRigEditorToolkit> IKRigEditor;
+};
diff --git a/Engine/Plugins/Animation/IKRig/Source/IKRigEditor/Public/RigEditor/IKRigSkeletonCommands.h b/Engine/Plugins/Animation/IKRig/Source/IKRigEditor/Public/RigEditor/IKRigSkeletonCommands.h
--- a/Engine/Plugins/Animation/IKRig/Source/IKRigEditor/Public/RigEditor/IKRigSkeletonCommands.h
+++ b/Engine/Plugins/Animation/IKRig/Source/IKRigEditor/Public/RigEditor/IKRigSkeletonCommands.h
@@ -34,6 +34,16 @@ public:
 	TSharedPtr< FUICommandInfo > AddBoneSettings;
 	/** remove bone setting from bone */
 	TSharedPtr< FUICommandInfo > RemoveBoneSettings;
+	
+	/** remove bone from solve */
+	TSharedPtr< FUICommandInfo > ExcludeBone;
+	/** add bone to solve */
+	TSharedPtr< FUICommandInfo > IncludeBone;
+
+	/** create retarget chain */
+	TSharedPtr< FUICommandInfo > NewRetargetChain;
+	/** set the root of the retargeting */
+	TSharedPtr< FUICommandInfo > SetRetargetRoot;
 
 	/** initialize commands */
 	virtual void RegisterCommands() override;
diff --git a/Engine/Plugins/Animation/IKRig/Source/IKRigEditor/Public/RigEditor/SIKRigRetargetChainList.h b/Engine/Plugins/Animation/IKRig/Source/IKRigEditor/Public/RigEditor/SIKRigRetargetChainList.h
new file mode 100644
--- /dev/null
+++ b/Engine/Plugins/Animation/IKRig/Source/IKRigEditor/Public/RigEditor/SIKRigRetargetChainList.h
@@ -0,0 +1,151 @@
+// Copyright Epic Games, Inc. All Rights Reserved.
+
+#pragma once
+
+#include "CoreMinimal.h"
+#include "EditorUndoClient.h"
+#include "DragAndDrop/DecoratedDragDropOp.h"
+#include "Widgets/SCompoundWidget.h"
+#include "Widgets/Views/STableViewBase.h"
+#include "Widgets/Views/STableRow.h"
+#include "Widgets/Views/SListView.h"
+#include "Framework/Commands/UICommandList.h"
+
+#include "SIKRigRetargetChainList.generated.h"
+
+class UIKRigEffectorGoal;
+class FIKRigEditorController;
+class SIKRigRetargetChainList;
+class FIKRigEditorToolkit;
+struct FBoneChain;
+
+
+class FRetargetChainElement
+{
+public:
+
+	TSharedRef<ITableRow> MakeListRowWidget(
+		const TSharedRef<STableViewBase>& InOwnerTable,
+        TSharedRef<FRetargetChainElement> InStackElement,
+        TSharedPtr<SIKRigRetargetChainList> InChainList);
+
+	static TSharedRef<FRetargetChainElement> Make(FBoneChain* InChain)
+	{
+		return MakeShareable(new FRetargetChainElement(InChain));
+	}
+
+	FBoneChain* Chain;
+
+private:
+	
+	/** Hidden constructor, always use Make above */
+	FRetargetChainElement(FBoneChain* InChain) : Chain(InChain) {}
+
+	/** Hidden constructor, always use Make above */
+	FRetargetChainElement() {}
+};
+
+
+class SIKRigRetargetChainRow : public SMultiColumnTableRow<TSharedPtr<FRetargetChainElement>>
+{
+public:
+	
+	void Construct(
+        const FArguments& InArgs,
+        const TSharedRef<STableViewBase>& InOwnerTableView,
+        TSharedRef<FRetargetChainElement> InChainElement,
+        TSharedPtr<SIKRigRetargetChainList> InChainList);
+
+	/** Overridden from SMultiColumnTableRow.  Generates a widget for this column of the table row. */
+	virtual TSharedRef<SWidget> GenerateWidgetForColumn(const FName& ColumnName) override;
+
+	TSharedRef<SWidget> MakeBoneComboEntryWidget(FName InItem) const;
+	TSharedRef<SWidget> MakeGoalComboEntryWidget(TSharedPtr<FString> InItem) const;
+	
+	void OnStartBoneComboSelectionChanged(FName InName, ESelectInfo::Type SelectInfo);
+	void OnEndBoneComboSelectionChanged(FName InName, ESelectInfo::Type SelectInfo);
+	void OnGoalComboSelectionChanged(TSharedPtr<FString> InGoalName, ESelectInfo::Type SelectInfo);
+	
+	FText GetStartBoneName() const;
+	FText GetEndBoneName() const;
+	FText GetGoalName() const;
+	
+	void OnRenameChain(const FText& InText, ETextCommit::Type ) const;
+
+private:
+
+	TArray<TSharedPtr<FString>> GoalOptions;
+	TWeakPtr<FRetargetChainElement> ChainElement;
+	TWeakPtr<SIKRigRetargetChainList> ChainList;
+
+	friend SIKRigRetargetChainList;
+};
+
+USTRUCT()
+struct FIKRigRetargetChainSettings
+{
+	GENERATED_BODY()
+
+	FIKRigRetargetChainSettings(){};
+	
+	FIKRigRetargetChainSettings(FName InChainName, FName InStartBone, FName InEndBone)
+	:	ChainName(InChainName),
+		StartBone(InStartBone),
+		EndBone(InEndBone){}
+
+	UPROPERTY(EditAnywhere, Category = "Bone Chain")
+	FName ChainName;
+
+	UPROPERTY(VisibleAnywhere, Category = "Bone Chain")
+	FName StartBone;
+
+	UPROPERTY(VisibleAnywhere, Category = "Bone Chain")
+	FName EndBone;
+};
+
+typedef SListView< TSharedPtr<FRetargetChainElement> > SRetargetChainListViewType;
+
+class SIKRigRetargetChainList : public SCompoundWidget, public FEditorUndoClient
+{
+public:
+	SLATE_BEGIN_ARGS(SIKRigRetargetChainList) {}
+	SLATE_END_ARGS()
+
+    ~SIKRigRetargetChainList();
+
+	void Construct(const FArguments& InArgs, TSharedRef<FIKRigEditorController> InEditorController);
+
+private:
+	
+	/** menu for adding new solver commands */
+	TSharedPtr<FUICommandList> CommandList;
+	
+	/** editor controller */
+	TWeakPtr<FIKRigEditorController> EditorController;
+
+	/** retarget root widget */
+	TSharedPtr<SEditableTextBox> RetargetRootTextBox;
+	
+	/** list view */
+	TSharedPtr<SRetargetChainListViewType> ListView;
+	TArray< TSharedPtr<FRetargetChainElement> > ListViewItems;
+	/** END list view */
+
+	/** callback when "Add New Chain" clicked */
+	bool IsAddChainEnabled() const;
+	/** when a chain is clicked on in the stack view */
+	void OnItemClicked(TSharedPtr<FRetargetChainElement> InItem);
+
+	/** list view generate row callback */
+	TSharedRef<ITableRow> MakeListRowWidget(TSharedPtr<FRetargetChainElement> InElement, const TSharedRef<STableViewBase>& OwnerTable);
+
+	/** call to refresh the list view */
+	void RefreshView();
+
+	/** SWidget interface */
+	virtual FReply OnKeyDown(const FGeometry& MyGeometry, const FKeyEvent& InKeyEvent);
+	/** END SWidget interface */
+
+	friend SIKRigRetargetChainRow;
+	friend FIKRigEditorController;
+};
diff --git a/Engine/Plugins/Animation/IKRig/Source/IKRigEditor/Public/RigEditor/SIKRigSkeleton.h b/Engine/Plugins/Animation/IKRig/Source/IKRigEditor/Public/RigEditor/SIKRigSkeleton.h
--- a/Engine/Plugins/Animation/IKRig/Source/IKRigEditor/Public/RigEditor/SIKRigSkeleton.h
+++ b/Engine/Plugins/Animation/IKRig/Source/IKRigEditor/Public/RigEditor/SIKRigSkeleton.h
@@ -4,6 +4,7 @@
 
 #include "CoreMinimal.h"
 #include "EditorUndoClient.h"
+#include "SIKRigRetargetChainList.h"
 #include "DragAndDrop/DecoratedDragDropOp.h"
 #include "Widgets/SCompoundWidget.h"
 #include "Widgets/Views/STreeView.h"
@@ -12,6 +13,7 @@
 
 #include "SIKRigSkeleton.generated.h"
 
+struct FIKRigSkeletonChain;
 class FIKRigEditorController;
 class SIKRigSkeleton;
 class FIKRigEditorToolkit;
@@ -106,7 +108,7 @@ public:
 	/** Restore the expansion infos map from the saved snapshot after tree reconstruction */
 	void RestoreSparseItemInfos(TSharedPtr<FIKRigTreeElement> ItemPtr)
 	{
-		for (const auto& Pair : OldSparseItemInfos)
+		for (const TTuple<TSharedPtr<FIKRigTreeElement>, FSparseItemInfo>& Pair : OldSparseItemInfos)
 		{
 			if (Pair.Key->Key == ItemPtr->Key)
 			{
@@ -184,11 +186,26 @@ private:
 	void HandleRemoveBoneSettings();
 	bool CanRemoveBoneSettings();
 	/** END per-bone settings */
+	
+	/** exclude/include bones */
+	void HandleExcludeBone();
+	bool CanExcludeBone();
+	void HandleIncludeBone();
+	bool CanIncludeBone();
+	/** END exclude/include bones */
+
+	/** retarget chains */
+	void HandleNewRetargetChain();
+	bool CanAddNewRetargetChain();
+	void HandleSetRetargetRoot();
+	bool CanSetRetargetRoot();
+	/** END retarget chains */
 
 	/** selection state queries */
 	void GetSelectedBones(TArray<TSharedPtr<FIKRigTreeElement>>& OutBoneItems);
 	void GetSelectedGoals(TArray<TSharedPtr<FIKRigTreeElement>>& OutSelectedGoals) const;
 	void SetSelectedGoalsFromViewport(const TArray<FName>& GoalNames);
+	void GetSelectedBoneChains(TArray<FIKRigSkeletonChain>& OutChains);
 	/** END selection state queries */
 
 	/** skeleton import */
@@ -230,4 +247,5 @@ private:
 	friend SIKRigSkeletonItem;
 	friend FIKRigEditorController;
 	friend SIKRigSolverStack;
+	friend SIKRigRetargetChainList;
 };
diff --git a/Engine/Plugins/Experimental/FullBodyIK/Source/PBIK/Private/Core/PBIKConstraint.cpp b/Engine/Plugins/Experimental/FullBodyIK/Source/PBIK/Private/Core/PBIKConstraint.cpp
--- a/Engine/Plugins/Experimental/FullBodyIK/Source/PBIK/Private/Core/PBIKConstraint.cpp
+++ b/Engine/Plugins/Experimental/FullBodyIK/Source/PBIK/Private/Core/PBIKConstraint.cpp
@@ -88,8 +88,8 @@ FVector FJointConstraint::GetPositionCorrection(FVector& OutBodyToA, FVector& Ou
 
 void FJointConstraint::ApplyRotationCorrection(FQuat DeltaQ) const
 {
-	A->ApplyRotationDelta(DeltaQ);
-	B->ApplyRotationDelta(DeltaQ * -1.0f);
+	A->ApplyRotationDelta(DeltaQ * (1.0f - A->J.RotationStiffness));
+	B->ApplyRotationDelta(DeltaQ * -1.0f * (1.0f - B->J.RotationStiffness));
 }
 
 void FJointConstraint::UpdateJointLimits()
diff --git a/Engine/Plugins/Experimental/FullBodyIK/Source/PBIK/Private/Core/PBIKSolver.cpp b/Engine/Plugins/Experimental/FullBodyIK/Source/PBIK/Private/Core/PBIKSolver.cpp
--- a/Engine/Plugins/Experimental/FullBodyIK/Source/PBIK/Private/Core/PBIKSolver.cpp
+++ b/Engine/Plugins/Experimental/FullBodyIK/Source/PBIK/Private/Core/PBIKSolver.cpp
@@ -12,7 +12,7 @@ namespace PBIK
 	{
 		check(InBone);
 		Bone = InBone;
-		SetGoal(Bone->Position, Bone->Rotation, 1.0f, 1.0f, 1.0f, false);
+		SetGoal(Bone->Position, Bone->Rotation, 1.0f, 1.0f, 1.0f, 0.0f);
 	}
 
 	void FEffector::SetGoal(
@@ -21,7 +21,7 @@ namespace PBIK
 		float InTransformAlpha,
 		float InStrengthAlpha,
 		float InPullChainAlpha,
-		bool bInPinRotation)
+		float InPinRotation)
 	{
 		PositionOrig = Bone->Position;
 		RotationOrig = Bone->Rotation;
@@ -32,7 +32,7 @@ namespace PBIK
 		TransformAlpha = InTransformAlpha;
 		StrengthAlpha = InStrengthAlpha;
 		PullChainAlpha = InPullChainAlpha;
-		bPinRotation = bInPinRotation;
+		PinRotation = InPinRotation;
 	}
 
 	void FEffector::UpdateFromInputs(const FBone& SolverRoot)
@@ -312,7 +312,7 @@ void FPBIKSolver::ApplyPreferredAngles()
 void FPBIKSolver::UpdateBodies(const FPBIKSolverSettings& Settings)
 {
 	// this creates a better pose from which to start the constraint solving
-	if (Settings.bPrePullRoot)
+	if (Settings.bPrePullRoot && !Settings.bPinRoot)
 	{
 		// pull ALL bodies towards the effectors
 		PullRootTowardsEffectors();
@@ -389,9 +389,11 @@ void FPBIKSolver::UpdateBonesFromBodies()
 		
 		Bone->Position = Bone->Parent->Position + Bone->Parent->Rotation * Bone->LocalPositionOrig;
 
-		if (Effector.bPinRotation)
+		if (Effector.PinRotation > SMALL_NUMBER)
 		{
-			Bone->Rotation = Effector.Rotation; // optionally pin rotation to that of effector
+			// optionally pin rotation to that of effector
+			const float RotAmount = FMath::Clamp(0.0f, 1.0f, Effector.PinRotation);
+			Bone->Rotation = FQuat::FastLerp(Bone->Rotation, Effector.Rotation, RotAmount).GetNormalized();
 		}else
 		{
 			Bone->Rotation = Bone->Parent->Rotation * Bone->LocalRotationOrig;
@@ -773,8 +775,8 @@ void FPBIKSolver::SetEffectorGoal(
 	const float OffsetAlpha,
 	const float StrengthAlpha,
 	const float PullChainAlpha,
-	const bool bPinRotation)
+	const float PinRotation)
 {
 	check(Index >= 0 && Index < Effectors.Num());
-	Effectors[Index].SetGoal(InPosition, InRotation, OffsetAlpha, StrengthAlpha, PullChainAlpha, bPinRotation);
+	Effectors[Index].SetGoal(InPosition, InRotation, OffsetAlpha, StrengthAlpha, PullChainAlpha, PinRotation);
 }
diff --git a/Engine/Plugins/Experimental/FullBodyIK/Source/PBIK/Private/RigUnit_PBIK.cpp b/Engine/Plugins/Experimental/FullBodyIK/Source/PBIK/Private/RigUnit_PBIK.cpp
--- a/Engine/Plugins/Experimental/FullBodyIK/Source/PBIK/Private/RigUnit_PBIK.cpp
+++ b/Engine/Plugins/Experimental/FullBodyIK/Source/PBIK/Private/RigUnit_PBIK.cpp
@@ -7,16 +7,30 @@ FRigUnit_PBIK_Execute()
 {
 	DECLARE_SCOPE_HIERARCHICAL_COUNTER_RIGUNIT()
 
+	if (Context.State == EControlRigState::Init)
+	{
+		BoneSettingToSolverBoneIndex.Reset();
+		Solver.Reset();
+		SolverBoneToElementIndex.Reset();
+		EffectorSolverIndices.Reset();
+		bNeedsInit = true;
+		return;
+	}
+
+	// only updates from here on...
+	if (Context.State != EControlRigState::Update)
+	{
+		return;
+	}
+
 	URigHierarchy* Hierarchy = ExecuteContext.Hierarchy;
 	if (Hierarchy == nullptr)
 	{
 		return;
 	}
 
-	if (Context.State == EControlRigState::Init)
+	if (bNeedsInit)
 	{
-		BoneSettingToSolverBoneIndex.Reset();
-		
 		// check how many effectors are assigned to a bone
 		int NumEffectors = 0;
 		for (const FPBIKEffector& Effector : Effectors)
@@ -34,13 +48,9 @@ FRigUnit_PBIK_Execute()
 		{
 			return; // not setup yet
 		}
-
-		// reset all internal data
-		Solver.Reset();
 		
 		// create solver bone index to elements index map
 		TArray<FRigBoneElement*> BoneElements = Hierarchy->GetBones(true);
-		SolverBoneToElementIndex.Reset();
 		for (int B = 0; B < BoneElements.Num(); ++B)
 		{
 			SolverBoneToElementIndex.Add(BoneElements[B]->GetIndex());
@@ -51,8 +61,20 @@ FRigUnit_PBIK_Execute()
 		{
 			FName Name = BoneElements[B]->GetName();
 
+			// get the first parent that is not excluded
+			int32 ParentElementIndex = INDEX_NONE;
+			FRigBaseElement* ParentElement = Hierarchy->GetFirstParent(BoneElements[B]);
+			while (ParentElement)
+			{
+				if (!ExcludedBones.Contains(ParentElement->GetKey().Name))
+				{
+					ParentElementIndex = ParentElement->GetIndex();
+					break;
+				}
+				ParentElement = Hierarchy->GetFirstParent(ParentElement);
+			}
+
 			// get the parent bone solver index
-			const int ParentElementIndex = Hierarchy->GetFirstParent(BoneElements[B]->GetIndex());
 			int ParentIndex = -1;
 			for (int P=0; P<SolverBoneToElementIndex.Num(); ++P)
 			{
@@ -71,7 +93,6 @@ FRigUnit_PBIK_Execute()
 		}
 		
 		// create effectors
-		EffectorSolverIndices.Reset();
 		for (const FPBIKEffector& Effector : Effectors)
 		{
 			int32 IndexInSolver = Solver.AddEffector(Effector.Bone);
@@ -80,8 +101,7 @@ FRigUnit_PBIK_Execute()
 		
 		// initialize
 		Solver.Initialize();
-
-		return; // don't update during init
+		bNeedsInit = false;
 	}
 
 	if (!Solver.IsReadyToSimulate())
@@ -145,7 +165,7 @@ FRigUnit_PBIK_Execute()
 			Effector.OffsetAlpha,
 			Effector.StrengthAlpha,
 			Effector.PullChainAlpha,
-			Effector.bPinRotation);
+			Effector.PinRotation);
 	}
 
 	// solve
diff --git a/Engine/Plugins/Experimental/FullBodyIK/Source/PBIK/Public/Core/PBIKSolver.h b/Engine/Plugins/Experimental/FullBodyIK/Source/PBIK/Public/Core/PBIKSolver.h
--- a/Engine/Plugins/Experimental/FullBodyIK/Source/PBIK/Public/Core/PBIKSolver.h
+++ b/Engine/Plugins/Experimental/FullBodyIK/Source/PBIK/Public/Core/PBIKSolver.h
@@ -37,7 +37,7 @@ struct FEffector
 
 	FBone* Bone;
 	TWeakPtr<FPinConstraint> Pin;
-	bool bPinRotation;
+	float PinRotation;
 	FRigidBody* ParentSubRoot = nullptr;
 	float DistanceToSubRootInInputPose;
 	float DistToRootAlongBones;
@@ -58,7 +58,7 @@ struct FEffector
 		float InTransformAlpha,
 		float InStrengthAlpha,
 		float InPullChainAlpha,
-		bool bInPinRotation);
+		float InPinRotation);
 
 	void UpdateFromInputs(const FBone& SolverRoot);
 	void ApplyPreferredAngles();
@@ -91,7 +91,9 @@ struct PBIK_API FPBIKSolverSettings
 	UPROPERTY(EditAnywhere, Category = SolverSettings)
 	bool bPinRoot = false;
 
-	/** If true, solver will kinematically rotate skeleton limbs to approximate a solved pose before running constraint iterations. This can decrease the amount of iterations needed to achieve a converged pose. Default is true. */
+	/** (Default is true) If true, solver runs a pre-pass (before the constraint solve) which will translate the whole body by the average motion of the stretched effectors.
+	 *This can decrease the amount of iterations needed to achieve a converged pose when the effectors are pulling the body far.
+	 *This option has no effect if PinRoot is set to True.*/
 	UPROPERTY(EditAnywhere, Category = SolverSettings)
 	bool bPrePullRoot = true;
 
@@ -136,7 +138,7 @@ public:
 		const float OffsetAlpha, 
 		const float StrengthAlpha,
 		const float PullChainAlpha,
-		const bool bPinRotation);
+		const float PinRotation);
 
 	void GetBoneGlobalTransform(const int32 Index, FTransform& OutTransform);
 
diff --git a/Engine/Plugins/Experimental/FullBodyIK/Source/PBIK/Public/PBIK_Shared.h b/Engine/Plugins/Experimental/FullBodyIK/Source/PBIK/Public/PBIK_Shared.h
--- a/Engine/Plugins/Experimental/FullBodyIK/Source/PBIK/Public/PBIK_Shared.h
+++ b/Engine/Plugins/Experimental/FullBodyIK/Source/PBIK/Public/PBIK_Shared.h
@@ -28,37 +28,65 @@ struct PBIK_API FPBIKBoneSetting
 		Z(EPBIKLimitType::Free),
 		PreferredAngles(FVector::ZeroVector){}
 
+	/** The Bone that these settings will be applied to. */
 	UPROPERTY(EditAnywhere, Category = Bone, meta = (Constant, CustomWidget = "BoneName"))
 	FName Bone;
 
+	/** Range is 0 to 1 (Default is 0). At higher values, the bone will resist rotating (forcing other bones to compensate). */
 	UPROPERTY(EditAnywhere, Category = Stiffness, meta = (ClampMin = "0", ClampMax = "1", UIMin = "0.0", UIMax = "1.0"))
 	float RotationStiffness = 0.0f;
+
+	/** Range is 0 to 1 (Default is 0). At higher values, the bone will resist translational motion (forcing other bones to compensate). */
 	UPROPERTY(EditAnywhere, Category = Stiffness, meta = (ClampMin = "0", ClampMax = "1", UIMin = "0.0", UIMax = "1.0"))
 	float PositionStiffness = 0.0f;
 
+	/** Limit the rotation angle of the bone on the X axis.
+	 *Free: can rotate freely in this axis.
+	 *Limited: rotation is clamped between the min/max angles relative to the Skeletal Mesh reference pose.
+	 *Locked: no rotation is allowed in this axis (will remain at reference pose angle). */
 	UPROPERTY(EditAnywhere, Category = Limits)
 	EPBIKLimitType X;
+	/**Range is -180 to 180 (Default is 0). Degrees of rotation in the negative X direction to allow when joint is in "Limited" mode. */
 	UPROPERTY(EditAnywhere, Category = Limits, meta = (ClampMin = "-180", ClampMax = "0", UIMin = "-180.0", UIMax = "0.0"))
 	float MinX = 0.0f;
+	/**Range is -180 to 180 (Default is 0). Degrees of rotation in the positive X direction to allow when joint is in "Limited" mode. */
 	UPROPERTY(EditAnywhere, Category = Limits, meta = (ClampMin = "0", ClampMax = "180", UIMin = "0.0", UIMax = "180.0"))
 	float MaxX = 0.0f;
 
+	/** Limit the rotation angle of the bone on the Y axis.
+	*Free: can rotate freely in this axis.
+	*Limited: rotation is clamped between the min/max angles relative to the Skeletal Mesh reference pose.
+	*Locked: no rotation is allowed in this axis (will remain at input pose angle). */
 	UPROPERTY(EditAnywhere, Category = Limits)
 	EPBIKLimitType Y;
+	/**Range is -180 to 180 (Default is 0). Degrees of rotation in the negative Y direction to allow when joint is in "Limited" mode. */
 	UPROPERTY(EditAnywhere, Category = Limits, meta = (ClampMin = "-180", ClampMax = "0", UIMin = "-180.0", UIMax = "0.0"))
 	float MinY = 0.0f;
+	/**Range is -180 to 180 (Default is 0). Degrees of rotation in the positive Y direction to allow when joint is in "Limited" mode. */
 	UPROPERTY(EditAnywhere, Category = Limits, meta = (ClampMin = "0", ClampMax = "180", UIMin = "0.0", UIMax = "180.0"))
 	float MaxY = 0.0f;
 
+	/** Limit the rotation angle of the bone on the Z axis.
+	*Free: can rotate freely in this axis.
+	*Limited: rotation is clamped between the min/max angles relative to the Skeletal Mesh reference pose.
+	*Locked: no rotation is allowed in this axis (will remain at input pose angle). */
 	UPROPERTY(EditAnywhere, Category = Limits)
 	EPBIKLimitType Z;
+	/**Range is -180 to 180 (Default is 0). Degrees of rotation in the negative Z direction to allow when joint is in "Limited" mode. */
 	UPROPERTY(EditAnywhere, Category = Limits, meta = (ClampMin = "-180", ClampMax = "0", UIMin = "-180.0", UIMax = "0.0"))
 	float MinZ = 0.0f;
+	/**Range is -180 to 180 (Default is 0). Degrees of rotation in the positive Z direction to allow when joint is in "Limited" mode. */
 	UPROPERTY(EditAnywhere, Category = Limits, meta = (ClampMin = "0", ClampMax = "180", UIMin = "0.0", UIMax = "180.0"))
 	float MaxZ = 0.0f;
 
+	/**When true, this bone will "prefer" to rotate in the direction specified by the Preferred Angles when the chain it belongs to is compressed.
+	 * Preferred Angles should be the first method used to fix bones that bend in the wrong direction (rather than limits).
+	 * The resulting angles can be visualized on their own by temporarily setting the Solver iterations to 0 and moving the effectors.*/
 	UPROPERTY(EditAnywhere, Category = PreferredAngles)
 	bool bUsePreferredAngles = false;
+	/**The local Euler angles (in degrees) used to rotate this bone when the chain it belongs to is squashed.
+	 * This happens by moving the effector at the tip of the chain towards the root of the chain.
+	 * This can be used to coerce knees and elbows to bend in the anatomically "correct" direction without resorting to limits (which may require more iterations to converge).*/
 	UPROPERTY(EditAnywhere, Category = PreferredAngles)
 	FVector PreferredAngles;
 
diff --git a/Engine/Plugins/Experimental/FullBodyIK/Source/PBIK/Public/RigUnit_PBIK.h b/Engine/Plugins/Experimental/FullBodyIK/Source/PBIK/Public/RigUnit_PBIK.h
--- a/Engine/Plugins/Experimental/FullBodyIK/Source/PBIK/Public/RigUnit_PBIK.h
+++ b/Engine/Plugins/Experimental/FullBodyIK/Source/PBIK/Public/RigUnit_PBIK.h
@@ -19,9 +19,11 @@ struct FPBIKDebug
 {
 	GENERATED_BODY()
 
+	/** The scale of the debug drawing. */
 	UPROPERTY()
 	float DrawScale = 1.0f;
 
+	/** If true, turns on debug drawing for the node. */
 	UPROPERTY()
 	bool bDrawDebug = false;
 
@@ -52,23 +54,39 @@ struct FPBIKEffector
 
 	FPBIKEffector()	: Bone(NAME_None) {}
 
+	/** The bone that this effector will pull on. */
 	UPROPERTY(meta = (Constant, CustomWidget = "BoneName"))
 	FName Bone;
 
+	/** The target location and rotation for this effector. The solver will try to get the specified bone to reach this location.*/
 	UPROPERTY()
 	FTransform Transform;
 
+	/** Range 0-1 (default is 1.0). This blends the effector translation and rotation between its transform and the transform of the bone from the input pose.
+	 * At 0.0, the effector will be placed at the location of the bone in the input pose (effectively pinning the skeleton to the input pose).
+	 * At 1.0, the effector will be placed at the supplied location.*/
 	UPROPERTY()
 	float OffsetAlpha = 1.0f;
 
+	/** Range 0-1 (default is 1.0). The strength of the effector when pulling the bone towards it's target location.
+	 * At 0.0, the effector does not pull at all, but the bones between the effector and the root will still slightly resist motion from other effectors.
+	 * This can thus act as a "stabilizer" of sorts for parts of the body that you do not want to behave in a pure FK fashion.
+	 */
 	UPROPERTY()
 	float StrengthAlpha = 1.0f;
-	
+
+	/** Range 0-1 (default is 1.0). When enabled (greater than 0.0), the solver internally partitions the skeleton into 'chains' which extend from the effector to the nearest fork in the skeleton.
+	 *These chains are pre-rotated and translated, as a whole, towards the effector targets.
+	 *This can improve the results for sparse bone chains, and significantly improve convergence on dense bone chains.
+	 *But it may cause undesirable results in highly constrained bone chains (like robot arms).
+	 */
 	UPROPERTY()
 	float PullChainAlpha = 1.0f;
 
+	/** Range 0-1 (default is 1.0).
+	 *Blends the effector bone rotation between the rotation of the effector transform (1.0) and the rotation of the input bone (0.0).*/
 	UPROPERTY()
-	bool bPinRotation = true;
+	float PinRotation = 1.0f;
 };
 
 USTRUCT(meta=(DisplayName="Full Body IK", Category="Hierarchy", Keywords="Position Based, PBIK, IK, Full Body, Multi, Effector, N-Chain, FB"))
@@ -79,25 +97,39 @@ struct FRigUnit_PBIK : public FRigUnit_HighlevelBaseMutable
 	RIGVM_METHOD()
 	virtual void Execute(const FRigUnitContext& Context) override;
 
-	FRigUnit_PBIK()
-		: Root(NAME_None)
+	FRigUnit_PBIK() :
+		Root(NAME_None),
+		bNeedsInit(true)
 	{
 	}
 
+	/**This is usually the top-most skinned bone; often the "Pelvis" or "Hips", but can be set to any bone.
+	 *Bones above the root will be ignored by the solver.
+	 *Bones that are located *between* the Root and the effectors will be included in the solve.*/
 	UPROPERTY(meta = (Input, Constant, CustomWidget = "BoneName"))
 	FName Root;
 
+	/** An array of effectors. These specify target transforms for different parts of the skeleton. */
 	UPROPERTY(meta = (Input))
 	TArray<FPBIKEffector> Effectors;
+	
 	UPROPERTY(transient)
 	TArray<int32> EffectorSolverIndices;
 
+	/** Per-bone settings to control the resulting pose. Includes limits and preferred angles. */
 	UPROPERTY(meta = (Input))
 	TArray<FPBIKBoneSetting> BoneSettings;
 
+	/** These bones will be excluded from the solver. They will not bend and will not contribute to the constraint set.
+	 * Use the ExcludedBones array instead of setting Rotation Stiffness to very high values or Rotation Limits with zero range. */
+	UPROPERTY(meta = (Input, Constant, CustomWidget = "BoneName"))
+	TArray<FName> ExcludedBones;
+
+	/** Global solver settings. */
 	UPROPERTY(meta = (Input))
 	FPBIKSolverSettings Settings;
 
+	/** Debug drawing options. */
 	UPROPERTY(meta = (Input))
 	FPBIKDebug Debug;
 
@@ -109,4 +141,7 @@ struct FRigUnit_PBIK : public FRigUnit_HighlevelBaseMutable
 
 	UPROPERTY(transient)
 	FPBIKSolver Solver;
+
+	UPROPERTY(transient)
+	bool bNeedsInit;
 };
