IK Rig UX Improvements
Added pop-up to select Skeletal Mesh when creating asset.
Added pop-up to add solver when creating first goal.
Added "Preview Mode" to Goals for ability to preview additively while playing animation.
Added Animation Asset Browser
Uses standard Reset button icon

diff --git a/Engine/Plugins/Animation/IKRig/Source/IKRig/Private/IKRigDataTypes.cpp b/Engine/Plugins/Animation/IKRig/Source/IKRig/Private/IKRigDataTypes.cpp
--- a/Engine/Plugins/Animation/IKRig/Source/IKRig/Private/IKRigDataTypes.cpp
+++ b/Engine/Plugins/Animation/IKRig/Source/IKRig/Private/IKRigDataTypes.cpp
@@ -20,8 +20,22 @@ void FIKRigGoalContainer::SetIKGoal(const UIKRigEffectorGoal* InEffectorGoal)
 {
 	if (FIKRigGoal* Goal = Goals.Find(InEffectorGoal->GoalName))
 	{
-		Goal->Position = InEffectorGoal->CurrentTransform.GetTranslation();
-        Goal->Rotation = InEffectorGoal->CurrentTransform.Rotator();
+		if (InEffectorGoal->PreviewMode == EIKRigGoalPreviewMode::Absolute)
+		{
+			Goal->Position = InEffectorGoal->CurrentTransform.GetTranslation();
+			Goal->Rotation = InEffectorGoal->CurrentTransform.Rotator();
+			Goal->PositionSpace = EIKRigGoalSpace::Component;
+			Goal->RotationSpace = EIKRigGoalSpace::Component;
+		}
+		else
+		{
+			Goal->Position = InEffectorGoal->CurrentTransform.GetTranslation() - InEffectorGoal->InitialTransform.GetTranslation();
+			const FQuat RelativeRotation = InEffectorGoal->CurrentTransform.GetRotation() * InEffectorGoal->InitialTransform.GetRotation().Inverse();
+			Goal->Rotation = RelativeRotation.Rotator();
+			Goal->PositionSpace = EIKRigGoalSpace::Additive;
+			Goal->RotationSpace = EIKRigGoalSpace::Additive;
+		}
+		
         Goal->PositionAlpha = InEffectorGoal->PositionAlpha;
         Goal->RotationAlpha = InEffectorGoal->RotationAlpha;
 	}
diff --git a/Engine/Plugins/Animation/IKRig/Source/IKRig/Private/IKRigDefinition.cpp b/Engine/Plugins/Animation/IKRig/Source/IKRig/Private/IKRigDefinition.cpp
--- a/Engine/Plugins/Animation/IKRig/Source/IKRig/Private/IKRigDefinition.cpp
+++ b/Engine/Plugins/Animation/IKRig/Source/IKRig/Private/IKRigDefinition.cpp
@@ -93,4 +93,14 @@ void UIKRigDefinition::SortRetargetChains()
 	});
 }
 
+USkeleton* UIKRigDefinition::GetSkeletonAsset() const
+{
+	if (PreviewSkeletalMesh.IsNull())
+	{
+		return nullptr;
+	}
+
+	return PreviewSkeletalMesh->GetSkeleton();
+}
+
 #undef LOCTEXT_NAMESPACE
diff --git a/Engine/Plugins/Animation/IKRig/Source/IKRig/Public/IKRigDefinition.h b/Engine/Plugins/Animation/IKRig/Source/IKRig/Public/IKRigDefinition.h
--- a/Engine/Plugins/Animation/IKRig/Source/IKRig/Public/IKRigDefinition.h
+++ b/Engine/Plugins/Animation/IKRig/Source/IKRig/Public/IKRigDefinition.h
@@ -10,6 +10,13 @@
 
 class UIKRigSolver;
 
+UENUM(BlueprintType)
+enum class EIKRigGoalPreviewMode : uint8
+{
+	Additive		UMETA(DisplayName = "Additive"),
+	Absolute		UMETA(DisplayName = "Absolute"),
+};
+
 UCLASS()
 class IKRIG_API UIKRigEffectorGoal : public UObject
 {
@@ -21,27 +28,43 @@ public:
 	
 	UIKRigEffectorGoal(const FName& InGoalName, const FName& InBoneName) :  GoalName(InGoalName), BoneName(InBoneName){}
 
+	/** The name used to refer to this goal from outside systems.
+	 *This is the name to use when referring to this Goal from Blueprint, Anim Graph, Control Rig or IK Retargeter.*/
 	UPROPERTY(VisibleAnywhere, Category = "Goal Settings")
 	FName GoalName;
-	
+
+	/**The name of the bone that this Goal is located at.*/
 	UPROPERTY(VisibleAnywhere, Category = "Goal Settings")
 	FName BoneName;
 
+	/**Range 0-1, default is 1. Blend between the Goal position between the input bone pose (0.0) and the current goal transform (1.0).*/
 	UPROPERTY(EditAnywhere, Category = "Goal Settings", meta = (ClampMin = "0.0", ClampMax = "1.0", UIMin = "0.0", UIMax = "1.0"))
 	float PositionAlpha = 1.0f;
 
+	/**Range 0-1, default is 1. Blend between the Goal rotation between the input bone pose (0.0) and the current goal transform (1.0).*/
 	UPROPERTY(EditAnywhere, Category = "Goal Settings", meta = (ClampMin = "0.0", ClampMax = "1.0", UIMin = "0.0", UIMax = "1.0"))
 	float RotationAlpha = 1.0f;
 
+	/**Effects how this Goal transform is previewed in the IK Rig editor.
+	 *"Additive" interprets the Goal transform as being relative to the input pose. Useful for previewing animations.
+	 *"Absolute" pins the Goal transform to the Gizmo in the viewport.
+	 */
+	UPROPERTY(EditAnywhere, Category = "Goal Settings")
+	EIKRigGoalPreviewMode PreviewMode;
+
+	/**The current transform of this Goal, in the Global Space of the character.*/
 	UPROPERTY(EditAnywhere, Transient, Category = "Goal Settings")
 	FTransform CurrentTransform;
-	
+
+	/**The initial transform of this Goal, as defined by the initial transform of the Goal's bone in the reference pose.*/
 	UPROPERTY(VisibleAnywhere, Transient, Category = "Goal Settings")
 	FTransform InitialTransform;
 
+	/**The size of the Goal gizmo drawing in the editor viewport.*/
 	UPROPERTY(EditAnywhere, Category = "Goal Gizmo", meta = (ClampMin = "0.1", ClampMax = "1000.0", UIMin = "0.1", UIMax = "100.0"))
 	float GizmoSize = 7.0f;
 
+	/**The thickness of the Goal gizmo drawing in the editor viewport.*/
 	UPROPERTY(EditAnywhere, Category = "Goal Gizmo",  meta = (ClampMin = "0.0", ClampMax = "10.0", UIMin = "0.0", UIMax = "5.0"))
 	float GizmoThickness = 0.7f;
 
@@ -142,6 +165,8 @@ public:
 	virtual USkeletalMesh* GetPreviewMesh() const override;
 	/** END IInterface_PreviewMeshProvider interface */
 
+	USkeleton* GetSkeletonAsset() const;
+
 private:
 	
 	/** UObject interface */
diff --git a/Engine/Plugins/Animation/IKRig/Source/IKRigEditor/Private/RigEditor/IKRigDefinitionFactory.cpp b/Engine/Plugins/Animation/IKRig/Source/IKRigEditor/Private/RigEditor/IKRigDefinitionFactory.cpp
--- a/Engine/Plugins/Animation/IKRig/Source/IKRigEditor/Private/RigEditor/IKRigDefinitionFactory.cpp
+++ b/Engine/Plugins/Animation/IKRig/Source/IKRigEditor/Private/RigEditor/IKRigDefinitionFactory.cpp
@@ -3,6 +3,9 @@
 #include "RigEditor/IKRigDefinitionFactory.h"
 #include "IKRigDefinition.h"
 #include "AssetTypeCategories.h"
+#include "ContentBrowserModule.h"
+#include "IContentBrowserSingleton.h"
+#include "RigEditor/IKRigController.h"
 
 #define LOCTEXT_NAMESPACE "IKRigDefinitionFactory"
 
@@ -14,9 +17,27 @@ UIKRigDefinitionFactory::UIKRigDefinitionFactory()
 	SupportedClass = UIKRigDefinition::StaticClass();
 }
 
-UObject* UIKRigDefinitionFactory::FactoryCreateNew(UClass* Class, UObject* InParent, FName Name, EObjectFlags Flags, UObject* Context, FFeedbackContext* Warn)
+UObject* UIKRigDefinitionFactory::FactoryCreateNew(
+	UClass* Class,
+	UObject* InParent,
+	FName Name,
+	EObjectFlags Flags,
+	UObject* Context, 
+	FFeedbackContext* Warn)
 {
-	return NewObject<UIKRigDefinition>(InParent, Name, Flags | RF_Transactional);;
+	if (!SkeletalMesh)
+	{
+		UE_LOG(LogTemp, Warning, TEXT("Unable to create IK Rig. No Skeletal Mesh asset supplied."));
+		return nullptr;
+	}
+	
+	UIKRigDefinition* IKRig = NewObject<UIKRigDefinition>(InParent, Name, Flags | RF_Transactional);
+	
+	// imports the skeleton data into the IK Rig
+	UIKRigController* Controller = UIKRigController::GetIKRigController(IKRig);
+	Controller->SetSourceSkeletalMesh(SkeletalMesh, true);
+	
+	return IKRig;
 }
 
 bool UIKRigDefinitionFactory::ShouldShowInNewMenu() const
@@ -24,9 +45,47 @@ bool UIKRigDefinitionFactory::ShouldShowInNewMenu() const
 	return true;
 }
 
+void UIKRigDefinitionFactory::OnSkeletalMeshSelected(const FAssetData& SelectedAsset)
+{
+	SkeletalMesh = Cast<USkeletalMesh>(SelectedAsset.GetAsset());
+	PickerWindow->RequestDestroyWindow();
+}
+
 bool UIKRigDefinitionFactory::ConfigureProperties()
 {
-	return true;
+	// Null so we can check for selection later
+	SkeletalMesh = nullptr;
+
+	// Load the content browser module to display an asset picker
+	FContentBrowserModule& ContentBrowserModule = FModuleManager::LoadModuleChecked<FContentBrowserModule>("ContentBrowser");
+
+	FAssetPickerConfig AssetPickerConfig;
+
+	/** The asset picker will only show skeletal meshes */
+	AssetPickerConfig.Filter.ClassNames.Add(USkeletalMesh::StaticClass()->GetFName());
+
+	/** The delegate that fires when an asset was selected */
+	AssetPickerConfig.OnAssetSelected = FOnAssetSelected::CreateUObject(this, &UIKRigDefinitionFactory::OnSkeletalMeshSelected);
+
+	/** The default view mode should be a list view */
+	AssetPickerConfig.InitialAssetViewType = EAssetViewType::List;
+
+	PickerWindow = SNew(SWindow)
+	.Title(LOCTEXT("CreateIKRigOptions", "Pick Skeletal Mesh"))
+	.ClientSize(FVector2D(500, 600))
+	.SupportsMinimize(false) .SupportsMaximize(false)
+	[
+		SNew(SBorder)
+		.BorderImage( FEditorStyle::GetBrush("Menu.Background") )
+		[
+			ContentBrowserModule.Get().CreateAssetPicker(AssetPickerConfig)
+		]
+	];
+
+	GEditor->EditorAddModalWindow(PickerWindow.ToSharedRef());
+	PickerWindow.Reset();
+
+	return SkeletalMesh != nullptr;
 }
 
 FText UIKRigDefinitionFactory::GetDisplayName() const
diff --git a/Engine/Plugins/Animation/IKRig/Source/IKRigEditor/Private/RigEditor/IKRigEditorController.cpp b/Engine/Plugins/Animation/IKRig/Source/IKRigEditor/Private/RigEditor/IKRigEditorController.cpp
--- a/Engine/Plugins/Animation/IKRig/Source/IKRigEditor/Private/RigEditor/IKRigEditorController.cpp
+++ b/Engine/Plugins/Animation/IKRig/Source/IKRigEditor/Private/RigEditor/IKRigEditorController.cpp
@@ -2,6 +2,7 @@
 
 #include "RigEditor/IKRigEditorController.h"
 
+#include "IPersonaToolkit.h"
 #include "SKismetInspector.h"
 #include "Dialogs/Dialogs.h"
 
@@ -91,7 +92,7 @@ int32 FIKRigEditorController::GetSelectedSolverIndex()
 	return SelectedSolvers[0]->IndexInStack;
 }
 
-void FIKRigEditorController::PromptToAddSolverAtStartup() const
+void FIKRigEditorController::PromptToAddSolver() const
 {
 	if (AssetController->GetNumSolvers() > 0)
 	{
diff --git a/Engine/Plugins/Animation/IKRig/Source/IKRigEditor/Private/RigEditor/IKRigMode.cpp b/Engine/Plugins/Animation/IKRig/Source/IKRigEditor/Private/RigEditor/IKRigMode.cpp
--- a/Engine/Plugins/Animation/IKRig/Source/IKRigEditor/Private/RigEditor/IKRigMode.cpp
+++ b/Engine/Plugins/Animation/IKRig/Source/IKRigEditor/Private/RigEditor/IKRigMode.cpp
@@ -32,6 +32,7 @@ FIKRigMode::FIKRigMode(
 	TabFactories.RegisterFactory(PersonaModule.CreatePersonaViewportTabFactory(InHostingApp, ViewportArgs));
 	TabFactories.RegisterFactory(PersonaModule.CreateAdvancedPreviewSceneTabFactory(InHostingApp, InPreviewScene));
 	TabFactories.RegisterFactory(PersonaModule.CreateDetailsTabFactory(InHostingApp, FOnDetailsCreated::CreateSP(&IKRigEditor.Get(), &FIKRigEditorToolkit::HandleDetailsCreated)));
+	TabFactories.RegisterFactory(PersonaModule.CreateAnimationAssetBrowserTabFactory(InHostingApp, IKRigEditor->GetPersonaToolkit(), FOnOpenNewAsset::CreateSP(&IKRigEditor.Get(), &FIKRigEditorToolkit::HandleOpenNewAsset), FOnAnimationSequenceBrowserCreated::CreateSP(&IKRigEditor.Get(), &FIKRigEditorToolkit::HandleAnimationSequenceBrowserCreated), true));
 
 	// register custom tabs
 	TabFactories.RegisterFactory(MakeShared<FIKRigSkeletonTabSummoner>(IKRigEditor));
@@ -39,7 +40,7 @@ FIKRigMode::FIKRigMode(
 	TabFactories.RegisterFactory(MakeShared<FIKRigRetargetChainTabSummoner>(IKRigEditor));
 
 	// create tab layout
-	TabLayout = FTabManager::NewLayout("Standalone_IKRigEditor_Layout_v1.120")
+	TabLayout = FTabManager::NewLayout("Standalone_IKRigEditor_Layout_v1.122")
 		->AddArea
 		(
 			FTabManager::NewPrimaryArea()
@@ -91,6 +92,7 @@ FIKRigMode::FIKRigMode(
                     (
 						FTabManager::NewStack()
 						->SetSizeCoefficient(0.6f)
+						->AddTab(FPersonaTabs::AssetBrowserID, ETabState::OpenedTab)
 						->AddTab(FIKRigRetargetChainTabSummoner::TabID, ETabState::OpenedTab)
 						->SetForegroundTab(FIKRigRetargetChainTabSummoner::TabID)
                     )
diff --git a/Engine/Plugins/Animation/IKRig/Source/IKRigEditor/Private/RigEditor/IKRigToolkit.cpp b/Engine/Plugins/Animation/IKRig/Source/IKRigEditor/Private/RigEditor/IKRigToolkit.cpp
--- a/Engine/Plugins/Animation/IKRig/Source/IKRigEditor/Private/RigEditor/IKRigToolkit.cpp
+++ b/Engine/Plugins/Animation/IKRig/Source/IKRigEditor/Private/RigEditor/IKRigToolkit.cpp
@@ -50,7 +50,7 @@ void FIKRigEditorToolkit::InitAssetEditor(
 	PersonaToolkitArgs.OnPreviewSceneCreated = FOnPreviewSceneCreated::FDelegate::CreateSP(this, &FIKRigEditorToolkit::HandlePreviewSceneCreated);
 	
 	FPersonaModule& PersonaModule = FModuleManager::LoadModuleChecked<FPersonaModule>("Persona");
-	PersonaToolkit = PersonaModule.CreatePersonaToolkit(IKRigAsset, PersonaToolkitArgs);
+	PersonaToolkit = PersonaModule.CreatePersonaToolkit(IKRigAsset, PersonaToolkitArgs, IKRigAsset->GetSkeletonAsset());
 	
 	// when/if preview mesh is changed, we need to reinitialize the anim instance
     PersonaToolkit->GetPreviewScene()->RegisterOnPreviewMeshChanged(FOnPreviewMeshChanged::CreateSP(this, &FIKRigEditorToolkit::HandlePreviewMeshChanged));
@@ -125,8 +125,12 @@ void FIKRigEditorToolkit::FillToolbar(FToolBarBuilder& ToolbarBuilder)
 {
 	ToolbarBuilder.BeginSection("Reset");
 	{
-		ToolbarBuilder.AddToolBarButton(FIKRigCommands::Get().Reset,
-			NAME_None, TAttribute<FText>(), TAttribute<FText>(), FSlateIcon(FIKRigEditorStyle::Get().GetStyleSetName(), "IKRig.Reset"));
+		ToolbarBuilder.AddToolBarButton(
+			FIKRigCommands::Get().Reset,
+			NAME_None,
+			TAttribute<FText>(),
+			TAttribute<FText>(),
+			FSlateIcon(FAppStyle::Get().GetStyleSetName(),"Icons.Refresh"));
 	}
 	ToolbarBuilder.EndSection();
 }
@@ -183,6 +187,20 @@ void FIKRigEditorToolkit::PostRedo(bool bSuccess)
 	EditorController->RefreshAllViews();
 }
 
+void FIKRigEditorToolkit::HandleOpenNewAsset(UObject* InNewAsset)
+{
+	if (UAnimationAsset* NewAnimationAsset = Cast<UAnimationAsset>(InNewAsset))
+	{
+		EditorController->AnimInstance->SetAnimationAsset(NewAnimationAsset);
+	}
+}
+
+void FIKRigEditorToolkit::HandleAnimationSequenceBrowserCreated(
+	const TSharedRef<IAnimationSequenceBrowser>& InSequenceBrowser)
+{
+	EditorController->SequenceBrowser = InSequenceBrowser;
+}
+
 void FIKRigEditorToolkit::HandlePreviewSceneCreated(const TSharedRef<IPersonaPreviewScene>& InPersonaPreviewScene)
 {
 	AAnimationEditorPreviewActor* Actor = InPersonaPreviewScene->GetWorld()->SpawnActor<AAnimationEditorPreviewActor>(AAnimationEditorPreviewActor::StaticClass(), FTransform::Identity);
@@ -203,9 +221,9 @@ void FIKRigEditorToolkit::HandlePreviewSceneCreated(const TSharedRef<IPersonaPre
 
 	// apply mesh to the preview scene
 	InPersonaPreviewScene->SetPreviewMeshComponent(EditorController->SkelMeshComponent);
-	InPersonaPreviewScene->AddComponent(EditorController->SkelMeshComponent, FTransform::Identity);
 	InPersonaPreviewScene->SetAdditionalMeshesSelectable(false);
 	InPersonaPreviewScene->SetPreviewMesh(Mesh);
+	InPersonaPreviewScene->AddComponent(EditorController->SkelMeshComponent, FTransform::Identity);
 }
 
 void FIKRigEditorToolkit::HandlePreviewMeshChanged(USkeletalMesh* InOldSkeletalMesh, USkeletalMesh* InNewSkeletalMesh)
diff --git a/Engine/Plugins/Animation/IKRig/Source/IKRigEditor/Private/RigEditor/SIKRigSkeleton.cpp b/Engine/Plugins/Animation/IKRig/Source/IKRigEditor/Private/RigEditor/SIKRigSkeleton.cpp
--- a/Engine/Plugins/Animation/IKRig/Source/IKRigEditor/Private/RigEditor/SIKRigSkeleton.cpp
+++ b/Engine/Plugins/Animation/IKRig/Source/IKRigEditor/Private/RigEditor/SIKRigSkeleton.cpp
@@ -229,38 +229,7 @@ void SIKRigSkeleton::Construct(const FArguments& InArgs, TSharedRef<FIKRigEditor
 	ChildSlot
     [
         SNew(SVerticalBox)
-        +SVerticalBox::Slot()
-        .AutoHeight()
-        .VAlign(VAlign_Top)
-        .Padding(0.0f)
-        [
-            SNew(SBorder)
-            .Padding(0.0f)
-            .BorderImage(FEditorStyle::GetBrush("DetailsView.CategoryTop"))
-            .BorderBackgroundColor(FLinearColor(0.6f, 0.6f, 0.6f, 1.0f))
-            [
-                SNew(SVerticalBox)
-                + SVerticalBox::Slot()
-                .AutoHeight()
-                .VAlign(VAlign_Top)
-                [
-                    SNew(SHorizontalBox)
-                    .Visibility(this, &SIKRigSkeleton::IsImportButtonVisible)
-                    + SHorizontalBox::Slot()
-                    .VAlign(VAlign_Center)
-                    .HAlign(HAlign_Left)
-                    .FillWidth(1.f)
-                    .Padding(3.0f, 1.0f)
-                    [
-						SNew(SEditorHeaderButton)
-                        .Icon(FAppStyle::Get().GetBrush("Icons.Plus"))
-                        .Text(LOCTEXT("ImportSkeletonLabel", "Import Skeleton"))
-                        .ToolTipText(LOCTEXT("ImportSkeletonToolTip", "Import a skeletal mesh to apply IK to."))
-                        .OnClicked(this, &SIKRigSkeleton::OnImportSkeletonClicked)
-                    ]
-                ]
-            ]
-        ]
+
         +SVerticalBox::Slot()
         .Padding(0.0f, 0.0f)
         [
@@ -424,6 +393,12 @@ void SIKRigSkeleton::HandleNewGoal()
 	{
 		return; 
 	}
+
+	// add a default solver if there isn't one already
+	if (Controller->AssetController->GetNumSolvers() == 0)
+	{
+		Controller->PromptToAddSolver();	
+	}
 	
 	// create goal for each selected bone
 	FName LastCreatedGoalName = NAME_None;
@@ -1005,49 +980,6 @@ void SIKRigSkeleton::HandleRenameElement() const
 	SelectedGoals[0]->RequestRename();
 }
 
-FReply SIKRigSkeleton::OnImportSkeletonClicked()
-{
-	const TSharedPtr<FIKRigEditorController> Controller = EditorController.Pin();
-	if (!Controller.IsValid())
-	{
-		return FReply::Unhandled();
-	}
-	
-	FIKRigSkeletonImportSettings Settings;
-	const TSharedPtr<FStructOnScope> StructToDisplay = MakeShareable(new FStructOnScope(FIKRigSkeletonImportSettings::StaticStruct(), (uint8*)&Settings));
-	TSharedRef<SKismetInspector> KismetInspector = SNew(SKismetInspector);
-	KismetInspector->ShowSingleStruct(StructToDisplay);
-
-	SGenericDialogWidget::OpenDialog(LOCTEXT("IKRigSkeletonImport", "Import Skeleton"), KismetInspector, SGenericDialogWidget::FArguments(), true);
-
-	if (Settings.Mesh != nullptr)
-	{
-		ImportSkeleton(FAssetData(Settings.Mesh));
-		Controller->PromptToAddSolverAtStartup();
-	}
-	
-	return FReply::Handled();
-}
-
-EVisibility SIKRigSkeleton::IsImportButtonVisible() const
-{
-	const TSharedPtr<FIKRigEditorController> Controller = EditorController.Pin();
-	if (!Controller.IsValid())
-	{
-		return EVisibility::Hidden;
-	}
-	
-	if (UIKRigController* AssetController = Controller->AssetController)
-	{
-		if (AssetController->GetSkeleton().BoneNames.Num() > 0)
-		{
-			return EVisibility::Collapsed;
-		}
-	}
-	
-	return EVisibility::Visible;
-}
-
 void SIKRigSkeleton::RefreshTreeView(bool IsInitialSetup)
 {
 	const TSharedPtr<FIKRigEditorController> Controller = EditorController.Pin();
@@ -1298,32 +1230,6 @@ void SIKRigSkeleton::SetExpansionRecursive(
     }
 }
 
-void SIKRigSkeleton::ImportSkeleton(const FAssetData& InAssetData)
-{
-	USkeletalMesh* Mesh = Cast<USkeletalMesh>(InAssetData.GetAsset());
-	if (!Mesh)
-	{
-		return;
-	}
-	
-	const TSharedPtr<FIKRigEditorController> Controller = EditorController.Pin();
-	if (!Controller.IsValid())
-	{
-		return;
-	}
-	
-	// load the skeletal mesh in the viewport
-	TSharedRef<IPersonaToolkit> Persona = Controller->EditorToolkit.Pin()->GetPersonaToolkit();
-	Persona->SetPreviewMesh(Mesh, true);
-	Persona->GetPreviewScene()->FocusViews();
-
-	// import the skeleton data into the IKRig asset
-	Controller->AssetController->SetSourceSkeletalMesh(Mesh, true);
-
-	// update the bone hierarchy view
-	RefreshTreeView(true);
-}
-
 FReply SIKRigSkeleton::OnDragDetected(const FGeometry& MyGeometry, const FPointerEvent& MouseEvent)
 {
 	const TArray<TSharedPtr<FIKRigTreeElement>> SelectedItems = TreeView.Get()->GetSelectedItems();
diff --git a/Engine/Plugins/Animation/IKRig/Source/IKRigEditor/Public/RigEditor/IKRigDefinitionFactory.h b/Engine/Plugins/Animation/IKRig/Source/IKRigEditor/Public/RigEditor/IKRigDefinitionFactory.h
--- a/Engine/Plugins/Animation/IKRig/Source/IKRigEditor/Public/RigEditor/IKRigDefinitionFactory.h
+++ b/Engine/Plugins/Animation/IKRig/Source/IKRigEditor/Public/RigEditor/IKRigDefinitionFactory.h
@@ -12,6 +12,9 @@ class UIKRigDefinitionFactory : public UFactory
 {
 	GENERATED_BODY()
 
+	UPROPERTY()
+	TObjectPtr<class USkeletalMesh>	SkeletalMesh;
+
 public:
 
 	UIKRigDefinitionFactory();
@@ -31,4 +34,10 @@ public:
 	virtual bool ConfigureProperties() override;
 	virtual bool ShouldShowInNewMenu() const override;
 	// END UFactory interface
+
+private:
+	
+	void OnSkeletalMeshSelected(const FAssetData& SelectedAsset);
+
+	TSharedPtr<SWindow> PickerWindow;
 };
diff --git a/Engine/Plugins/Animation/IKRig/Source/IKRigEditor/Public/RigEditor/IKRigEditorController.h b/Engine/Plugins/Animation/IKRig/Source/IKRigEditor/Public/RigEditor/IKRigEditorController.h
--- a/Engine/Plugins/Animation/IKRig/Source/IKRigEditor/Public/RigEditor/IKRigEditorController.h
+++ b/Engine/Plugins/Animation/IKRig/Source/IKRigEditor/Public/RigEditor/IKRigEditorController.h
@@ -21,7 +21,7 @@ class UDebugSkelMeshComponent;
 class FIKRigEditorController
 {
 public:
-
+	
 	/** get array of the names of currently selected goals */
 	const TArray<FName>& GetSelectedGoals() const {return SelectedGoals;};
 	/** return true if a goal with the given name is selected */
@@ -42,7 +42,7 @@ public:
 	/** get index of the first selected solver, return INDEX_None if nothing selected */
 	int32 GetSelectedSolverIndex();
 	/** right after importing a skeleton, we ask user what solver they want to use */
-	void PromptToAddSolverAtStartup() const;
+	void PromptToAddSolver() const;
 	/** determine if the element is connected to the selected solver */
 	bool IsElementConnectedToSolver(TSharedRef<FIKRigTreeElement> TreeElement, int32 SolverIndex);
 	/** determine if the element is connected to ANY solver */
@@ -76,6 +76,9 @@ public:
 	UPROPERTY(transient, NonTransactional)
 	TObjectPtr<class UIKRigAnimInstance> AnimInstance;
 
+	/** Sequence Browser **/
+	TWeakPtr<class IAnimationSequenceBrowser> SequenceBrowser;
+
 	/** asset properties tab */
 	TSharedPtr<class IDetailsView> DetailsView;
 
@@ -100,7 +103,7 @@ private:
 	TArray<FName> SelectedGoals;	
 };
 
-/** only used for pop-up window to selected a solver to add after importing a skeleton */
+/** only used for pop-up window to selected a first solver to add*/
 USTRUCT()
 struct FIKRigAddFirstSolverSettings
 {
diff --git a/Engine/Plugins/Animation/IKRig/Source/IKRigEditor/Public/RigEditor/IKRigToolkit.h b/Engine/Plugins/Animation/IKRig/Source/IKRigEditor/Public/RigEditor/IKRigToolkit.h
--- a/Engine/Plugins/Animation/IKRig/Source/IKRigEditor/Public/RigEditor/IKRigToolkit.h
+++ b/Engine/Plugins/Animation/IKRig/Source/IKRigEditor/Public/RigEditor/IKRigToolkit.h
@@ -12,6 +12,7 @@
 #include "PersonaAssetEditorToolkit.h"
 #include "SIKRigSkeleton.h"
 
+class IAnimationSequenceBrowser;
 class UIKRigSkeletalMeshComponent;
 class FSolverStackElement;
 class SIKRigSolverStack;
@@ -77,6 +78,11 @@ public:
 	virtual void PostRedo( bool bSuccess );
 	/** END FSelfRegisteringEditorUndoClient interface */
 
+	/** animation asset browser */
+	void HandleOpenNewAsset(UObject* InNewAsset);
+	void HandleAnimationSequenceBrowserCreated(const TSharedRef<IAnimationSequenceBrowser>& InSequenceBrowser);
+	/* END animation asset browser */
+
 	TSharedRef<FIKRigEditorController> GetController() const {return EditorController;};
 
 private:
diff --git a/Engine/Plugins/Animation/IKRig/Source/IKRigEditor/Public/RigEditor/SIKRigSkeleton.h b/Engine/Plugins/Animation/IKRig/Source/IKRigEditor/Public/RigEditor/SIKRigSkeleton.h
--- a/Engine/Plugins/Animation/IKRig/Source/IKRigEditor/Public/RigEditor/SIKRigSkeleton.h
+++ b/Engine/Plugins/Animation/IKRig/Source/IKRigEditor/Public/RigEditor/SIKRigSkeleton.h
@@ -208,12 +208,6 @@ private:
 	void GetSelectedBoneChains(TArray<FIKRigSkeletonChain>& OutChains);
 	/** END selection state queries */
 
-	/** skeleton import */
-	FReply OnImportSkeletonClicked();
-	void ImportSkeleton(const FAssetData& InAssetData);
-	EVisibility IsImportButtonVisible() const;
-	/** END skeleton import */
-
 	/** centralized editor controls (facilitate cross-communication between multiple UI elements)*/
 	TWeakPtr<FIKRigEditorController> EditorController;
 	
