Fixing static analysis warning

diff --git a/Engine/Plugins/Animation/IKRig/Source/IKRig/Private/IKRigProcessor.cpp b/Engine/Plugins/Animation/IKRig/Source/IKRig/Private/IKRigProcessor.cpp
--- a/Engine/Plugins/Animation/IKRig/Source/IKRig/Private/IKRigProcessor.cpp
+++ b/Engine/Plugins/Animation/IKRig/Source/IKRig/Private/IKRigProcessor.cpp
@@ -8,10 +8,17 @@ void UIKRigProcessor::Initialize(UIKRigDefinition* InRigAsset, const FReferenceS
 {
 	// we instantiate UObjects here which MUST be done on game thread...
 	check(IsInGameThread());
+	check(InRigAsset);
 	
 	bInitialized = false;
 	InitializedWithIKRigAssetVersion = -1;
 
+	if (!InRigAsset)
+	{
+		UE_LOG(LogTemp, Warning, TEXT("Trying to initialize IKRigProcessor with a null IKRigDefinition asset."));
+		return;
+	}
+
 	// bail out if we've already tried initializing with this exact version of the rig asset
 	if (LastVersionTried == InRigAsset->GetAssetVersion())
 	{
@@ -19,12 +26,6 @@ void UIKRigProcessor::Initialize(UIKRigDefinition* InRigAsset, const FReferenceS
 	}
 	LastVersionTried = InRigAsset->GetAssetVersion();
 
-	if (!InRigAsset)
-	{
-		UE_LOG(LogTemp, Warning, TEXT("Trying to initialize IKRigProcessor with a null IKRigDefinition asset."));
-		return;
-	}
-
 	if (InRigAsset->Skeleton.BoneNames.IsEmpty())
 	{
 		UE_LOG(LogTemp, Warning, TEXT("Trying to initialize IKRigProcessor with a IKRigDefinition that has no skeleton: %s"), *InRigAsset->GetName());
diff --git a/Engine/Plugins/Animation/IKRig/Source/IKRig/Public/Retargeter/IKRetargeter.h b/Engine/Plugins/Animation/IKRig/Source/IKRig/Public/Retargeter/IKRetargeter.h
--- a/Engine/Plugins/Animation/IKRig/Source/IKRig/Public/Retargeter/IKRetargeter.h
+++ b/Engine/Plugins/Animation/IKRig/Source/IKRig/Public/Retargeter/IKRetargeter.h
@@ -328,10 +328,10 @@ struct IKRIG_API FIKRetargetPose
 	FIKRetargetPose(){}
 	
 	UPROPERTY(EditAnywhere, Category = RetargetPose)
-	FVector RootTranslationOffset;
+	FVector RootTranslationOffset = FVector::ZeroVector;
 	
 	UPROPERTY(EditAnywhere, Category = RetargetPose)
-	FQuat RootRotationOffset;
+	FQuat RootRotationOffset = FQuat::Identity;
 	
 	UPROPERTY(EditAnywhere, Category = RetargetPose)
 	TMap<FName, FQuat> BoneRotationOffsets;
