Added "Duplicate and Retarget" batch process functionality to IK Rig plugin

diff --git a/Engine/Plugins/Animation/IKRig/IKRig.uplugin b/Engine/Plugins/Animation/IKRig/IKRig.uplugin
--- a/Engine/Plugins/Animation/IKRig/IKRig.uplugin
+++ b/Engine/Plugins/Animation/IKRig/IKRig.uplugin
@@ -29,17 +29,17 @@
 		{
 			"Name": "IKRig",
 			"Type": "Runtime",
-			"LoadingPhase": "Default"
+			"LoadingPhase": "PreDefault"
 		},
 		{
 			"Name": "IKRigDeveloper",
 			"Type": "UncookedOnly",
-			"LoadingPhase": "Default"
+			"LoadingPhase": "PreDefault"
 		},
 		{
 			"Name": "IKRigEditor",
 			"Type": "Editor",
-			"LoadingPhase": "Default"
+			"LoadingPhase": "PostEngineInit"
 		}
 	]
 }
\ No newline at end of file
diff --git a/Engine/Plugins/Animation/IKRig/Source/IKRig/Private/AnimNodes/AnimNode_IKRetargeter.cpp b/Engine/Plugins/Animation/IKRig/Source/IKRig/Private/AnimNodes/AnimNode_IKRetargeter.cpp
--- a/Engine/Plugins/Animation/IKRig/Source/IKRig/Private/AnimNodes/AnimNode_IKRetargeter.cpp
+++ b/Engine/Plugins/Animation/IKRig/Source/IKRig/Private/AnimNodes/AnimNode_IKRetargeter.cpp
@@ -37,14 +37,14 @@ void FAnimNode_IKRetargeter::Evaluate_AnyThread(FPoseContext& Output)
 	}
 
 	// run the retargeter
-	CurrentlyUsedRetargeter->RunRetargeter(SourceMeshComponentSpaceBoneTransforms);
+	const TArray<FTransform>& RetargetedPose = CurrentlyUsedRetargeter->RunRetargeter(SourceMeshComponentSpaceBoneTransforms);
 
 	// copy pose back
 	FCSPose<FCompactPose> ComponentPose;
 	ComponentPose.InitPose(Output.Pose);
 	for (FCompactPoseBoneIndex CompactBoneIndex : Output.Pose.ForEachBoneIndex())
 	{
-		ComponentPose.SetComponentSpaceTransform(CompactBoneIndex, CurrentlyUsedRetargeter->TargetSkeleton.OutputGlobalPose[CompactBoneIndex.GetInt()]);
+		ComponentPose.SetComponentSpaceTransform(CompactBoneIndex, RetargetedPose[CompactBoneIndex.GetInt()]);
 	}
 
 	// convert to local space
diff --git a/Engine/Plugins/Animation/IKRig/Source/IKRig/Private/Retargeter/IKRetargeter.cpp b/Engine/Plugins/Animation/IKRig/Source/IKRig/Private/Retargeter/IKRetargeter.cpp
--- a/Engine/Plugins/Animation/IKRig/Source/IKRig/Private/Retargeter/IKRetargeter.cpp
+++ b/Engine/Plugins/Animation/IKRig/Source/IKRig/Private/Retargeter/IKRetargeter.cpp
@@ -338,6 +338,11 @@ FTransform FChainDecoderFK::GetTransformAtParam(
 	const TArray<float>& InParams,
 	const float Param) const
 {
+	if (InParams.Num() == 1)
+	{
+		return Transforms[0];
+	}
+	
 	if (Param < KINDA_SMALL_NUMBER)
 	{
 		return Transforms[0];
@@ -906,7 +911,7 @@ bool UIKRetargeter::InitializeIKRig(UObject* Outer, const FReferenceSkeleton& In
 	return true;
 }
 
-void UIKRetargeter::RunRetargeter(const TArray<FTransform>& InSourceGlobalPose)
+TArray<FTransform>&  UIKRetargeter::RunRetargeter(const TArray<FTransform>& InSourceGlobalPose)
 {
 	check(bIsLoadedAndValid);
 
@@ -918,7 +923,7 @@ void UIKRetargeter::RunRetargeter(const TArray<FTransform>& InSourceGlobalPose)
 	// then the TargetSkeleton.RetargetGlobalPose will contain the updated retarget pose.
 	if (bEditReferencePoseMode)
 	{
-		return; 
+		return TargetSkeleton.OutputGlobalPose; 
 	}
 
 	// ROOT retargeting
@@ -940,7 +945,9 @@ void UIKRetargeter::RunRetargeter(const TArray<FTransform>& InSourceGlobalPose)
 	if (bRetargetIK)
 	{
 		RunIKRetarget(InSourceGlobalPose, TargetSkeleton.OutputGlobalPose);
-	}	
+	}
+
+	return TargetSkeleton.OutputGlobalPose;
 }
 
 void UIKRetargeter::RunRootRetarget(
diff --git a/Engine/Plugins/Animation/IKRig/Source/IKRig/Public/Retargeter/IKRetargeter.h b/Engine/Plugins/Animation/IKRig/Source/IKRig/Public/Retargeter/IKRetargeter.h
--- a/Engine/Plugins/Animation/IKRig/Source/IKRig/Public/Retargeter/IKRetargeter.h
+++ b/Engine/Plugins/Animation/IKRig/Source/IKRig/Public/Retargeter/IKRetargeter.h
@@ -57,7 +57,7 @@ struct IKRIG_API FRetargetChainMap
 	//FName IkRelativeTargetBone;	
 };
 
-struct FRetargetSkeleton
+struct IKRIG_API FRetargetSkeleton
 {
 	TArray<FName> BoneNames;
 	
@@ -79,13 +79,13 @@ struct FRetargetSkeleton
 
 	void UpdateGlobalTransformsBelowBone(
 		const int32 StartBoneIndex,
-		TArray<FTransform>& InOutLocalPose,
-		TArray<FTransform>& InOutGlobalPose) const;
+		TArray<FTransform>& InLocalPose,
+		TArray<FTransform>& OutGlobalPose) const;
 
 	void UpdateLocalTransformsBelowBone(
 		const int32 StartBoneIndex,
-		TArray<FTransform>& InOutLocalPose,
-		TArray<FTransform>& InOutGlobalPose) const;
+		TArray<FTransform>& OutLocalPose,
+		TArray<FTransform>& InGlobalPose) const;
 	
 	void UpdateGlobalTransformOfSingleBone(
 		const int32 BoneIndex,
@@ -404,10 +404,22 @@ public:
 	int32 AssetVersion = 0;
 
 	UIKRetargeter();
-	
+
+	/**
+	* Initialize the retargeter to enable running it.
+	* @param SourceSkeleton - the skeletal mesh to poses FROM
+	* @param TargetSkeleton - the skeletal mesh to poses TO
+	* @param Outer - a UObject that will live for the duration of the retargeting workload
+	* @warning - Initialization does a lot of validation and can fail for many reasons. Check bIsLoadedAndValid afterwards.
+	*/
 	void Initialize(USkeletalMesh *SourceSkeleton, USkeletalMesh *TargetSkeleton, UObject* Outer);
-	
-	void RunRetargeter(const TArray<FTransform>& InSourceGlobalPose);
+
+	/**
+	 * Run the retarget to generate a new pose.
+	 * @param InSourceGlobalPose -  is the source mesh input pose in Component/Global space
+	 * @return The retargeted Component/Global space pose for the target skeleton
+	 */
+	TArray<FTransform>& RunRetargeter(const TArray<FTransform>& InSourceGlobalPose);
 
 	FTransform GetTargetBoneRetargetPoseGlobalTransform(const FName& ChainName) const;
 
@@ -416,6 +428,7 @@ public:
 	FRetargetChainMap* GetChainMap(const FName& TargetChainName);
 
 #if WITH_EDITOR
+	
 	// UObject
 	virtual bool Modify(bool bAlwaysMarkDirty = true) override;
 	virtual void PostEditChangeProperty( struct FPropertyChangedEvent& PropertyChangedEvent) override;
diff --git a/Engine/Plugins/Animation/IKRig/Source/IKRigEditor/IKRigEditor.Build.cs b/Engine/Plugins/Animation/IKRig/Source/IKRigEditor/IKRigEditor.Build.cs
--- a/Engine/Plugins/Animation/IKRig/Source/IKRigEditor/IKRigEditor.Build.cs
+++ b/Engine/Plugins/Animation/IKRig/Source/IKRigEditor/IKRigEditor.Build.cs
@@ -34,6 +34,8 @@ namespace UnrealBuildTool.Rules
                     "EditorStyle",
                     "Persona",
                     "SkeletonEditor",
+                    "ContentBrowser",
+                    "AnimationBlueprintLibrary",
                     
 					"PropertyEditor",
 					"BlueprintGraph",
diff --git a/Engine/Plugins/Animation/IKRig/Source/IKRigEditor/Private/IKRigEditor.cpp b/Engine/Plugins/Animation/IKRig/Source/IKRigEditor/Private/IKRigEditor.cpp
--- a/Engine/Plugins/Animation/IKRig/Source/IKRigEditor/Private/IKRigEditor.cpp
+++ b/Engine/Plugins/Animation/IKRig/Source/IKRigEditor/Private/IKRigEditor.cpp
@@ -36,6 +36,9 @@ void FIKRigEditor::StartupModule()
 	IKRetargeterAssetAction = MakeShareable(new FAssetTypeActions_IKRetargeter);
 	FModuleManager::LoadModuleChecked<FAssetToolsModule>("AssetTools").Get().RegisterAssetTypeActions(IKRetargeterAssetAction.ToSharedRef());
 
+	// extend the content browser menu
+	FAssetTypeActions_IKRetargeter::ExtendAnimSequenceToolMenu();
+
 	// register custom editor modes
 	FEditorModeRegistry::Get().RegisterMode<FIKRigEditMode>(FIKRigEditMode::ModeName, LOCTEXT("IKRigEditMode", "IKRig"), FSlateIcon(), false);
 	FEditorModeRegistry::Get().RegisterMode<FIKRetargetEditMode>(FIKRetargetEditMode::ModeName, LOCTEXT("IKRetargetEditMode", "IKRetarget"), FSlateIcon(), false);
@@ -72,7 +75,6 @@ void FIKRigEditor::ShutdownModule()
 		}
 		IKRetargeterAssetAction.Reset();
 	}
-
 }
 
 #undef LOCTEXT_NAMESPACE
\ No newline at end of file
diff --git a/Engine/Plugins/Animation/IKRig/Source/IKRigEditor/Private/RetargetEditor/AssetTypeActions_IKRetargeter.cpp b/Engine/Plugins/Animation/IKRig/Source/IKRigEditor/Private/RetargetEditor/AssetTypeActions_IKRetargeter.cpp
--- a/Engine/Plugins/Animation/IKRig/Source/IKRigEditor/Private/RetargetEditor/AssetTypeActions_IKRetargeter.cpp
+++ b/Engine/Plugins/Animation/IKRig/Source/IKRigEditor/Private/RetargetEditor/AssetTypeActions_IKRetargeter.cpp
@@ -1,7 +1,12 @@
 // Copyright Epic Games, Inc. All Rights Reserved.
 
 #include "RetargetEditor/AssetTypeActions_IKRetargeter.h"
+
+#include "ContentBrowserMenuContexts.h"
+#include "ToolMenus.h"
 #include "RetargetEditor/IKRetargetEditor.h"
+#include "RetargetEditor/IKRetargetEditorStyle.h"
+#include "RetargetEditor/SRetargetAnimAssetsWindow.h"
 #include "Retargeter/IKRetargeter.h"
 #include "ThumbnailRendering/SceneThumbnailInfo.h"
 
@@ -39,4 +44,89 @@ UThumbnailInfo* FAssetTypeActions_IKRetargeter::GetThumbnailInfo(UObject* Asset)
 	return NewObject<USceneThumbnailInfo>(IKRetargeter, NAME_None, RF_Transactional);
 }
 
+void FAssetTypeActions_IKRetargeter::ExtendAnimSequenceToolMenu()
+{
+	TArray<UToolMenu*> MenusToExtend;
+	MenusToExtend.Add(UToolMenus::Get()->ExtendMenu("ContentBrowser.AssetContextMenu.AnimSequence"));
+	MenusToExtend.Add(UToolMenus::Get()->ExtendMenu("ContentBrowser.AssetContextMenu.BlendSpace"));
+	MenusToExtend.Add(UToolMenus::Get()->ExtendMenu("ContentBrowser.AssetContextMenu.AimOffsetBlendSpace"));
+	MenusToExtend.Add(UToolMenus::Get()->ExtendMenu("ContentBrowser.AssetContextMenu.BlendSpace1D"));
+	MenusToExtend.Add(UToolMenus::Get()->ExtendMenu("ContentBrowser.AssetContextMenu.PoseAsset"));
+	MenusToExtend.Add(UToolMenus::Get()->ExtendMenu("ContentBrowser.AssetContextMenu.AnimBlueprint"));
+
+	for (UToolMenu* Menu : MenusToExtend)
+	{
+		if (Menu == nullptr)
+		{
+			continue;
+		}
+
+		FToolMenuSection& Section = Menu->FindOrAddSection("GetAssetActions");
+		Section.AddSubMenu(
+		"IKRetargetAnimSubmenu", 
+		LOCTEXT("IKRetargetAnimSubmenu", "Retarget Animation Assets"),
+		LOCTEXT("IKRetargetAnimSubmenu_ToolTip", "Opens the batch retargeting menu."),
+		FNewToolMenuDelegate::CreateLambda([](UToolMenu* AlignmentMenu)
+		{
+			FToolMenuSection& InSection = AlignmentMenu->AddSection("IKRetargetMenu", LOCTEXT("RetargetHeader", "IK Retargeting"));
+			InSection.AddDynamicEntry("IKRigActions", FNewToolMenuSectionDelegate::CreateLambda([](FToolMenuSection& InSection)
+			{
+				FAssetTypeActions_IKRetargeter::CreateRetargetSubMenu(InSection);
+			}));
+		}),
+		false,
+		FSlateIcon(FEditorStyle::GetStyleSetName(), "Persona.AssetActions.RetargetSkeleton")
+		);
+	}
+}
+
+void FAssetTypeActions_IKRetargeter::CreateRetargetSubMenu(FToolMenuSection& InSection)
+{
+	UContentBrowserAssetContextMenuContext* Context = InSection.FindContext<UContentBrowserAssetContextMenuContext>();
+	if (!Context)
+	{
+		return;			
+	}
+	
+	TArray<UObject*> SelectedObjects = Context->GetSelectedObjects();
+	if (SelectedObjects.IsEmpty())
+	{
+		return;
+	}
+
+	// change menu label if anim blueprint is selected
+	FText MenuLabel = FText(LOCTEXT("RetargetAnimation", "Duplicate and Retarget Animation Assets"));
+	if (Cast<UAnimBlueprint>(SelectedObjects[0]))
+	{
+		MenuLabel = FText(LOCTEXT("RetargetAnimation", "Duplicate and Retarget Animation Blueprint"));
+	}
+	
+	InSection.AddMenuEntry(
+		"IKRetargetToDifferentSkeleton",
+		MenuLabel,
+		LOCTEXT("RetargetAnimation_ToolTip", "Duplicate an animation asset and retarget to a different skeleton."),
+		FSlateIcon(FEditorStyle::GetStyleSetName(), "GenericCurveEditor.TabIcon"),
+		FUIAction(FExecuteAction::CreateLambda([SelectedObjects]()
+		{
+			SRetargetAnimAssetsWindow::ShowWindow(SelectedObjects);
+		}),
+		FCanExecuteAction::CreateLambda([InSection]()
+		{
+			if (UContentBrowserAssetContextMenuContext* Context = InSection.FindContext<UContentBrowserAssetContextMenuContext>())
+			{
+				TArray<UObject*> SelectedObjects = Context->GetSelectedObjects();
+				for (UObject* SelectedObject : SelectedObjects)
+				{
+					if (Cast<UAnimationAsset>(SelectedObject) || Cast<UAnimBlueprint>(SelectedObject))
+					{
+						return true;
+					}
+				}		
+			}
+			
+			return false;
+		}))
+	);
+}
+
 #undef LOCTEXT_NAMESPACE
diff --git a/Engine/Plugins/Animation/IKRig/Source/IKRigEditor/Private/RetargetEditor/SRetargetAnimAssetsWindow.cpp b/Engine/Plugins/Animation/IKRig/Source/IKRigEditor/Private/RetargetEditor/SRetargetAnimAssetsWindow.cpp
new file mode 100644
--- /dev/null
+++ b/Engine/Plugins/Animation/IKRig/Source/IKRigEditor/Private/RetargetEditor/SRetargetAnimAssetsWindow.cpp
@@ -0,0 +1,1051 @@
+﻿// Copyright Epic Games, Inc. All Rights Reserved.
+
+#include "RetargetEditor/SRetargetAnimAssetsWindow.h"
+
+#include "AnimationBlueprintLibrary.h"
+#include "AnimPose.h"
+#include "AnimPreviewInstance.h"
+#include "ContentBrowserModule.h"
+#include "EditorReimportHandler.h"
+#include "IContentBrowserSingleton.h"
+#include "ObjectEditorUtils.h"
+#include "SSkeletonWidget.h"
+#include "Widgets/SWindow.h"
+#include "Widgets/Layout/SSeparator.h"
+#include "PropertyCustomizationHelpers.h"
+#include "Animation/DebugSkelMeshComponent.h"
+#include "EditorFramework/AssetImportData.h"
+#include "Framework/Notifications/NotificationManager.h"
+#include "Kismet2/BlueprintEditorUtils.h"
+#include "Kismet2/KismetEditorUtilities.h"
+#include "Misc/ScopedSlowTask.h"
+#include "Retargeter/IKRetargeter.h"
+#include "Widgets/Notifications/SNotificationList.h"
+
+#define LOCTEXT_NAMESPACE "RetargetAnimAssetWindow"
+
+int32 FIKRetargetAnimAssetsContext::GenerateAssetLists()
+{
+	// re-generate lists of selected and referenced assets
+	AnimationAssetsToRetarget.Reset();
+	AnimBlueprintsToRetarget.Reset();
+
+	for (TWeakObjectPtr<UObject> AssetPtr : AssetsToRetarget)
+	{
+		UObject* Asset = AssetPtr.Get();
+		if (UAnimationAsset* AnimAsset = Cast<UAnimationAsset>(Asset))
+		{
+			AnimationAssetsToRetarget.AddUnique(AnimAsset);
+		}
+		else if (UAnimBlueprint* AnimBlueprint = Cast<UAnimBlueprint>(Asset))
+		{
+			// Add parents blueprint. 
+			UAnimBlueprint* ParentBP = Cast<UAnimBlueprint>(AnimBlueprint->ParentClass->ClassGeneratedBy);
+			while (ParentBP)
+			{
+				AnimBlueprintsToRetarget.AddUnique(ParentBP);
+				ParentBP = Cast<UAnimBlueprint>(ParentBP->ParentClass->ClassGeneratedBy);
+			}
+				
+			AnimBlueprintsToRetarget.AddUnique(AnimBlueprint);				
+		}
+	}
+
+	if (bRemapReferencedAssets)
+	{
+		// Grab assets from the blueprint.
+		// Do this first as it can add complex assets to the retarget array which will need to be processed next.
+		for (UAnimBlueprint* AnimBlueprint : AnimBlueprintsToRetarget)
+		{
+			GetAllAnimationSequencesReferredInBlueprint(AnimBlueprint, AnimationAssetsToRetarget);
+		}
+
+		int32 AssetIndex = 0;
+		while (AssetIndex < AnimationAssetsToRetarget.Num())
+		{
+			UAnimationAsset* AnimAsset = AnimationAssetsToRetarget[AssetIndex++];
+			AnimAsset->HandleAnimReferenceCollection(AnimationAssetsToRetarget, true);
+		}
+	}
+
+	return AnimationAssetsToRetarget.Num();
+}
+
+void FIKRetargetAnimAssetsContext::DuplicateRetargetAssets()
+{
+	UPackage* DestinationPackage = TargetMesh->GetOutermost();
+
+	TArray<UAnimationAsset*> AnimationAssetsToDuplicate = AnimationAssetsToRetarget;
+	TArray<UAnimBlueprint*> AnimBlueprintsToDuplicate = AnimBlueprintsToRetarget;
+
+	// We only want to duplicate unmapped assets, so we remove mapped assets from the list we're duplicating
+	for (TPair<UAnimationAsset*, UAnimationAsset*>& Pair : RemappedAnimAssets)
+	{
+		AnimationAssetsToDuplicate.Remove(Pair.Key);
+	}
+
+	DuplicatedAnimAssets = DuplicateAssets<UAnimationAsset>(AnimationAssetsToDuplicate, DestinationPackage, &NameRule);
+	DuplicatedBlueprints = DuplicateAssets<UAnimBlueprint>(AnimBlueprintsToDuplicate, DestinationPackage, &NameRule);
+
+	// If we are moving the new asset to a different directory we need to fixup the reimport path.
+	// This should only effect source FBX paths within the project.
+	if (!NameRule.FolderPath.IsEmpty())
+	{
+		for (TPair<UAnimationAsset*, UAnimationAsset*>& Pair : DuplicatedAnimAssets)
+		{
+			UAnimSequence* SourceSequence = Cast<UAnimSequence>(Pair.Key);
+			UAnimSequence* DestinationSequence = Cast<UAnimSequence>(Pair.Value);
+			if (!(SourceSequence && DestinationSequence))
+			{
+				continue;
+			}
+			
+			for (int index = 0; index < SourceSequence->AssetImportData->SourceData.SourceFiles.Num(); index++)
+			{
+				const FString& RelativeFilename = SourceSequence->AssetImportData->SourceData.SourceFiles[index].RelativeFilename;
+				const FString OldPackagePath = FPackageName::GetLongPackagePath(SourceSequence->GetPathName()) / TEXT("");
+				const FString NewPackagePath = FPackageName::GetLongPackagePath(DestinationSequence->GetPathName()) / TEXT("");
+				const FString AbsoluteSrcPath = FPaths::ConvertRelativePathToFull(FPackageName::LongPackageNameToFilename(OldPackagePath));
+				const FString SrcFile = AbsoluteSrcPath / RelativeFilename;
+				const bool bSrcFileExists = FPlatformFileManager::Get().GetPlatformFile().FileExists(*SrcFile);
+				if (!bSrcFileExists || (NewPackagePath == OldPackagePath))
+				{
+					continue;
+				}
+
+				FString BasePath = FPackageName::LongPackageNameToFilename(OldPackagePath);
+				FString OldSourceFilePath = FPaths::ConvertRelativePathToFull(BasePath, RelativeFilename);
+				TArray<FString> Paths;
+				Paths.Add(OldSourceFilePath);
+				
+				// update the FBX reimport file path
+				FReimportManager::Instance()->UpdateReimportPaths(DestinationSequence, Paths);
+			}
+		}
+	}
+
+	// Remapped assets needs the duplicated ones added
+	RemappedAnimAssets.Append(DuplicatedAnimAssets);
+
+	DuplicatedAnimAssets.GenerateValueArray(AnimationAssetsToRetarget);
+	DuplicatedBlueprints.GenerateValueArray(AnimBlueprintsToRetarget);
+}
+
+void FIKRetargetAnimAssetsContext::RetargetAssets()
+{
+	USkeleton* OldSkeleton = SourceMesh->GetSkeleton();
+	USkeleton* NewSkeleton = TargetMesh->GetSkeleton();
+
+	for (UAnimationAsset* AssetToRetarget : AnimationAssetsToRetarget)
+	{
+		// synchronize curves between old/new asset
+		UAnimSequence* AnimSequenceToRetarget = Cast<UAnimSequence>(AssetToRetarget);
+		if (AnimSequenceToRetarget)
+		{
+			// copy curve data from source asset, preserving data in the target if present.
+			UAnimationBlueprintLibrary::CopyAnimationCurveNamesToSkeleton(OldSkeleton, NewSkeleton, AnimSequenceToRetarget, ERawCurveTrackTypes::RCT_Float);	
+			// clear transform curves since those curves won't work in new skeleton
+			IAnimationDataController& Controller = AnimSequenceToRetarget->GetController();
+			Controller.RemoveAllCurvesOfType(ERawCurveTrackTypes::RCT_Transform);
+		}
+
+		// replace references to other animation
+		AssetToRetarget->ReplaceReferredAnimations(RemappedAnimAssets);
+		AssetToRetarget->SetSkeleton(NewSkeleton);
+		AssetToRetarget->MarkPackageDirty();
+	}
+
+	// convert the animation using the IK retargeter
+	ConvertAnimation();
+
+	// convert all Animation Blueprints and compile 
+	for (UAnimBlueprint* AnimBlueprint : AnimBlueprintsToRetarget)
+	{
+		// replace skeleton
+		AnimBlueprint->TargetSkeleton = NewSkeleton;
+
+		// if they have parent blueprint, make sure to re-link to the new one also
+		UAnimBlueprint* CurrentParentBP = Cast<UAnimBlueprint>(AnimBlueprint->ParentClass->ClassGeneratedBy);
+		if (CurrentParentBP)
+		{
+			UAnimBlueprint* const * ParentBP = DuplicatedBlueprints.Find(CurrentParentBP);
+			if (ParentBP)
+			{
+				AnimBlueprint->ParentClass = (*ParentBP)->GeneratedClass;
+			}
+		}
+
+		if(RemappedAnimAssets.Num() > 0)
+		{
+			ReplaceReferredAnimationsInBlueprint(AnimBlueprint, RemappedAnimAssets);
+		}
+
+		FBlueprintEditorUtils::RefreshAllNodes(AnimBlueprint);
+		FKismetEditorUtilities::CompileBlueprint(AnimBlueprint, EBlueprintCompileOptions::SkipGarbageCollection);
+		AnimBlueprint->PostEditChange();
+		AnimBlueprint->MarkPackageDirty();
+	}
+}
+
+void FIKRetargetAnimAssetsContext::ConvertAnimation()
+{
+	// initialize the retargeter
+	UObject* TransientOuter = Cast<UObject>(GetTransientPackage());
+	UIKRetargeter* Retargeter = DuplicateObject(IKRetargetAsset, TransientOuter);	
+	Retargeter->Initialize(SourceMesh, TargetMesh, TransientOuter);
+	if (!Retargeter->bIsLoadedAndValid)
+	{
+		UE_LOG(LogTemp, Warning, TEXT("Unable to initialize the IK Retargeter. Newly created animations were not retargeted!"));
+		return;
+	}
+
+	// for each pair of source / target animation sequences
+	for (TPair<UAnimationAsset*, UAnimationAsset*>& Pair : DuplicatedAnimAssets)
+	{
+		UAnimSequence* SourceSequence = Cast<UAnimSequence>(Pair.Key);
+		UAnimSequence* DestinationSequence = Cast<UAnimSequence>(Pair.Value);
+		if (!(SourceSequence && DestinationSequence))
+		{
+			continue;
+		}
+
+		// remove all keys from the destination animation sequence
+		IAnimationDataController& TargetSeqController = DestinationSequence->GetController();
+		TargetSeqController.RemoveAllBoneTracks();
+
+		// number of frames in this animation
+		const int32 NumFrames = SourceSequence->GetNumberOfSampledKeys();
+
+		// make space for the target keyframe data
+		const int32 NumTargetBones = Retargeter->TargetSkeleton.BoneNames.Num();
+		TArray<FRawAnimSequenceTrack> BoneTracks;
+		BoneTracks.SetNumZeroed(NumTargetBones);
+
+		// retarget each frame's pose from source to target
+		for (int32 FrameIndex=0; FrameIndex<NumFrames; ++FrameIndex)
+		{
+			// get the source global pose
+			FAnimPose SourcePoseAtFrame;
+			UAnimPoseExtensions::GetAnimPoseAtFrame(SourceSequence, FrameIndex, FAnimPoseEvaluationOptions(), SourcePoseAtFrame);
+			TArray<FName> BoneNames;
+			UAnimPoseExtensions::GetBoneNames(SourcePoseAtFrame,BoneNames);
+			TArray<FTransform> SourceComponentPose;
+			for (const FName& BoneName : BoneNames)
+			{
+				FTransform BonePose = UAnimPoseExtensions::GetBonePose(SourcePoseAtFrame, BoneName, EAnimPoseSpaces::World);
+				SourceComponentPose.Add(BonePose);
+			}
+
+			// run the retarget
+			TArray<FTransform>& TargetComponentPose = Retargeter->RunRetargeter(SourceComponentPose);
+
+			// convert to a local-space pose
+			TArray<FTransform> TargetLocalPose = TargetComponentPose;
+			Retargeter->TargetSkeleton.UpdateLocalTransformsBelowBone(0,TargetLocalPose, TargetComponentPose);
+
+			// store key data for each bone
+			for (int32 TargetBoneIndex=0; TargetBoneIndex<TargetLocalPose.Num(); ++TargetBoneIndex)
+			{
+				BoneTracks[TargetBoneIndex].PosKeys.Add(TargetLocalPose[TargetBoneIndex].GetLocation());
+				BoneTracks[TargetBoneIndex].RotKeys.Add(TargetLocalPose[TargetBoneIndex].GetRotation());
+				BoneTracks[TargetBoneIndex].ScaleKeys.Add(TargetLocalPose[TargetBoneIndex].GetScale3D());
+			}
+		}
+
+		// add keys to bone tracks
+		const bool bShouldTransact = false;
+		for (int32 TargetBoneIndex=0; TargetBoneIndex<NumTargetBones; ++TargetBoneIndex)
+		{
+			FName TargetBoneName = Retargeter->TargetSkeleton.BoneNames[TargetBoneIndex];
+			FRawAnimSequenceTrack& RawTrack = BoneTracks[TargetBoneIndex];
+			TargetSeqController.AddBoneTrack(TargetBoneName, bShouldTransact);
+			TargetSeqController.SetBoneTrackKeys(TargetBoneName, RawTrack.PosKeys, RawTrack.RotKeys, RawTrack.ScaleKeys);
+		}
+	}
+}
+
+void FIKRetargetAnimAssetsContext::NotifyUserOfResults() const
+{
+	// gather newly created objects
+	TArray<UObject*> NewAssets;
+	GetNewAssets(NewAssets);
+
+	// log details of what assets were created
+	for (UObject* NewAsset : NewAssets)
+	{
+		UE_LOG(LogTemp, Display, TEXT("Duplicate and Retarget - New Asset Created: %s"), *NewAsset->GetName());
+	}
+	
+	// notify user
+	FNotificationInfo Notification(FText::GetEmpty());
+	Notification.ExpireDuration = 5.f;
+	Notification.Text = FText::Format(
+		LOCTEXT("MultiNonDuplicatedAsset", "{0} assets were retargeted to new skeleton {1}. See Output for details."),
+		FText::AsNumber(NewAssets.Num()),
+		FText::FromString(TargetMesh->GetName()));
+	FSlateNotificationManager::Get().AddNotification(Notification);
+	
+	// select all new assets
+	TArray<FAssetData> CurrentSelection;
+	for(UObject* NewObject : NewAssets)
+	{
+		CurrentSelection.Add(FAssetData(NewObject));
+	}
+
+	// show assets in browser
+	FContentBrowserModule& ContentBrowserModule = FModuleManager::Get().LoadModuleChecked<FContentBrowserModule>("ContentBrowser");
+	ContentBrowserModule.Get().SyncBrowserToAssets(CurrentSelection);
+}
+
+void FIKRetargetAnimAssetsContext::GetNewAssets(TArray<UObject*>& NewAssets) const
+{
+	TArray<UAnimationAsset*> NewAnims;
+	DuplicatedAnimAssets.GenerateValueArray(NewAnims);
+	for (UAnimationAsset* NewAnim : NewAnims)
+	{
+		NewAssets.Add(Cast<UObject>(NewAnim));
+	}
+
+	TArray<UAnimBlueprint*> NewBlueprints;
+	DuplicatedBlueprints.GenerateValueArray(NewBlueprints);
+	for (UAnimBlueprint* NewBlueprint : NewBlueprints)
+	{
+		NewAssets.Add(Cast<UObject>(NewBlueprint));
+	}
+}
+
+void FIKRetargetAnimAssetsContext::Reset()
+{
+	SourceMesh = nullptr;
+	TargetMesh = nullptr;
+	IKRetargetAsset = nullptr;
+	bRemapReferencedAssets = true;
+	NameRule.Prefix = "";
+	NameRule.Suffix = "";
+	NameRule.ReplaceFrom = "";
+	NameRule.ReplaceTo = "";
+}
+
+bool FIKRetargetAnimAssetsContext::IsValid() const
+{
+	// todo validate compatibility
+	return SourceMesh && TargetMesh && IKRetargetAsset && (SourceMesh != TargetMesh);
+}
+
+void FIKRetargetAnimAssetsContext::RunRetarget()
+{
+	// todo progess bar not showing up, why?
+	//FScopedSlowTask Progress(2.f, LOCTEXT("GatheringBatchRetarget", "Gathering animation assets..."));
+	//Progress.MakeDialog();
+	
+	const int32 NumAssets = GenerateAssetLists();
+	
+	//Progress.EnterProgressFrame(1.f, FText::Format(LOCTEXT("DuplicatingBatchRetarget", "Duplicating {0} animation assets..."), NumAssets));
+
+	DuplicateRetargetAssets();
+	
+	//Progress.EnterProgressFrame(1.f, FText::Format(LOCTEXT("RunningBatchRetarget", "Retargeting {0} animation assets..."), NumAssets));
+
+	RetargetAssets();
+	
+	NotifyUserOfResults();
+}
+
+/**
+* Duplicates the supplied AssetsToDuplicate and returns a map of original asset to duplicate. Templated wrapper that calls DuplicateAssetInternal.
+*
+* @param	AssetsToDuplicate	The animations to duplicate
+* @param	DestinationPackage	The package that the duplicates should be placed in
+*
+* @return	TMap of original animation to duplicate
+*/
+template<class AssetType>
+TMap<AssetType*, AssetType*> FIKRetargetAnimAssetsContext::DuplicateAssets(
+	const TArray<AssetType*>& AssetsToDuplicate,
+	UPackage* DestinationPackage,
+	const FNameDuplicationRule* NameRule)
+{
+	TArray<UObject*> Assets;
+	for (AssetType* Asset : AssetsToDuplicate)
+	{
+		Assets.Add(Asset);
+	}
+
+	// duplicate assets
+	TMap<UObject*, UObject*> DuplicateAssetsMap = DuplicateAssetsInternal(Assets, DestinationPackage, NameRule);
+
+	// cast to AssetType
+	TMap<AssetType*, AssetType*> ReturnMap;
+	for (const TTuple<UObject*, UObject*>& DuplicateAsset : DuplicateAssetsMap)
+	{
+		ReturnMap.Add(Cast<AssetType>(DuplicateAsset.Key), Cast<AssetType>(DuplicateAsset.Value));
+	}
+	
+	return ReturnMap;
+}
+
+void SRetargetPoseViewport::Construct(const FArguments& InArgs)
+{
+	SEditorViewport::Construct(SEditorViewport::FArguments());
+
+	PreviewComponent = NewObject<UDebugSkelMeshComponent>();
+	PreviewComponent->VisibilityBasedAnimTickOption = EVisibilityBasedAnimTickOption::AlwaysTickPoseAndRefreshBones;
+	PreviewScene.AddComponent(PreviewComponent, FTransform::Identity);
+
+	SetSkeletalMesh(InArgs._SkeletalMesh);
+}
+
+void SRetargetPoseViewport::SetSkeletalMesh(USkeletalMesh* InSkeltalMesh)
+{
+	if(InSkeltalMesh == Mesh)
+	{
+		return;
+	}
+	
+	Mesh = InSkeltalMesh;
+
+	if(Mesh)
+	{
+		PreviewComponent->SetSkeletalMesh(Mesh);
+		PreviewComponent->EnablePreview(true, nullptr);
+		// todo add IK retargeter and set it to output the retarget pose
+		PreviewComponent->PreviewInstance->SetForceRetargetBasePose(true);
+		PreviewComponent->RefreshBoneTransforms(nullptr);
+
+		//Place the camera at a good viewer position
+		FBoxSphereBounds Bounds = Mesh->GetBounds();
+		Client->FocusViewportOnBox(Bounds.GetBox(), true);
+	}
+	else
+	{
+		PreviewComponent->SetSkeletalMesh(nullptr);
+	}
+
+	Client->Invalidate();
+}
+
+SRetargetPoseViewport::SRetargetPoseViewport()
+: PreviewScene(FPreviewScene::ConstructionValues())
+{
+}
+
+bool SRetargetPoseViewport::IsVisible() const
+{
+	return true;
+}
+
+TSharedRef<FEditorViewportClient> SRetargetPoseViewport::MakeEditorViewportClient()
+{
+	TSharedPtr<FEditorViewportClient> EditorViewportClient = MakeShareable(new FRetargetPoseViewportClient(PreviewScene, SharedThis(this)));
+
+	EditorViewportClient->ViewportType = LVT_Perspective;
+	EditorViewportClient->bSetListenerPosition = false;
+	EditorViewportClient->SetViewLocation(EditorViewportDefs::DefaultPerspectiveViewLocation);
+	EditorViewportClient->SetViewRotation(EditorViewportDefs::DefaultPerspectiveViewRotation);
+
+	EditorViewportClient->SetRealtime(false);
+	EditorViewportClient->VisibilityDelegate.BindSP(this, &SRetargetPoseViewport::IsVisible);
+	EditorViewportClient->SetViewMode(VMI_Lit);
+
+	return EditorViewportClient.ToSharedRef();
+}
+
+TSharedPtr<SWidget> SRetargetPoseViewport::MakeViewportToolbar()
+{
+	return nullptr;
+}
+
+
+TSharedPtr<SWindow> SRetargetAnimAssetsWindow::DialogWindow;
+
+void SRetargetAnimAssetsWindow::Construct(const FArguments& InArgs)
+{
+	AssetThumbnailPool = MakeShareable( new FAssetThumbnailPool(1024) );
+	
+	this->ChildSlot
+	[
+		SNew (SHorizontalBox)
+		
+		+SHorizontalBox::Slot()
+		.HAlign(HAlign_Center)
+		.VAlign(VAlign_Top)
+		.AutoWidth()
+		[
+			SNew(SVerticalBox)
+			+SVerticalBox::Slot()
+			.AutoHeight()
+			.VAlign(VAlign_Top)
+			.Padding(0, 5)
+			[
+				SNew(SHorizontalBox)
+				+SHorizontalBox::Slot()
+				[
+					SNew(SVerticalBox)
+					+ SVerticalBox::Slot()
+					.AutoHeight()
+					.HAlign(HAlign_Center)
+					[
+						SNew(STextBlock)
+						.Text(LOCTEXT("DuplicateAndRetarget_SourceTitle", "Source Skeletal Mesh"))
+						.Font(FEditorStyle::GetFontStyle("Persona.RetargetManager.BoldFont"))
+						.AutoWrapText(true)
+					]
+
+					+ SVerticalBox::Slot()
+					.AutoHeight()
+					.Padding(5, 5)
+					[
+						SAssignNew(SourceViewport, SRetargetPoseViewport)
+						.SkeletalMesh(RetargetContext.SourceMesh)
+					]
+
+					+ SVerticalBox::Slot()
+					.AutoHeight()
+					.Padding(5, 5)
+					[
+						SNew(SObjectPropertyEntryBox)
+						.AllowedClass(USkeletalMesh::StaticClass())
+						.AllowClear(true)
+						.DisplayUseSelected(true)
+						.DisplayBrowse(true)
+						.DisplayThumbnail(true)
+						.ThumbnailPool(AssetThumbnailPool)
+						.IsEnabled_Lambda([this]()
+						{
+							if (!RetargetContext.IKRetargetAsset)
+							{
+								return false;
+							}
+							
+							return RetargetContext.IKRetargetAsset->SourceIKRigAsset != nullptr;
+						})
+						.ObjectPath(this, &SRetargetAnimAssetsWindow::GetCurrentSourceMeshPath)
+						.OnObjectChanged(this, &SRetargetAnimAssetsWindow::SourceMeshAssigned)
+						.OnShouldFilterAsset_Lambda([this](const FAssetData& AssetData)
+						{
+							if (!RetargetContext.IKRetargetAsset)
+							{
+								return true;
+							}
+							
+							USkeletalMesh* Mesh = Cast<USkeletalMesh>(AssetData.GetAsset());
+							if (!Mesh)
+							{
+								return true;
+							}
+							
+							USkeletalMesh* PreviewMesh = RetargetContext.IKRetargetAsset->SourceIKRigAsset->GetPreviewMesh();
+							if (!PreviewMesh)
+							{
+								return true;
+							}
+							
+							return Mesh->GetSkeleton() != PreviewMesh->GetSkeleton();
+						})
+					]
+				]
+
+				+SHorizontalBox::Slot()
+				.Padding(5)
+				.AutoWidth()
+				[
+					SNew(SSeparator)
+					.Orientation(Orient_Vertical)
+				]
+
+				+SHorizontalBox::Slot()
+				[
+					SNew(SVerticalBox)
+					+ SVerticalBox::Slot()
+					.AutoHeight()
+					.HAlign(HAlign_Center)
+					[
+						SNew(STextBlock)
+						.Text(LOCTEXT("DuplicateAndRetarget_TargetTitle", "Target Skeletal Mesh"))
+						.Font(FEditorStyle::GetFontStyle("Persona.RetargetManager.BoldFont"))
+						.AutoWrapText(true)
+					]
+				
+					+SVerticalBox::Slot()
+					.AutoHeight()
+					.Padding(5, 5)
+					[
+						SAssignNew(TargetViewport, SRetargetPoseViewport)
+						.SkeletalMesh(nullptr)
+					]
+
+					+ SVerticalBox::Slot()
+					.AutoHeight()
+					.Padding(5, 5)
+					[
+						SNew(SObjectPropertyEntryBox)
+						.AllowedClass(USkeletalMesh::StaticClass())
+						.AllowClear(true)
+						.DisplayUseSelected(true)
+						.DisplayBrowse(true)
+						.DisplayThumbnail(true)
+						.ThumbnailPool(AssetThumbnailPool)
+						.IsEnabled_Lambda([this]()
+						{
+							if (!RetargetContext.IKRetargetAsset)
+							{
+								return false;
+							}
+							
+							return RetargetContext.IKRetargetAsset->TargetIKRigAsset != nullptr;
+						})
+						.ObjectPath(this, &SRetargetAnimAssetsWindow::GetCurrentTargetMeshPath)
+						.OnObjectChanged(this, &SRetargetAnimAssetsWindow::TargetMeshAssigned)
+						.OnShouldFilterAsset_Lambda([this](const FAssetData& AssetData)
+						{
+							/*
+							if (!RetargetContext.IKRetargetAsset)
+							{
+								return true;
+							}
+							
+							USkeletalMesh* Mesh = Cast<USkeletalMesh>(AssetData.GetAsset());
+							if (!Mesh)
+							{
+								return true;
+							}
+
+							
+							USkeletalMesh* PreviewMesh = RetargetContext.IKRetargetAsset->TargetIKRigAsset->GetPreviewMesh();
+							if (!PreviewMesh)
+							{
+								return true;
+							}
+							
+							return Mesh->GetSkeleton() != PreviewMesh->GetSkeleton();*/
+							return false;
+						})
+					]
+				]
+			]
+		]
+
+		+SHorizontalBox::Slot()
+		.Padding(5)
+		.AutoWidth()
+		[
+			SNew(SSeparator)
+			.Orientation(Orient_Vertical)
+		]
+			
+		+SHorizontalBox::Slot()
+		.AutoWidth()
+		[
+			SNew(SVerticalBox)
+
+			+ SVerticalBox::Slot()
+			.AutoHeight()
+			.HAlign(HAlign_Center)
+			.Padding(0, 5)
+			[
+				SNew(STextBlock)
+				.Text(LOCTEXT("DuplicateAndRetarget_RetargetAsset", "IK Retargeter"))
+				.Font(FEditorStyle::GetFontStyle("Persona.RetargetManager.BoldFont"))
+				.AutoWrapText(true)
+			]
+
+			+SVerticalBox::Slot()
+			.AutoHeight()
+			.HAlign(HAlign_Fill)
+			.Padding(2)
+			[
+				SNew(SObjectPropertyEntryBox)
+				.AllowedClass(UIKRetargeter::StaticClass())
+				.AllowClear(true)
+				.DisplayUseSelected(true)
+				.DisplayBrowse(true)
+				.DisplayThumbnail(true)
+				.ThumbnailPool(AssetThumbnailPool)
+				.ObjectPath(this, &SRetargetAnimAssetsWindow::GetCurrentRetargeterPath)
+				.OnObjectChanged(this, &SRetargetAnimAssetsWindow::RetargeterAssigned)
+			]
+
+			+SVerticalBox::Slot()
+			.Padding(5)
+			.AutoHeight()
+			[
+				SNew(SSeparator)
+				.Orientation(Orient_Horizontal)
+			]
+			
+			+SVerticalBox::Slot()
+			[	
+				SNew(SVerticalBox)
+				+SVerticalBox::Slot()
+				.AutoHeight()
+				.HAlign(HAlign_Center)
+				.Padding(2, 3)
+				[
+					SNew(STextBlock)
+					.AutoWrapText(true)
+					.Font(FEditorStyle::GetFontStyle("Persona.RetargetManager.SmallBoldFont"))
+					.Text(LOCTEXT("DuplicateAndRetarget_RenameLabel", "Rename New Assets"))
+				]
+
+				+SVerticalBox::Slot()
+				.AutoHeight()
+				.Padding(2, 1)
+				[
+					SNew(SHorizontalBox)
+					+SHorizontalBox::Slot()
+					[
+						SNew(STextBlock).Text(LOCTEXT("DuplicateAndRetarget_Prefix", "Prefix"))
+					]
+
+					+SHorizontalBox::Slot()
+					[
+						SNew(SEditableTextBox)
+							.Text(this, &SRetargetAnimAssetsWindow::GetPrefixName)
+							.MinDesiredWidth(100)
+							.OnTextChanged(this, &SRetargetAnimAssetsWindow::SetPrefixName)
+							.IsReadOnly(false)
+							.RevertTextOnEscape(true)
+					]
+				]
+
+				+SVerticalBox::Slot()
+				.AutoHeight()
+				.Padding(2, 1)
+				[
+					SNew(SHorizontalBox)
+					+SHorizontalBox::Slot()
+					[
+						SNew(STextBlock).Text(LOCTEXT("DuplicateAndRetarget_Suffix", "Suffix"))
+					]
+
+					+SHorizontalBox::Slot()
+					[
+						SNew(SEditableTextBox)
+							.Text(this, &SRetargetAnimAssetsWindow::GetSuffixName)
+							.MinDesiredWidth(100)
+							.OnTextChanged(this, &SRetargetAnimAssetsWindow::SetSuffixName)
+							.IsReadOnly(false)
+							.RevertTextOnEscape(true)
+					]
+				]
+
+				+SVerticalBox::Slot()
+				.AutoHeight()
+				.Padding(2, 1)
+				[
+					SNew(SHorizontalBox)
+					+SHorizontalBox::Slot()
+					[
+						SNew(STextBlock).Text(LOCTEXT("DuplicateAndRetarget_Search", "Search "))
+					]
+
+					+SHorizontalBox::Slot()
+					[
+						SNew(SEditableTextBox)
+							.Text(this, &SRetargetAnimAssetsWindow::GetReplaceFrom)
+							.MinDesiredWidth(100)
+							.OnTextChanged(this, &SRetargetAnimAssetsWindow::SetReplaceFrom)
+							.IsReadOnly(false)
+							.RevertTextOnEscape(true)
+					]
+				]
+
+				+SVerticalBox::Slot()
+				.AutoHeight()
+				.Padding(2, 1)
+				[
+					SNew(SHorizontalBox)
+					+SHorizontalBox::Slot()
+					[
+						SNew(STextBlock).Text(LOCTEXT("DuplicateAndRetarget_Replace", "Replace "))
+					]
+
+					+SHorizontalBox::Slot()
+					[
+						SNew(SEditableTextBox)
+							.Text(this, &SRetargetAnimAssetsWindow::GetReplaceTo)
+							.MinDesiredWidth(100)
+							.OnTextChanged(this, &SRetargetAnimAssetsWindow::SetReplaceTo)
+							.IsReadOnly(false)
+							.RevertTextOnEscape(true)
+					]
+				]
+
+				+SVerticalBox::Slot()
+				.AutoHeight()
+				.Padding(2, 3)
+				[
+					SNew(SHorizontalBox)
+					+SHorizontalBox::Slot()
+					.Padding(5, 5)
+					[
+						SNew(STextBlock)
+						.Text(this,  &SRetargetAnimAssetsWindow::GetExampleText)
+						.Font(FEditorStyle::GetFontStyle("Persona.RetargetManager.ItalicFont"))
+					]
+				]
+
+				+SVerticalBox::Slot()
+				.AutoHeight()
+				.Padding(2, 3)
+				[
+					SNew(SHorizontalBox)
+
+					+SHorizontalBox::Slot()
+					.AutoWidth()
+					[
+						SNew(STextBlock)
+						.Text(LOCTEXT("DuplicateAndRetarget_Folder", "Folder "))
+						.Font(FEditorStyle::GetFontStyle("Persona.RetargetManager.SmallBoldFont"))
+					]
+
+					+SHorizontalBox::Slot()
+					.FillWidth(1)
+					.HAlign(HAlign_Center)
+					[
+						SNew(STextBlock)
+						.Text(this, &SRetargetAnimAssetsWindow::GetFolderPath)
+					]
+
+					+SHorizontalBox::Slot()
+					.AutoWidth()
+					[
+						SNew(SButton)
+						.HAlign(HAlign_Center)
+						.Text(LOCTEXT("DuplicateAndRetarget_ChangeFolder", "Change..."))
+						//.OnClicked(this, &SAnimationRemapSkeleton::ShowFolderOption)
+					]
+				]
+
+				+SVerticalBox::Slot()
+				.Padding(5)
+				.AutoHeight()
+				[
+					SNew(SSeparator)
+					.Orientation(Orient_Horizontal)
+				]
+
+				+SVerticalBox::Slot()
+				.AutoHeight()
+				.HAlign(HAlign_Fill)
+				.Padding(2)
+				[
+					SNew(SCheckBox)
+					.IsChecked(this, &SRetargetAnimAssetsWindow::IsRemappingReferencedAssets)
+					.OnCheckStateChanged(this, &SRetargetAnimAssetsWindow::OnRemappingReferencedAssetsChanged)
+					[
+						SNew(STextBlock).Text(LOCTEXT("DuplicateAndRetarget_AllowRemap", "Remap Referenced Assets"))
+					]
+				]
+			]
+
+			+SVerticalBox::Slot()
+			.HAlign(HAlign_Right)
+			.VAlign(VAlign_Bottom)
+			.Padding(2)
+			[
+				SNew(SUniformGridPanel)
+				.SlotPadding(FEditorStyle::GetMargin("StandardDialog.SlotPadding"))
+				.MinDesiredSlotWidth(FEditorStyle::GetFloat("StandardDialog.MinDesiredSlotWidth"))
+				.MinDesiredSlotHeight(FEditorStyle::GetFloat("StandardDialog.MinDesiredSlotHeight"))
+				+SUniformGridPanel::Slot(0, 0)
+				[
+					SNew(SButton).HAlign(HAlign_Center)
+					.Text(LOCTEXT("RetargetOptions_Cancel", "Cancel"))
+					.ContentPadding(FEditorStyle::GetMargin("StandardDialog.ContentPadding"))
+					.OnClicked(this, &SRetargetAnimAssetsWindow::OnCancel)
+				]
+				+SUniformGridPanel::Slot(1, 0)
+				[
+					SNew(SButton).HAlign(HAlign_Center)
+					.Text(LOCTEXT("RetargetOptions_Apply", "Retarget"))
+					.IsEnabled(this, &SRetargetAnimAssetsWindow::CanApply)
+					.OnClicked(this, &SRetargetAnimAssetsWindow::OnApply)
+					.HAlign(HAlign_Center)
+					.ContentPadding(FEditorStyle::GetMargin("StandardDialog.ContentPadding"))
+				]
+				
+			]
+		]
+	];
+
+	UpdateExampleText();
+}
+
+bool SRetargetAnimAssetsWindow::CanApply() const
+{
+	return RetargetContext.IsValid();
+}
+
+FReply SRetargetAnimAssetsWindow::OnApply()
+{
+	CloseWindow();
+	RetargetContext.RunRetarget();
+	return FReply::Handled();
+}
+
+FReply SRetargetAnimAssetsWindow::OnCancel()
+{
+	CloseWindow();
+	return FReply::Handled();
+}
+
+void SRetargetAnimAssetsWindow::CloseWindow()
+{
+	if ( DialogWindow.IsValid() )
+	{
+		DialogWindow->RequestDestroyWindow();
+	}
+}
+
+void SRetargetAnimAssetsWindow::ShowWindow(TArray<UObject*> InSelectedAssets)
+{	
+	if(DialogWindow.IsValid())
+	{
+		FSlateApplication::Get().DestroyWindowImmediately(DialogWindow.ToSharedRef());
+	}
+	
+	DialogWindow = SNew(SWindow)
+		.Title(LOCTEXT("RetargetAssets", "Duplicate and Retarget Animation Assets"))
+		.SupportsMinimize(false)
+		.SupportsMaximize(false)
+		.HasCloseButton(true)
+		.MaxWidth(1024.0f)
+		.IsTopmostWindow(true)
+		.SizingRule(ESizingRule::Autosized);
+	
+	TSharedPtr<class SRetargetAnimAssetsWindow> DialogWidget;
+	TSharedPtr<SBorder> DialogWrapper =
+		SNew(SBorder)
+		.BorderImage(FEditorStyle::GetBrush("ToolPanel.GroupBorder"))
+		.Padding(4.0f)
+		[
+			SAssignNew(DialogWidget, SRetargetAnimAssetsWindow)
+		];
+
+	DialogWidget->RetargetContext.AssetsToRetarget = FObjectEditorUtils::GetTypedWeakObjectPtrs<UObject>(InSelectedAssets);
+	DialogWindow->SetOnWindowClosed(FRequestDestroyWindowOverride::CreateSP(DialogWidget.Get(), &SRetargetAnimAssetsWindow::OnDialogClosed));
+	DialogWindow->SetContent(DialogWrapper.ToSharedRef());
+	
+	FSlateApplication::Get().AddWindow(DialogWindow.ToSharedRef());
+}
+
+void SRetargetAnimAssetsWindow::OnDialogClosed(const TSharedRef<SWindow>& Window)
+{
+	DialogWindow = nullptr;
+}
+
+void SRetargetAnimAssetsWindow::SourceMeshAssigned(const FAssetData& InAssetData)
+{
+	USkeletalMesh* Mesh = Cast<USkeletalMesh>(InAssetData.GetAsset());
+	RetargetContext.SourceMesh = Mesh;
+	SourceViewport->SetSkeletalMesh(RetargetContext.SourceMesh);
+}
+
+void SRetargetAnimAssetsWindow::TargetMeshAssigned(const FAssetData& InAssetData)
+{
+	USkeletalMesh* Mesh = Cast<USkeletalMesh>(InAssetData.GetAsset());
+	RetargetContext.TargetMesh = Mesh;
+	TargetViewport->SetSkeletalMesh(RetargetContext.TargetMesh);
+}
+
+FString SRetargetAnimAssetsWindow::GetCurrentSourceMeshPath() const
+{
+	return RetargetContext.SourceMesh ? RetargetContext.SourceMesh->GetPathName() : FString("");
+}
+
+FString SRetargetAnimAssetsWindow::GetCurrentTargetMeshPath() const
+{
+	return RetargetContext.TargetMesh ? RetargetContext.TargetMesh->GetPathName() : FString("");
+}
+
+FString SRetargetAnimAssetsWindow::GetCurrentRetargeterPath() const
+{
+	return RetargetContext.IKRetargetAsset ? RetargetContext.IKRetargetAsset->GetPathName() : FString("");
+}
+
+void SRetargetAnimAssetsWindow::RetargeterAssigned(const FAssetData& InAssetData)
+{
+	UIKRetargeter* InRetargeter = Cast<UIKRetargeter>(InAssetData.GetAsset());
+	RetargetContext.IKRetargetAsset = InRetargeter;
+	const UIKRigDefinition* SourceIKRig = InRetargeter ? InRetargeter->SourceIKRigAsset : nullptr;
+	const UIKRigDefinition* TargetIKRig = InRetargeter ? InRetargeter->TargetIKRigAsset : nullptr;
+	USkeletalMesh* SourceMesh =  SourceIKRig ? SourceIKRig->GetPreviewMesh() : nullptr;
+	USkeletalMesh* TargetMesh =  TargetIKRig ? TargetIKRig->GetPreviewMesh() : nullptr;
+	SourceMeshAssigned(FAssetData(SourceMesh));
+	TargetMeshAssigned(FAssetData(TargetMesh));
+}
+
+ECheckBoxState SRetargetAnimAssetsWindow::IsRemappingReferencedAssets() const
+{
+	return RetargetContext.bRemapReferencedAssets ? ECheckBoxState::Checked : ECheckBoxState::Unchecked;
+}
+
+void SRetargetAnimAssetsWindow::OnRemappingReferencedAssetsChanged(ECheckBoxState InNewRadioState)
+{
+	RetargetContext.bRemapReferencedAssets = (InNewRadioState == ECheckBoxState::Checked);
+}
+
+FText SRetargetAnimAssetsWindow::GetPrefixName() const
+{
+	return FText::FromString(RetargetContext.NameRule.Prefix);
+}
+
+void SRetargetAnimAssetsWindow::SetPrefixName(const FText &InText)
+{
+	RetargetContext.NameRule.Prefix = InText.ToString();
+	UpdateExampleText();
+}
+
+FText SRetargetAnimAssetsWindow::GetSuffixName() const
+{
+	return FText::FromString(RetargetContext.NameRule.Suffix);
+}
+
+void SRetargetAnimAssetsWindow::SetSuffixName(const FText &InText)
+{
+	RetargetContext.NameRule.Suffix = InText.ToString();
+	UpdateExampleText();
+}
+
+FText SRetargetAnimAssetsWindow::GetReplaceFrom() const
+{
+	return FText::FromString(RetargetContext.NameRule.ReplaceFrom);
+}
+
+void SRetargetAnimAssetsWindow::SetReplaceFrom(const FText &InText)
+{
+	RetargetContext.NameRule.ReplaceFrom = InText.ToString();
+	UpdateExampleText();
+}
+
+FText SRetargetAnimAssetsWindow::GetReplaceTo() const
+{
+	return FText::FromString(RetargetContext.NameRule.ReplaceTo);
+}
+
+void SRetargetAnimAssetsWindow::SetReplaceTo(const FText &InText)
+{
+	RetargetContext.NameRule.ReplaceTo = InText.ToString();
+	UpdateExampleText();
+}
+
+FText SRetargetAnimAssetsWindow::GetExampleText() const
+{
+	return ExampleText;
+}
+
+void SRetargetAnimAssetsWindow::UpdateExampleText()
+{
+	const FString ReplaceFrom = FString::Printf(TEXT("Old Name : ###%s###"), *RetargetContext.NameRule.ReplaceFrom);
+	const FString ReplaceTo = FString::Printf(TEXT("New Name : %s###%s###%s"), *RetargetContext.NameRule.Prefix, *RetargetContext.NameRule.ReplaceTo, *RetargetContext.NameRule.Suffix);
+
+	ExampleText = FText::FromString(FString::Printf(TEXT("%s\n%s"), *ReplaceFrom, *ReplaceTo));
+}
+
+FText SRetargetAnimAssetsWindow::GetFolderPath() const
+{
+	return FText::FromString(RetargetContext.NameRule.FolderPath);
+}
+
+#undef LOCTEXT_NAMESPACE
diff --git a/Engine/Plugins/Animation/IKRig/Source/IKRigEditor/Public/IKRigEditor.h b/Engine/Plugins/Animation/IKRig/Source/IKRigEditor/Public/IKRigEditor.h
--- a/Engine/Plugins/Animation/IKRig/Source/IKRigEditor/Public/IKRigEditor.h
+++ b/Engine/Plugins/Animation/IKRig/Source/IKRigEditor/Public/IKRigEditor.h
@@ -12,7 +12,7 @@ public:
 	void ShutdownModule() override;
 
 private:
+	TSharedPtr<class FAssetTypeActions_AnimationAssetRetarget> RetargetAnimationAssetAction;
 	TSharedPtr<class FAssetTypeActions_IKRigDefinition> IKRigDefinitionAssetAction;
 	TSharedPtr<class FAssetTypeActions_IKRetargeter> IKRetargeterAssetAction;
-	TSharedPtr<class FAssetTypeActions_IKRetargetRig> IKRetargetRigAssetAction;
 };
diff --git a/Engine/Plugins/Animation/IKRig/Source/IKRigEditor/Public/RetargetEditor/AssetTypeActions_IKRetargeter.h b/Engine/Plugins/Animation/IKRig/Source/IKRigEditor/Public/RetargetEditor/AssetTypeActions_IKRetargeter.h
--- a/Engine/Plugins/Animation/IKRig/Source/IKRigEditor/Public/RetargetEditor/AssetTypeActions_IKRetargeter.h
+++ b/Engine/Plugins/Animation/IKRig/Source/IKRigEditor/Public/RetargetEditor/AssetTypeActions_IKRetargeter.h
@@ -5,6 +5,9 @@
 #include "CoreMinimal.h"
 #include "Toolkits/IToolkitHost.h"
 #include "AssetTypeActions_Base.h"
+#include "EditorAnimUtils.h"
+
+class UIKRetargeter;
 
 class FAssetTypeActions_IKRetargeter : public FAssetTypeActions_Base
 {
@@ -17,4 +20,7 @@ public:
 	virtual void OpenAssetEditor(const TArray<UObject*>& InObjects, TSharedPtr<class IToolkitHost> EditWithinLevelEditor = TSharedPtr<IToolkitHost>()) override;
 	virtual uint32 GetCategories() override { return EAssetTypeCategories::Animation; }
 	virtual UThumbnailInfo* GetThumbnailInfo(UObject* Asset) const override;
+
+	static void ExtendAnimSequenceToolMenu();
+	static void CreateRetargetSubMenu(FToolMenuSection& InSection);
 };
diff --git a/Engine/Plugins/Animation/IKRig/Source/IKRigEditor/Public/RetargetEditor/SRetargetAnimAssetsWindow.h b/Engine/Plugins/Animation/IKRig/Source/IKRigEditor/Public/RetargetEditor/SRetargetAnimAssetsWindow.h
new file mode 100644
--- /dev/null
+++ b/Engine/Plugins/Animation/IKRig/Source/IKRigEditor/Public/RetargetEditor/SRetargetAnimAssetsWindow.h
@@ -0,0 +1,284 @@
+﻿// Copyright Epic Games, Inc. All Rights Reserved.
+
+#pragma once
+
+#include "EditorAnimUtils.h"
+#include "SEditorViewport.h"
+#include "Settings/SkeletalMeshEditorSettings.h"
+
+class UIKRetargeter;
+
+//** Data needed to initiate a batch duplicate and retarget of animation assets */
+struct FIKRetargetAnimAssetsContext
+{
+	/** The selected assets when batch process was initiated */
+	TArray<TWeakObjectPtr<UObject>> AssetsToRetarget;
+
+	/** Source mesh to use to copy animation FROM. */
+	USkeletalMesh* SourceMesh = nullptr;
+
+	/** Target mesh to use to copy animation TO. */
+	USkeletalMesh* TargetMesh = nullptr;
+
+	/** The retargeter used to copy animation */
+	UIKRetargeter* IKRetargetAsset = nullptr;
+
+	/** Whether we are remapping assets that are referenced by the assets the user selects to remap */
+	bool bRemapReferencedAssets = true;
+
+	/* Rename rules for duplicated assets */
+	EditorAnimUtils::FNameDuplicationRule NameRule;
+
+	/* The base folder path to put newly duplicated assets */
+	FString FolderPath = "";
+
+	/* Reset all data (called when window re-opened */
+	void Reset();
+
+	/* Is the data configured in such a way that we could run the retarget? */
+	bool IsValid() const;
+
+	/* Actually run the process to duplicate and retarget the assets. */
+	void RunRetarget();
+
+private:
+
+
+	/**
+	* Initialize set of referenced assets to retarget.
+	* @return	Number of assets that need retargeting.
+	*/
+	int32 GenerateAssetLists();
+
+	/* Duplicate all the assets to retarget */
+	void DuplicateRetargetAssets();
+
+	/* Retarget skeleton and animation on all the duplicates */
+	void RetargetAssets();
+
+	/* Convert animation on all the duplicates */
+	void ConvertAnimation();
+
+	/* Output notifications of results */
+	void NotifyUserOfResults() const;
+
+	void GetNewAssets(TArray<UObject*>& NewAssets) const;
+
+	/**
+	* Duplicates the supplied AssetsToDuplicate and returns a map of original asset to duplicate. Templated wrapper that calls DuplicateAssetInternal.
+	*
+	* @param	AssetsToDuplicate	The animations to duplicate
+	* @param	DestinationPackage	The package that the duplicates should be placed in
+	*
+	* @return	TMap of original animation to duplicate
+	*/
+	template<class AssetType>
+	static TMap<AssetType*, AssetType*> DuplicateAssets(
+		const TArray<AssetType*>& AssetsToDuplicate,
+		UPackage* DestinationPackage,
+		const EditorAnimUtils::FNameDuplicationRule* NameRule);
+	
+	/** Lists of assets to retarget. Populated from selection during init */
+	TArray<UAnimationAsset*>	AnimationAssetsToRetarget;
+	TArray<UAnimBlueprint*>		AnimBlueprintsToRetarget;
+
+	/** Lists of original assets map to duplicate assets */
+	TMap<UAnimationAsset*, UAnimationAsset*>	DuplicatedAnimAssets;
+	TMap<UAnimBlueprint*, UAnimBlueprint*>		DuplicatedBlueprints;
+
+	TMap<UAnimationAsset*, UAnimationAsset*>	RemappedAnimAssets;
+
+	/** If we only chose one object to retarget store it here */
+	UObject* SingleTargetObject = nullptr;
+};
+
+//** Viewport in retarget dialog that shows source / target retarget pose */
+class SRetargetPoseViewport: public SEditorViewport
+{
+	
+public:
+	
+	SLATE_BEGIN_ARGS(SRetargetPoseViewport)
+	{}
+	SLATE_ARGUMENT(USkeletalMesh*, SkeletalMesh)
+	SLATE_END_ARGS()
+
+	SRetargetPoseViewport();
+
+	void Construct(const FArguments& InArgs);
+	
+	void SetSkeletalMesh(USkeletalMesh* SkeletalMesh);
+
+protected:
+	
+	/** SEditorViewport interface */
+	virtual TSharedRef<FEditorViewportClient> MakeEditorViewportClient() override;
+	virtual TSharedPtr<SWidget> MakeViewportToolbar() override;
+	/** END SEditorViewport interface */
+
+private:
+
+	USkeletalMesh* Mesh;
+
+	FPreviewScene PreviewScene;
+
+	class UDebugSkelMeshComponent* PreviewComponent;
+
+	virtual bool IsVisible() const override;
+};
+
+//** Client for SRetargetPoseViewport */
+class FRetargetPoseViewportClient: public FEditorViewportClient
+{
+	
+public:
+	
+	FRetargetPoseViewportClient(
+		FPreviewScene& InPreviewScene,
+		const TSharedRef<SRetargetPoseViewport>& InRetargetPoseViewport)
+		: FEditorViewportClient(
+			nullptr,
+			&InPreviewScene,
+			StaticCastSharedRef<SEditorViewport>(InRetargetPoseViewport))
+	{
+		SetViewMode(VMI_Lit);
+
+		// Always composite editor objects after post processing in the editor
+		EngineShowFlags.SetCompositeEditorPrimitives(true);
+		EngineShowFlags.DisableAdvancedFeatures();
+
+		UpdateLighting();
+
+		// Setup defaults for the common draw helper.
+		DrawHelper.bDrawPivot = false;
+		DrawHelper.bDrawWorldBox = false;
+		DrawHelper.bDrawKillZ = false;
+		DrawHelper.bDrawGrid = true;
+		DrawHelper.GridColorAxis = FColor(70, 70, 70);
+		DrawHelper.GridColorMajor = FColor(40, 40, 40);
+		DrawHelper.GridColorMinor =  FColor(20, 20, 20);
+		DrawHelper.PerspectiveGridSize = HALF_WORLD_MAX1;
+
+		bDisableInput = true;
+	}
+
+	/** FEditorViewportClient interface */
+	virtual void Tick(float DeltaTime) override
+	{
+		if (PreviewScene)
+		{
+			PreviewScene->GetWorld()->Tick(LEVELTICK_All, DeltaTime);
+		}
+	}
+
+	virtual FSceneInterface* GetScene() const override
+	{
+		return PreviewScene->GetScene();
+	}
+
+	virtual FLinearColor GetBackgroundColor() const override 
+	{ 
+		return FLinearColor::White; 
+	}
+	/** END FEditorViewportClient interface */
+
+	void UpdateLighting() const
+	{
+		const USkeletalMeshEditorSettings* Options = GetDefault<USkeletalMeshEditorSettings>();
+
+		PreviewScene->SetLightDirection(Options->AnimPreviewLightingDirection);
+		PreviewScene->SetLightColor(Options->AnimPreviewDirectionalColor);
+		PreviewScene->SetLightBrightness(Options->AnimPreviewLightBrightness);
+	}
+};
+
+//** Window to display when configuring batch duplicate & retarget process */
+class SRetargetAnimAssetsWindow : public SCompoundWidget
+{
+	
+public:
+
+	SLATE_BEGIN_ARGS(SRetargetAnimAssetsWindow)
+		: _CurrentSkeletalMesh(nullptr)
+		, _WidgetWindow(nullptr)
+		, _ShowRemapOption(false)
+	{
+	}
+
+	SLATE_ARGUMENT( USkeletalMesh*, CurrentSkeletalMesh )
+	SLATE_ARGUMENT( TSharedPtr<SWindow>, WidgetWindow )
+	SLATE_ARGUMENT( bool, ShowRemapOption )
+	SLATE_END_ARGS()	
+
+	/** Constructs this widget for the window */
+	void Construct( const FArguments& InArgs );
+
+private:
+
+	/** Retarget or Cancel buttons */
+	bool CanApply() const;
+	FReply OnApply();
+	FReply OnCancel();
+	void CloseWindow();
+
+	/** Handler for dialog window close button */
+	void OnDialogClosed(const TSharedRef<SWindow>& Window);
+
+	/** Modifying Source Mesh */
+	void SourceMeshAssigned(const FAssetData& InAssetData);
+	FString GetCurrentSourceMeshPath() const;
+
+	/** Modifying Target Mesh */
+	void TargetMeshAssigned(const FAssetData& InAssetData);
+	FString GetCurrentTargetMeshPath() const;
+
+	/** Modifying Retargeter */
+	FString GetCurrentRetargeterPath() const;
+	void RetargeterAssigned(const FAssetData& InAssetData);
+	
+	/** Modifying "Remap Assets" checkbox */
+	ECheckBoxState IsRemappingReferencedAssets() const;
+	void OnRemappingReferencedAssetsChanged(ECheckBoxState InNewRadioState);
+
+	/** Modifying Rename Prefix */
+	FText GetPrefixName() const;
+	void SetPrefixName(const FText &InText);
+
+	/** Modifying Rename Suffix */
+	FText GetSuffixName() const;
+	void SetSuffixName(const FText &InText);
+
+	/** Modifying Search/Replace text */
+	FText GetReplaceFrom() const;
+	void SetReplaceFrom(const FText &InText);
+	FText GetReplaceTo() const;
+	void SetReplaceTo(const FText &InText);
+
+	/** Example rename text */
+	FText GetExampleText() const;
+	void UpdateExampleText();
+
+	/** Modify folder output path */
+	FText GetFolderPath() const;
+
+	/** Necessary data collected from UI to run retarget. */
+	FIKRetargetAnimAssetsContext RetargetContext;
+
+	/** The rename rule sample text */
+	FText ExampleText;
+
+	/** Pool for maintaining and rendering thumbnails */
+	TSharedPtr<FAssetThumbnailPool> AssetThumbnailPool;
+
+	/** Viewport displaying SOURCE mesh in retarget pose */
+	TSharedPtr<SRetargetPoseViewport> SourceViewport;
+
+	/** Viewport displaying TARGET mesh in retarget pose */
+	TSharedPtr<SRetargetPoseViewport> TargetViewport;
+
+public:
+
+	static void ShowWindow(TArray<UObject*> InAnimAssets);
+
+	static TSharedPtr<SWindow> DialogWindow;
+};
\ No newline at end of file
diff --git a/Engine/Source/Editor/UnrealEd/Public/EditorAnimUtils.h b/Engine/Source/Editor/UnrealEd/Public/EditorAnimUtils.h
--- a/Engine/Source/Editor/UnrealEd/Public/EditorAnimUtils.h
+++ b/Engine/Source/Editor/UnrealEd/Public/EditorAnimUtils.h
@@ -120,10 +120,10 @@ namespace EditorAnimUtils
 	UNREALED_API UObject* RetargetAnimations(USkeleton* OldSkeleton, USkeleton* NewSkeleton, FAnimationRetargetContext& RetargetContext, bool bRetargetReferredAssets, const FNameDuplicationRule* NameRule);
 
 	// Populates the supplied TArrays with any animation assets that this blueprint refers too
-	void GetAllAnimationSequencesReferredInBlueprint(UAnimBlueprint* AnimBlueprint, TArray<UAnimationAsset*>& AnimationAsset);
+	UNREALED_API void GetAllAnimationSequencesReferredInBlueprint(UAnimBlueprint* AnimBlueprint, TArray<UAnimationAsset*>& AnimationAsset);
 
 	// Replaces references to any animations found with the match animation from the map
-	void ReplaceReferredAnimationsInBlueprint(UAnimBlueprint* AnimBlueprint, const TMap<UAnimationAsset*, UAnimationAsset*>& AnimAssetReplacementMap);
+	UNREALED_API void ReplaceReferredAnimationsInBlueprint(UAnimBlueprint* AnimBlueprint, const TMap<UAnimationAsset*, UAnimationAsset*>& AnimAssetReplacementMap);
 
 	/**
 	 * Duplicates the supplied AssetsToDuplicate and returns a map of original asset to duplicate
@@ -133,7 +133,7 @@ namespace EditorAnimUtils
 	 *
 	 * @return	TMap of original animation to duplicate
 	 */
-	TMap<UObject*, UObject*> DuplicateAssetsInternal(const TArray<UObject*>& AssetsToDuplicate, UPackage* DestinationPackage, const FNameDuplicationRule* NameRule);
+	UNREALED_API TMap<UObject*, UObject*> DuplicateAssetsInternal(const TArray<UObject*>& AssetsToDuplicate, UPackage* DestinationPackage, const FNameDuplicationRule* NameRule);
 
 	/**
 	 * Duplicates the supplied AssetsToDuplicate and returns a map of original asset to duplicate. Templated wrapper that calls DuplicateAssetInternal.
