IK Retargeter fixes
- refactored batch operation into standalone system
- applied batch operation to export button
- moved export animation buttno above asset browser
- replaced retarget editor asset browser with custom one (simplified)
- fixed bug, root retarget orientation wrong
- fixed bug, IK retarget orientation wrong
- fixed bug, IK Rig crash after goal deleted and moved in viewport
- fixed bug, progress dialog now displays correctly during batch op
- replaced mention of "gizmo" with "debug draw"
- added padding to chain-map elements UI
- renamed "IK Retargeter" anim node to "Retarget Pose from Mesh"

diff --git a/Engine/Plugins/Animation/IKRig/Config/BaseIKRig.ini b/Engine/Plugins/Animation/IKRig/Config/BaseIKRig.ini
--- a/Engine/Plugins/Animation/IKRig/Config/BaseIKRig.ini
+++ b/Engine/Plugins/Animation/IKRig/Config/BaseIKRig.ini
@@ -0,0 +1 @@
++ClassRedirects=(OldName="AnimNode_IKRetargeter",NewName="AnimNode_RetargetPoseFromMesh")
diff --git a/Engine/Plugins/Animation/IKRig/Source/IKRig/Private/AnimNodes/AnimNode_IKRetargeter.cpp b/Engine/Plugins/Animation/IKRig/Source/IKRig/Private/AnimNodes/AnimNode_IKRetargeter.cpp
deleted file mode 100644
diff --git a/Engine/Plugins/Animation/IKRig/Source/IKRig/Private/AnimNodes/AnimNode_RetargetPoseFromMesh.cpp b/Engine/Plugins/Animation/IKRig/Source/IKRig/Private/AnimNodes/AnimNode_RetargetPoseFromMesh.cpp
new file mode 100644
--- /dev/null
+++ b/Engine/Plugins/Animation/IKRig/Source/IKRig/Private/AnimNodes/AnimNode_RetargetPoseFromMesh.cpp
@@ -0,0 +1,205 @@
+// Copyright Epic Games, Inc. All Rights Reserved.
+
+#include "AnimNodes/AnimNode_RetargetPoseFromMesh.h"
+#include "Animation/AnimInstanceProxy.h"
+
+
+void FAnimNode_RetargetPoseFromMesh::Initialize_AnyThread(const FAnimationInitializeContext& Context)
+{
+	DECLARE_SCOPE_HIERARCHICAL_COUNTER_ANIMNODE(Initialize_AnyThread)
+    FAnimNode_Base::Initialize_AnyThread(Context);
+
+	// Initial update of the node, so we dont have a frame-delay on setup
+	GetEvaluateGraphExposedInputs().Execute(Context);
+}
+
+void FAnimNode_RetargetPoseFromMesh::CacheBones_AnyThread(const FAnimationCacheBonesContext& Context)
+{
+	DECLARE_SCOPE_HIERARCHICAL_COUNTER_ANIMNODE(CacheBones_AnyThread)
+}
+
+void FAnimNode_RetargetPoseFromMesh::Update_AnyThread(const FAnimationUpdateContext& Context)
+{
+	DECLARE_SCOPE_HIERARCHICAL_COUNTER_ANIMNODE(Update_AnyThread)
+	FAnimNode_Base::Update_AnyThread(Context);
+    // this introduces a frame of latency in setting the pin-driven source component,
+    // but we cannot do the work to extract transforms on a worker thread as it is not thread safe.
+    GetEvaluateGraphExposedInputs().Execute(Context);
+}
+void FAnimNode_RetargetPoseFromMesh::Evaluate_AnyThread(FPoseContext& Output)
+{
+	DECLARE_SCOPE_HIERARCHICAL_COUNTER_ANIMNODE(Evaluate_AnyThread)
+
+	if (!bIsInitialized)
+	{
+		Output.ResetToRefPose();
+		return;
+	}
+
+	// run the retargeter
+	const TArray<FTransform>& RetargetedPose = CurrentlyUsedRetargeter->RunRetargeter(SourceMeshComponentSpaceBoneTransforms);
+
+	// copy pose back
+	FCSPose<FCompactPose> ComponentPose;
+	ComponentPose.InitPose(Output.Pose);
+	for (FCompactPoseBoneIndex CompactBoneIndex : Output.Pose.ForEachBoneIndex())
+	{
+		ComponentPose.SetComponentSpaceTransform(CompactBoneIndex, RetargetedPose[CompactBoneIndex.GetInt()]);
+	}
+
+	// convert to local space
+	FCSPose<FCompactPose>::ConvertComponentPosesToLocalPoses(ComponentPose, Output.Pose);
+}
+
+void FAnimNode_RetargetPoseFromMesh::PreUpdate(const UAnimInstance* InAnimInstance)
+{
+	DECLARE_SCOPE_HIERARCHICAL_COUNTER_ANIMNODE(PreUpdate)
+	EnsureInitialized(InAnimInstance);
+	if (bIsInitialized)
+	{
+		CopyBoneTransformsFromSource(InAnimInstance->GetSkelMeshComponent());
+	}
+}
+
+UIKRetargeter* FAnimNode_RetargetPoseFromMesh::GetCurrentlyUsedRetargeter() const
+{
+	return CurrentlyUsedRetargeter.Get();
+}
+
+void FAnimNode_RetargetPoseFromMesh::EnsureInitialized(const UAnimInstance* InAnimInstance)
+{
+	// has user supplied a retargeter asset?
+	if (!IKRetargeterAsset)
+	{
+		bIsInitialized = false;
+		return;
+	}
+
+	// if user hasn't explicitly connected a source mesh, optionally use the parent mesh component (if there is one) 
+	if (!SourceMeshComponent.IsValid() && bUseAttachedParent)
+	{
+		USkeletalMeshComponent* TargetMesh = InAnimInstance->GetSkelMeshComponent();
+		USkeletalMeshComponent* ParentComponent = Cast<USkeletalMeshComponent>(TargetMesh->GetAttachParent());
+		if (ParentComponent)
+		{
+			SourceMeshComponent = ParentComponent;
+		}
+	}
+	
+	// has a source mesh been plugged in or found?
+	if (!SourceMeshComponent.IsValid())
+	{
+		bIsInitialized = false;
+		return; // can't do anything if we don't have a source mesh
+	}
+
+	// has the source asset been updated? this can happen in the editor when changing source asset
+	if (!CurrentlyUsedRetargeter
+		|| (CurrentlyUsedRetargeter->AssetVersion != IKRetargeterAsset->AssetVersion))
+	{
+		InitializeRetargetData(InAnimInstance);
+		return;	
+	}
+}
+
+void FAnimNode_RetargetPoseFromMesh::InitializeRetargetData(const UAnimInstance* InAnimInstance)
+{
+	// assume we fail until we don't
+	bIsInitialized = false;
+	
+	// don't try to reinitialize with same asset version
+	if (LastVersionTried == IKRetargeterAsset->AssetVersion)
+	{
+		return;
+	}
+	LastVersionTried = IKRetargeterAsset->AssetVersion;
+
+	// store all the components that were used to initialize
+	// if in future updates, any of this are mismatched, we have to re-initialize
+	CurrentlyUsedSourceMeshComponent = SourceMeshComponent;
+	CurrentlyUsedSourceMesh = SourceMeshComponent->SkeletalMesh;
+	CurrentlyUsedTargetMesh = InAnimInstance->GetSkelMeshComponent()->SkeletalMesh;
+	CurrentlyUsedRetargeter = DuplicateObject(IKRetargeterAsset, const_cast<UAnimInstance*>(InAnimInstance));
+	CurrentlyUsedRetargeter->AssetVersion = IKRetargeterAsset->AssetVersion;
+	CurrentlyUsedSourceIKRig = CurrentlyUsedRetargeter->SourceIKRigAsset;
+	CurrentlyUsedTargetIKRig = CurrentlyUsedRetargeter->TargetIKRigAsset;
+
+	const bool bMeshesAreValid = CurrentlyUsedSourceMesh.IsValid() && CurrentlyUsedTargetMesh.IsValid();
+	if (!bMeshesAreValid)
+	{
+		return; // cannot initialize if components are missing skeletal mesh references
+	}
+
+	const bool bRetargeterIsValid =
+	   CurrentlyUsedRetargeter
+	&& CurrentlyUsedSourceIKRig.IsValid()
+	&& CurrentlyUsedTargetIKRig.IsValid();
+	if (!bRetargeterIsValid)
+	{
+		return; // cannot initialize unless we have a retargeter with BOTH source AND target IK Rigs
+	}
+
+	// initialize retargeter with source and target skeletal meshes
+	// (anim instance is passed in as outer UObject for new UIKRigProcessor) 
+	CurrentlyUsedRetargeter->Initialize(
+		CurrentlyUsedSourceMesh.Get(),
+		CurrentlyUsedTargetMesh.Get(),
+		const_cast<UAnimInstance*>(InAnimInstance));
+	
+	// made it!
+	bIsInitialized = CurrentlyUsedRetargeter->bIsLoadedAndValid;
+}
+
+void FAnimNode_RetargetPoseFromMesh::CopyBoneTransformsFromSource(USkeletalMeshComponent* TargetMeshComponent)
+{
+	if (!CurrentlyUsedSourceMeshComponent.IsValid())
+	{
+		return; 
+	}
+
+	USkeletalMeshComponent* SourceMeshComp =  CurrentlyUsedSourceMeshComponent.Get();
+	
+	// is the source mesh ticking?
+	if (!SourceMeshComp->IsRegistered())
+		
+	{
+		CurrentlyUsedSourceMesh.Reset(); // forces reinitialization when re-registered
+		return; // skip copying pose when component is no longer ticking
+	}
+	
+	// if our source is running under master-pose, then get bone data from there
+	if(USkeletalMeshComponent* MasterPoseComponent = Cast<USkeletalMeshComponent>(SourceMeshComponent->MasterPoseComponent.Get()))
+	{
+		SourceMeshComp = MasterPoseComponent;
+	}
+
+	// re-check mesh component validity as it may have changed to master
+	if (!(SourceMeshComp->SkeletalMesh && SourceMeshComp->IsRegistered()))
+	{
+		return; // master pose either missing skeletal mesh reference or not ticking, either way, we aren't copying from it
+	}
+	
+	const bool bUROInSync = SourceMeshComp->ShouldUseUpdateRateOptimizations() && SourceMeshComp->AnimUpdateRateParams != nullptr && SourceMeshComponent->AnimUpdateRateParams == TargetMeshComponent->AnimUpdateRateParams;
+	const bool bUsingExternalInterpolation = SourceMeshComp->IsUsingExternalInterpolation();
+	const TArray<FTransform>& CachedComponentSpaceTransforms = SourceMeshComp->GetCachedComponentSpaceTransforms();
+	const bool bArraySizesMatch = CachedComponentSpaceTransforms.Num() == SourceMeshComp->GetComponentSpaceTransforms().Num();
+
+	// copy source array from the appropriate location
+	SourceMeshComponentSpaceBoneTransforms.Reset();
+	if ((bUROInSync || bUsingExternalInterpolation) && bArraySizesMatch)
+	{
+		SourceMeshComponentSpaceBoneTransforms.Append(CachedComponentSpaceTransforms); // copy from source's cache
+	}
+	else
+	{
+		SourceMeshComponentSpaceBoneTransforms.Append(SourceMeshComp->GetComponentSpaceTransforms()); // copy directly
+	}
+
+	// ref skeleton is need for parent index lookups later, so store it now
+	CurrentlyUsedSourceMesh = SourceMeshComp->SkeletalMesh;
+}
+
+USkeletalMeshComponent* FAnimNode_RetargetPoseFromMesh::GetSourceMesh() const
+{
+	return nullptr;
+}
diff --git a/Engine/Plugins/Animation/IKRig/Source/IKRig/Private/Retargeter/IKRetargeter.cpp b/Engine/Plugins/Animation/IKRig/Source/IKRig/Private/Retargeter/IKRetargeter.cpp
--- a/Engine/Plugins/Animation/IKRig/Source/IKRig/Private/Retargeter/IKRetargeter.cpp
+++ b/Engine/Plugins/Animation/IKRig/Source/IKRig/Private/Retargeter/IKRetargeter.cpp
@@ -124,6 +124,17 @@ void FRetargetSkeleton::UpdateLocalTransformOfSingleBone(
 	OutLocalPose[BoneIndex] = ChildGlobalTransform.GetRelativeTransform(ParentGlobalTransform);
 }
 
+void FRetargetSkeleton::GetChildrenIndices(const int32 BoneIndex, TArray<int32>& OutChildren) const
+{
+	for (int32 ChildBoneIndex=0; ChildBoneIndex<ParentIndices.Num(); ++ChildBoneIndex)
+	{
+		if (ParentIndices[ChildBoneIndex] == BoneIndex)
+		{
+			OutChildren.Add(ChildBoneIndex);
+		}
+	}
+}
+
 int32 FRetargetSkeleton::GetParentIndex(const int32 BoneIndex) const
 {
 	if (BoneIndex < 0 || BoneIndex>ParentIndices.Num() || BoneIndex == INDEX_NONE)
@@ -376,75 +387,74 @@ FTransform FChainDecoderFK::GetTransformAtParam(
 	return FTransform::Identity;
 }
 
-bool FChainEncoderIK::Initialize(
+bool FChainRetargeterIK::InitializeSource(
 	const TArray<int32>& BoneIndices,
-	const TArray<FTransform>& InitialGlobalPose)
+	const TArray<FTransform>& SourceInitialGlobalPose)
 {
 	if (BoneIndices.Num() < 3)
 	{
-		UE_LOG(LogTemp, Warning, TEXT("IK Retargeter trying to retarget bone chain with IK but it has less than 3 joints."));
+		UE_LOG(LogTemp, Warning, TEXT("IK Retargeter trying to retarget source bone chain with IK but it has less than 3 joints."));
 		return false;
 	}
 	
-	BoneIndexA = BoneIndices[0];
-	BoneIndexB = BoneIndices[1];
-	BoneIndexC = BoneIndices.Last();
+	Source.BoneIndexA = BoneIndices[0];
+	Source.BoneIndexB = BoneIndices[1];
+	Source.BoneIndexC = BoneIndices.Last();
 	
-	const FTransform& End = InitialGlobalPose[BoneIndexC];
-	EndPositionOrig = End.GetTranslation();
-	EndRotationOrig = End.GetRotation();
+	const FTransform& End = SourceInitialGlobalPose[Source.BoneIndexC];
+	Source.InitialEndPosition = End.GetTranslation();
+	Source.InitialEndRotation = End.GetRotation();
 
-	const FTransform& Start = InitialGlobalPose[BoneIndexA];
-	const float Length = (Start.GetTranslation() - EndPositionOrig).Size();
+	const FTransform& Start = SourceInitialGlobalPose[Source.BoneIndexA];
+	const float Length = (Start.GetTranslation() - Source.InitialEndPosition).Size();
 
 	if (Length <= KINDA_SMALL_NUMBER)
 	{
-		UE_LOG(LogTemp, Warning, TEXT("IK Retargeter trying to retarget bone chain with IK, but it is zero length!"));
+		UE_LOG(LogTemp, Warning, TEXT("IK Retargeter trying to retarget source bone chain with IK, but it is zero length!"));
     	return false;
 	}
 
-	InverseLength = 1.0f / Length;
+	Source.InvInitialLength = 1.0f / Length;
 	
 	return true;
 }
 
-void FChainEncoderIK::EncodePose(const TArray<FTransform>& InputGlobalPose)
+void FChainRetargeterIK::EncodePose(const TArray<FTransform>& InputGlobalPose)
 {
-	const FVector A = InputGlobalPose[BoneIndexA].GetTranslation();
+	const FVector A = InputGlobalPose[Source.BoneIndexA].GetTranslation();
 	//FVector B = InputGlobalPose[BoneIndexB].GetTranslation(); TODO use for pole vector 
-	const FVector C = InputGlobalPose[BoneIndexC].GetTranslation();
+	const FVector C = InputGlobalPose[Source.BoneIndexC].GetTranslation();
 
     // get the normalized direction / length of the IK limb (how extended it is as percentage of original length)
     const FVector AC = C - A;
 	float ACLength;
 	FVector ACDirection;
 	AC.ToDirectionAndLength(ACDirection, ACLength);
-	const float NormalizedLimbLength = ACLength * InverseLength;
+	const float NormalizedLimbLength = ACLength * Source.InvInitialLength;
 
-	EncodedResult.EndDirectionNormalized = ACDirection * NormalizedLimbLength;
-	EncodedResult.EndRotation = InputGlobalPose[BoneIndexC].GetRotation();
-	EncodedResult.EndRotationOrig = EndRotationOrig;
-	EncodedResult.HeightFromGroundNormalized = (C.Z - EndPositionOrig.Z)  * InverseLength;
-	EncodedResult.PoleVectorDirection = FVector::OneVector; // TBD
+	Source.CurrentEndDirectionNormalized = ACDirection * NormalizedLimbLength;
+	Source.CurrentEndRotation = InputGlobalPose[Source.BoneIndexC].GetRotation();
+	Source.CurrentHeightFromGroundNormalized = (C.Z - Source.InitialEndPosition.Z)  * Source.InvInitialLength;
+	Source.PoleVectorDirection = FVector::OneVector; // TBD
 }
 
-bool FChainDecoderIK::Initialize(
+bool FChainRetargeterIK::InitializeTarget(
 	const TArray<int32>& BoneIndices,
-	const TArray<FTransform> &InitialGlobalPose)
+	const TArray<FTransform> &TargetInitialGlobalPose)
 {
 	if (BoneIndices.Num() < 3)
 	{
-		UE_LOG(LogTemp, Warning, TEXT("IK Retargeter trying to retarget bone chain with IK but it has less than 3 joints."));
+		UE_LOG(LogTemp, Warning, TEXT("IK Retargeter trying to retarget target bone chain with IK but it has less than 3 joints."));
 		return false;
 	}
 	
-	BoneIndexA = BoneIndices[0];
-	const FTransform& Last = InitialGlobalPose[BoneIndices.Last()];
-	EndPositionOrig = Last.GetTranslation();
-	EndRotationOrig = Last.GetRotation();
-	Length = (InitialGlobalPose[BoneIndexA].GetTranslation() - Last.GetTranslation()).Size();
+	Target.BoneIndexA = BoneIndices[0];
+	const FTransform& Last = TargetInitialGlobalPose[BoneIndices.Last()];
+	Target.InitialEndPosition = Last.GetTranslation();
+	Target.InitialEndRotation = Last.GetRotation();
+	Target.InitialLength = (TargetInitialGlobalPose[Target.BoneIndexA].GetTranslation() - Last.GetTranslation()).Size();
 
-	if (Length <= KINDA_SMALL_NUMBER)
+	if (Target.InitialLength <= KINDA_SMALL_NUMBER)
 	{
 		UE_LOG(LogTemp, Warning, TEXT("IK Retargeter trying to retarget bone chain with IK, but it is zero length!"));
 		return false;
@@ -453,14 +463,13 @@ bool FChainDecoderIK::Initialize(
 	return true;
 }
 	
-void FChainDecoderIK::DecodePose(
-	const FEncodedIKChain& EncodedChain,
+void FChainRetargeterIK::DecodePose(
     const TArray<FTransform>& OutGlobalPose,
     FDecodedIKChain& OutResults) const
 {
 	// calculate end effector position
-	const FVector Start = OutGlobalPose[BoneIndexA].GetTranslation();
-	OutResults.EndEffectorPosition = Start + EncodedChain.EndDirectionNormalized * Length; // * LimbScale (todo)
+	const FVector Start = OutGlobalPose[Target.BoneIndexA].GetTranslation();
+	OutResults.EndEffectorPosition = Start + Source.CurrentEndDirectionNormalized * Target.InitialLength; // * LimbScale (todo)
 	// optionally factor in ground height
 	//const bool IsGrounded = false;
 	//if (IsGrounded) // todo expose this
@@ -469,8 +478,8 @@ void FChainDecoderIK::DecodePose(
 	//}
 
 	// calculate end effector rotation
-	const FQuat Delta = EndRotationOrig * EncodedChain.EndRotationOrig.Inverse();
-	OutResults.EndEffectorRotation = Delta * EncodedChain.EndRotation;
+	const FQuat RotationDelta = Source.CurrentEndRotation * Source.InitialEndRotation.Inverse();
+	OutResults.EndEffectorRotation = RotationDelta * Target.InitialEndRotation;
 
 	// TBD calc pole vector position
 	OutResults.PoleVectorPosition = FVector::OneVector;
@@ -617,7 +626,7 @@ bool FRetargetChainPairIK::Initialize(
 	}
 
 	// initialize SOURCE IK chain encoder with retarget pose
-	const bool bIKEncoderInitialized = IKEncoder.Initialize(SourceBoneIndices, SourceSkeleton.RetargetGlobalPose);
+	const bool bIKEncoderInitialized = IKChainRetargeter.InitializeSource(SourceBoneIndices, SourceSkeleton.RetargetGlobalPose);
 	if (!bIKEncoderInitialized)
 	{
 		UE_LOG(LogTemp, Warning, TEXT("IK Retargeter failed to initialize IK encoder, '%s', on Skeletal Mesh: '%s'"),
@@ -626,7 +635,7 @@ bool FRetargetChainPairIK::Initialize(
 	}
 
 	// initialize TARGET IK chain decoder with retarget pose
-	const bool bIKDecoderInitialized = IKDecoder.Initialize(TargetBoneIndices, TargetSkeleton.RetargetGlobalPose);
+	const bool bIKDecoderInitialized = IKChainRetargeter.InitializeTarget(TargetBoneIndices, TargetSkeleton.RetargetGlobalPose);
 	if (!bIKDecoderInitialized)
 	{
 		UE_LOG(LogTemp, Warning, TEXT("IK Retargeter failed to initialize IK decoder, '%s', on Skeletal Mesh: '%s'"),
@@ -651,21 +660,21 @@ void FIKRetargetPose::AddRotationDeltaToBone(FName BoneName, FQuat RotationDelta
 	*RotOffset = RotationDelta * (*RotOffset);
 }
 
-bool FRootEncoder::Initialize(const FName RootBoneName, const FRetargetSkeleton& SourceSkeleton)
+bool FRootRetargeter::InitializeSource(const FName SourceRootBoneName, const FRetargetSkeleton& SourceSkeleton)
 {
-	// validate root bone exists
-	BoneIndex = SourceSkeleton.FindBoneIndexByName(RootBoneName);
-	if (BoneIndex == INDEX_NONE)
+	// validate target root bone exists
+	Source.BoneIndex = SourceSkeleton.FindBoneIndexByName(SourceRootBoneName);
+	if (Source.BoneIndex == INDEX_NONE)
 	{
 		UE_LOG(LogTemp, Warning, TEXT("IK Retargeter could not find source root bone, %s in mesh %s"),
-			*RootBoneName.ToString(), *SourceSkeleton.SkeletalMeshName);
+			*SourceRootBoneName.ToString(), *SourceSkeleton.SkeletalMeshName);
 		return false;
 	}
 	
 	// record initial root data
-	const FTransform InitialTransform = SourceSkeleton.RetargetGlobalPose[BoneIndex]; 
+	const FTransform InitialTransform = SourceSkeleton.RetargetGlobalPose[Source.BoneIndex]; 
 	float InitialHeight = InitialTransform.GetTranslation().Z;
-	EncodedResult.InitialRotation = InitialTransform.GetRotation();
+	Source.InitialRotation = InitialTransform.GetRotation();
 
 	// ensure root height is not at origin, this happens if user sets root to ACTUAL skeleton root and not pelvis
 	if (InitialHeight < KINDA_SMALL_NUMBER)
@@ -676,55 +685,54 @@ bool FRootEncoder::Initialize(const FName RootBoneName, const FRetargetSkeleton&
 	}
 
 	// invert height
-	InvInitialHeight = 1.0f / InitialHeight;
+	Source.InitialHeightInverse = 1.0f / InitialHeight;
 
 	return true;
 }
 
-void FRootEncoder::EncodePose(const TArray<FTransform>& InputComponentSpaceBoneTransforms)
-{
-	const FTransform& Transform = InputComponentSpaceBoneTransforms[BoneIndex];
-	EncodedResult.NormalizedPosition = Transform.GetTranslation() * InvInitialHeight;
-	EncodedResult.Rotation = Transform.GetRotation();	
-}
-
-bool FRootDecoder::Initialize(const FName RootBoneName, const FTargetSkeleton& TargetSkeleton)
+bool FRootRetargeter::InitializeTarget(const FName TargetRootBoneName, const FTargetSkeleton& TargetSkeleton)
 {
-	// validate root bone exists
-	BoneIndex = TargetSkeleton.FindBoneIndexByName(RootBoneName);
-	if (BoneIndex == INDEX_NONE)
+	// validate target root bone exists
+	Target.BoneIndex = TargetSkeleton.FindBoneIndexByName(TargetRootBoneName);
+	if (Target.BoneIndex == INDEX_NONE)
 	{
 		UE_LOG(LogTemp, Warning, TEXT("IK Retargeter could not find target root bone, %s in mesh %s"),
-            *RootBoneName.ToString(), *TargetSkeleton.SkeletalMeshName);
+            *TargetRootBoneName.ToString(), *TargetSkeleton.SkeletalMeshName);
 		return false;
 	}
 
-	const FTransform InitialTransform = TargetSkeleton.RetargetGlobalPose[BoneIndex];
-	InitialHeight = InitialTransform.GetTranslation().Z;
-	InitialRotation = InitialTransform.GetRotation();
+	const FTransform TargetInitialTransform = TargetSkeleton.RetargetGlobalPose[Target.BoneIndex];
+	Target.InitialHeight = TargetInitialTransform.GetTranslation().Z;
+	Target.InitialRotation = TargetInitialTransform.GetRotation();
 
 	return true;
 }
 
-void FRootDecoder::DecodePose(
-	const FEncodedRoot& EncodedResult,
-	TArray<FTransform>& InOutTargetBoneTransforms,
+void FRootRetargeter::EncodePose(const TArray<FTransform>& SourceGlobalPose)
+{
+	const FTransform& SourceTransform = SourceGlobalPose[Source.BoneIndex];
+	Source.CurrentPositionNormalized = SourceTransform.GetTranslation() * Source.InitialHeightInverse;
+	Source.CurrentRotation = SourceTransform.GetRotation();	
+}
+
+void FRootRetargeter::DecodePose(
+	TArray<FTransform>& OutTargetGlobalPose,
 	const float StrideScale) const
 {
 	// scale normalized position by root height
-	FVector Position = EncodedResult.NormalizedPosition * InitialHeight;
-	// scale horizontal displacement by stride scale factor
-	Position.X *= StrideScale;
-	Position.Y *= StrideScale;
-	// apply position to target
-	InOutTargetBoneTransforms[BoneIndex].SetTranslation(Position);
+	FVector Position = Source.CurrentPositionNormalized * Target.InitialHeight;
+	// scale horizontal displacement by stride scale
+	Position *= FVector(StrideScale, StrideScale, 1.0f);
 
 	// calc offset between initial source/target root rotations
-	const FQuat RotationOffset = InitialRotation * EncodedResult.InitialRotation.Inverse();
+	const FQuat RotationDelta = Source.CurrentRotation * Source.InitialRotation.Inverse();
 	// add offset to the current source rotation
-	const FQuat Rotation = RotationOffset * EncodedResult.Rotation;
-	// apply rotation to target
-	InOutTargetBoneTransforms[BoneIndex].SetRotation(Rotation);	
+	const FQuat Rotation = RotationDelta * Target.InitialRotation;
+
+	// apply to target
+	FTransform& TargetRootTransform = OutTargetGlobalPose[Target.BoneIndex];
+	TargetRootTransform.SetTranslation(Position);
+	TargetRootTransform.SetRotation(Rotation);
 }
 
 UIKRetargeter::UIKRetargeter()
@@ -732,6 +740,7 @@ UIKRetargeter::UIKRetargeter()
     TargetIKRigAsset(nullptr),
     bIsLoadedAndValid(false)
 {
+	CleanPoseList();
 }
 
 void UIKRetargeter::Initialize(USkeletalMesh* SourceSkeletalMesh, USkeletalMesh* TargetSkeletalMesh, UObject* Outer)
@@ -784,7 +793,7 @@ bool UIKRetargeter::InitializeRoots()
 
 	// initialize root encoder
 	const FName SourceRootBoneName = SourceIKRigAsset->RetargetDefinition.RootBone;
-	const bool bRootEncoderInit = RootEncoder.Initialize(SourceRootBoneName, SourceSkeleton);
+	const bool bRootEncoderInit = RootRetargeter.InitializeSource(SourceRootBoneName, SourceSkeleton);
 	if (!bRootEncoderInit)
 	{
 		UE_LOG(LogTemp, Warning, TEXT("IK Retargeter unable to initialize source root, '%s' on skeletal mesh: '%s'"),
@@ -794,7 +803,7 @@ bool UIKRetargeter::InitializeRoots()
 
 	// initialize root decoder
 	const FName TargetRootBoneName = TargetIKRigAsset->RetargetDefinition.RootBone;
-	const bool bRootDecoderInit = RootDecoder.Initialize(TargetRootBoneName, TargetSkeleton);
+	const bool bRootDecoderInit = RootRetargeter.InitializeTarget(TargetRootBoneName, TargetSkeleton);
 	if (!bRootDecoderInit)
 	{
 		UE_LOG(LogTemp, Warning, TEXT("IK Retargeter unable to initialize target root, '%s' on skeletal mesh: '%s'"),
@@ -869,13 +878,13 @@ bool UIKRetargeter::InitializeBoneChainPairs()
 	for (FRetargetChainPairFK& FKChainPair : ChainPairsFK)
 	{
 		FKChainPair.FKDecoder.InitializeIntermediateParentIndices(
-			RootDecoder.BoneIndex,
+			RootRetargeter.Target.BoneIndex,
 			FKChainPair.TargetBoneIndices[0],
 			TargetSkeleton);
 	}
 
 	// root is updated before IK as well
-	TargetSkeleton.SetBoneIsRetargeted(RootDecoder.BoneIndex, true);
+	TargetSkeleton.SetBoneIsRetargeted(RootRetargeter.Target.BoneIndex, true);
 
 	// return true if at least 1 pair of bone chains were initialized
 	return !(ChainPairsIK.IsEmpty() && ChainPairsFK.IsEmpty());
@@ -931,8 +940,9 @@ TArray<FTransform>&  UIKRetargeter::RunRetargeter(const TArray<FTransform>& InSo
 	{
 		RunRootRetarget(InSourceGlobalPose, TargetSkeleton.OutputGlobalPose);
 		// update global transforms below root
-		TargetSkeleton.UpdateGlobalTransformsBelowBone(RootDecoder.BoneIndex, TargetSkeleton.RetargetLocalPose, TargetSkeleton.OutputGlobalPose);
-	}	
+		TargetSkeleton.UpdateGlobalTransformsBelowBone(RootRetargeter.Target.BoneIndex, TargetSkeleton.RetargetLocalPose, TargetSkeleton.OutputGlobalPose);
+	}
+	
 	// FK CHAIN retargeting
 	if (bRetargetFK)
 	{
@@ -954,9 +964,9 @@ void UIKRetargeter::RunRootRetarget(
 	const TArray<FTransform>& InGlobalTransforms,
     TArray<FTransform>& OutGlobalTransforms)
 {
-	RootEncoder.EncodePose(InGlobalTransforms);
+	RootRetargeter.EncodePose(InGlobalTransforms);
 	const float StrideScale = 1.0f;
-	RootDecoder.DecodePose(RootEncoder.EncodedResult, OutGlobalTransforms, StrideScale);
+	RootRetargeter.DecodePose(OutGlobalTransforms, StrideScale);
 }
 
 void UIKRetargeter::RunFKRetarget(
@@ -984,10 +994,10 @@ void UIKRetargeter::RunIKRetarget(
 	for (FRetargetChainPairIK& ChainPair : ChainPairsIK)
 	{
 		// encode them all using the input pose
-		ChainPair.IKEncoder.EncodePose(InGlobalPose);
+		ChainPair.IKChainRetargeter.EncodePose(InGlobalPose);
 		// decode the IK goal and apply to IKRig
 		FDecodedIKChain OutIKGoal;
-		ChainPair.IKDecoder.DecodePose(ChainPair.IKEncoder.EncodedResult, OutGlobalPose, OutIKGoal);
+		ChainPair.IKChainRetargeter.DecodePose(OutGlobalPose, OutIKGoal);
 		// set the goal transform on the IK Rig
 		FIKRigGoal Goal = FIKRigGoal(
 			ChainPair.IKGoalName,
@@ -1006,6 +1016,48 @@ void UIKRetargeter::RunIKRetarget(
 	IKRigProcessor->CopyOutputGlobalPoseToArray(OutGlobalPose);
 }
 
+void UIKRetargeter::CleanPoseList()
+{
+	// enforce the existence of a default pose
+	const bool HasDefaultPose = RetargetPoses.Contains(DefaultPoseName);
+	if (!HasDefaultPose)
+	{
+		RetargetPoses.Emplace(DefaultPoseName);
+	}
+	
+	// use default pose unless set to something else
+	if (CurrentRetargetPose == NAME_None)
+	{
+		CurrentRetargetPose = DefaultPoseName;
+	}
+
+	// remove all bone offsets that are no longer part of the target skeleton
+	if (TargetIKRigAsset)
+	{
+		const TArray<FName> AllowedBoneNames = TargetIKRigAsset->Skeleton.BoneNames;
+		for (TTuple<FName, FIKRetargetPose>& Pose : RetargetPoses)
+		{
+			// find bone offsets no longer in target skeleton
+			TArray<FName> BonesToRemove;
+			for (TTuple<FName, FQuat>& BoneOffset : Pose.Value.BoneRotationOffsets)
+			{
+				if (!AllowedBoneNames.Contains(BoneOffset.Key))
+				{
+					BonesToRemove.Add(BoneOffset.Key);
+				}
+			}
+			
+			// remove bone offsets
+			for (const FName& BoneToRemove : BonesToRemove)
+			{
+				Pose.Value.BoneRotationOffsets.Remove(BoneToRemove);
+			}
+		}
+	}
+
+	Modify();
+}
+
 #if WITH_EDITOR
 
 FName UIKRetargeter::GetSourceRootBone()
@@ -1116,48 +1168,6 @@ void UIKRetargeter::CleanChainMapping()
 	Modify();
 }
 
-void UIKRetargeter::CleanPoseList()
-{
-	// enforce the existence of a default pose
-	const bool HasDefaultPose = RetargetPoses.Contains(DefaultPoseName);
-	if (!HasDefaultPose)
-	{
-		RetargetPoses.Emplace(DefaultPoseName);
-	}
-	
-	// use default pose unless set to something else
-	if (CurrentRetargetPose == NAME_None)
-	{
-		CurrentRetargetPose = DefaultPoseName;
-	}
-
-	// remove all bone offsets that are no longer part of the target skeleton
-	if (TargetIKRigAsset)
-	{
-		const TArray<FName> AllowedBoneNames = TargetIKRigAsset->Skeleton.BoneNames;
-		for (TTuple<FName, FIKRetargetPose>& Pose : RetargetPoses)
-		{
-			// find bone offsets no longer in target skeleton
-			TArray<FName> BonesToRemove;
-			for (TTuple<FName, FQuat>& BoneOffset : Pose.Value.BoneRotationOffsets)
-			{
-				if (!AllowedBoneNames.Contains(BoneOffset.Key))
-				{
-					BonesToRemove.Add(BoneOffset.Key);
-				}
-			}
-			
-			// remove bone offsets
-			for (const FName& BoneToRemove : BonesToRemove)
-			{
-				Pose.Value.BoneRotationOffsets.Remove(BoneToRemove);
-			}
-		}
-	}
-
-	Modify();
-}
-
 void UIKRetargeter::AutoMapChains()
 {
 	TArray<FName> SourceChainNames;
@@ -1302,53 +1312,13 @@ void UIKRetargeter::PostLoad()
 }
 #endif
 
-FTransform UIKRetargeter::GetTargetBoneRetargetPoseGlobalTransform(const FName& TargetBoneName) const
+FTransform UIKRetargeter::GetTargetBoneRetargetPoseGlobalTransform(const int32& TargetBoneIndex) const
 {
 	check(bIsLoadedAndValid);
-
-	// makes sure is actually in skeleton
-	const int32 BoneIndex = TargetSkeleton.FindBoneIndexByName(TargetBoneName);
-	if (BoneIndex == INDEX_NONE)
-	{
-		return FTransform::Identity;
-	}
+	check(TargetSkeleton.BoneNames.IsValidIndex(TargetBoneIndex))
 
 	// get the current retarget pose
-	return TargetSkeleton.RetargetGlobalPose[BoneIndex];
-}
-
-void UIKRetargeter::GetTargetBoneStartAndEnd(const FName& TargetBoneName, FVector& OutStart, FVector& OutEnd) const
-{
-	check(bIsLoadedAndValid);
-
-	// find the bone
-	const int32 BoneIndex = TargetSkeleton.FindBoneIndexByName(TargetBoneName);
-	if (BoneIndex == INDEX_NONE)
-	{
-		return;
-	}
-
-	// find the end bone of the bone chain
-	// todo, return list of lines to each child
-	int32 FirstChildIndex = INDEX_NONE;
-	for (int32 I=0; I<TargetSkeleton.ParentIndices.Num(); ++I)
-	{
-		const int32 ParentIndex = TargetSkeleton.ParentIndices[I];
-		if (ParentIndex == BoneIndex)
-		{
-			FirstChildIndex = I;
-			break;
-		}
-	}
-
-	if (FirstChildIndex == INDEX_NONE)
-	{
-		return;
-	}
-
-	// get the origin of the bone chain
-	OutStart = TargetSkeleton.RetargetGlobalPose[BoneIndex].GetTranslation();
-	OutEnd = TargetSkeleton.RetargetGlobalPose[FirstChildIndex].GetTranslation();
+	return TargetSkeleton.RetargetGlobalPose[TargetBoneIndex];
 }
 
 FRetargetChainMap* UIKRetargeter::GetChainMap(const FName& TargetChainName)
diff --git a/Engine/Plugins/Animation/IKRig/Source/IKRig/Public/AnimNodes/AnimNode_IKRetargeter.h b/Engine/Plugins/Animation/IKRig/Source/IKRig/Public/AnimNodes/AnimNode_IKRetargeter.h
deleted file mode 100644
diff --git a/Engine/Plugins/Animation/IKRig/Source/IKRig/Public/AnimNodes/AnimNode_RetargetPoseFromMesh.h b/Engine/Plugins/Animation/IKRig/Source/IKRig/Public/AnimNodes/AnimNode_RetargetPoseFromMesh.h
new file mode 100644
--- /dev/null
+++ b/Engine/Plugins/Animation/IKRig/Source/IKRig/Public/AnimNodes/AnimNode_RetargetPoseFromMesh.h
@@ -0,0 +1,67 @@
+// Copyright Epic Games, Inc. All Rights Reserved.
+
+#pragma once
+
+#include "CoreMinimal.h"
+
+#include "Animation/AnimNodeBase.h"
+#include "Retargeter/IKRetargeter.h"
+#include "AnimNode_RetargetPoseFromMesh.generated.h"
+
+USTRUCT(BlueprintInternalUseOnly)
+struct IKRIG_API FAnimNode_RetargetPoseFromMesh : public FAnimNode_Base
+{
+	GENERATED_BODY()
+
+	/** The Skeletal Mesh Component to retarget animation from. Assumed to be animated and tick BEFORE this anim instance.*/
+	UPROPERTY(BlueprintReadWrite, transient, Category=Settings, meta=(PinShownByDefault))
+	TWeakObjectPtr<USkeletalMeshComponent> SourceMeshComponent = nullptr;
+
+	/* If SourceMeshComponent is not valid, and if this is true, it will look for attached parent as a source */
+	UPROPERTY(EditAnywhere, BlueprintReadWrite, Category=Settings, meta = (NeverAsPin))
+	bool bUseAttachedParent = true;
+	
+	/** Retarget asset to use. Must define a Source and Target IK Rig compatible with the SourceMeshComponent and current anim instance.*/
+	UPROPERTY(EditAnywhere, Category = Settings)
+	UIKRetargeter* IKRetargeterAsset = nullptr;
+	
+	// FAnimNode_Base interface
+	virtual void Initialize_AnyThread(const FAnimationInitializeContext& Context) override;
+	virtual void CacheBones_AnyThread(const FAnimationCacheBonesContext& Context) override;
+	virtual void Update_AnyThread(const FAnimationUpdateContext& Context) override;
+	virtual void Evaluate_AnyThread(FPoseContext& Output) override;
+	//virtual void GatherDebugData(FNodeDebugData& DebugData) override;
+	virtual bool HasPreUpdate() const override { return true; }
+	virtual void PreUpdate(const UAnimInstance* InAnimInstance) override;
+	// End of FAnimNode_Base interface
+
+	UIKRetargeter* GetCurrentlyUsedRetargeter() const;
+
+private:
+	
+	void EnsureInitialized(const UAnimInstance* InAnimInstance);
+	void InitializeRetargetData(const UAnimInstance* InAnimInstance);
+	void CopyBoneTransformsFromSource(USkeletalMeshComponent* TargetMeshComponent);
+	USkeletalMeshComponent* GetSourceMesh() const;
+
+	// indicates that all prerequisites are met and node is ready to operate
+	UPROPERTY(Transient)
+	bool bIsInitialized = false;
+	/** which version of the IK Retargeter asset was this instance last initialized with?
+	* this allows the IK Retargeter asset to undergo modifications at runtime via the editor and trigger
+	* re-initialization by incrementing this transient asset version.
+	 */
+	int32 LastVersionTried = -2;
+	
+	// source mesh references, cached during init so that we can compare and see if it has changed
+	TWeakObjectPtr<USkeletalMeshComponent>	CurrentlyUsedSourceMeshComponent;
+	TWeakObjectPtr<USkeletalMesh>			CurrentlyUsedSourceMesh;
+	TWeakObjectPtr<USkeletalMesh>			CurrentlyUsedTargetMesh;
+	UPROPERTY(Transient)
+	TObjectPtr<UIKRetargeter>				CurrentlyUsedRetargeter = nullptr;
+	TWeakObjectPtr<UIKRigDefinition>		CurrentlyUsedSourceIKRig;
+	TWeakObjectPtr<UIKRigDefinition>		CurrentlyUsedTargetIKRig;
+
+	// cached transforms, copied on the game thread
+	TArray<FTransform> SourceMeshComponentSpaceBoneTransforms;
+};
diff --git a/Engine/Plugins/Animation/IKRig/Source/IKRig/Public/Retargeter/IKRetargeter.h b/Engine/Plugins/Animation/IKRig/Source/IKRig/Public/Retargeter/IKRetargeter.h
--- a/Engine/Plugins/Animation/IKRig/Source/IKRig/Public/Retargeter/IKRetargeter.h
+++ b/Engine/Plugins/Animation/IKRig/Source/IKRig/Public/Retargeter/IKRetargeter.h
@@ -96,6 +96,8 @@ struct IKRIG_API FRetargetSkeleton
 		const int32 BoneIndex,
 		TArray<FTransform>& OutLocalPose,
 		TArray<FTransform>& InGlobalPose) const;
+
+	void GetChildrenIndices(const int32 BoneIndex, TArray<int32>& OutChildren) const;
 };
 
 struct FTargetSkeleton : public FRetargetSkeleton
@@ -112,42 +114,41 @@ struct FTargetSkeleton : public FRetargetSkeleton
 	void UpdateGlobalTransformsAllNonRetargetedBones(TArray<FTransform>& InOutGlobalPose);
 };
 
-struct FEncodedRoot
-{
-	FVector NormalizedPosition;
-	
-	FQuat Rotation;
-
-	FQuat InitialRotation;
-};
-
-struct FRootEncoder
+struct FRootSource
 {
 	int32 BoneIndex;
 
-	float InvInitialHeight;
-
-	FEncodedRoot EncodedResult;
+	FQuat InitialRotation;
 
-	bool Initialize(const FName RootBoneName, const FRetargetSkeleton& SourceSkeleton);
+	float InitialHeightInverse;
 
-	void EncodePose(const TArray<FTransform> &InputGlobalPose);
+	FVector CurrentPositionNormalized;
+	
+	FQuat CurrentRotation;
 };
 
-struct FRootDecoder
+struct FRootTarget
 {
 	int32 BoneIndex;
 	
+	FQuat InitialRotation;
+	
 	float InitialHeight;
+};
 
-	FQuat InitialRotation;
+struct FRootRetargeter
+{
+	FRootSource Source;
+	
+	FRootTarget Target;
 
-	bool Initialize(const FName RootBoneName, const FTargetSkeleton& TargetSkeleton);
+	bool InitializeSource(const FName SourceRootBoneName, const FRetargetSkeleton& SourceSkeleton);
+	
+	bool InitializeTarget(const FName TargetRootBoneName, const FTargetSkeleton& TargetSkeleton);
 
-	void DecodePose(
-        const FEncodedRoot& EncodedResult,
-        TArray<FTransform> &InOutTargetBoneTransforms,
-        const float StrideScale) const;
+	void EncodePose(const TArray<FTransform> &SourceGlobalPose);
+	
+	void DecodePose( TArray<FTransform> &OutTargetGlobalPose, const float StrideScale) const;
 };
 
 struct FChainFK
@@ -201,67 +202,64 @@ private:
 	TArray<int32> IntermediateParentIndices;
 };
 
-struct FEncodedIKChain
+struct FDecodedIKChain
 {
-	FVector EndDirectionNormalized;
-	
-	FQuat EndRotation;
-	
-	FQuat EndRotationOrig;
+	FVector EndEffectorPosition;
 	
-	float HeightFromGroundNormalized;
+	FQuat EndEffectorRotation;
 	
-	FVector PoleVectorDirection;
+	FVector PoleVectorPosition;
 };
 
-struct FChainEncoderIK
+struct FSourceChainIK
 {
 	int32 BoneIndexA;
 	
 	int32 BoneIndexB;
 	
 	int32 BoneIndexC;
-
-	FVector EndPositionOrig;
 	
-	FQuat EndRotationOrig;
+	FVector InitialEndPosition;
 	
-	float InverseLength;
-
-	FEncodedIKChain EncodedResult;
-
-	bool Initialize(const TArray<int32>& BoneIndices, const TArray<FTransform> &InitialGlobalPose);
+	FQuat InitialEndRotation;
+	
+	float InvInitialLength;
 
-	void EncodePose(const TArray<FTransform> &InputGlobalPose);
+	// results after encoding...
+	
+	FVector CurrentEndDirectionNormalized;
+	
+	FQuat CurrentEndRotation;
+	
+	float CurrentHeightFromGroundNormalized;
+	
+	FVector PoleVectorDirection;
 };
 
-struct FDecodedIKChain
+struct FTargetChainIK
 {
-	FVector EndEffectorPosition;
+	int32 BoneIndexA;
 	
-	FQuat EndEffectorRotation;
+	float InitialLength;
 	
-	FVector PoleVectorPosition;
+	FVector InitialEndPosition;
+	
+	FQuat InitialEndRotation;
 };
 
-struct FChainDecoderIK
+struct FChainRetargeterIK
 {
-	int32 BoneIndexA;
-
-	float Length;
-
-	FVector EndPositionOrig;
+	FSourceChainIK Source;
 	
-	FQuat EndRotationOrig;
+	FTargetChainIK Target;
 
-	FDecodedIKChain DecodedResults;
+	bool InitializeSource(const TArray<int32>& BoneIndices, const TArray<FTransform> &SourceInitialGlobalPose);
+	
+	bool InitializeTarget(const TArray<int32>& BoneIndices, const TArray<FTransform> &TargetInitialGlobalPose);
 
-	bool Initialize(const TArray<int32>& BoneIndices, const TArray<FTransform> &InitialGlobalPose);
+	void EncodePose(const TArray<FTransform> &SourceInputGlobalPose);
 	
-	void DecodePose(
-		const FEncodedIKChain& EncodedChain,
-		const TArray<FTransform>& OutGlobalPose,
-		FDecodedIKChain& OutResults) const;
+	void DecodePose(const TArray<FTransform>& OutGlobalPose, FDecodedIKChain& OutResults) const;
 };
 
 struct FRetargetChainPair
@@ -304,10 +302,8 @@ struct FRetargetChainPairFK : FRetargetChainPair
 };
 
 struct FRetargetChainPairIK : FRetargetChainPair
-{
-	FChainEncoderIK IKEncoder;
-	
-	FChainDecoderIK IKDecoder;
+{	
+	FChainRetargeterIK IKChainRetargeter;
 	
 	FName IKGoalName;
 	
@@ -359,12 +355,12 @@ public:
 	TArray<FRetargetChainMap> ChainMapping;
 
 	/** The set of retarget poses available as options for retargeting.*/
-	UPROPERTY(EditAnywhere, Category = RetargetPose)
-	FName CurrentRetargetPose = NAME_None;
+	UPROPERTY()
+	FName CurrentRetargetPose = DefaultPoseName;
 	static const FName DefaultPoseName;
 	
 	/** The set of retarget poses available as options for retargeting.*/
-	UPROPERTY(EditAnywhere, Category = RetargetPose)
+	UPROPERTY(VisibleAnywhere, Category = RetargetPose)
 	TMap<FName, FIKRetargetPose> RetargetPoses;
 
 	/** When false, translational motion of skeleton root is not copied. Useful for debugging.*/
@@ -379,27 +375,37 @@ public:
 	UPROPERTY(EditAnywhere, Category = Settings)
 	bool bRetargetIK = true;
 
-	UPROPERTY(EditAnywhere, Category = DebugPreview, meta = (UIMin = "-500.0", UIMax = "500.0"))
+	/** Move the target actor in the viewport for easier visualization next to the source actor.*/
+	UPROPERTY(EditAnywhere, Category = DebugPreview, meta = (UIMin = "-2000.0", UIMax = "2000.0"))
 	float TargetActorOffset = 150.0f;
 
-	UPROPERTY(EditAnywhere, Category = DebugPreview, meta = (UIMin = "0.01", UIMax = "2.0"))
+	/** Scale the target actor in the viewport for easier visualization next to the source actor.*/
+	UPROPERTY(EditAnywhere, Category = DebugPreview, meta = (UIMin = "0.01", UIMax = "10.0"))
 	float TargetActorScale = 1.0f;
 
+	/** The visual size of the bones in the viewport when editing the retarget pose.*/
 	UPROPERTY(EditAnywhere, Category = DebugPreview, meta = (ClampMin = "0.0", UIMin = "0.01", UIMax = "10.0"))
-	float BoneGizmoSize = 8.0f;
-	
+	float BoneDrawSize = 8.0f;
+
+	/** The visual thickness of the bones in the viewport when editing the retarget pose.*/
 	UPROPERTY(EditAnywhere, Category = DebugPreview, meta = (ClampMin = "0.0", UIMin = "0.01", UIMax = "10.0"))
-	float BoneGizmoThickness = 1.0f;
+	float BoneDrawThickness = 1.0f;
 
+	/** The internal data structures used to represent the SOURCE skeleton / pose during retargeter.*/
 	FRetargetSkeleton SourceSkeleton;
-	
+
+	/** The internal data structures used to represent the TARGET skeleton / pose during retargeter.*/
 	FTargetSkeleton TargetSkeleton;
-	
+
+	/** Only true once Initialize() has successfully completed.*/
 	bool bIsLoadedAndValid = false;
 
+	/** A special editor-only mode which forces the retargeter to output the current retarget reference pose,
+	 * rather than actually running the retarget and outputting the retargeted pose. Used in Edit-Pose mode.*/
 	UPROPERTY()
 	bool bEditReferencePoseMode = false;
-
+	/** Used in-editor to force reinitialization when the asset version of the source asset differs from the currently
+	 *running asset version.*/
 	UPROPERTY()
 	int32 AssetVersion = 0;
 
@@ -421,12 +427,23 @@ public:
 	 */
 	TArray<FTransform>& RunRetargeter(const TArray<FTransform>& InSourceGlobalPose);
 
-	FTransform GetTargetBoneRetargetPoseGlobalTransform(const FName& ChainName) const;
-
-	void GetTargetBoneStartAndEnd(const FName& TargetBoneName, FVector& OutStart, FVector& OutEnd) const;
+	/**
+	* Get the Global transform, in the currently used retarget pose, for a bone in the target skeletal mesh
+	* @param TargetBoneIndex - The index of a bone in the target skeleton.
+	* @return The Global space transform for the bone or Identity if bone not found.
+	* @warning This function is only valid to call after the retargeter has been initialized (which generates global retarget pose)
+	*/
+	FTransform GetTargetBoneRetargetPoseGlobalTransform(const int32& TargetBoneIndex) const;
 
+	/**
+	* Get the mapping (between source/target) for a given target chain
+	* @param TargetChainName - The name of chain on the target IK Rig
+	* @return The mapping between the target chain and the source
+	*/
 	FRetargetChainMap* GetChainMap(const FName& TargetChainName);
 
+	void CleanPoseList();
+	
 #if WITH_EDITOR
 	
 	// UObject
@@ -442,7 +459,6 @@ public:
 	void GetSourceChainNames(TArray<FName>& OutNames) const;
 
 	void CleanChainMapping();
-	void CleanPoseList();
 	void AutoMapChains();
 
 	USkeleton* GetSourceSkeletonAsset() const;
@@ -463,10 +479,8 @@ private:
 	TArray<FRetargetChainPairFK> ChainPairsFK;
 	
 	TArray<FRetargetChainPairIK> ChainPairsIK;
-
-	FRootEncoder RootEncoder;
 	
-	FRootDecoder RootDecoder;
+	FRootRetargeter RootRetargeter;
 
 	bool InitializeRoots();
 		
diff --git a/Engine/Plugins/Animation/IKRig/Source/IKRigDeveloper/Private/AnimGraphNode_IKRetargeter.cpp b/Engine/Plugins/Animation/IKRig/Source/IKRigDeveloper/Private/AnimGraphNode_IKRetargeter.cpp
deleted file mode 100644
diff --git a/Engine/Plugins/Animation/IKRig/Source/IKRigDeveloper/Private/AnimGraphNode_RetargetPoseFromMesh.cpp b/Engine/Plugins/Animation/IKRig/Source/IKRigDeveloper/Private/AnimGraphNode_RetargetPoseFromMesh.cpp
new file mode 100644
--- /dev/null
+++ b/Engine/Plugins/Animation/IKRig/Source/IKRigDeveloper/Private/AnimGraphNode_RetargetPoseFromMesh.cpp
@@ -0,0 +1,110 @@
+// Copyright Epic Games, Inc. All Rights Reserved.
+
+#include "AnimGraphNode_RetargetPoseFromMesh.h"
+#include "Animation/AnimInstance.h"
+#include "Kismet2/CompilerResultsLog.h"
+
+#define LOCTEXT_NAMESPACE "AnimGraphNode_IKRig"
+const FName UAnimGraphNode_RetargetPoseFromMesh::AnimModeName(TEXT("IKRig.IKRigEditor.IKRigEditMode"));
+
+void UAnimGraphNode_RetargetPoseFromMesh::Draw(FPrimitiveDrawInterface* PDI, USkeletalMeshComponent* PreviewSkelMeshComp) const
+{
+}
+
+FText UAnimGraphNode_RetargetPoseFromMesh::GetNodeTitle(ENodeTitleType::Type TitleType) const
+{
+	return LOCTEXT("AnimGraphNode_IKRetargeter_Title", "Retarget Pose From Mesh");
+}
+
+void UAnimGraphNode_RetargetPoseFromMesh::CopyNodeDataToPreviewNode(FAnimNode_Base* InPreviewNode)
+{
+	FAnimNode_RetargetPoseFromMesh* IKRetargeterNode = static_cast<FAnimNode_RetargetPoseFromMesh*>(InPreviewNode);
+}
+
+FEditorModeID UAnimGraphNode_RetargetPoseFromMesh::GetEditorMode() const
+{
+	return AnimModeName;
+}
+
+void UAnimGraphNode_RetargetPoseFromMesh::CustomizePinData(UEdGraphPin* Pin, FName SourcePropertyName, int32 ArrayIndex) const
+{
+	Super::CustomizePinData(Pin, SourcePropertyName, ArrayIndex);
+}
+
+void UAnimGraphNode_RetargetPoseFromMesh::PostEditChangeProperty(struct FPropertyChangedEvent& PropertyChangedEvent)
+{
+	const FName PropertyName = (PropertyChangedEvent.Property ? PropertyChangedEvent.Property->GetFName() : NAME_None);
+	if ((PropertyName == GET_MEMBER_NAME_STRING_CHECKED(FAnimNode_RetargetPoseFromMesh, bUseAttachedParent)))
+	{
+		ReconstructNode();
+	}
+}
+
+void UAnimGraphNode_RetargetPoseFromMesh::ValidateAnimNodeDuringCompilation(USkeleton* ForSkeleton,	FCompilerResultsLog& MessageLog)
+{
+	Super::ValidateAnimNodeDuringCompilation(ForSkeleton, MessageLog);
+
+	// validate source mesh component is not null
+	if (!Node.bUseAttachedParent)
+	{
+		if (!IsPinExposedAndLinked(GET_MEMBER_NAME_STRING_CHECKED(FAnimNode_RetargetPoseFromMesh, SourceMeshComponent)))
+		{
+			MessageLog.Warning(TEXT("@@ is missing a Source Skeletal Mesh Component reference."), this);
+			return;
+		}
+	}
+
+	// validate IK Rig asset has been assigned
+	if (!Node.IKRetargeterAsset)
+	{
+		MessageLog.Warning(TEXT("@@ is missing an IKRetargeter asset."), this);
+		return;
+	}
+
+	// validate SOURCE IK Rig asset has been assigned
+	if (!Node.IKRetargeterAsset->SourceIKRigAsset)
+	{
+		MessageLog.Warning(TEXT("@@ has IK Retargeter that is missing a source IK Rig asset."), this);
+	}
+
+	// validate TARGET IK Rig asset has been assigned
+	if (!Node.IKRetargeterAsset->TargetIKRigAsset)
+	{
+		MessageLog.Warning(TEXT("@@ has IK Retargeter that is missing a target IK Rig asset."), this);
+	}
+
+	if (!(Node.IKRetargeterAsset->SourceIKRigAsset && Node.IKRetargeterAsset->TargetIKRigAsset))
+	{
+		return;
+	}
+	
+	// validate that target bone chains exist on this skeleton
+	const FReferenceSkeleton &RefSkel = ForSkeleton->GetReferenceSkeleton();
+	const TArray<FBoneChain> &TargetBoneChains = Node.IKRetargeterAsset->TargetIKRigAsset->RetargetDefinition.BoneChains;
+    for (const FBoneChain &Chain : TargetBoneChains)
+    {
+        if (RefSkel.FindBoneIndex(Chain.StartBone) == INDEX_NONE)
+        {
+        	MessageLog.Warning(*LOCTEXT("StartBoneNotFound", "@@ - Start Bone in target IK Rig Bone Chain not found.").ToString(), this);
+        }
+
+    	if (RefSkel.FindBoneIndex(Chain.EndBone) == INDEX_NONE)
+    	{
+    		MessageLog.Warning(*LOCTEXT("EndBoneNotFound", "@@ - End Bone in target IK Rig Bone Chain not found.").ToString(), this);
+    	}
+    }
+}
+
+void UAnimGraphNode_RetargetPoseFromMesh::PreloadRequiredAssets()
+{
+	Super::PreloadRequiredAssets();
+	
+	if (Node.IKRetargeterAsset)
+	{
+		PreloadObject(Node.IKRetargeterAsset);
+		PreloadObject(Node.IKRetargeterAsset->SourceIKRigAsset);
+		PreloadObject(Node.IKRetargeterAsset->TargetIKRigAsset);
+	}
+}
+
+#undef LOCTEXT_NAMESPACE
diff --git a/Engine/Plugins/Animation/IKRig/Source/IKRigDeveloper/Public/AnimGraphNode_IKRetargeter.h b/Engine/Plugins/Animation/IKRig/Source/IKRigDeveloper/Public/AnimGraphNode_IKRetargeter.h
deleted file mode 100644
diff --git a/Engine/Plugins/Animation/IKRig/Source/IKRigDeveloper/Public/AnimGraphNode_RetargetPoseFromMesh.h b/Engine/Plugins/Animation/IKRig/Source/IKRigDeveloper/Public/AnimGraphNode_RetargetPoseFromMesh.h
new file mode 100644
--- /dev/null
+++ b/Engine/Plugins/Animation/IKRig/Source/IKRigDeveloper/Public/AnimGraphNode_RetargetPoseFromMesh.h
@@ -0,0 +1,39 @@
+// Copyright Epic Games, Inc. All Rights Reserved.
+
+#pragma once
+
+#include "CoreMinimal.h"
+#include "UObject/ObjectMacros.h"
+#include "AnimGraphNode_Base.h"
+#include "AnimNodes/AnimNode_RetargetPoseFromMesh.h"
+#include "AnimGraphNode_RetargetPoseFromMesh.generated.h"
+
+class FPrimitiveDrawInterface;
+class USkeletalMeshComponent;
+
+// Editor node for IKRig 
+UCLASS()
+class IKRIGDEVELOPER_API UAnimGraphNode_RetargetPoseFromMesh : public UAnimGraphNode_Base
+{
+	GENERATED_BODY()
+
+	UPROPERTY(EditAnywhere, Category = Settings)
+	FAnimNode_RetargetPoseFromMesh Node;
+
+public:
+	// UEdGraphNode interface
+	virtual FText GetNodeTitle(ENodeTitleType::Type TitleType) const override;
+	virtual void PostEditChangeProperty(struct FPropertyChangedEvent& PropertyChangedEvent) override;
+	virtual void ValidateAnimNodeDuringCompilation(USkeleton* ForSkeleton, FCompilerResultsLog& MessageLog);
+	// End of UEdGraphNode interface
+
+	// UAnimGraphNode_Base interface
+	virtual void PreloadRequiredAssets() override;
+	virtual void CopyNodeDataToPreviewNode(FAnimNode_Base* AnimNode) override;
+	virtual FEditorModeID GetEditorMode() const override;
+	virtual void Draw(FPrimitiveDrawInterface* PDI, USkeletalMeshComponent* PreviewSkelMeshComp) const override;
+	virtual void CustomizePinData(UEdGraphPin* Pin, FName SourcePropertyName, int32 ArrayIndex) const override;
+	// End of UAnimGraphNode_Base interface
+
+	static const FName AnimModeName;
+};
diff --git a/Engine/Plugins/Animation/IKRig/Source/IKRigEditor/Private/RetargetEditor/IKRetargetAnimInstanceProxy.cpp b/Engine/Plugins/Animation/IKRig/Source/IKRigEditor/Private/RetargetEditor/IKRetargetAnimInstanceProxy.cpp
--- a/Engine/Plugins/Animation/IKRig/Source/IKRigEditor/Private/RetargetEditor/IKRetargetAnimInstanceProxy.cpp
+++ b/Engine/Plugins/Animation/IKRig/Source/IKRigEditor/Private/RetargetEditor/IKRetargetAnimInstanceProxy.cpp
@@ -4,7 +4,7 @@
 #include "RetargetEditor/IKRetargetAnimInstance.h"
 
 
-FIKRetargetAnimInstanceProxy::FIKRetargetAnimInstanceProxy(UAnimInstance* InAnimInstance, FAnimNode_IKRetargeter* InNode)
+FIKRetargetAnimInstanceProxy::FIKRetargetAnimInstanceProxy(UAnimInstance* InAnimInstance, FAnimNode_RetargetPoseFromMesh* InNode)
 	: FAnimPreviewInstanceProxy(InAnimInstance),
 	IKRetargetNode(InNode)
 {
diff --git a/Engine/Plugins/Animation/IKRig/Source/IKRigEditor/Private/RetargetEditor/IKRetargetAssetBrowserTabSummoner.cpp b/Engine/Plugins/Animation/IKRig/Source/IKRigEditor/Private/RetargetEditor/IKRetargetAssetBrowserTabSummoner.cpp
new file mode 100644
--- /dev/null
+++ b/Engine/Plugins/Animation/IKRig/Source/IKRigEditor/Private/RetargetEditor/IKRetargetAssetBrowserTabSummoner.cpp
@@ -0,0 +1,36 @@
+﻿// Copyright Epic Games, Inc. All Rights Reserved.
+#include "RetargetEditor/IKRetargetAssetBrowserTabSummoner.h"
+
+#include "IDocumentation.h"
+#include "RigEditor/IKRigEditorStyle.h"
+#include "RetargetEditor/IKRetargetEditor.h"
+#include "RetargetEditor/SIKRetargetAssetBrowser.h"
+
+#define LOCTEXT_NAMESPACE "RetargetAssetBrowserTabSummoner"
+
+const FName FIKRetargetAssetBrowserTabSummoner::TabID(TEXT("AssetBrowser"));
+
+FIKRetargetAssetBrowserTabSummoner::FIKRetargetAssetBrowserTabSummoner(const TSharedRef<FIKRetargetEditor>& InRetargetEditor)
+	: FWorkflowTabFactory(TabID, InRetargetEditor),
+	IKRetargetEditor(InRetargetEditor)
+{
+	bIsSingleton = true; // only allow a single instance of this tab
+	
+	TabLabel = LOCTEXT("IKRetargetAssetBrowserTabLabel", "Asset Browser");
+	TabIcon = FSlateIcon(FIKRigEditorStyle::Get().GetStyleSetName(), "IKRig.TabIcon");
+
+	ViewMenuDescription = LOCTEXT("IKRetargetAssetBrowser_ViewMenu_Desc", "Asset Browser");
+	ViewMenuTooltip = LOCTEXT("IKRetargetAssetBrowser_ViewMenu_ToolTip", "Show the Asset Browser Tab");
+}
+
+TSharedPtr<SToolTip> FIKRetargetAssetBrowserTabSummoner::CreateTabToolTipWidget(const FWorkflowTabSpawnInfo& Info) const
+{
+	return  IDocumentation::Get()->CreateToolTip(LOCTEXT("IKRetargetAssetBrowserTooltip", "Select an animation asset to preview the retarget results."), NULL, TEXT("Shared/Editors/Persona"), TEXT("IKRetargetAssetBrowser_Window"));
+}
+
+TSharedRef<SWidget> FIKRetargetAssetBrowserTabSummoner::CreateTabBody(const FWorkflowTabSpawnInfo& Info) const
+{
+	return SNew(SIKRetargetAssetBrowser, IKRetargetEditor.Pin()->GetController());
+}
+
+#undef LOCTEXT_NAMESPACE 
diff --git a/Engine/Plugins/Animation/IKRig/Source/IKRigEditor/Private/RetargetEditor/IKRetargetBatchOperation.cpp b/Engine/Plugins/Animation/IKRig/Source/IKRigEditor/Private/RetargetEditor/IKRetargetBatchOperation.cpp
new file mode 100644
--- /dev/null
+++ b/Engine/Plugins/Animation/IKRig/Source/IKRigEditor/Private/RetargetEditor/IKRetargetBatchOperation.cpp
@@ -0,0 +1,385 @@
+﻿// Copyright Epic Games, Inc. All Rights Reserved.
+
+#include "RetargetEditor/IKRetargetBatchOperation.h"
+
+#include "AnimationBlueprintLibrary.h"
+#include "AnimPose.h"
+#include "AnimPreviewInstance.h"
+#include "ContentBrowserModule.h"
+#include "EditorReimportHandler.h"
+#include "IContentBrowserSingleton.h"
+#include "ObjectEditorUtils.h"
+#include "SSkeletonWidget.h"
+#include "PropertyCustomizationHelpers.h"
+#include "Animation/DebugSkelMeshComponent.h"
+#include "EditorFramework/AssetImportData.h"
+#include "Framework/Notifications/NotificationManager.h"
+#include "Kismet2/BlueprintEditorUtils.h"
+#include "Kismet2/KismetEditorUtilities.h"
+#include "Misc/ScopedSlowTask.h"
+#include "Retargeter/IKRetargeter.h"
+#include "Widgets/Notifications/SNotificationList.h"
+
+#define LOCTEXT_NAMESPACE "RetargetBatchOperation"
+
+int32 FIKRetargetBatchOperation::GenerateAssetLists(const FIKRetargetBatchOperationContext& Context)
+{
+	// re-generate lists of selected and referenced assets
+	AnimationAssetsToRetarget.Reset();
+	AnimBlueprintsToRetarget.Reset();
+
+	for (TWeakObjectPtr<UObject> AssetPtr : Context.AssetsToRetarget)
+	{
+		UObject* Asset = AssetPtr.Get();
+		if (UAnimationAsset* AnimAsset = Cast<UAnimationAsset>(Asset))
+		{
+			AnimationAssetsToRetarget.AddUnique(AnimAsset);
+		}
+		else if (UAnimBlueprint* AnimBlueprint = Cast<UAnimBlueprint>(Asset))
+		{
+			// Add parents blueprint. 
+			UAnimBlueprint* ParentBP = Cast<UAnimBlueprint>(AnimBlueprint->ParentClass->ClassGeneratedBy);
+			while (ParentBP)
+			{
+				AnimBlueprintsToRetarget.AddUnique(ParentBP);
+				ParentBP = Cast<UAnimBlueprint>(ParentBP->ParentClass->ClassGeneratedBy);
+			}
+				
+			AnimBlueprintsToRetarget.AddUnique(AnimBlueprint);				
+		}
+	}
+
+	if (Context.bRemapReferencedAssets)
+	{
+		// Grab assets from the blueprint.
+		// Do this first as it can add complex assets to the retarget array which will need to be processed next.
+		for (UAnimBlueprint* AnimBlueprint : AnimBlueprintsToRetarget)
+		{
+			GetAllAnimationSequencesReferredInBlueprint(AnimBlueprint, AnimationAssetsToRetarget);
+		}
+
+		int32 AssetIndex = 0;
+		while (AssetIndex < AnimationAssetsToRetarget.Num())
+		{
+			UAnimationAsset* AnimAsset = AnimationAssetsToRetarget[AssetIndex++];
+			AnimAsset->HandleAnimReferenceCollection(AnimationAssetsToRetarget, true);
+		}
+	}
+
+	return AnimationAssetsToRetarget.Num();
+}
+
+void FIKRetargetBatchOperation::DuplicateRetargetAssets(
+	const FIKRetargetBatchOperationContext& Context,
+	FScopedSlowTask& Progress)
+{
+	Progress.EnterProgressFrame(1.f, FText(LOCTEXT("DuplicatingBatchRetarget", "Duplicating animation assets...")));
+
+	UPackage* DestinationPackage = Context.TargetMesh->GetOutermost();
+
+	TArray<UAnimationAsset*> AnimationAssetsToDuplicate = AnimationAssetsToRetarget;
+	TArray<UAnimBlueprint*> AnimBlueprintsToDuplicate = AnimBlueprintsToRetarget;
+
+	// We only want to duplicate unmapped assets, so we remove mapped assets from the list we're duplicating
+	for (TPair<UAnimationAsset*, UAnimationAsset*>& Pair : RemappedAnimAssets)
+	{
+		AnimationAssetsToDuplicate.Remove(Pair.Key);
+	}
+
+	DuplicatedAnimAssets = DuplicateAssets<UAnimationAsset>(AnimationAssetsToDuplicate, DestinationPackage, &Context.NameRule);
+	DuplicatedBlueprints = DuplicateAssets<UAnimBlueprint>(AnimBlueprintsToDuplicate, DestinationPackage, &Context.NameRule);
+
+	// If we are moving the new asset to a different directory we need to fixup the reimport path.
+	// This should only effect source FBX paths within the project.
+	if (!Context.NameRule.FolderPath.IsEmpty())
+	{
+		for (TPair<UAnimationAsset*, UAnimationAsset*>& Pair : DuplicatedAnimAssets)
+		{
+			UAnimSequence* SourceSequence = Cast<UAnimSequence>(Pair.Key);
+			UAnimSequence* DestinationSequence = Cast<UAnimSequence>(Pair.Value);
+			if (!(SourceSequence && DestinationSequence))
+			{
+				continue;
+			}
+			
+			for (int index = 0; index < SourceSequence->AssetImportData->SourceData.SourceFiles.Num(); index++)
+			{
+				const FString& RelativeFilename = SourceSequence->AssetImportData->SourceData.SourceFiles[index].RelativeFilename;
+				const FString OldPackagePath = FPackageName::GetLongPackagePath(SourceSequence->GetPathName()) / TEXT("");
+				const FString NewPackagePath = FPackageName::GetLongPackagePath(DestinationSequence->GetPathName()) / TEXT("");
+				const FString AbsoluteSrcPath = FPaths::ConvertRelativePathToFull(FPackageName::LongPackageNameToFilename(OldPackagePath));
+				const FString SrcFile = AbsoluteSrcPath / RelativeFilename;
+				const bool bSrcFileExists = FPlatformFileManager::Get().GetPlatformFile().FileExists(*SrcFile);
+				if (!bSrcFileExists || (NewPackagePath == OldPackagePath))
+				{
+					continue;
+				}
+
+				FString BasePath = FPackageName::LongPackageNameToFilename(OldPackagePath);
+				FString OldSourceFilePath = FPaths::ConvertRelativePathToFull(BasePath, RelativeFilename);
+				TArray<FString> Paths;
+				Paths.Add(OldSourceFilePath);
+				
+				// update the FBX reimport file path
+				FReimportManager::Instance()->UpdateReimportPaths(DestinationSequence, Paths);
+			}
+		}
+	}
+
+	// Remapped assets needs the duplicated ones added
+	RemappedAnimAssets.Append(DuplicatedAnimAssets);
+
+	DuplicatedAnimAssets.GenerateValueArray(AnimationAssetsToRetarget);
+	DuplicatedBlueprints.GenerateValueArray(AnimBlueprintsToRetarget);
+}
+
+void FIKRetargetBatchOperation::RetargetAssets(
+	const FIKRetargetBatchOperationContext& Context,
+	FScopedSlowTask& Progress)
+{
+	USkeleton* OldSkeleton = Context.SourceMesh->GetSkeleton();
+	USkeleton* NewSkeleton = Context.TargetMesh->GetSkeleton();
+
+	for (UAnimationAsset* AssetToRetarget : AnimationAssetsToRetarget)
+	{
+		// synchronize curves between old/new asset
+		UAnimSequence* AnimSequenceToRetarget = Cast<UAnimSequence>(AssetToRetarget);
+		if (AnimSequenceToRetarget)
+		{
+			// copy curve data from source asset, preserving data in the target if present.
+			UAnimationBlueprintLibrary::CopyAnimationCurveNamesToSkeleton(OldSkeleton, NewSkeleton, AnimSequenceToRetarget, ERawCurveTrackTypes::RCT_Float);	
+			// clear transform curves since those curves won't work in new skeleton
+			IAnimationDataController& Controller = AnimSequenceToRetarget->GetController();
+			Controller.RemoveAllCurvesOfType(ERawCurveTrackTypes::RCT_Transform);
+		}
+
+		// replace references to other animation
+		AssetToRetarget->ReplaceReferredAnimations(RemappedAnimAssets);
+		AssetToRetarget->SetSkeleton(NewSkeleton);
+		AssetToRetarget->MarkPackageDirty();
+	}
+
+	// convert the animation using the IK retargeter
+	ConvertAnimation(Context,Progress);
+
+	// convert all Animation Blueprints and compile 
+	for (UAnimBlueprint* AnimBlueprint : AnimBlueprintsToRetarget)
+	{
+		// replace skeleton
+		AnimBlueprint->TargetSkeleton = NewSkeleton;
+
+		// if they have parent blueprint, make sure to re-link to the new one also
+		UAnimBlueprint* CurrentParentBP = Cast<UAnimBlueprint>(AnimBlueprint->ParentClass->ClassGeneratedBy);
+		if (CurrentParentBP)
+		{
+			UAnimBlueprint* const * ParentBP = DuplicatedBlueprints.Find(CurrentParentBP);
+			if (ParentBP)
+			{
+				AnimBlueprint->ParentClass = (*ParentBP)->GeneratedClass;
+			}
+		}
+
+		if(RemappedAnimAssets.Num() > 0)
+		{
+			ReplaceReferredAnimationsInBlueprint(AnimBlueprint, RemappedAnimAssets);
+		}
+
+		FBlueprintEditorUtils::RefreshAllNodes(AnimBlueprint);
+		FKismetEditorUtilities::CompileBlueprint(AnimBlueprint, EBlueprintCompileOptions::SkipGarbageCollection);
+		AnimBlueprint->PostEditChange();
+		AnimBlueprint->MarkPackageDirty();
+	}
+}
+
+void FIKRetargetBatchOperation::ConvertAnimation(
+	const FIKRetargetBatchOperationContext& Context,
+	FScopedSlowTask& Progress)
+{
+	// initialize the retargeter
+	UObject* TransientOuter = Cast<UObject>(GetTransientPackage());
+	UIKRetargeter* Retargeter = DuplicateObject(Context.IKRetargetAsset, TransientOuter);	
+	Retargeter->Initialize(Context.SourceMesh, Context.TargetMesh, TransientOuter);
+	if (!Retargeter->bIsLoadedAndValid)
+	{
+		UE_LOG(LogTemp, Warning, TEXT("Unable to initialize the IK Retargeter. Newly created animations were not retargeted!"));
+		return;
+	}
+
+	// for each pair of source / target animation sequences
+	for (TPair<UAnimationAsset*, UAnimationAsset*>& Pair : DuplicatedAnimAssets)
+	{
+		UAnimSequence* SourceSequence = Cast<UAnimSequence>(Pair.Key);
+		UAnimSequence* DestinationSequence = Cast<UAnimSequence>(Pair.Value);
+		if (!(SourceSequence && DestinationSequence))
+		{
+			continue;
+		}
+
+		// increment progress bar
+		FString AssetName = DestinationSequence->GetName();
+		Progress.EnterProgressFrame(1.f, FText::Format(LOCTEXT("RunningBatchRetarget", "Retargeting animation asset: {Asset}"), FText::FromString(AssetName)));
+
+		// remove all keys from the destination animation sequence
+		IAnimationDataController& TargetSeqController = DestinationSequence->GetController();
+		const bool ShouldTransactAnimEdits = false;
+		TargetSeqController.OpenBracket(FText::FromString("Generating Retargeted Animation Data"), ShouldTransactAnimEdits);
+		TargetSeqController.RemoveAllBoneTracks();
+
+		// number of frames in this animation
+		const int32 NumFrames = SourceSequence->GetNumberOfSampledKeys();
+
+		// make space for the target keyframe data
+		const int32 NumTargetBones = Retargeter->TargetSkeleton.BoneNames.Num();
+		TArray<FRawAnimSequenceTrack> BoneTracks;
+		BoneTracks.SetNumZeroed(NumTargetBones);
+
+		// retarget each frame's pose from source to target
+		for (int32 FrameIndex=0; FrameIndex<NumFrames; ++FrameIndex)
+		{
+			// get the source global pose
+			FAnimPose SourcePoseAtFrame;
+			UAnimPoseExtensions::GetAnimPoseAtFrame(SourceSequence, FrameIndex, FAnimPoseEvaluationOptions(), SourcePoseAtFrame);
+			TArray<FName> BoneNames;
+			UAnimPoseExtensions::GetBoneNames(SourcePoseAtFrame,BoneNames);
+			TArray<FTransform> SourceComponentPose;
+			for (const FName& BoneName : BoneNames)
+			{
+				FTransform BonePose = UAnimPoseExtensions::GetBonePose(SourcePoseAtFrame, BoneName, EAnimPoseSpaces::World);
+				SourceComponentPose.Add(BonePose);
+			}
+
+			// run the retarget
+			TArray<FTransform>& TargetComponentPose = Retargeter->RunRetargeter(SourceComponentPose);
+
+			// convert to a local-space pose
+			TArray<FTransform> TargetLocalPose = TargetComponentPose;
+			Retargeter->TargetSkeleton.UpdateLocalTransformsBelowBone(0,TargetLocalPose, TargetComponentPose);
+
+			// store key data for each bone
+			for (int32 TargetBoneIndex=0; TargetBoneIndex<TargetLocalPose.Num(); ++TargetBoneIndex)
+			{
+				BoneTracks[TargetBoneIndex].PosKeys.Add(TargetLocalPose[TargetBoneIndex].GetLocation());
+				BoneTracks[TargetBoneIndex].RotKeys.Add(TargetLocalPose[TargetBoneIndex].GetRotation());
+				BoneTracks[TargetBoneIndex].ScaleKeys.Add(TargetLocalPose[TargetBoneIndex].GetScale3D());
+			}
+		}
+
+		// add keys to bone tracks
+		const bool bShouldTransact = false;
+		for (int32 TargetBoneIndex=0; TargetBoneIndex<NumTargetBones; ++TargetBoneIndex)
+		{
+			FName TargetBoneName = Retargeter->TargetSkeleton.BoneNames[TargetBoneIndex];
+			FRawAnimSequenceTrack& RawTrack = BoneTracks[TargetBoneIndex];
+			TargetSeqController.AddBoneTrack(TargetBoneName, bShouldTransact);
+			TargetSeqController.SetBoneTrackKeys(TargetBoneName, RawTrack.PosKeys, RawTrack.RotKeys, RawTrack.ScaleKeys);
+		}
+
+		// done editing sequence data, close bracket
+		TargetSeqController.CloseBracket(ShouldTransactAnimEdits);
+	}
+}
+
+void FIKRetargetBatchOperation::NotifyUserOfResults(
+	const FIKRetargetBatchOperationContext& Context,
+	FScopedSlowTask& Progress) const
+{
+	Progress.EnterProgressFrame(1.f, FText(LOCTEXT("DoneBatchRetarget", "Duplicate and retarget complete!")));
+
+	// gather newly created objects
+	TArray<UObject*> NewAssets;
+	GetNewAssets(NewAssets);
+
+	// log details of what assets were created
+	for (UObject* NewAsset : NewAssets)
+	{
+		UE_LOG(LogTemp, Display, TEXT("Duplicate and Retarget - New Asset Created: %s"), *NewAsset->GetName());
+	}
+	
+	// notify user
+	FNotificationInfo Notification(FText::GetEmpty());
+	Notification.ExpireDuration = 5.f;
+	Notification.Text = FText::Format(
+		LOCTEXT("MultiNonDuplicatedAsset", "{0} assets were retargeted to new skeleton {1}. See Output for details."),
+		FText::AsNumber(NewAssets.Num()),
+		FText::FromString(Context.TargetMesh->GetName()));
+	FSlateNotificationManager::Get().AddNotification(Notification);
+	
+	// select all new assets
+	TArray<FAssetData> CurrentSelection;
+	for(UObject* NewObject : NewAssets)
+	{
+		CurrentSelection.Add(FAssetData(NewObject));
+	}
+
+	// show assets in browser
+	FContentBrowserModule& ContentBrowserModule = FModuleManager::Get().LoadModuleChecked<FContentBrowserModule>("ContentBrowser");
+	ContentBrowserModule.Get().SyncBrowserToAssets(CurrentSelection);
+}
+
+void FIKRetargetBatchOperation::GetNewAssets(TArray<UObject*>& NewAssets) const
+{
+	TArray<UAnimationAsset*> NewAnims;
+	DuplicatedAnimAssets.GenerateValueArray(NewAnims);
+	for (UAnimationAsset* NewAnim : NewAnims)
+	{
+		NewAssets.Add(Cast<UObject>(NewAnim));
+	}
+
+	TArray<UAnimBlueprint*> NewBlueprints;
+	DuplicatedBlueprints.GenerateValueArray(NewBlueprints);
+	for (UAnimBlueprint* NewBlueprint : NewBlueprints)
+	{
+		NewAssets.Add(Cast<UObject>(NewBlueprint));
+	}
+}
+
+
+
+void FIKRetargetBatchOperation::RunRetarget(FIKRetargetBatchOperationContext& Context)
+{
+	const int32 NumAssets = GenerateAssetLists(Context);
+	
+	// show progress bar
+	FScopedSlowTask Progress(NumAssets + 2, LOCTEXT("GatheringBatchRetarget", "Gathering animation assets..."));
+	Progress.MakeDialog();
+	
+	DuplicateRetargetAssets(Context, Progress);
+	RetargetAssets(Context, Progress);
+	NotifyUserOfResults(Context, Progress);
+}
+
+/**
+* Duplicates the supplied AssetsToDuplicate and returns a map of original asset to duplicate. Templated wrapper that calls DuplicateAssetInternal.
+*
+* @param	AssetsToDuplicate	The animations to duplicate
+* @param	DestinationPackage	The package that the duplicates should be placed in
+* @param	NameRule			The rules for how to rename the duplicated assets
+*
+* @return	TMap of original animation to duplicate
+*/
+template<class AssetType>
+TMap<AssetType*, AssetType*> FIKRetargetBatchOperation::DuplicateAssets(
+	const TArray<AssetType*>& AssetsToDuplicate,
+	UPackage* DestinationPackage,
+	const FNameDuplicationRule* NameRule)
+{
+	TArray<UObject*> Assets;
+	for (AssetType* Asset : AssetsToDuplicate)
+	{
+		Assets.Add(Asset);
+	}
+
+	// duplicate assets
+	TMap<UObject*, UObject*> DuplicateAssetsMap = DuplicateAssetsInternal(Assets, DestinationPackage, NameRule);
+
+	// cast to AssetType
+	TMap<AssetType*, AssetType*> ReturnMap;
+	for (const TTuple<UObject*, UObject*>& DuplicateAsset : DuplicateAssetsMap)
+	{
+		ReturnMap.Add(Cast<AssetType>(DuplicateAsset.Key), Cast<AssetType>(DuplicateAsset.Value));
+	}
+	
+	return ReturnMap;
+}
+
+#undef LOCTEXT_NAMESPACE
diff --git a/Engine/Plugins/Animation/IKRig/Source/IKRigEditor/Private/RetargetEditor/IKRetargetEditMode.cpp b/Engine/Plugins/Animation/IKRig/Source/IKRigEditor/Private/RetargetEditor/IKRetargetEditMode.cpp
--- a/Engine/Plugins/Animation/IKRig/Source/IKRigEditor/Private/RetargetEditor/IKRetargetEditMode.cpp
+++ b/Engine/Plugins/Animation/IKRig/Source/IKRigEditor/Private/RetargetEditor/IKRetargetEditMode.cpp
@@ -49,29 +49,39 @@ void FIKRetargetEditMode::Render(const FSceneView* View, FViewport* Viewport, FP
 	}
 	
 	UIKRetargeter* Retargeter = Controller->Asset;
-	for (const FName& BoneName : Retargeter->TargetIKRigAsset->Skeleton.BoneNames)
+	
+	for (int32 BoneIndex = 0; BoneIndex<Retargeter->TargetIKRigAsset->Skeleton.BoneNames.Num(); ++BoneIndex)
 	{
+		// filter out bones that are not part of a retargeted chain
+		if (!Controller->IsTargetBoneRetargeted(BoneIndex))
+		{
+			continue;
+		}
+		
+		const FName BoneName = Retargeter->TargetIKRigAsset->Skeleton.BoneNames[BoneIndex];
+
 		// is this chain currently selected?
 		const bool bIsSelected = IsBoneSelected(BoneName);
 
 		// get the location of the bone on the currently initialized target skeletal mesh
-		const FTransform BoneTransform = Controller->GetTargetBoneTransform(BoneName);
+		const FTransform BoneTransform = Controller->GetTargetBoneTransform(BoneIndex);
 
 		FVector Start;
-		FVector End;
-		Controller->GetTargetBoneStartAndEnd(BoneName, Start, End);
-
-		// draw the chain rotation gizmo
-		PDI->SetHitProxy(new HIKRetargetEditorBoneProxy(BoneName));
-		DrawBoneGizmo(
-			PDI,
-			BoneTransform,
-			Start,
-			End,
-			Retargeter->BoneGizmoSize,
-			Retargeter->BoneGizmoThickness,
-			bIsSelected);
-		PDI->SetHitProxy(NULL);
+		TArray<FVector> End;
+		if (Controller->GetTargetBoneLineSegments(BoneIndex, Start, End))
+		{
+			// draw the chain rotation gizmo
+			PDI->SetHitProxy(new HIKRetargetEditorBoneProxy(BoneName));
+			DrawBoneProxy(
+				PDI,
+				BoneTransform,
+				Start,
+				End,
+				Retargeter->BoneDrawSize,
+				Retargeter->BoneDrawThickness,
+				bIsSelected);
+			PDI->SetHitProxy(NULL);
+		}
 	}
 }
 
@@ -107,8 +117,14 @@ FVector FIKRetargetEditMode::GetWidgetLocation() const
 	{
 		return FVector::ZeroVector; 
 	}
+
+	const int32 BoneIndex = Controller->GetCurrentlyRunningRetargeter()->TargetSkeleton.FindBoneIndexByName(SelectedBones.Last());
+	if (BoneIndex == INDEX_NONE)
+	{
+		return FVector::ZeroVector;
+	}
 	
-	return Controller->GetTargetBoneTransform(SelectedBones.Last()).GetTranslation();
+	return Controller->GetTargetBoneTransform(BoneIndex).GetTranslation();
 }
 
 bool FIKRetargetEditMode::HandleClick(FEditorViewportClient* InViewportClient, HHitProxy* HitProxy, const FViewportClick& Click)
@@ -203,7 +219,8 @@ bool FIKRetargetEditMode::GetCustomDrawingCoordinateSystem(FMatrix& InMatrix, vo
 		return false; 
 	}
 
-	const FTransform CurrentTransform = Controller->GetTargetBoneTransform(SelectedBones[0]);
+	const int32 BoneIndex = Controller->Asset->TargetIKRigAsset->Skeleton.GetBoneIndexFromName(SelectedBones[0]);
+	const FTransform CurrentTransform = Controller->GetTargetBoneTransform(BoneIndex);
 	InMatrix = CurrentTransform.ToMatrixNoScale().RemoveTranslation();
 	return true;
 }
@@ -213,11 +230,11 @@ bool FIKRetargetEditMode::GetCustomInputCoordinateSystem(FMatrix& InMatrix, void
 	return GetCustomDrawingCoordinateSystem(InMatrix, InData);
 }
 
-void FIKRetargetEditMode::DrawBoneGizmo(
+void FIKRetargetEditMode::DrawBoneProxy(
 	FPrimitiveDrawInterface* PDI,
 	const FTransform& BoneTransform,
 	const FVector& Start,
-	const FVector& End,
+	const TArray<FVector>& ChildPoints,
 	float Size,
 	float Thickness,
 	bool bIsSelected) const
@@ -226,10 +243,14 @@ void FIKRetargetEditMode::DrawBoneGizmo(
 	Size = FMath::Clamp(Size, 0.1f, 1000.0f);
 	Thickness = FMath::Clamp(Thickness, 0.1f, 1000.0f);
 
-	// draw a thick line along length of the limb
-	PDI->DrawLine(Start, End, Color, SDPG_Foreground, Thickness * 2.0f);
-	// draw a coordinate frame at the tip of the limb
-	DrawCoordinateSystem(PDI, End, BoneTransform.GetRotation().Rotator(), Size, SDPG_Foreground, Thickness);
+	// draw a line to each child
+	for (const FVector& ChildPosition : ChildPoints)
+	{
+		PDI->DrawLine(Start, ChildPosition, Color, SDPG_Foreground, Thickness * 1.0f);
+	}
+	
+	// draw a coordinate frame at the bone
+	DrawCoordinateSystem(PDI, Start, BoneTransform.GetRotation().Rotator(), Size, SDPG_Foreground, Thickness);
 }
 
 void FIKRetargetEditMode::Tick(FEditorViewportClient* ViewportClient, float DeltaTime)
diff --git a/Engine/Plugins/Animation/IKRig/Source/IKRigEditor/Private/RetargetEditor/IKRetargetEditor.cpp b/Engine/Plugins/Animation/IKRig/Source/IKRigEditor/Private/RetargetEditor/IKRetargetEditor.cpp
--- a/Engine/Plugins/Animation/IKRig/Source/IKRigEditor/Private/RetargetEditor/IKRetargetEditor.cpp
+++ b/Engine/Plugins/Animation/IKRig/Source/IKRigEditor/Private/RetargetEditor/IKRetargetEditor.cpp
@@ -220,17 +220,6 @@ void FIKRetargetEditor::FillToolbar(FToolBarBuilder& ToolbarBuilder)
 			FSlateIcon(FAppStyle::Get().GetStyleSetName(),"Icons.Refresh"));
 	}
 	ToolbarBuilder.EndSection();
-
-	ToolbarBuilder.BeginSection("Export");
-	{
-		ToolbarBuilder.AddToolBarButton(
-			FIKRetargetCommands::Get().ExportAnimation,
-			NAME_None,
-			TAttribute<FText>(),
-			TAttribute<FText>(),
-			FSlateIcon(FAppStyle::Get().GetStyleSetName(),"Icons.Save"));
-	}
-	ToolbarBuilder.EndSection();	
 }
 
 FName FIKRetargetEditor::GetToolkitFName() const
@@ -311,7 +300,10 @@ void FIKRetargetEditor::HandlePreviewSceneCreated(const TSharedRef<IPersonaPrevi
 
 	// apply mesh to the preview scene
 	InPersonaPreviewScene->SetPreviewMeshComponent(EditorController->SourceSkelMeshComponent);
+	InPersonaPreviewScene->SetAllowMeshHitProxies(false);
 	InPersonaPreviewScene->SetAdditionalMeshesSelectable(false);
+	EditorController->SourceSkelMeshComponent->bSelectable = false;
+	EditorController->TargetSkelMeshComponent->bSelectable = false;
 	InPersonaPreviewScene->SetPreviewMesh(SourceMesh);
 	InPersonaPreviewScene->AddComponent(EditorController->SourceSkelMeshComponent, FTransform::Identity);
 	InPersonaPreviewScene->AddComponent(EditorController->TargetSkelMeshComponent, FTransform::Identity);
@@ -329,21 +321,6 @@ void FIKRetargetEditor::SetupAnimInstance()
 	EditorController->TargetAnimInstance->InitializeAnimation();
 }
 
-void FIKRetargetEditor::HandleOpenNewAsset(UObject* InNewAsset)
-{
-	if (UAnimationAsset* NewAnimationAsset = Cast<UAnimationAsset>(InNewAsset))
-	{
-		EditorController->SourceAnimInstance->SetAnimationAsset(NewAnimationAsset);
-		PreviousAsset = NewAnimationAsset;
-	}
-}
-
-void FIKRetargetEditor::HandleAnimationSequenceBrowserCreated(
-	const TSharedRef<IAnimationSequenceBrowser>& InSequenceBrowser)
-{
-	EditorController->SequenceBrowser = InSequenceBrowser;
-}
-
 FText FIKRetargetEditor::GetCurrentPoseName() const
 {
 	return FText::FromName(EditorController->Asset->CurrentRetargetPose);
@@ -419,10 +396,7 @@ void FIKRetargetEditor::HandleEditPose() const
 	else
 	{
 		GetEditorModeManager().DeactivateMode(FIKRetargetEditMode::ModeName);
-		if (PreviousAsset)
-		{
-			EditorController->SourceSkelMeshComponent->EnablePreview(true, PreviousAsset);
-		}
+		EditorController->PlayPreviousAnimationAsset();
 	}
 
 	Asset->Modify();
diff --git a/Engine/Plugins/Animation/IKRig/Source/IKRigEditor/Private/RetargetEditor/IKRetargetEditorController.cpp b/Engine/Plugins/Animation/IKRig/Source/IKRigEditor/Private/RetargetEditor/IKRetargetEditorController.cpp
--- a/Engine/Plugins/Animation/IKRig/Source/IKRigEditor/Private/RetargetEditor/IKRetargetEditorController.cpp
+++ b/Engine/Plugins/Animation/IKRig/Source/IKRigEditor/Private/RetargetEditor/IKRetargetEditorController.cpp
@@ -28,7 +28,7 @@ USkeletalMesh* FIKRetargetEditorController::GetTargetSkeletalMesh() const
 	return Asset->TargetIKRigAsset->PreviewSkeletalMesh;
 }
 
-FTransform FIKRetargetEditorController::GetTargetBoneTransform(const FName& BoneName) const
+FTransform FIKRetargetEditorController::GetTargetBoneTransform(const int32& TargetBoneIndex) const
 {
 	UIKRetargetAnimInstance* AnimInstance = TargetAnimInstance.Get();
 	check(AnimInstance);
@@ -37,50 +37,68 @@ FTransform FIKRetargetEditorController::GetTargetBoneTransform(const FName& Bone
 	check(CurrentRetargeter)
 
 	// get transform of root of chain
-	FTransform BoneTransform = CurrentRetargeter->GetTargetBoneRetargetPoseGlobalTransform(BoneName);
+	FTransform BoneTransform = CurrentRetargeter->GetTargetBoneRetargetPoseGlobalTransform(TargetBoneIndex);
 
+	// scale and offset
 	BoneTransform.ScaleTranslation(Asset->TargetActorScale);
-
-	// add the target translation offset
-	const FVector TargetOffset(Asset->TargetActorOffset, 0.f, 0.f);
-	BoneTransform.AddToTranslation(TargetOffset);
+	BoneTransform.AddToTranslation(FVector(Asset->TargetActorOffset, 0.f, 0.f));
 
 	return BoneTransform;
 }
 
-void FIKRetargetEditorController::GetTargetBoneStartAndEnd(const FName& BoneName, FVector& Start, FVector& End) const
+bool FIKRetargetEditorController::GetTargetBoneLineSegments(
+	const int32& TargetBoneIndex,
+	FVector& OutStart,
+	TArray<FVector>& OutChildren) const
 {
 	UIKRetargeter* CurrentRetargeter = GetCurrentlyRunningRetargeter();
-	check(CurrentRetargeter)
+	check(CurrentRetargeter && CurrentRetargeter->bIsLoadedAndValid)
+	check(CurrentRetargeter->TargetSkeleton.BoneNames.IsValidIndex(TargetBoneIndex))
+	
+	// get the target skeleton we want to draw
+	const FTargetSkeleton& TargetSkeleton = CurrentRetargeter->TargetSkeleton;
 
-	// get line of chain
-	CurrentRetargeter->GetTargetBoneStartAndEnd(BoneName, Start, End);
+	// get the origin of the bone chain
+	OutStart = TargetSkeleton.RetargetGlobalPose[TargetBoneIndex].GetTranslation();
 
-	// add the target translation offset
-	const FVector TargetOffset(Asset->TargetActorOffset, 0.f, 0.f);
+	// get children
+	TArray<int32> ChildIndices;
+	TargetSkeleton.GetChildrenIndices(TargetBoneIndex, ChildIndices);
+	for (const int32& ChildIndex : ChildIndices)
+	{
+		OutChildren.Emplace(TargetSkeleton.RetargetGlobalPose[ChildIndex].GetTranslation());
+	}
 
-	Start *= Asset->TargetActorScale;
-	End *= Asset->TargetActorScale;
+	// add the target translation offset and scale
+	const FVector TargetOffset(Asset->TargetActorOffset, 0.f, 0.f);
+	OutStart *= Asset->TargetActorScale;
+	OutStart += TargetOffset;
+	for (FVector& ChildPoint : OutChildren)
+	{
+		ChildPoint *= Asset->TargetActorScale;
+		ChildPoint += TargetOffset;
+	}
 	
-	Start += TargetOffset;
-	End += TargetOffset;
+	return true;
 }
 
-UIKRetargeter* FIKRetargetEditorController::GetCurrentlyRunningRetargeter() const
+bool FIKRetargetEditorController::IsTargetBoneRetargeted(const int32& TargetBoneIndex)
 {
-	UIKRetargetAnimInstance* AnimInstance = TargetAnimInstance.Get();
-	if(!AnimInstance)
-	{
-		return nullptr;
-	}
+	UIKRetargeter* CurrentRetargeter = GetCurrentlyRunningRetargeter();
+	check(CurrentRetargeter && CurrentRetargeter->bIsLoadedAndValid)
+	check(CurrentRetargeter->TargetSkeleton.BoneNames.IsValidIndex(TargetBoneIndex))
 
-	UIKRetargeter* CurrentRetargeter = AnimInstance->GetCurrentlyUsedRetargeter();
-	if (!CurrentRetargeter)
+	return CurrentRetargeter->TargetSkeleton.IsBoneRetargeted[TargetBoneIndex];
+}
+
+UIKRetargeter* FIKRetargetEditorController::GetCurrentlyRunningRetargeter() const
+{	
+	if(UIKRetargetAnimInstance* AnimInstance = TargetAnimInstance.Get())
 	{
-		return nullptr;
+		return AnimInstance->GetCurrentlyUsedRetargeter();
 	}
 
-	return CurrentRetargeter;
+	return nullptr;	
 }
 
 void FIKRetargetEditorController::RefreshAllViews()
@@ -88,4 +106,21 @@ void FIKRetargetEditorController::RefreshAllViews()
 	ChainsView.Get()->RefreshView();
 }
 
+void FIKRetargetEditorController::PlayAnimationAsset(UAnimationAsset* AssetToPlay)
+{
+	if (AssetToPlay && SourceAnimInstance)
+	{
+		SourceAnimInstance->SetAnimationAsset(AssetToPlay);
+		PreviousAsset = AssetToPlay;
+	}
+}
+
+void FIKRetargetEditorController::PlayPreviousAnimationAsset()
+{
+	if (PreviousAsset)
+	{
+		SourceAnimInstance->SetAnimationAsset(PreviousAsset);
+	}
+}
+
 #undef LOCTEXT_NAMESPACE
diff --git a/Engine/Plugins/Animation/IKRig/Source/IKRigEditor/Private/RetargetEditor/IKRetargetMode.cpp b/Engine/Plugins/Animation/IKRig/Source/IKRigEditor/Private/RetargetEditor/IKRetargetMode.cpp
--- a/Engine/Plugins/Animation/IKRig/Source/IKRigEditor/Private/RetargetEditor/IKRetargetMode.cpp
+++ b/Engine/Plugins/Animation/IKRig/Source/IKRigEditor/Private/RetargetEditor/IKRetargetMode.cpp
@@ -6,6 +6,7 @@
 #include "PersonaModule.h"
 #include "Modules/ModuleManager.h"
 #include "PersonaTabs.h"
+#include "RetargetEditor/IKRetargetAssetBrowserTabSummoner.h"
 
 #include "RetargetEditor/IKRetargetChainTabSummoner.h"
 #include "RetargetEditor/IKRetargetEditor.h"
@@ -32,13 +33,13 @@ FIKRetargetMode::FIKRetargetMode(
 	TabFactories.RegisterFactory(PersonaModule.CreatePersonaViewportTabFactory(InHostingApp, ViewportArgs));
 	TabFactories.RegisterFactory(PersonaModule.CreateAdvancedPreviewSceneTabFactory(InHostingApp, InPreviewScene));
 	TabFactories.RegisterFactory(PersonaModule.CreateDetailsTabFactory(InHostingApp, FOnDetailsCreated::CreateSP(&IKRetargetEditor.Get(), &FIKRetargetEditor::HandleDetailsCreated)));
-	TabFactories.RegisterFactory(PersonaModule.CreateAnimationAssetBrowserTabFactory(InHostingApp, IKRetargetEditor->GetPersonaToolkit(), FOnOpenNewAsset::CreateSP(&IKRetargetEditor.Get(), &FIKRetargetEditor::HandleOpenNewAsset), FOnAnimationSequenceBrowserCreated::CreateSP(&IKRetargetEditor.Get(), &FIKRetargetEditor::HandleAnimationSequenceBrowserCreated), true));
 
 	// register custom tabs
 	TabFactories.RegisterFactory(MakeShared<FIKRetargetChainTabSummoner>(IKRetargetEditor));
+	TabFactories.RegisterFactory(MakeShared<FIKRetargetAssetBrowserTabSummoner>(IKRetargetEditor));
 
 	// create tab layout
-	TabLayout = FTabManager::NewLayout("Standalone_IKRetargetEditor_Layout_v1.005")
+	TabLayout = FTabManager::NewLayout("Standalone_IKRetargetEditor_Layout_v1.006")
 		->AddArea
 		(
 			FTabManager::NewPrimaryArea()
@@ -73,7 +74,7 @@ FIKRetargetMode::FIKRetargetMode(
 					FTabManager::NewStack()
 						->SetSizeCoefficient(0.6f)
 						->AddTab(FIKRetargetChainTabSummoner::TabID, ETabState::OpenedTab)
-						->AddTab(FPersonaTabs::AssetBrowserID, ETabState::OpenedTab)
+						->AddTab(FIKRetargetAssetBrowserTabSummoner::TabID, ETabState::OpenedTab)
 						->SetForegroundTab(FIKRetargetChainTabSummoner::TabID)
 					)
 					
diff --git a/Engine/Plugins/Animation/IKRig/Source/IKRigEditor/Private/RetargetEditor/SIKRetargetAssetBrowser.cpp b/Engine/Plugins/Animation/IKRig/Source/IKRigEditor/Private/RetargetEditor/SIKRetargetAssetBrowser.cpp
new file mode 100644
--- /dev/null
+++ b/Engine/Plugins/Animation/IKRig/Source/IKRigEditor/Private/RetargetEditor/SIKRetargetAssetBrowser.cpp
@@ -0,0 +1,173 @@
+﻿// Copyright Epic Games, Inc. All Rights Reserved.
+
+#include "RetargetEditor/SIKRetargetAssetBrowser.h"
+
+#include "SEditorHeaderButton.h"
+#include "AnimPreviewInstance.h"
+#include "ContentBrowserModule.h"
+#include "IContentBrowserSingleton.h"
+#include "Animation/AnimMontage.h"
+#include "Animation/PoseAsset.h"
+#include "RetargetEditor/IKRetargetBatchOperation.h"
+#include "RetargetEditor/IKRetargetEditorController.h"
+#include "Retargeter/IKRetargeter.h"
+
+#define LOCTEXT_NAMESPACE "IKRetargeterAssetBrowser"
+
+void SIKRetargetAssetBrowser::Construct(
+	const FArguments& InArgs,
+	TSharedRef<FIKRetargetEditorController> InEditorController)
+{
+	EditorController = InEditorController;
+	
+	ChildSlot
+    [
+        SNew(SVerticalBox)
+        
+        + SVerticalBox::Slot()
+		.AutoHeight()
+		.Padding(5)
+		[
+			SNew(SEditorHeaderButton)
+			.IsEnabled(this, &SIKRetargetAssetBrowser::IsExportButtonEnabled)
+			.Icon(FAppStyle::Get().GetBrush("Icons.Save"))
+			.Text(LOCTEXT("ExportButtonLabel", "Export Selected Animations"))
+			.ToolTipText(LOCTEXT("ExportButtonToolTip", "Generate new retargeted sequence assets on target skeletal mesh (uses current retargeting configuration)."))
+			.OnClicked(this, &SIKRetargetAssetBrowser::OnExportButtonClicked)
+		]
+
+		+SVerticalBox::Slot()
+		[
+			SAssignNew(AssetBrowserBox, SBox)
+		]
+    ];
+
+	AddAssetBrowser();
+}
+
+void SIKRetargetAssetBrowser::AddAssetBrowser()
+{
+	FAssetPickerConfig AssetPickerConfig;
+	AssetPickerConfig.Filter.ClassNames.Add(UAnimSequence::StaticClass()->GetFName());
+	AssetPickerConfig.Filter.ClassNames.Add(UAnimMontage::StaticClass()->GetFName());
+	AssetPickerConfig.Filter.ClassNames.Add(UPoseAsset::StaticClass()->GetFName());
+	AssetPickerConfig.OnAssetDoubleClicked = FOnAssetSelected::CreateSP(this, &SIKRetargetAssetBrowser::OnAssetDoubleClicked);
+	AssetPickerConfig.GetCurrentSelectionDelegates.Add(&GetCurrentSelectionDelegate);
+	AssetPickerConfig.bAllowNullSelection = false;
+	AssetPickerConfig.InitialAssetViewType = EAssetViewType::Column;
+	AssetPickerConfig.OnShouldFilterAsset = FOnShouldFilterAsset::CreateSP(this, &SIKRetargetAssetBrowser::OnShouldFilterAsset);
+	AssetPickerConfig.bShowPathInColumnView = true;
+	AssetPickerConfig.bShowTypeInColumnView = false;
+
+	// hide all asset registry columns by default (we only really want the name and path)
+	TArray<UObject::FAssetRegistryTag> AssetRegistryTags;
+	UAnimSequence::StaticClass()->GetDefaultObject()->GetAssetRegistryTags(AssetRegistryTags);
+	for(UObject::FAssetRegistryTag& AssetRegistryTag : AssetRegistryTags)
+	{
+		AssetPickerConfig.HiddenColumnNames.Add(AssetRegistryTag.Name.ToString());
+	}
+
+	// Also hide the type column by default (but allow users to enable it, so don't use bShowTypeInColumnView)
+	AssetPickerConfig.HiddenColumnNames.Add(TEXT("Class"));
+
+	FContentBrowserModule& ContentBrowserModule = FModuleManager::Get().LoadModuleChecked<FContentBrowserModule>(TEXT("ContentBrowser"));
+
+	AssetBrowserBox->SetContent(ContentBrowserModule.Get().CreateAssetPicker(AssetPickerConfig));
+}
+
+FReply SIKRetargetAssetBrowser::OnExportButtonClicked() const
+{
+	FIKRetargetEditorController* Controller = EditorController.Pin().Get();
+	if (!Controller)
+	{
+		return FReply::Handled();
+	}
+	
+	// assemble the data for the assets we want to batch duplicate/retarget
+	FIKRetargetBatchOperationContext BatchContext;
+
+	// add selected assets to dup/retarget
+	TArray<FAssetData> SelectedAssets = GetCurrentSelectionDelegate.Execute();
+	for (const FAssetData& Asset : SelectedAssets)
+	{
+		UE_LOG(LogTemp, Display, TEXT("Duplicating and Retargeting: %s"), *Asset.GetFullName());
+
+		BatchContext.AssetsToRetarget.Add(Asset.GetAsset());
+	}
+
+	BatchContext.SourceMesh = Controller->GetSourceSkeletalMesh();
+	BatchContext.TargetMesh = Controller->GetTargetSkeletalMesh();
+	BatchContext.IKRetargetAsset = Controller->Asset;
+	BatchContext.bRemapReferencedAssets = false;
+	BatchContext.NameRule.Suffix = "_Retargeted";
+
+	// actually run the retarget
+	FIKRetargetBatchOperation BatchOperation;
+	BatchOperation.RunRetarget(BatchContext);
+	
+	return FReply::Handled();
+}
+
+bool SIKRetargetAssetBrowser::IsExportButtonEnabled() const
+{
+	if (!EditorController.Pin().IsValid())
+	{
+		return false;
+	}
+
+	UIKRetargeter* CurrentRetargeter = EditorController.Pin()->GetCurrentlyRunningRetargeter();
+	if (!IsValid(CurrentRetargeter))
+	{
+		return false;
+	}
+
+	return CurrentRetargeter->bIsLoadedAndValid;
+}
+
+void SIKRetargetAssetBrowser::OnAssetDoubleClicked(const FAssetData& AssetData)
+{
+	if (!AssetData.GetAsset())
+	{
+		return;
+	}
+	
+	UAnimationAsset* NewAnimationAsset = Cast<UAnimationAsset>(AssetData.GetAsset());
+	if (NewAnimationAsset && EditorController.Pin().IsValid())
+	{
+		EditorController.Pin()->PlayAnimationAsset(NewAnimationAsset);
+	}
+}
+
+bool SIKRetargetAssetBrowser::OnShouldFilterAsset(const struct FAssetData& AssetData)
+{
+	// is this an animation asset?
+	if (!AssetData.GetClass()->IsChildOf(UAnimationAsset::StaticClass()))
+	{
+		return true;
+	}
+	
+	// controller setup
+	FIKRetargetEditorController* Controller = EditorController.Pin().Get();
+	if (!Controller)
+	{
+		return true;
+	}
+
+	// get source mesh
+	USkeletalMesh* SourceMesh = Controller->GetSourceSkeletalMesh();
+	if (!SourceMesh)
+	{
+		return true;
+	}
+
+	// get source skeleton
+	USkeleton* DesiredSkeleton = SourceMesh->GetSkeleton();
+	if (!DesiredSkeleton)
+	{
+		return true;
+	}
+
+	return !DesiredSkeleton->IsCompatibleSkeletonByAssetData(AssetData);
+}
+
+#undef LOCTEXT_NAMESPACE
diff --git a/Engine/Plugins/Animation/IKRig/Source/IKRigEditor/Private/RetargetEditor/SIKRetargetChainMapList.cpp b/Engine/Plugins/Animation/IKRig/Source/IKRigEditor/Private/RetargetEditor/SIKRetargetChainMapList.cpp
--- a/Engine/Plugins/Animation/IKRig/Source/IKRigEditor/Private/RetargetEditor/SIKRetargetChainMapList.cpp
+++ b/Engine/Plugins/Animation/IKRig/Source/IKRigEditor/Private/RetargetEditor/SIKRetargetChainMapList.cpp
@@ -122,11 +122,6 @@ FText SIKRetargetChainMapRow::GetSourceChainName() const
 	return FText::FromName(Mapping.SourceChain);
 }
 
-SIKRetargetChainMapList::~SIKRetargetChainMapList()
-{
-	
-}
-
 void SIKRetargetChainMapList::Construct(
 	const FArguments& InArgs,
 	TSharedRef<FIKRetargetEditorController> InEditorController)
@@ -136,12 +131,10 @@ void SIKRetargetChainMapList::Construct(
 
 	ChildSlot
     [
-        SNew(SVerticalBox)
-        
-        +SVerticalBox::Slot()
-        .AutoHeight()
-        .VAlign(VAlign_Top)
-        [
+	    SNew(SVerticalBox)
+		+SVerticalBox::Slot()
+		.AutoHeight()
+		[
 			SNew(SHorizontalBox)
 			+ SHorizontalBox::Slot()
 			.AutoWidth()
@@ -155,6 +148,7 @@ void SIKRetargetChainMapList::Construct(
 
 			+ SHorizontalBox::Slot()
 			.AutoWidth()
+			.Padding(5, 0)
 			.HAlign(HAlign_Left)
 			.VAlign(VAlign_Center)
 			[
@@ -162,13 +156,7 @@ void SIKRetargetChainMapList::Construct(
 				.Text(this, &SIKRetargetChainMapList::GetSourceRootBone)
 				.IsEnabled(false)
 			]
-        ]
-
-		+SVerticalBox::Slot()
-		.AutoHeight()
-		.VAlign(VAlign_Top)
-		[
-			SNew(SHorizontalBox)
+				
 			+ SHorizontalBox::Slot()
 			.AutoWidth()
 			.HAlign(HAlign_Left)
@@ -181,6 +169,7 @@ void SIKRetargetChainMapList::Construct(
 
 			+ SHorizontalBox::Slot()
 			.AutoWidth()
+			.Padding(5, 0)
 			.HAlign(HAlign_Left)
 			.VAlign(VAlign_Center)
 			[
@@ -190,10 +179,11 @@ void SIKRetargetChainMapList::Construct(
 			]
 		]
 
-		+ SVerticalBox::Slot()
+		+SVerticalBox::Slot()
 		.AutoHeight()
-		.VAlign(VAlign_Top)
-		.HAlign(HAlign_Left)
+		.VAlign(VAlign_Center)
+		.HAlign(HAlign_Fill)
+		.Padding(5)
 		[
 			SNew(SEditorHeaderButton)
 			.Visibility(this, &SIKRetargetChainMapList::IsAutoMapButtonVisible)
@@ -202,8 +192,9 @@ void SIKRetargetChainMapList::Construct(
 			.ToolTipText(LOCTEXT("AutoMapButtonToolTip", "Automatically assign source chains based on fuzzy string match"))
 			.OnClicked(this, &SIKRetargetChainMapList::OnAutoMapButtonClicked)
 		]
-		
+	    
         +SVerticalBox::Slot()
+        .AutoHeight()
         [
 			SAssignNew(ListView, SRetargetChainMapListViewType )
 			.SelectionMode(ESelectionMode::Multi)
diff --git a/Engine/Plugins/Animation/IKRig/Source/IKRigEditor/Private/RetargetEditor/SRetargetAnimAssetsWindow.cpp b/Engine/Plugins/Animation/IKRig/Source/IKRigEditor/Private/RetargetEditor/SRetargetAnimAssetsWindow.cpp
--- a/Engine/Plugins/Animation/IKRig/Source/IKRigEditor/Private/RetargetEditor/SRetargetAnimAssetsWindow.cpp
+++ b/Engine/Plugins/Animation/IKRig/Source/IKRigEditor/Private/RetargetEditor/SRetargetAnimAssetsWindow.cpp
@@ -2,11 +2,8 @@
 
 #include "RetargetEditor/SRetargetAnimAssetsWindow.h"
 
-#include "AnimationBlueprintLibrary.h"
 #include "AnimPose.h"
 #include "AnimPreviewInstance.h"
-#include "ContentBrowserModule.h"
-#include "EditorReimportHandler.h"
 #include "IContentBrowserSingleton.h"
 #include "ObjectEditorUtils.h"
 #include "SSkeletonWidget.h"
@@ -14,376 +11,12 @@
 #include "Widgets/Layout/SSeparator.h"
 #include "PropertyCustomizationHelpers.h"
 #include "Animation/DebugSkelMeshComponent.h"
-#include "EditorFramework/AssetImportData.h"
-#include "Framework/Notifications/NotificationManager.h"
-#include "Kismet2/BlueprintEditorUtils.h"
-#include "Kismet2/KismetEditorUtilities.h"
 #include "Misc/ScopedSlowTask.h"
 #include "Retargeter/IKRetargeter.h"
 #include "Widgets/Notifications/SNotificationList.h"
 
 #define LOCTEXT_NAMESPACE "RetargetAnimAssetWindow"
 
-int32 FIKRetargetAnimAssetsContext::GenerateAssetLists()
-{
-	// re-generate lists of selected and referenced assets
-	AnimationAssetsToRetarget.Reset();
-	AnimBlueprintsToRetarget.Reset();
-
-	for (TWeakObjectPtr<UObject> AssetPtr : AssetsToRetarget)
-	{
-		UObject* Asset = AssetPtr.Get();
-		if (UAnimationAsset* AnimAsset = Cast<UAnimationAsset>(Asset))
-		{
-			AnimationAssetsToRetarget.AddUnique(AnimAsset);
-		}
-		else if (UAnimBlueprint* AnimBlueprint = Cast<UAnimBlueprint>(Asset))
-		{
-			// Add parents blueprint. 
-			UAnimBlueprint* ParentBP = Cast<UAnimBlueprint>(AnimBlueprint->ParentClass->ClassGeneratedBy);
-			while (ParentBP)
-			{
-				AnimBlueprintsToRetarget.AddUnique(ParentBP);
-				ParentBP = Cast<UAnimBlueprint>(ParentBP->ParentClass->ClassGeneratedBy);
-			}
-				
-			AnimBlueprintsToRetarget.AddUnique(AnimBlueprint);				
-		}
-	}
-
-	if (bRemapReferencedAssets)
-	{
-		// Grab assets from the blueprint.
-		// Do this first as it can add complex assets to the retarget array which will need to be processed next.
-		for (UAnimBlueprint* AnimBlueprint : AnimBlueprintsToRetarget)
-		{
-			GetAllAnimationSequencesReferredInBlueprint(AnimBlueprint, AnimationAssetsToRetarget);
-		}
-
-		int32 AssetIndex = 0;
-		while (AssetIndex < AnimationAssetsToRetarget.Num())
-		{
-			UAnimationAsset* AnimAsset = AnimationAssetsToRetarget[AssetIndex++];
-			AnimAsset->HandleAnimReferenceCollection(AnimationAssetsToRetarget, true);
-		}
-	}
-
-	return AnimationAssetsToRetarget.Num();
-}
-
-void FIKRetargetAnimAssetsContext::DuplicateRetargetAssets()
-{
-	UPackage* DestinationPackage = TargetMesh->GetOutermost();
-
-	TArray<UAnimationAsset*> AnimationAssetsToDuplicate = AnimationAssetsToRetarget;
-	TArray<UAnimBlueprint*> AnimBlueprintsToDuplicate = AnimBlueprintsToRetarget;
-
-	// We only want to duplicate unmapped assets, so we remove mapped assets from the list we're duplicating
-	for (TPair<UAnimationAsset*, UAnimationAsset*>& Pair : RemappedAnimAssets)
-	{
-		AnimationAssetsToDuplicate.Remove(Pair.Key);
-	}
-
-	DuplicatedAnimAssets = DuplicateAssets<UAnimationAsset>(AnimationAssetsToDuplicate, DestinationPackage, &NameRule);
-	DuplicatedBlueprints = DuplicateAssets<UAnimBlueprint>(AnimBlueprintsToDuplicate, DestinationPackage, &NameRule);
-
-	// If we are moving the new asset to a different directory we need to fixup the reimport path.
-	// This should only effect source FBX paths within the project.
-	if (!NameRule.FolderPath.IsEmpty())
-	{
-		for (TPair<UAnimationAsset*, UAnimationAsset*>& Pair : DuplicatedAnimAssets)
-		{
-			UAnimSequence* SourceSequence = Cast<UAnimSequence>(Pair.Key);
-			UAnimSequence* DestinationSequence = Cast<UAnimSequence>(Pair.Value);
-			if (!(SourceSequence && DestinationSequence))
-			{
-				continue;
-			}
-			
-			for (int index = 0; index < SourceSequence->AssetImportData->SourceData.SourceFiles.Num(); index++)
-			{
-				const FString& RelativeFilename = SourceSequence->AssetImportData->SourceData.SourceFiles[index].RelativeFilename;
-				const FString OldPackagePath = FPackageName::GetLongPackagePath(SourceSequence->GetPathName()) / TEXT("");
-				const FString NewPackagePath = FPackageName::GetLongPackagePath(DestinationSequence->GetPathName()) / TEXT("");
-				const FString AbsoluteSrcPath = FPaths::ConvertRelativePathToFull(FPackageName::LongPackageNameToFilename(OldPackagePath));
-				const FString SrcFile = AbsoluteSrcPath / RelativeFilename;
-				const bool bSrcFileExists = FPlatformFileManager::Get().GetPlatformFile().FileExists(*SrcFile);
-				if (!bSrcFileExists || (NewPackagePath == OldPackagePath))
-				{
-					continue;
-				}
-
-				FString BasePath = FPackageName::LongPackageNameToFilename(OldPackagePath);
-				FString OldSourceFilePath = FPaths::ConvertRelativePathToFull(BasePath, RelativeFilename);
-				TArray<FString> Paths;
-				Paths.Add(OldSourceFilePath);
-				
-				// update the FBX reimport file path
-				FReimportManager::Instance()->UpdateReimportPaths(DestinationSequence, Paths);
-			}
-		}
-	}
-
-	// Remapped assets needs the duplicated ones added
-	RemappedAnimAssets.Append(DuplicatedAnimAssets);
-
-	DuplicatedAnimAssets.GenerateValueArray(AnimationAssetsToRetarget);
-	DuplicatedBlueprints.GenerateValueArray(AnimBlueprintsToRetarget);
-}
-
-void FIKRetargetAnimAssetsContext::RetargetAssets()
-{
-	USkeleton* OldSkeleton = SourceMesh->GetSkeleton();
-	USkeleton* NewSkeleton = TargetMesh->GetSkeleton();
-
-	for (UAnimationAsset* AssetToRetarget : AnimationAssetsToRetarget)
-	{
-		// synchronize curves between old/new asset
-		UAnimSequence* AnimSequenceToRetarget = Cast<UAnimSequence>(AssetToRetarget);
-		if (AnimSequenceToRetarget)
-		{
-			// copy curve data from source asset, preserving data in the target if present.
-			UAnimationBlueprintLibrary::CopyAnimationCurveNamesToSkeleton(OldSkeleton, NewSkeleton, AnimSequenceToRetarget, ERawCurveTrackTypes::RCT_Float);	
-			// clear transform curves since those curves won't work in new skeleton
-			IAnimationDataController& Controller = AnimSequenceToRetarget->GetController();
-			Controller.RemoveAllCurvesOfType(ERawCurveTrackTypes::RCT_Transform);
-		}
-
-		// replace references to other animation
-		AssetToRetarget->ReplaceReferredAnimations(RemappedAnimAssets);
-		AssetToRetarget->SetSkeleton(NewSkeleton);
-		AssetToRetarget->MarkPackageDirty();
-	}
-
-	// convert the animation using the IK retargeter
-	ConvertAnimation();
-
-	// convert all Animation Blueprints and compile 
-	for (UAnimBlueprint* AnimBlueprint : AnimBlueprintsToRetarget)
-	{
-		// replace skeleton
-		AnimBlueprint->TargetSkeleton = NewSkeleton;
-
-		// if they have parent blueprint, make sure to re-link to the new one also
-		UAnimBlueprint* CurrentParentBP = Cast<UAnimBlueprint>(AnimBlueprint->ParentClass->ClassGeneratedBy);
-		if (CurrentParentBP)
-		{
-			UAnimBlueprint* const * ParentBP = DuplicatedBlueprints.Find(CurrentParentBP);
-			if (ParentBP)
-			{
-				AnimBlueprint->ParentClass = (*ParentBP)->GeneratedClass;
-			}
-		}
-
-		if(RemappedAnimAssets.Num() > 0)
-		{
-			ReplaceReferredAnimationsInBlueprint(AnimBlueprint, RemappedAnimAssets);
-		}
-
-		FBlueprintEditorUtils::RefreshAllNodes(AnimBlueprint);
-		FKismetEditorUtilities::CompileBlueprint(AnimBlueprint, EBlueprintCompileOptions::SkipGarbageCollection);
-		AnimBlueprint->PostEditChange();
-		AnimBlueprint->MarkPackageDirty();
-	}
-}
-
-void FIKRetargetAnimAssetsContext::ConvertAnimation()
-{
-	// initialize the retargeter
-	UObject* TransientOuter = Cast<UObject>(GetTransientPackage());
-	UIKRetargeter* Retargeter = DuplicateObject(IKRetargetAsset, TransientOuter);	
-	Retargeter->Initialize(SourceMesh, TargetMesh, TransientOuter);
-	if (!Retargeter->bIsLoadedAndValid)
-	{
-		UE_LOG(LogTemp, Warning, TEXT("Unable to initialize the IK Retargeter. Newly created animations were not retargeted!"));
-		return;
-	}
-
-	// for each pair of source / target animation sequences
-	for (TPair<UAnimationAsset*, UAnimationAsset*>& Pair : DuplicatedAnimAssets)
-	{
-		UAnimSequence* SourceSequence = Cast<UAnimSequence>(Pair.Key);
-		UAnimSequence* DestinationSequence = Cast<UAnimSequence>(Pair.Value);
-		if (!(SourceSequence && DestinationSequence))
-		{
-			continue;
-		}
-
-		// remove all keys from the destination animation sequence
-		IAnimationDataController& TargetSeqController = DestinationSequence->GetController();
-		TargetSeqController.RemoveAllBoneTracks();
-
-		// number of frames in this animation
-		const int32 NumFrames = SourceSequence->GetNumberOfSampledKeys();
-
-		// make space for the target keyframe data
-		const int32 NumTargetBones = Retargeter->TargetSkeleton.BoneNames.Num();
-		TArray<FRawAnimSequenceTrack> BoneTracks;
-		BoneTracks.SetNumZeroed(NumTargetBones);
-
-		// retarget each frame's pose from source to target
-		for (int32 FrameIndex=0; FrameIndex<NumFrames; ++FrameIndex)
-		{
-			// get the source global pose
-			FAnimPose SourcePoseAtFrame;
-			UAnimPoseExtensions::GetAnimPoseAtFrame(SourceSequence, FrameIndex, FAnimPoseEvaluationOptions(), SourcePoseAtFrame);
-			TArray<FName> BoneNames;
-			UAnimPoseExtensions::GetBoneNames(SourcePoseAtFrame,BoneNames);
-			TArray<FTransform> SourceComponentPose;
-			for (const FName& BoneName : BoneNames)
-			{
-				FTransform BonePose = UAnimPoseExtensions::GetBonePose(SourcePoseAtFrame, BoneName, EAnimPoseSpaces::World);
-				SourceComponentPose.Add(BonePose);
-			}
-
-			// run the retarget
-			TArray<FTransform>& TargetComponentPose = Retargeter->RunRetargeter(SourceComponentPose);
-
-			// convert to a local-space pose
-			TArray<FTransform> TargetLocalPose = TargetComponentPose;
-			Retargeter->TargetSkeleton.UpdateLocalTransformsBelowBone(0,TargetLocalPose, TargetComponentPose);
-
-			// store key data for each bone
-			for (int32 TargetBoneIndex=0; TargetBoneIndex<TargetLocalPose.Num(); ++TargetBoneIndex)
-			{
-				BoneTracks[TargetBoneIndex].PosKeys.Add(TargetLocalPose[TargetBoneIndex].GetLocation());
-				BoneTracks[TargetBoneIndex].RotKeys.Add(TargetLocalPose[TargetBoneIndex].GetRotation());
-				BoneTracks[TargetBoneIndex].ScaleKeys.Add(TargetLocalPose[TargetBoneIndex].GetScale3D());
-			}
-		}
-
-		// add keys to bone tracks
-		const bool bShouldTransact = false;
-		for (int32 TargetBoneIndex=0; TargetBoneIndex<NumTargetBones; ++TargetBoneIndex)
-		{
-			FName TargetBoneName = Retargeter->TargetSkeleton.BoneNames[TargetBoneIndex];
-			FRawAnimSequenceTrack& RawTrack = BoneTracks[TargetBoneIndex];
-			TargetSeqController.AddBoneTrack(TargetBoneName, bShouldTransact);
-			TargetSeqController.SetBoneTrackKeys(TargetBoneName, RawTrack.PosKeys, RawTrack.RotKeys, RawTrack.ScaleKeys);
-		}
-	}
-}
-
-void FIKRetargetAnimAssetsContext::NotifyUserOfResults() const
-{
-	// gather newly created objects
-	TArray<UObject*> NewAssets;
-	GetNewAssets(NewAssets);
-
-	// log details of what assets were created
-	for (UObject* NewAsset : NewAssets)
-	{
-		UE_LOG(LogTemp, Display, TEXT("Duplicate and Retarget - New Asset Created: %s"), *NewAsset->GetName());
-	}
-	
-	// notify user
-	FNotificationInfo Notification(FText::GetEmpty());
-	Notification.ExpireDuration = 5.f;
-	Notification.Text = FText::Format(
-		LOCTEXT("MultiNonDuplicatedAsset", "{0} assets were retargeted to new skeleton {1}. See Output for details."),
-		FText::AsNumber(NewAssets.Num()),
-		FText::FromString(TargetMesh->GetName()));
-	FSlateNotificationManager::Get().AddNotification(Notification);
-	
-	// select all new assets
-	TArray<FAssetData> CurrentSelection;
-	for(UObject* NewObject : NewAssets)
-	{
-		CurrentSelection.Add(FAssetData(NewObject));
-	}
-
-	// show assets in browser
-	FContentBrowserModule& ContentBrowserModule = FModuleManager::Get().LoadModuleChecked<FContentBrowserModule>("ContentBrowser");
-	ContentBrowserModule.Get().SyncBrowserToAssets(CurrentSelection);
-}
-
-void FIKRetargetAnimAssetsContext::GetNewAssets(TArray<UObject*>& NewAssets) const
-{
-	TArray<UAnimationAsset*> NewAnims;
-	DuplicatedAnimAssets.GenerateValueArray(NewAnims);
-	for (UAnimationAsset* NewAnim : NewAnims)
-	{
-		NewAssets.Add(Cast<UObject>(NewAnim));
-	}
-
-	TArray<UAnimBlueprint*> NewBlueprints;
-	DuplicatedBlueprints.GenerateValueArray(NewBlueprints);
-	for (UAnimBlueprint* NewBlueprint : NewBlueprints)
-	{
-		NewAssets.Add(Cast<UObject>(NewBlueprint));
-	}
-}
-
-void FIKRetargetAnimAssetsContext::Reset()
-{
-	SourceMesh = nullptr;
-	TargetMesh = nullptr;
-	IKRetargetAsset = nullptr;
-	bRemapReferencedAssets = true;
-	NameRule.Prefix = "";
-	NameRule.Suffix = "";
-	NameRule.ReplaceFrom = "";
-	NameRule.ReplaceTo = "";
-}
-
-bool FIKRetargetAnimAssetsContext::IsValid() const
-{
-	// todo validate compatibility
-	return SourceMesh && TargetMesh && IKRetargetAsset && (SourceMesh != TargetMesh);
-}
-
-void FIKRetargetAnimAssetsContext::RunRetarget()
-{
-	// todo progess bar not showing up, why?
-	//FScopedSlowTask Progress(2.f, LOCTEXT("GatheringBatchRetarget", "Gathering animation assets..."));
-	//Progress.MakeDialog();
-	
-	const int32 NumAssets = GenerateAssetLists();
-	
-	//Progress.EnterProgressFrame(1.f, FText::Format(LOCTEXT("DuplicatingBatchRetarget", "Duplicating {0} animation assets..."), NumAssets));
-
-	DuplicateRetargetAssets();
-	
-	//Progress.EnterProgressFrame(1.f, FText::Format(LOCTEXT("RunningBatchRetarget", "Retargeting {0} animation assets..."), NumAssets));
-
-	RetargetAssets();
-	
-	NotifyUserOfResults();
-}
-
-/**
-* Duplicates the supplied AssetsToDuplicate and returns a map of original asset to duplicate. Templated wrapper that calls DuplicateAssetInternal.
-*
-* @param	AssetsToDuplicate	The animations to duplicate
-* @param	DestinationPackage	The package that the duplicates should be placed in
-*
-* @return	TMap of original animation to duplicate
-*/
-template<class AssetType>
-TMap<AssetType*, AssetType*> FIKRetargetAnimAssetsContext::DuplicateAssets(
-	const TArray<AssetType*>& AssetsToDuplicate,
-	UPackage* DestinationPackage,
-	const FNameDuplicationRule* NameRule)
-{
-	TArray<UObject*> Assets;
-	for (AssetType* Asset : AssetsToDuplicate)
-	{
-		Assets.Add(Asset);
-	}
-
-	// duplicate assets
-	TMap<UObject*, UObject*> DuplicateAssetsMap = DuplicateAssetsInternal(Assets, DestinationPackage, NameRule);
-
-	// cast to AssetType
-	TMap<AssetType*, AssetType*> ReturnMap;
-	for (const TTuple<UObject*, UObject*>& DuplicateAsset : DuplicateAssetsMap)
-	{
-		ReturnMap.Add(Cast<AssetType>(DuplicateAsset.Key), Cast<AssetType>(DuplicateAsset.Value));
-	}
-	
-	return ReturnMap;
-}
-
 void SRetargetPoseViewport::Construct(const FArguments& InArgs)
 {
 	SEditorViewport::Construct(SEditorViewport::FArguments());
@@ -496,7 +129,7 @@ void SRetargetAnimAssetsWindow::Construct(const FArguments& InArgs)
 					.Padding(5, 5)
 					[
 						SAssignNew(SourceViewport, SRetargetPoseViewport)
-						.SkeletalMesh(RetargetContext.SourceMesh)
+						.SkeletalMesh(BatchContext.SourceMesh)
 					]
 
 					+ SVerticalBox::Slot()
@@ -512,18 +145,18 @@ void SRetargetAnimAssetsWindow::Construct(const FArguments& InArgs)
 						.ThumbnailPool(AssetThumbnailPool)
 						.IsEnabled_Lambda([this]()
 						{
-							if (!RetargetContext.IKRetargetAsset)
+							if (!BatchContext.IKRetargetAsset)
 							{
 								return false;
 							}
 							
-							return RetargetContext.IKRetargetAsset->SourceIKRigAsset != nullptr;
+							return BatchContext.IKRetargetAsset->SourceIKRigAsset != nullptr;
 						})
 						.ObjectPath(this, &SRetargetAnimAssetsWindow::GetCurrentSourceMeshPath)
 						.OnObjectChanged(this, &SRetargetAnimAssetsWindow::SourceMeshAssigned)
 						.OnShouldFilterAsset_Lambda([this](const FAssetData& AssetData)
 						{
-							if (!RetargetContext.IKRetargetAsset)
+							if (!BatchContext.IKRetargetAsset)
 							{
 								return true;
 							}
@@ -534,7 +167,7 @@ void SRetargetAnimAssetsWindow::Construct(const FArguments& InArgs)
 								return true;
 							}
 							
-							USkeletalMesh* PreviewMesh = RetargetContext.IKRetargetAsset->SourceIKRigAsset->GetPreviewMesh();
+							USkeletalMesh* PreviewMesh = BatchContext.IKRetargetAsset->SourceIKRigAsset->GetPreviewMesh();
 							if (!PreviewMesh)
 							{
 								return true;
@@ -587,19 +220,19 @@ void SRetargetAnimAssetsWindow::Construct(const FArguments& InArgs)
 						.ThumbnailPool(AssetThumbnailPool)
 						.IsEnabled_Lambda([this]()
 						{
-							if (!RetargetContext.IKRetargetAsset)
+							if (!BatchContext.IKRetargetAsset)
 							{
 								return false;
 							}
 							
-							return RetargetContext.IKRetargetAsset->TargetIKRigAsset != nullptr;
+							return BatchContext.IKRetargetAsset->TargetIKRigAsset != nullptr;
 						})
 						.ObjectPath(this, &SRetargetAnimAssetsWindow::GetCurrentTargetMeshPath)
 						.OnObjectChanged(this, &SRetargetAnimAssetsWindow::TargetMeshAssigned)
 						.OnShouldFilterAsset_Lambda([this](const FAssetData& AssetData)
 						{
 							/*
-							if (!RetargetContext.IKRetargetAsset)
+							if (!BatchContext.IKRetargetAsset)
 							{
 								return true;
 							}
@@ -611,7 +244,7 @@ void SRetargetAnimAssetsWindow::Construct(const FArguments& InArgs)
 							}
 
 							
-							USkeletalMesh* PreviewMesh = RetargetContext.IKRetargetAsset->TargetIKRigAsset->GetPreviewMesh();
+							USkeletalMesh* PreviewMesh = BatchContext.IKRetargetAsset->TargetIKRigAsset->GetPreviewMesh();
 							if (!PreviewMesh)
 							{
 								return true;
@@ -874,13 +507,14 @@ void SRetargetAnimAssetsWindow::Construct(const FArguments& InArgs)
 
 bool SRetargetAnimAssetsWindow::CanApply() const
 {
-	return RetargetContext.IsValid();
+	return BatchContext.IsValid();
 }
 
 FReply SRetargetAnimAssetsWindow::OnApply()
 {
 	CloseWindow();
-	RetargetContext.RunRetarget();
+	FIKRetargetBatchOperation BatchOperation;
+	BatchOperation.RunRetarget(BatchContext);
 	return FReply::Handled();
 }
 
@@ -923,7 +557,7 @@ void SRetargetAnimAssetsWindow::ShowWindow(TArray<UObject*> InSelectedAssets)
 			SAssignNew(DialogWidget, SRetargetAnimAssetsWindow)
 		];
 
-	DialogWidget->RetargetContext.AssetsToRetarget = FObjectEditorUtils::GetTypedWeakObjectPtrs<UObject>(InSelectedAssets);
+	DialogWidget->BatchContext.AssetsToRetarget = FObjectEditorUtils::GetTypedWeakObjectPtrs<UObject>(InSelectedAssets);
 	DialogWindow->SetOnWindowClosed(FRequestDestroyWindowOverride::CreateSP(DialogWidget.Get(), &SRetargetAnimAssetsWindow::OnDialogClosed));
 	DialogWindow->SetContent(DialogWrapper.ToSharedRef());
 	
@@ -938,36 +572,36 @@ void SRetargetAnimAssetsWindow::OnDialogClosed(const TSharedRef<SWindow>& Window
 void SRetargetAnimAssetsWindow::SourceMeshAssigned(const FAssetData& InAssetData)
 {
 	USkeletalMesh* Mesh = Cast<USkeletalMesh>(InAssetData.GetAsset());
-	RetargetContext.SourceMesh = Mesh;
-	SourceViewport->SetSkeletalMesh(RetargetContext.SourceMesh);
+	BatchContext.SourceMesh = Mesh;
+	SourceViewport->SetSkeletalMesh(BatchContext.SourceMesh);
 }
 
 void SRetargetAnimAssetsWindow::TargetMeshAssigned(const FAssetData& InAssetData)
 {
 	USkeletalMesh* Mesh = Cast<USkeletalMesh>(InAssetData.GetAsset());
-	RetargetContext.TargetMesh = Mesh;
-	TargetViewport->SetSkeletalMesh(RetargetContext.TargetMesh);
+	BatchContext.TargetMesh = Mesh;
+	TargetViewport->SetSkeletalMesh(BatchContext.TargetMesh);
 }
 
 FString SRetargetAnimAssetsWindow::GetCurrentSourceMeshPath() const
 {
-	return RetargetContext.SourceMesh ? RetargetContext.SourceMesh->GetPathName() : FString("");
+	return BatchContext.SourceMesh ? BatchContext.SourceMesh->GetPathName() : FString("");
 }
 
 FString SRetargetAnimAssetsWindow::GetCurrentTargetMeshPath() const
 {
-	return RetargetContext.TargetMesh ? RetargetContext.TargetMesh->GetPathName() : FString("");
+	return BatchContext.TargetMesh ? BatchContext.TargetMesh->GetPathName() : FString("");
 }
 
 FString SRetargetAnimAssetsWindow::GetCurrentRetargeterPath() const
 {
-	return RetargetContext.IKRetargetAsset ? RetargetContext.IKRetargetAsset->GetPathName() : FString("");
+	return BatchContext.IKRetargetAsset ? BatchContext.IKRetargetAsset->GetPathName() : FString("");
 }
 
 void SRetargetAnimAssetsWindow::RetargeterAssigned(const FAssetData& InAssetData)
 {
 	UIKRetargeter* InRetargeter = Cast<UIKRetargeter>(InAssetData.GetAsset());
-	RetargetContext.IKRetargetAsset = InRetargeter;
+	BatchContext.IKRetargetAsset = InRetargeter;
 	const UIKRigDefinition* SourceIKRig = InRetargeter ? InRetargeter->SourceIKRigAsset : nullptr;
 	const UIKRigDefinition* TargetIKRig = InRetargeter ? InRetargeter->TargetIKRigAsset : nullptr;
 	USkeletalMesh* SourceMesh =  SourceIKRig ? SourceIKRig->GetPreviewMesh() : nullptr;
@@ -978,55 +612,55 @@ void SRetargetAnimAssetsWindow::RetargeterAssigned(const FAssetData& InAssetData
 
 ECheckBoxState SRetargetAnimAssetsWindow::IsRemappingReferencedAssets() const
 {
-	return RetargetContext.bRemapReferencedAssets ? ECheckBoxState::Checked : ECheckBoxState::Unchecked;
+	return BatchContext.bRemapReferencedAssets ? ECheckBoxState::Checked : ECheckBoxState::Unchecked;
 }
 
 void SRetargetAnimAssetsWindow::OnRemappingReferencedAssetsChanged(ECheckBoxState InNewRadioState)
 {
-	RetargetContext.bRemapReferencedAssets = (InNewRadioState == ECheckBoxState::Checked);
+	BatchContext.bRemapReferencedAssets = (InNewRadioState == ECheckBoxState::Checked);
 }
 
 FText SRetargetAnimAssetsWindow::GetPrefixName() const
 {
-	return FText::FromString(RetargetContext.NameRule.Prefix);
+	return FText::FromString(BatchContext.NameRule.Prefix);
 }
 
 void SRetargetAnimAssetsWindow::SetPrefixName(const FText &InText)
 {
-	RetargetContext.NameRule.Prefix = InText.ToString();
+	BatchContext.NameRule.Prefix = InText.ToString();
 	UpdateExampleText();
 }
 
 FText SRetargetAnimAssetsWindow::GetSuffixName() const
 {
-	return FText::FromString(RetargetContext.NameRule.Suffix);
+	return FText::FromString(BatchContext.NameRule.Suffix);
 }
 
 void SRetargetAnimAssetsWindow::SetSuffixName(const FText &InText)
 {
-	RetargetContext.NameRule.Suffix = InText.ToString();
+	BatchContext.NameRule.Suffix = InText.ToString();
 	UpdateExampleText();
 }
 
 FText SRetargetAnimAssetsWindow::GetReplaceFrom() const
 {
-	return FText::FromString(RetargetContext.NameRule.ReplaceFrom);
+	return FText::FromString(BatchContext.NameRule.ReplaceFrom);
 }
 
 void SRetargetAnimAssetsWindow::SetReplaceFrom(const FText &InText)
 {
-	RetargetContext.NameRule.ReplaceFrom = InText.ToString();
+	BatchContext.NameRule.ReplaceFrom = InText.ToString();
 	UpdateExampleText();
 }
 
 FText SRetargetAnimAssetsWindow::GetReplaceTo() const
 {
-	return FText::FromString(RetargetContext.NameRule.ReplaceTo);
+	return FText::FromString(BatchContext.NameRule.ReplaceTo);
 }
 
 void SRetargetAnimAssetsWindow::SetReplaceTo(const FText &InText)
 {
-	RetargetContext.NameRule.ReplaceTo = InText.ToString();
+	BatchContext.NameRule.ReplaceTo = InText.ToString();
 	UpdateExampleText();
 }
 
@@ -1037,15 +671,15 @@ FText SRetargetAnimAssetsWindow::GetExampleText() const
 
 void SRetargetAnimAssetsWindow::UpdateExampleText()
 {
-	const FString ReplaceFrom = FString::Printf(TEXT("Old Name : ###%s###"), *RetargetContext.NameRule.ReplaceFrom);
-	const FString ReplaceTo = FString::Printf(TEXT("New Name : %s###%s###%s"), *RetargetContext.NameRule.Prefix, *RetargetContext.NameRule.ReplaceTo, *RetargetContext.NameRule.Suffix);
+	const FString ReplaceFrom = FString::Printf(TEXT("Old Name : ###%s###"), *BatchContext.NameRule.ReplaceFrom);
+	const FString ReplaceTo = FString::Printf(TEXT("New Name : %s###%s###%s"), *BatchContext.NameRule.Prefix, *BatchContext.NameRule.ReplaceTo, *BatchContext.NameRule.Suffix);
 
 	ExampleText = FText::FromString(FString::Printf(TEXT("%s\n%s"), *ReplaceFrom, *ReplaceTo));
 }
 
 FText SRetargetAnimAssetsWindow::GetFolderPath() const
 {
-	return FText::FromString(RetargetContext.NameRule.FolderPath);
+	return FText::FromString(BatchContext.NameRule.FolderPath);
 }
 
 #undef LOCTEXT_NAMESPACE
diff --git a/Engine/Plugins/Animation/IKRig/Source/IKRigEditor/Private/RigEditor/IKRigEditMode.cpp b/Engine/Plugins/Animation/IKRig/Source/IKRigEditor/Private/RigEditor/IKRigEditMode.cpp
--- a/Engine/Plugins/Animation/IKRig/Source/IKRigEditor/Private/RigEditor/IKRigEditMode.cpp
+++ b/Engine/Plugins/Animation/IKRig/Source/IKRigEditor/Private/RigEditor/IKRigEditMode.cpp
@@ -234,7 +234,7 @@ bool FIKRigEditMode::InputDelta(FEditorViewportClient* InViewportClient, FViewpo
 	{
 		for (const FName& GoalName : SelectedGoalNames)
 		{
-			FTransform CurrentTransform = AssetController->GetGoal(GoalName)->CurrentTransform;
+			FTransform CurrentTransform = AssetController->GetGoalCurrentTransform(GoalName);
 			CurrentTransform.AddToTranslation(InDrag);
 			AssetController->SetGoalCurrentTransform(GoalName, CurrentTransform);
 		}
@@ -245,7 +245,7 @@ bool FIKRigEditMode::InputDelta(FEditorViewportClient* InViewportClient, FViewpo
 	{
 		for (const FName& GoalName : SelectedGoalNames)
 		{
-			FTransform CurrentTransform = AssetController->GetGoal(GoalName)->CurrentTransform;
+			FTransform CurrentTransform = AssetController->GetGoalCurrentTransform(GoalName);
 			FQuat CurrentRotation = CurrentTransform.GetRotation();
 			CurrentRotation = (InRot.Quaternion() * CurrentRotation);
 			CurrentTransform.SetRotation(CurrentRotation);
diff --git a/Engine/Plugins/Animation/IKRig/Source/IKRigEditor/Private/RigEditor/IKRigEditorController.cpp b/Engine/Plugins/Animation/IKRig/Source/IKRigEditor/Private/RigEditor/IKRigEditorController.cpp
--- a/Engine/Plugins/Animation/IKRig/Source/IKRigEditor/Private/RigEditor/IKRigEditorController.cpp
+++ b/Engine/Plugins/Animation/IKRig/Source/IKRigEditor/Private/RigEditor/IKRigEditorController.cpp
@@ -30,6 +30,59 @@ void FIKRigEditorController::RefreshAllViews() const
 	RetargetingView->RefreshView();
 }
 
+void FIKRigEditorController::AddNewGoals(const TArray<FName>& GoalNames, const TArray<FName>& BoneNames)
+{
+	check(GoalNames.Num() == BoneNames.Num());
+
+	// add a default solver if there isn't one already
+	if (AssetController->GetNumSolvers() == 0)
+	{
+		PromptToAddSolver();	
+	}
+
+	// get selected solvers
+	TArray<TSharedPtr<FSolverStackElement>> SelectedSolvers;
+	GetSelectedSolvers(SelectedSolvers);
+
+	// create goals
+	FName LastCreatedGoalName = NAME_None;
+	for (int32 I=0; I<GoalNames.Num(); ++I)
+	{
+		const FName& GoalName = GoalNames[I];
+		const FName& BoneName = BoneNames[I];
+
+		// create a new goal
+		UIKRigEffectorGoal* NewGoal = AssetController->AddNewGoal(GoalName, BoneName);
+		if (!NewGoal)
+		{
+			continue; // already exists
+		}
+		
+		// connect the new goal to all the selected solvers
+		for (const TSharedPtr<FSolverStackElement>& SolverElement : SelectedSolvers)
+		{
+			AssetController->ConnectGoalToSolver(*NewGoal, SolverElement->IndexInStack);	
+		}
+
+		LastCreatedGoalName = GoalName;
+	}
+	
+	// were any goals created?
+	if (LastCreatedGoalName != NAME_None)
+	{
+		// show last created goal in details view
+		ShowDetailsForGoal(LastCreatedGoalName);
+		// update all views
+		RefreshAllViews();
+	}
+}
+
+void FIKRigEditorController::DeleteGoal(const FName& GoalToDelete)
+{
+	AssetController->RemoveGoal(GoalToDelete);
+	SelectedGoals.Remove(GoalToDelete);
+}
+
 bool FIKRigEditorController::IsGoalSelected(const FName& GoalName)
 {
 	return SelectedGoals.Contains(GoalName);
diff --git a/Engine/Plugins/Animation/IKRig/Source/IKRigEditor/Private/RigEditor/IKRigSkeletonCommands.cpp b/Engine/Plugins/Animation/IKRig/Source/IKRigEditor/Private/RigEditor/IKRigSkeletonCommands.cpp
--- a/Engine/Plugins/Animation/IKRig/Source/IKRigEditor/Private/RigEditor/IKRigSkeletonCommands.cpp
+++ b/Engine/Plugins/Animation/IKRig/Source/IKRigEditor/Private/RigEditor/IKRigSkeletonCommands.cpp
@@ -7,7 +7,7 @@
 void FIKRigSkeletonCommands::RegisterCommands()
 {
 	UI_COMMAND(NewGoal, "New IK Goal", "Add new IK goal at the selected bone.", EUserInterfaceActionType::Button, FInputChord());
-	UI_COMMAND(Delete, "Delete Goal or Effector", "Delete the selected IK Goal or Effector.", EUserInterfaceActionType::Button, FInputChord(EKeys::Delete));
+	UI_COMMAND(DeleteGoal, "Delete Goal or Effector", "Delete the selected IK Goal or Effector.", EUserInterfaceActionType::Button, FInputChord(EKeys::Delete));
 	UI_COMMAND(ConnectGoalToSolvers, "Connect Goal to Selected Solvers", "Make the selected goal an effector in the selected solvers.", EUserInterfaceActionType::Button, FInputChord());
 	UI_COMMAND(DisconnectGoalFromSolvers, "Disconnect Goal from Selected Solvers", "Remove effectors from the selected solvers that use this goal.", EUserInterfaceActionType::Button, FInputChord());
 	UI_COMMAND(SetRootBoneOnSolvers, "Set Root Bone on Selected Solvers", "Set the Root Bone setting on the selected solvers.", EUserInterfaceActionType::Button, FInputChord());
diff --git a/Engine/Plugins/Animation/IKRig/Source/IKRigEditor/Private/RigEditor/IKRigToolkit.cpp b/Engine/Plugins/Animation/IKRig/Source/IKRigEditor/Private/RigEditor/IKRigToolkit.cpp
--- a/Engine/Plugins/Animation/IKRig/Source/IKRigEditor/Private/RigEditor/IKRigToolkit.cpp
+++ b/Engine/Plugins/Animation/IKRig/Source/IKRigEditor/Private/RigEditor/IKRigToolkit.cpp
@@ -76,7 +76,6 @@ void FIKRigEditorToolkit::InitAssetEditor(
 	SetCurrentMode(IKRigEditorModes::IKRigEditorMode);
 
 	GetEditorModeManager().SetDefaultMode(FIKRigEditMode::ModeName);
-	GetEditorModeManager().ActivateMode(FPersonaEditModes::SkeletonSelection);
 	GetEditorModeManager().ActivateMode(FIKRigEditMode::ModeName);
 	static_cast<FIKRigEditMode*>(GetEditorModeManager().GetActiveMode(FIKRigEditMode::ModeName))->SetEditorController(EditorController);
 
@@ -221,7 +220,9 @@ void FIKRigEditorToolkit::HandlePreviewSceneCreated(const TSharedRef<IPersonaPre
 
 	// apply mesh to the preview scene
 	InPersonaPreviewScene->SetPreviewMeshComponent(EditorController->SkelMeshComponent);
+	InPersonaPreviewScene->SetAllowMeshHitProxies(false);
 	InPersonaPreviewScene->SetAdditionalMeshesSelectable(false);
+	EditorController->SkelMeshComponent->bSelectable = false;
 	InPersonaPreviewScene->SetPreviewMesh(Mesh);
 	InPersonaPreviewScene->AddComponent(EditorController->SkelMeshComponent, FTransform::Identity);
 }
diff --git a/Engine/Plugins/Animation/IKRig/Source/IKRigEditor/Private/RigEditor/SIKRigSkeleton.cpp b/Engine/Plugins/Animation/IKRig/Source/IKRigEditor/Private/RigEditor/SIKRigSkeleton.cpp
--- a/Engine/Plugins/Animation/IKRig/Source/IKRigEditor/Private/RigEditor/SIKRigSkeleton.cpp
+++ b/Engine/Plugins/Animation/IKRig/Source/IKRigEditor/Private/RigEditor/SIKRigSkeleton.cpp
@@ -308,9 +308,9 @@ void SIKRigSkeleton::BindCommands()
         FExecuteAction::CreateSP(this, &SIKRigSkeleton::HandleNewGoal),
         FCanExecuteAction::CreateSP(this, &SIKRigSkeleton::CanAddNewGoal));
 	
-	CommandList->MapAction(Commands.Delete,
-        FExecuteAction::CreateSP(this, &SIKRigSkeleton::HandleDeleteItem),
-        FCanExecuteAction::CreateSP(this, &SIKRigSkeleton::CanDeleteItem));
+	CommandList->MapAction(Commands.DeleteGoal,
+        FExecuteAction::CreateSP(this, &SIKRigSkeleton::HandleDeleteGoal),
+        FCanExecuteAction::CreateSP(this, &SIKRigSkeleton::CanDeleteGoal));
 
 	CommandList->MapAction(Commands.ConnectGoalToSolvers,
         FExecuteAction::CreateSP(this, &SIKRigSkeleton::HandleConnectGoalToSolver),
@@ -361,7 +361,7 @@ void SIKRigSkeleton::FillContextMenu(FMenuBuilder& MenuBuilder)
 
 	MenuBuilder.BeginSection("AddRemoveGoals", LOCTEXT("AddRemoveGoalOperations", "Goals"));
 	MenuBuilder.AddMenuEntry(Actions.NewGoal);
-	MenuBuilder.AddMenuEntry(Actions.Delete);
+	MenuBuilder.AddMenuEntry(Actions.DeleteGoal);
 	MenuBuilder.EndSection();
 	
 	MenuBuilder.BeginSection("ConnectGoals", LOCTEXT("ConnectGoalOperations", "Connect Goals To Solvers"));
@@ -394,14 +394,9 @@ void SIKRigSkeleton::HandleNewGoal()
 		return; 
 	}
 
-	// add a default solver if there isn't one already
-	if (Controller->AssetController->GetNumSolvers() == 0)
-	{
-		Controller->PromptToAddSolver();	
-	}
-	
-	// create goal for each selected bone
-	FName LastCreatedGoalName = NAME_None;
+	// get names of selected bones and default goal names for them
+	TArray<FName> GoalNames;
+	TArray<FName> BoneNames;
 	const TArray<TSharedPtr<FIKRigTreeElement>> SelectedItems = TreeView.Get()->GetSelectedItems();
 	for (const TSharedPtr<FIKRigTreeElement>& Item : SelectedItems)
 	{
@@ -410,37 +405,16 @@ void SIKRigSkeleton::HandleNewGoal()
 			continue; // can only add goals to bones
 		}
 
-		UIKRigController* AssetController = Controller->AssetController;
-
 		// build default name for the new goal
 		const FName BoneName = Item->Key;
 		const FName NewGoalName = FName(BoneName.ToString() + "_Goal");
-
-		// create a new goal
-		UIKRigEffectorGoal* NewGoal = AssetController->AddNewGoal(NewGoalName, BoneName);
-		if (NewGoal)
-		{
-			// get selected solvers
-			TArray<TSharedPtr<FSolverStackElement>> SelectedSolvers;
-			Controller->GetSelectedSolvers(SelectedSolvers);
-			// connect the new goal to all the selected solvers
-			for (const TSharedPtr<FSolverStackElement>& SolverElement : SelectedSolvers)
-			{
-				AssetController->ConnectGoalToSolver(*NewGoal, SolverElement->IndexInStack);	
-			}
-
-			LastCreatedGoalName = NewGoalName;
-		}
-	}
-
-	// were any goals created?
-	if (LastCreatedGoalName != NAME_None)
-	{
-		// show last created goal in details view
-		Controller->ShowDetailsForGoal(LastCreatedGoalName);
-		// update all views
-		Controller->RefreshAllViews();
+		
+		GoalNames.Add(NewGoalName);
+		BoneNames.Add(BoneName);
 	}
+	
+	// add new goals
+	Controller->AddNewGoals(GoalNames, BoneNames);
 }
 
 bool SIKRigSkeleton::CanAddNewGoal()
@@ -464,22 +438,21 @@ bool SIKRigSkeleton::CanAddNewGoal()
 	return true;
 }
 
-void SIKRigSkeleton::HandleDeleteItem()
+void SIKRigSkeleton::HandleDeleteGoal()
 {
 	const TSharedPtr<FIKRigEditorController> Controller = EditorController.Pin();
 	if (!Controller.IsValid())
 	{
 		return; 
 	}
-	
+
 	TArray<TSharedPtr<FIKRigTreeElement>> SelectedItems = TreeView->GetSelectedItems();
 	TSharedPtr<FIKRigTreeElement> ParentOfDeletedGoal;
 	for (const TSharedPtr<FIKRigTreeElement>& Item : SelectedItems)
 	{
 		if (Item->ElementType == IKRigTreeElementType::GOAL)
 		{
-			Controller->AssetController->RemoveGoal(Item->Key);
-			
+			Controller->DeleteGoal(Item->Key);
 		}
 		else if (Item->ElementType == IKRigTreeElementType::EFFECTOR)
 		{
@@ -491,7 +464,7 @@ void SIKRigSkeleton::HandleDeleteItem()
 	// update all views
 	Controller->RefreshAllViews();
 }
-bool SIKRigSkeleton::CanDeleteItem() const
+bool SIKRigSkeleton::CanDeleteGoal() const
 {
 	// is anything selected?
 	TArray<TSharedPtr<FIKRigTreeElement>> SelectedItems = TreeView->GetSelectedItems();
@@ -1310,25 +1283,25 @@ FReply SIKRigSkeleton::OnKeyDown(const FGeometry& MyGeometry, const FKeyEvent& I
 		return FReply::Handled();
 	}
 
-	// handle deleting selected goals
+	// handle deleting selected items
 	if (Key == EKeys::Delete)
 	{
 		TArray<TSharedPtr<FIKRigTreeElement>> SelectedItems = TreeView->GetSelectedItems();
 		for (const TSharedPtr<FIKRigTreeElement>& SelectedItem : SelectedItems)
 		{
-			if (SelectedItem->ElementType == IKRigTreeElementType::GOAL)
-			{
-				Controller->AssetController->RemoveGoal(SelectedItem->Key);
-			}
-
-			if (SelectedItem->ElementType == IKRigTreeElementType::EFFECTOR)
-			{
-				Controller->AssetController->DisconnectGoalFromSolver(SelectedItem->EffectorGoalName, SelectedItem->EffectorSolverIndex);
-			}
-			
-			if (SelectedItem->ElementType == IKRigTreeElementType::BONE_SETTINGS)
+			switch(SelectedItem->ElementType)
 			{
-				Controller->AssetController->RemoveBoneSetting(SelectedItem->BoneSettingBoneName, SelectedItem->BoneSettingsSolverIndex);
+				case IKRigTreeElementType::GOAL:
+					Controller->DeleteGoal(SelectedItem->Key);
+					break;
+				case IKRigTreeElementType::EFFECTOR:
+					Controller->AssetController->DisconnectGoalFromSolver(SelectedItem->EffectorGoalName, SelectedItem->EffectorSolverIndex);
+					break;
+				case IKRigTreeElementType::BONE_SETTINGS:
+					Controller->AssetController->RemoveBoneSetting(SelectedItem->BoneSettingBoneName, SelectedItem->BoneSettingsSolverIndex);
+					break;
+				default:
+					checkNoEntry()
 			}
 		}
 
diff --git a/Engine/Plugins/Animation/IKRig/Source/IKRigEditor/Public/RetargetEditor/IKRetargetAnimInstance.h b/Engine/Plugins/Animation/IKRig/Source/IKRigEditor/Public/RetargetEditor/IKRetargetAnimInstance.h
--- a/Engine/Plugins/Animation/IKRig/Source/IKRigEditor/Public/RetargetEditor/IKRetargetAnimInstance.h
+++ b/Engine/Plugins/Animation/IKRig/Source/IKRigEditor/Public/RetargetEditor/IKRetargetAnimInstance.h
@@ -3,7 +3,7 @@
 #pragma once
 
 #include "AnimPreviewInstance.h"
-#include "AnimNodes/AnimNode_IKRetargeter.h"
+#include "AnimNodes/AnimNode_RetargetPoseFromMesh.h"
 
 #include "IKRetargetAnimInstance.generated.h"
 
@@ -26,5 +26,5 @@ protected:
 	virtual FAnimInstanceProxy* CreateAnimInstanceProxy() override;
 
 	UPROPERTY(Transient)
-	FAnimNode_IKRetargeter IKRetargeterNode;
+	FAnimNode_RetargetPoseFromMesh IKRetargeterNode;
 };
diff --git a/Engine/Plugins/Animation/IKRig/Source/IKRigEditor/Public/RetargetEditor/IKRetargetAnimInstanceProxy.h b/Engine/Plugins/Animation/IKRig/Source/IKRigEditor/Public/RetargetEditor/IKRetargetAnimInstanceProxy.h
--- a/Engine/Plugins/Animation/IKRig/Source/IKRigEditor/Public/RetargetEditor/IKRetargetAnimInstanceProxy.h
+++ b/Engine/Plugins/Animation/IKRig/Source/IKRigEditor/Public/RetargetEditor/IKRetargetAnimInstanceProxy.h
@@ -4,7 +4,7 @@
 
 #include "AnimPreviewInstance.h"
 #include "Animation/AnimNode_LinkedInputPose.h"
-#include "AnimNodes/AnimNode_IKRetargeter.h"
+#include "AnimNodes/AnimNode_RetargetPoseFromMesh.h"
 
 #include "IKRetargetAnimInstanceProxy.generated.h"
 
@@ -17,7 +17,7 @@ struct FIKRetargetAnimInstanceProxy : public FAnimPreviewInstanceProxy
 public:
 	
 	FIKRetargetAnimInstanceProxy() = default;
-	FIKRetargetAnimInstanceProxy(UAnimInstance* InAnimInstance, FAnimNode_IKRetargeter* IKRetargeterNode);
+	FIKRetargetAnimInstanceProxy(UAnimInstance* InAnimInstance, FAnimNode_RetargetPoseFromMesh* IKRetargeterNode);
 	virtual ~FIKRetargetAnimInstanceProxy() override = default;
 
 	/** FAnimPreviewInstanceProxy interface */
@@ -36,5 +36,5 @@ public:
 		UIKRetargeter* InIKRetargetAsset,
 		TWeakObjectPtr<USkeletalMeshComponent> InSourceMeshComponent) const;
 
-	FAnimNode_IKRetargeter* IKRetargetNode;
+	FAnimNode_RetargetPoseFromMesh* IKRetargetNode;
 };
\ No newline at end of file
diff --git a/Engine/Plugins/Animation/IKRig/Source/IKRigEditor/Public/RetargetEditor/IKRetargetAssetBrowserTabSummoner.h b/Engine/Plugins/Animation/IKRig/Source/IKRigEditor/Public/RetargetEditor/IKRetargetAssetBrowserTabSummoner.h
new file mode 100644
--- /dev/null
+++ b/Engine/Plugins/Animation/IKRig/Source/IKRigEditor/Public/RetargetEditor/IKRetargetAssetBrowserTabSummoner.h
@@ -0,0 +1,21 @@
+﻿// Copyright Epic Games, Inc. All Rights Reserved.
+
+#pragma once
+
+#include "CoreMinimal.h"
+#include "WorkflowOrientedApp/WorkflowTabFactory.h"
+
+class FIKRetargetEditor;
+
+struct FIKRetargetAssetBrowserTabSummoner : public FWorkflowTabFactory
+{
+public:
+	static const FName TabID;
+	
+	FIKRetargetAssetBrowserTabSummoner(const TSharedRef<FIKRetargetEditor>& InIKRetargetEditor);
+	virtual TSharedRef<SWidget> CreateTabBody(const FWorkflowTabSpawnInfo& Info) const override;
+	virtual TSharedPtr<SToolTip> CreateTabToolTipWidget(const FWorkflowTabSpawnInfo& Info) const override;
+
+protected:
+	TWeakPtr<FIKRetargetEditor> IKRetargetEditor;
+};
diff --git a/Engine/Plugins/Animation/IKRig/Source/IKRigEditor/Public/RetargetEditor/IKRetargetBatchOperation.h b/Engine/Plugins/Animation/IKRig/Source/IKRigEditor/Public/RetargetEditor/IKRetargetBatchOperation.h
new file mode 100644
--- /dev/null
+++ b/Engine/Plugins/Animation/IKRig/Source/IKRigEditor/Public/RetargetEditor/IKRetargetBatchOperation.h
@@ -0,0 +1,116 @@
+﻿// Copyright Epic Games, Inc. All Rights Reserved.
+
+#pragma once
+
+#include "CoreMinimal.h"
+#include "EditorAnimUtils.h"
+
+class UIKRetargeter;
+
+//** Data needed to run a batch "duplicate and retarget" operation on a set of animation assets */
+struct FIKRetargetBatchOperationContext
+{
+	
+public:
+	
+	/** The source assets to duplicate and retarget */
+	TArray<TWeakObjectPtr<UObject>> AssetsToRetarget;
+
+	/** Source mesh to use to copy animation FROM. */
+	USkeletalMesh* SourceMesh = nullptr;
+
+	/** Target mesh to use to copy animation TO. */
+	USkeletalMesh* TargetMesh = nullptr;
+
+	/** The retargeter used to copy animation */
+	UIKRetargeter* IKRetargetAsset = nullptr;
+
+	/** Whether we are remapping assets that are referenced by the assets the user selects to remap */
+	bool bRemapReferencedAssets = true;
+
+	/* Rename rules for duplicated assets */
+	EditorAnimUtils::FNameDuplicationRule NameRule;
+
+	/* The base folder path to put newly duplicated assets */
+	FString FolderPath = "";
+
+	/* Reset all data (called when window re-opened */
+	void Reset()
+	{
+		SourceMesh = nullptr;
+		TargetMesh = nullptr;
+		IKRetargetAsset = nullptr;
+		bRemapReferencedAssets = true;
+		NameRule.Prefix = "";
+		NameRule.Suffix = "";
+		NameRule.ReplaceFrom = "";
+		NameRule.ReplaceTo = "";
+	}
+
+	/* Is the data configured in such a way that we could run the retarget? */
+	bool IsValid() const
+	{
+		// todo validate compatibility
+		return SourceMesh && TargetMesh && IKRetargetAsset && (SourceMesh != TargetMesh);
+	}
+};
+
+//** Encapsulate ability to batch duplicate and retarget a set of animation assets */
+struct FIKRetargetBatchOperation
+{
+	
+public:
+
+	/* Actually run the process to duplicate and retarget the assets for the given context */
+	void RunRetarget(FIKRetargetBatchOperationContext& Context);
+
+private:
+
+	/**
+	* Initialize set of referenced assets to retarget.
+	* @return	Number of assets that need retargeting.
+	*/
+	int32 GenerateAssetLists(const FIKRetargetBatchOperationContext& Context);
+
+	/* Duplicate all the assets to retarget */
+	void DuplicateRetargetAssets(const FIKRetargetBatchOperationContext& Context, FScopedSlowTask& Progress);
+
+	/* Retarget skeleton and animation on all the duplicates */
+	void RetargetAssets(const FIKRetargetBatchOperationContext& Context, FScopedSlowTask& Progress);
+
+	/* Convert animation on all the duplicates */
+	void ConvertAnimation(const FIKRetargetBatchOperationContext& Context, FScopedSlowTask& Progress);
+
+	/* Output notifications of results */
+	void NotifyUserOfResults(const FIKRetargetBatchOperationContext& Context, FScopedSlowTask& Progress) const;
+
+	/* Generate list of newly created assets to report to user */
+	void GetNewAssets(TArray<UObject*>& NewAssets) const;
+
+	/**
+	* Duplicates the supplied AssetsToDuplicate and returns a map of original asset to duplicate. Templated wrapper that calls DuplicateAssetInternal.
+	*
+	* @param	AssetsToDuplicate	The animations to duplicate
+	* @param	DestinationPackage	The package that the duplicates should be placed in
+	*
+	* @return	TMap of original animation to duplicate
+	*/
+	template<class AssetType>
+	static TMap<AssetType*, AssetType*> DuplicateAssets(
+		const TArray<AssetType*>& AssetsToDuplicate,
+		UPackage* DestinationPackage,
+		const EditorAnimUtils::FNameDuplicationRule* NameRule);
+	
+	/** Lists of assets to retarget. Populated from selection during init */
+	TArray<UAnimationAsset*>	AnimationAssetsToRetarget;
+	TArray<UAnimBlueprint*>		AnimBlueprintsToRetarget;
+
+	/** Lists of original assets map to duplicate assets */
+	TMap<UAnimationAsset*, UAnimationAsset*>	DuplicatedAnimAssets;
+	TMap<UAnimBlueprint*, UAnimBlueprint*>		DuplicatedBlueprints;
+
+	TMap<UAnimationAsset*, UAnimationAsset*>	RemappedAnimAssets;
+
+	/** If we only chose one object to retarget store it here */
+	UObject* SingleTargetObject = nullptr;
+};
diff --git a/Engine/Plugins/Animation/IKRig/Source/IKRigEditor/Public/RetargetEditor/IKRetargetEditMode.h b/Engine/Plugins/Animation/IKRig/Source/IKRigEditor/Public/RetargetEditor/IKRetargetEditMode.h
--- a/Engine/Plugins/Animation/IKRig/Source/IKRigEditor/Public/RetargetEditor/IKRetargetEditMode.h
+++ b/Engine/Plugins/Animation/IKRig/Source/IKRigEditor/Public/RetargetEditor/IKRetargetEditMode.h
@@ -53,11 +53,11 @@ private:
 	/** The hosting app */
 	TWeakPtr<FIKRetargetEditorController> EditorController;
 
-	void DrawBoneGizmo(
+	void DrawBoneProxy(
 		FPrimitiveDrawInterface* PDI,
 		const FTransform& BoneTransform,
 		const FVector& Start,
-		const FVector& End,
+		const TArray<FVector>& ChildPoints,
 		float Size,
 		float Thickness,
 		bool bIsSelected) const;
diff --git a/Engine/Plugins/Animation/IKRig/Source/IKRigEditor/Public/RetargetEditor/IKRetargetEditor.h b/Engine/Plugins/Animation/IKRig/Source/IKRigEditor/Public/RetargetEditor/IKRetargetEditor.h
--- a/Engine/Plugins/Animation/IKRig/Source/IKRigEditor/Public/RetargetEditor/IKRetargetEditor.h
+++ b/Engine/Plugins/Animation/IKRig/Source/IKRigEditor/Public/RetargetEditor/IKRetargetEditor.h
@@ -84,16 +84,10 @@ private:
 	void SetupAnimInstance();
 	/** END preview scene setup */
 
-	/** animation asset browser */
-	void HandleOpenNewAsset(UObject* InNewAsset);
-	void HandleAnimationSequenceBrowserCreated(const TSharedRef<IAnimationSequenceBrowser>& InSequenceBrowser);
-	/* END animation asset browser */
-
 	/** edit retarget pose */
 	void HandleEditPose() const;
 	bool CanEditPose() const;
 	bool IsEditingPose() const;
-	UAnimationAsset* PreviousAsset = nullptr;
 	/** END edit retarget pose*/
 
 	/** new/delete retarget pose */
diff --git a/Engine/Plugins/Animation/IKRig/Source/IKRigEditor/Public/RetargetEditor/IKRetargetEditorController.h b/Engine/Plugins/Animation/IKRig/Source/IKRigEditor/Public/RetargetEditor/IKRetargetEditorController.h
--- a/Engine/Plugins/Animation/IKRig/Source/IKRigEditor/Public/RetargetEditor/IKRetargetEditorController.h
+++ b/Engine/Plugins/Animation/IKRig/Source/IKRigEditor/Public/RetargetEditor/IKRetargetEditorController.h
@@ -36,9 +36,6 @@ public:
 	/** Preview scene to be supplied by IHasPersonaToolkit::GetPersonaToolkit */
 	TSharedPtr<IPersonaToolkit> PersonaToolkit;
 
-	/** Sequence Browser **/
-	TWeakPtr<class IAnimationSequenceBrowser> SequenceBrowser;
-
 	/** the persona toolkit */
 	TWeakPtr<FIKRetargetEditor> Editor;
 
@@ -47,11 +44,19 @@ public:
 	USkeletalMesh* GetTargetSkeletalMesh() const;
 
 	/** get current chain pose */
-	FTransform GetTargetBoneTransform(const FName& BoneName) const;
-	void GetTargetBoneStartAndEnd(const FName& BoneName, FVector& Start, FVector& End) const;
+	FTransform GetTargetBoneTransform(const int32& TargetBoneIndex) const;
+	/** get the line segments to draw from this bone to each child */
+	bool GetTargetBoneLineSegments(const int32& TargetBoneIndex, FVector& OutStart, TArray<FVector>& OutChildren) const;
+	/** get if the target bone is being retargeted or not */
+	bool IsTargetBoneRetargeted(const int32& TargetBoneIndex);
 
 	/** get the retargeter that is running in the viewport (which is a duplicate of the source asset) */
 	UIKRetargeter* GetCurrentlyRunningRetargeter() const;
 
+	/** Sequence Browser and Edit Pose mode **/
+	void PlayAnimationAsset(UAnimationAsset* AssetToPlay);
+	void PlayPreviousAnimationAsset();
+	UAnimationAsset* PreviousAsset = nullptr;
+
 	void RefreshAllViews();
-};
\ No newline at end of file
+};
diff --git a/Engine/Plugins/Animation/IKRig/Source/IKRigEditor/Public/RetargetEditor/SIKRetargetAssetBrowser.h b/Engine/Plugins/Animation/IKRig/Source/IKRigEditor/Public/RetargetEditor/SIKRetargetAssetBrowser.h
new file mode 100644
--- /dev/null
+++ b/Engine/Plugins/Animation/IKRig/Source/IKRigEditor/Public/RetargetEditor/SIKRetargetAssetBrowser.h
@@ -0,0 +1,41 @@
+﻿// Copyright Epic Games, Inc. All Rights Reserved.
+
+#pragma once
+
+#include "CoreMinimal.h"
+
+#include "ContentBrowserDelegates.h"
+#include "Widgets/Layout/SBox.h"
+
+
+class FIKRetargetEditorController;
+
+class SIKRetargetAssetBrowser : public SBox
+{
+public:
+	SLATE_BEGIN_ARGS(SIKRetargetAssetBrowser) {}
+	SLATE_END_ARGS()
+
+	void Construct(const FArguments& InArgs, TSharedRef<FIKRetargetEditorController> InEditorController);
+
+private:
+
+	void AddAssetBrowser();
+
+	FReply OnExportButtonClicked() const;
+	bool IsExportButtonEnabled() const;
+	
+	void OnAssetDoubleClicked(const FAssetData& AssetData);
+	bool OnShouldFilterAsset(const struct FAssetData& AssetData);
+
+	/** Used to get the currently selected assets */
+	FGetCurrentSelectionDelegate GetCurrentSelectionDelegate;
+		
+	/** editor controller */
+	TWeakPtr<FIKRetargetEditorController> EditorController;
+
+	/** the animation asset browser */
+	TSharedPtr<SBox> AssetBrowserBox;
+
+	friend FIKRetargetEditorController;
+};
diff --git a/Engine/Plugins/Animation/IKRig/Source/IKRigEditor/Public/RetargetEditor/SIKRetargetChainMapList.h b/Engine/Plugins/Animation/IKRig/Source/IKRigEditor/Public/RetargetEditor/SIKRetargetChainMapList.h
--- a/Engine/Plugins/Animation/IKRig/Source/IKRigEditor/Public/RetargetEditor/SIKRetargetChainMapList.h
+++ b/Engine/Plugins/Animation/IKRig/Source/IKRigEditor/Public/RetargetEditor/SIKRetargetChainMapList.h
@@ -79,8 +79,6 @@ public:
 	SLATE_BEGIN_ARGS(SIKRetargetChainMapList) {}
 	SLATE_END_ARGS()
 
-    virtual ~SIKRetargetChainMapList() override;
-
 	void Construct(const FArguments& InArgs, TSharedRef<FIKRetargetEditorController> InEditorController);
 
 private:
diff --git a/Engine/Plugins/Animation/IKRig/Source/IKRigEditor/Public/RetargetEditor/SRetargetAnimAssetsWindow.h b/Engine/Plugins/Animation/IKRig/Source/IKRigEditor/Public/RetargetEditor/SRetargetAnimAssetsWindow.h
--- a/Engine/Plugins/Animation/IKRig/Source/IKRigEditor/Public/RetargetEditor/SRetargetAnimAssetsWindow.h
+++ b/Engine/Plugins/Animation/IKRig/Source/IKRigEditor/Public/RetargetEditor/SRetargetAnimAssetsWindow.h
@@ -3,95 +3,12 @@
 #pragma once
 
 #include "EditorAnimUtils.h"
+#include "IKRetargetBatchOperation.h"
 #include "SEditorViewport.h"
 #include "Settings/SkeletalMeshEditorSettings.h"
 
 class UIKRetargeter;
 
-//** Data needed to initiate a batch duplicate and retarget of animation assets */
-struct FIKRetargetAnimAssetsContext
-{
-	/** The selected assets when batch process was initiated */
-	TArray<TWeakObjectPtr<UObject>> AssetsToRetarget;
-
-	/** Source mesh to use to copy animation FROM. */
-	USkeletalMesh* SourceMesh = nullptr;
-
-	/** Target mesh to use to copy animation TO. */
-	USkeletalMesh* TargetMesh = nullptr;
-
-	/** The retargeter used to copy animation */
-	UIKRetargeter* IKRetargetAsset = nullptr;
-
-	/** Whether we are remapping assets that are referenced by the assets the user selects to remap */
-	bool bRemapReferencedAssets = true;
-
-	/* Rename rules for duplicated assets */
-	EditorAnimUtils::FNameDuplicationRule NameRule;
-
-	/* The base folder path to put newly duplicated assets */
-	FString FolderPath = "";
-
-	/* Reset all data (called when window re-opened */
-	void Reset();
-
-	/* Is the data configured in such a way that we could run the retarget? */
-	bool IsValid() const;
-
-	/* Actually run the process to duplicate and retarget the assets. */
-	void RunRetarget();
-
-private:
-
-
-	/**
-	* Initialize set of referenced assets to retarget.
-	* @return	Number of assets that need retargeting.
-	*/
-	int32 GenerateAssetLists();
-
-	/* Duplicate all the assets to retarget */
-	void DuplicateRetargetAssets();
-
-	/* Retarget skeleton and animation on all the duplicates */
-	void RetargetAssets();
-
-	/* Convert animation on all the duplicates */
-	void ConvertAnimation();
-
-	/* Output notifications of results */
-	void NotifyUserOfResults() const;
-
-	void GetNewAssets(TArray<UObject*>& NewAssets) const;
-
-	/**
-	* Duplicates the supplied AssetsToDuplicate and returns a map of original asset to duplicate. Templated wrapper that calls DuplicateAssetInternal.
-	*
-	* @param	AssetsToDuplicate	The animations to duplicate
-	* @param	DestinationPackage	The package that the duplicates should be placed in
-	*
-	* @return	TMap of original animation to duplicate
-	*/
-	template<class AssetType>
-	static TMap<AssetType*, AssetType*> DuplicateAssets(
-		const TArray<AssetType*>& AssetsToDuplicate,
-		UPackage* DestinationPackage,
-		const EditorAnimUtils::FNameDuplicationRule* NameRule);
-	
-	/** Lists of assets to retarget. Populated from selection during init */
-	TArray<UAnimationAsset*>	AnimationAssetsToRetarget;
-	TArray<UAnimBlueprint*>		AnimBlueprintsToRetarget;
-
-	/** Lists of original assets map to duplicate assets */
-	TMap<UAnimationAsset*, UAnimationAsset*>	DuplicatedAnimAssets;
-	TMap<UAnimBlueprint*, UAnimBlueprint*>		DuplicatedBlueprints;
-
-	TMap<UAnimationAsset*, UAnimationAsset*>	RemappedAnimAssets;
-
-	/** If we only chose one object to retarget store it here */
-	UObject* SingleTargetObject = nullptr;
-};
-
 //** Viewport in retarget dialog that shows source / target retarget pose */
 class SRetargetPoseViewport: public SEditorViewport
 {
@@ -262,7 +179,7 @@ private:
 	FText GetFolderPath() const;
 
 	/** Necessary data collected from UI to run retarget. */
-	FIKRetargetAnimAssetsContext RetargetContext;
+	FIKRetargetBatchOperationContext BatchContext;
 
 	/** The rename rule sample text */
 	FText ExampleText;
diff --git a/Engine/Plugins/Animation/IKRig/Source/IKRigEditor/Public/RigEditor/IKRigEditorController.h b/Engine/Plugins/Animation/IKRig/Source/IKRigEditor/Public/RigEditor/IKRigEditorController.h
--- a/Engine/Plugins/Animation/IKRig/Source/IKRigEditor/Public/RigEditor/IKRigEditorController.h
+++ b/Engine/Plugins/Animation/IKRig/Source/IKRigEditor/Public/RigEditor/IKRigEditorController.h
@@ -21,7 +21,11 @@ class UDebugSkelMeshComponent;
 class FIKRigEditorController
 {
 public:
-	
+
+	/** create goals */
+	void AddNewGoals(const TArray<FName>& GoalNames, const TArray<FName>& BoneNames);
+	/** delete goals */
+	void DeleteGoal(const FName& GoalToDelete);
 	/** get array of the names of currently selected goals */
 	const TArray<FName>& GetSelectedGoals() const {return SelectedGoals;};
 	/** return true if a goal with the given name is selected */
diff --git a/Engine/Plugins/Animation/IKRig/Source/IKRigEditor/Public/RigEditor/IKRigSkeletonCommands.h b/Engine/Plugins/Animation/IKRig/Source/IKRigEditor/Public/RigEditor/IKRigSkeletonCommands.h
--- a/Engine/Plugins/Animation/IKRig/Source/IKRigEditor/Public/RigEditor/IKRigSkeletonCommands.h
+++ b/Engine/Plugins/Animation/IKRig/Source/IKRigEditor/Public/RigEditor/IKRigSkeletonCommands.h
@@ -21,7 +21,7 @@ public:
 	/** create new goal on selected bone */
 	TSharedPtr< FUICommandInfo > NewGoal;
 	/** delete selected goal */
-	TSharedPtr< FUICommandInfo > Delete;
+	TSharedPtr< FUICommandInfo > DeleteGoal;
 	/** add goal to selected solvers */
 	TSharedPtr< FUICommandInfo > ConnectGoalToSolvers;
 	/** remove goal to selected solvers */
diff --git a/Engine/Plugins/Animation/IKRig/Source/IKRigEditor/Public/RigEditor/SIKRigSkeleton.h b/Engine/Plugins/Animation/IKRig/Source/IKRigEditor/Public/RigEditor/SIKRigSkeleton.h
--- a/Engine/Plugins/Animation/IKRig/Source/IKRigEditor/Public/RigEditor/SIKRigSkeleton.h
+++ b/Engine/Plugins/Animation/IKRig/Source/IKRigEditor/Public/RigEditor/SIKRigSkeleton.h
@@ -162,8 +162,8 @@ private:
 	void HandleNewGoal();
 	bool CanAddNewGoal();
 	void HandleRenameElement() const;
-	void HandleDeleteItem();
-	bool CanDeleteItem() const;
+	void HandleDeleteGoal();
+	bool CanDeleteGoal() const;
 	/** END creating / renaming / deleting goals */
 
 	/** connecting/disconnecting goals to solvers */
