IK Rig / Retargeter refactor
- added controller/processor class to retargeter to separate data editing from runtime
- many smaller fixes to IK Rig / Retargeter and Full-body IK

diff --git a/Engine/Plugins/Animation/IKRig/Source/IKRig/IKRig.Build.cs b/Engine/Plugins/Animation/IKRig/Source/IKRig/IKRig.Build.cs
--- a/Engine/Plugins/Animation/IKRig/Source/IKRig/IKRig.Build.cs
+++ b/Engine/Plugins/Animation/IKRig/Source/IKRig/IKRig.Build.cs
@@ -13,33 +13,23 @@ public class IKRig : ModuleRules
 			{
 				"CoreUObject",
 				"Engine",
-				"Slate",
-				"SlateCore",
 				"AnimationCore",
-				"PBIK"
-			}
-			);
-
-		if (Target.bBuildEditor)
-		{
-			PrivateDependencyModuleNames.Add("UnrealEd");
-		}
+				"PBIK",
+			});
 
 		PublicDependencyModuleNames.AddRange(
 			new string[]
 			{
 				"ControlRig",
+				"RigVM",
 				"Core",
-				
 				"PBIK"
-			}
-			);
+			});
 
 		DynamicallyLoadedModuleNames.AddRange(
 			new string[]
 			{
 				// ... add any modules that your module loads dynamically here ...
-			}
-			);
+			});
 	}
 }
diff --git a/Engine/Plugins/Animation/IKRig/Source/IKRig/Private/AnimNodes/AnimNode_IKRig.cpp b/Engine/Plugins/Animation/IKRig/Source/IKRig/Private/AnimNodes/AnimNode_IKRig.cpp
--- a/Engine/Plugins/Animation/IKRig/Source/IKRig/Private/AnimNodes/AnimNode_IKRig.cpp
+++ b/Engine/Plugins/Animation/IKRig/Source/IKRig/Private/AnimNodes/AnimNode_IKRig.cpp
@@ -9,15 +9,6 @@
 #include "Animation/AnimInstanceProxy.h"
 
 
-FAnimNode_IKRig::FAnimNode_IKRig()
-	:	RigDefinitionAsset (nullptr),
-#if WITH_EDITORONLY_DATA
-		bEnableDebugDraw(false),
-#endif
-		IKRigProcessor(nullptr)
-{
-}
-
 void FAnimNode_IKRig::Evaluate_AnyThread(FPoseContext& Output) 
 {
 	DECLARE_SCOPE_HIERARCHICAL_COUNTER_FUNC()
@@ -36,7 +27,7 @@ void FAnimNode_IKRig::Evaluate_AnyThread(FPoseContext& Output)
 		return;
 	}
 
-	if (IKRigProcessor->NeedsInitialized(RigDefinitionAsset))
+	if (!IKRigProcessor->IsInitialized())
 	{
 		return;
 	}
@@ -88,11 +79,13 @@ void FAnimNode_IKRig::AssignGoalTargets()
 
 	// use the goal transforms from the source asset itself
 	// this is used to live preview results from the IK Rig editor
+	#if WITH_EDITOR
 	if (bDriveWithSourceAsset)
 	{
 		IKRigProcessor->CopyAllInputsFromSourceAssetAtRuntime(RigDefinitionAsset);
 		return;
 	}
+	#endif
 	
 	// copy transforms from this anim node's goal pins from blueprint
 	for (const FIKRigGoal& Goal : Goals)
@@ -169,13 +162,18 @@ void FAnimNode_IKRig::Update_AnyThread(const FAnimationUpdateContext& Context)
 
 void FAnimNode_IKRig::PreUpdate(const UAnimInstance* InAnimInstance)
 {
+	if (!IsValid(RigDefinitionAsset))
+	{
+		return;
+	}
+	
 	if (!IsValid(IKRigProcessor))
 	{
 		IKRigProcessor = NewObject<UIKRigProcessor>(InAnimInstance->GetOwningComponent());	
 	}
 	
 	// initialize the IK Rig (will only try once on the current version of the rig asset)
-	if (IKRigProcessor->NeedsInitialized(RigDefinitionAsset))
+	if (!IKRigProcessor->IsInitialized())
 	{
 		const FReferenceSkeleton& RefSkeleton = InAnimInstance->GetSkelMeshComponent()->SkeletalMesh->GetRefSkeleton();
 		IKRigProcessor->Initialize(RigDefinitionAsset, RefSkeleton);
@@ -208,6 +206,14 @@ void FAnimNode_IKRig::PreUpdate(const UAnimInstance* InAnimInstance)
 	}
 }
 
+void FAnimNode_IKRig::SetProcessorNeedsInitialized()
+{
+	if (IKRigProcessor)
+	{
+		IKRigProcessor->SetNeedsInitialized();
+	}
+}
+
 void FAnimNode_IKRig::CacheBones_AnyThread(const FAnimationCacheBonesContext& Context)
 {
 	FAnimNode_Base::CacheBones_AnyThread(Context);
@@ -247,7 +253,7 @@ void FAnimNode_IKRig::ConditionalDebugDraw(
 	FPrimitiveDrawInterface* PDI,
 	USkeletalMeshComponent* PreviewSkelMeshComp) const
 {
-#if WITH_EDITORONLY_DATA
+#if WITH_EDITOR
 
 	// is anim graph setup?
 	if (!(bEnableDebugDraw && PreviewSkelMeshComp && PreviewSkelMeshComp->GetWorld()))
@@ -256,15 +262,14 @@ void FAnimNode_IKRig::ConditionalDebugDraw(
 	}
 
 	// is node setup?
-	if (!(RigDefinitionAsset && IKRigProcessor && !IKRigProcessor->NeedsInitialized(RigDefinitionAsset)))
+	if (!(RigDefinitionAsset && IKRigProcessor && IKRigProcessor->IsInitialized()))
 	{
 		return;
 	}
 
-	const FIKRigGoalContainer& ProcessorGoals = IKRigProcessor->GetGoalContainer();
-	for (const TPair<FName, FIKRigGoal>& GoalPair : ProcessorGoals.Goals)
+	const TArray<FIKRigGoal>& ProcessorGoals = IKRigProcessor->GetGoalContainer().GetGoalArray();
+	for (const FIKRigGoal& Goal : ProcessorGoals)
 	{
-		const FIKRigGoal& Goal = GoalPair.Value;
 		DrawOrientedWireBox(PDI, Goal.FinalBlendedPosition, FVector::XAxisVector, FVector::YAxisVector, FVector::ZAxisVector, FVector::One() * DebugScale, FLinearColor::Yellow, SDPG_World);
 		DrawCoordinateSystem(PDI, Goal.Position, Goal.FinalBlendedRotation.Rotator(), DebugScale, SDPG_World);
 	}
diff --git a/Engine/Plugins/Animation/IKRig/Source/IKRig/Private/AnimNodes/AnimNode_RetargetPoseFromMesh.cpp b/Engine/Plugins/Animation/IKRig/Source/IKRig/Private/AnimNodes/AnimNode_RetargetPoseFromMesh.cpp
--- a/Engine/Plugins/Animation/IKRig/Source/IKRig/Private/AnimNodes/AnimNode_RetargetPoseFromMesh.cpp
+++ b/Engine/Plugins/Animation/IKRig/Source/IKRig/Private/AnimNodes/AnimNode_RetargetPoseFromMesh.cpp
@@ -30,14 +30,28 @@ void FAnimNode_RetargetPoseFromMesh::Evaluate_AnyThread(FPoseContext& Output)
 {
 	DECLARE_SCOPE_HIERARCHICAL_COUNTER_ANIMNODE(Evaluate_AnyThread)
 
-	if (!bIsInitialized)
+	if (!(IKRetargeterAsset && Processor))
+	{
+		return;
+	}
+	
+	if (!Processor->IsInitialized())
 	{
 		Output.ResetToRefPose();
 		return;
 	}
 
+#if WITH_EDITOR
+	// live preview IK Rig solver settings in the retarget, editor only
+	// NOTE: this copies goal targets as well, but these are overwritten by IK chain goals
+	if (bDriveTargetIKRigWithAsset)
+	{
+		Processor->CopyTargetIKRigSettingsFromAsset();
+	}
+#endif
+
 	// run the retargeter
-	const TArray<FTransform>& RetargetedPose = CurrentlyUsedRetargeter->RunRetargeter(SourceMeshComponentSpaceBoneTransforms);
+	const TArray<FTransform>& RetargetedPose = Processor->RunRetargeter(SourceMeshComponentSpaceBoneTransforms);
 
 	// copy pose back
 	FCSPose<FCompactPose> ComponentPose;
@@ -54,16 +68,37 @@ void FAnimNode_RetargetPoseFromMesh::Evaluate_AnyThread(FPoseContext& Output)
 void FAnimNode_RetargetPoseFromMesh::PreUpdate(const UAnimInstance* InAnimInstance)
 {
 	DECLARE_SCOPE_HIERARCHICAL_COUNTER_ANIMNODE(PreUpdate)
+	
+	if (!IsValid(IKRetargeterAsset))
+	{
+		return;
+	}
+	
+	if (!IsValid(Processor))
+	{
+		Processor = NewObject<UIKRetargetProcessor>(InAnimInstance->GetOwningComponent());	
+	}
+	
 	EnsureInitialized(InAnimInstance);
-	if (bIsInitialized)
+	if (Processor->IsInitialized())
 	{
 		CopyBoneTransformsFromSource(InAnimInstance->GetSkelMeshComponent());
 	}
 }
 
-UIKRetargeter* FAnimNode_RetargetPoseFromMesh::GetCurrentlyUsedRetargeter() const
+#if WITH_EDITOR
+void FAnimNode_RetargetPoseFromMesh::SetProcessorNeedsInitialized()
+{
+	if (Processor)
+	{
+		Processor->SetNeedsInitialized();
+	}
+}
+#endif
+
+const UIKRetargetProcessor* FAnimNode_RetargetPoseFromMesh::GetRetargetProcessor() const
 {
-	return CurrentlyUsedRetargeter.Get();
+	return Processor;
 }
 
 void FAnimNode_RetargetPoseFromMesh::EnsureInitialized(const UAnimInstance* InAnimInstance)
@@ -71,7 +106,6 @@ void FAnimNode_RetargetPoseFromMesh::EnsureInitialized(const UAnimInstance* InAn
 	// has user supplied a retargeter asset?
 	if (!IKRetargeterAsset)
 	{
-		bIsInitialized = false;
 		return;
 	}
 
@@ -89,75 +123,39 @@ void FAnimNode_RetargetPoseFromMesh::EnsureInitialized(const UAnimInstance* InAn
 	// has a source mesh been plugged in or found?
 	if (!SourceMeshComponent.IsValid())
 	{
-		bIsInitialized = false;
 		return; // can't do anything if we don't have a source mesh
 	}
 
-	// has the source asset been updated? this can happen in the editor when changing source asset
-	if (!CurrentlyUsedRetargeter
-		|| (CurrentlyUsedRetargeter->AssetVersion != IKRetargeterAsset->AssetVersion))
-	{
-		InitializeRetargetData(InAnimInstance);
-		return;	
-	}
-}
-
-void FAnimNode_RetargetPoseFromMesh::InitializeRetargetData(const UAnimInstance* InAnimInstance)
-{
-	// assume we fail until we don't
-	bIsInitialized = false;
-	
-	// don't try to reinitialize with same asset version
-	if (LastVersionTried == IKRetargeterAsset->AssetVersion)
-	{
-		return;
-	}
-	LastVersionTried = IKRetargeterAsset->AssetVersion;
-
 	// store all the components that were used to initialize
 	// if in future updates, any of this are mismatched, we have to re-initialize
-	CurrentlyUsedSourceMeshComponent = SourceMeshComponent;
 	CurrentlyUsedSourceMesh = SourceMeshComponent->SkeletalMesh;
 	CurrentlyUsedTargetMesh = InAnimInstance->GetSkelMeshComponent()->SkeletalMesh;
-	CurrentlyUsedRetargeter = DuplicateObject(IKRetargeterAsset, const_cast<UAnimInstance*>(InAnimInstance));
-	CurrentlyUsedRetargeter->AssetVersion = IKRetargeterAsset->AssetVersion;
-	CurrentlyUsedSourceIKRig = CurrentlyUsedRetargeter->SourceIKRigAsset;
-	CurrentlyUsedTargetIKRig = CurrentlyUsedRetargeter->TargetIKRigAsset;
-
 	const bool bMeshesAreValid = CurrentlyUsedSourceMesh.IsValid() && CurrentlyUsedTargetMesh.IsValid();
 	if (!bMeshesAreValid)
 	{
 		return; // cannot initialize if components are missing skeletal mesh references
 	}
 
-	const bool bRetargeterIsValid =
-	   CurrentlyUsedRetargeter
-	&& CurrentlyUsedSourceIKRig.IsValid()
-	&& CurrentlyUsedTargetIKRig.IsValid();
-	if (!bRetargeterIsValid)
+	// try initializing the processor
+	if (!Processor->IsInitialized())
 	{
-		return; // cannot initialize unless we have a retargeter with BOTH source AND target IK Rigs
+		// initialize retarget processor with source and target skeletal meshes
+		// (anim instance is passed in as outer UObject for new UIKRigProcessor) 
+		Processor->Initialize(
+			CurrentlyUsedSourceMesh.Get(),
+			CurrentlyUsedTargetMesh.Get(),
+			IKRetargeterAsset);
 	}
-
-	// initialize retargeter with source and target skeletal meshes
-	// (anim instance is passed in as outer UObject for new UIKRigProcessor) 
-	CurrentlyUsedRetargeter->Initialize(
-		CurrentlyUsedSourceMesh.Get(),
-		CurrentlyUsedTargetMesh.Get(),
-		const_cast<UAnimInstance*>(InAnimInstance));
-	
-	// made it!
-	bIsInitialized = CurrentlyUsedRetargeter->bIsLoadedAndValid;
 }
 
 void FAnimNode_RetargetPoseFromMesh::CopyBoneTransformsFromSource(USkeletalMeshComponent* TargetMeshComponent)
 {
-	if (!CurrentlyUsedSourceMeshComponent.IsValid())
+	if (!SourceMeshComponent.IsValid())
 	{
 		return; 
 	}
 
-	USkeletalMeshComponent* SourceMeshComp =  CurrentlyUsedSourceMeshComponent.Get();
+	USkeletalMeshComponent* SourceMeshComp =  SourceMeshComponent.Get();
 	
 	// is the source mesh ticking?
 	if (!SourceMeshComp->IsRegistered())
@@ -179,7 +177,10 @@ void FAnimNode_RetargetPoseFromMesh::CopyBoneTransformsFromSource(USkeletalMeshC
 		return; // master pose either missing skeletal mesh reference or not ticking, either way, we aren't copying from it
 	}
 	
-	const bool bUROInSync = SourceMeshComp->ShouldUseUpdateRateOptimizations() && SourceMeshComp->AnimUpdateRateParams != nullptr && SourceMeshComponent->AnimUpdateRateParams == TargetMeshComponent->AnimUpdateRateParams;
+	const bool bUROInSync =
+		SourceMeshComp->ShouldUseUpdateRateOptimizations() &&
+		SourceMeshComp->AnimUpdateRateParams != nullptr &&
+		SourceMeshComponent->AnimUpdateRateParams == TargetMeshComponent->AnimUpdateRateParams;
 	const bool bUsingExternalInterpolation = SourceMeshComp->IsUsingExternalInterpolation();
 	const TArray<FTransform>& CachedComponentSpaceTransforms = SourceMeshComp->GetCachedComponentSpaceTransforms();
 	const bool bArraySizesMatch = CachedComponentSpaceTransforms.Num() == SourceMeshComp->GetComponentSpaceTransforms().Num();
@@ -194,12 +195,4 @@ void FAnimNode_RetargetPoseFromMesh::CopyBoneTransformsFromSource(USkeletalMeshC
 	{
 		SourceMeshComponentSpaceBoneTransforms.Append(SourceMeshComp->GetComponentSpaceTransforms()); // copy directly
 	}
-
-	// ref skeleton is need for parent index lookups later, so store it now
-	CurrentlyUsedSourceMesh = SourceMeshComp->SkeletalMesh;
-}
-
-USkeletalMeshComponent* FAnimNode_RetargetPoseFromMesh::GetSourceMesh() const
-{
-	return nullptr;
 }
diff --git a/Engine/Plugins/Animation/IKRig/Source/IKRig/Private/IKRigDataTypes.cpp b/Engine/Plugins/Animation/IKRig/Source/IKRig/Private/IKRigDataTypes.cpp
--- a/Engine/Plugins/Animation/IKRig/Source/IKRig/Private/IKRigDataTypes.cpp
+++ b/Engine/Plugins/Animation/IKRig/Source/IKRig/Private/IKRigDataTypes.cpp
@@ -1,57 +1,58 @@
 // Copyright Epic Games, Inc. All Rights Reserved.
 
 #include "IKRigDataTypes.h"
+#include "IKRigDefinition.h"
 
-void FIKRigGoalContainer::InitializeFromGoals(const TArray<UIKRigEffectorGoal*>& InGoals)
+void FIKRigGoalContainer::SetIKGoal(const FIKRigGoal& InGoal)
 {
-	Goals.Empty(InGoals.Num());
-	for (const UIKRigEffectorGoal* Goal : InGoals)
+	FIKRigGoal* Goal = const_cast<FIKRigGoal*>(FindGoalByName(InGoal.Name));
+	if (!Goal)
 	{
-		Goals.Emplace(Goal->GoalName, Goal);
+		// container hasn't seen this goal before, create new one, copying the input goal
+		Goals.Emplace(InGoal);
+		return;
 	}
-}
 
-void FIKRigGoalContainer::SetIKGoal(const FIKRigGoal& InGoal)
-{
-	Goals.Add(InGoal.Name, InGoal);
+	// copy settings to existing goal
+	*Goal = InGoal;
 }
 
 void FIKRigGoalContainer::SetIKGoal(const UIKRigEffectorGoal* InEffectorGoal)
 {
-	if (FIKRigGoal* Goal = Goals.Find(InEffectorGoal->GoalName))
+	FIKRigGoal* Goal = const_cast<FIKRigGoal*>(FindGoalByName(InEffectorGoal->GoalName));
+	if (!Goal)
 	{
-		if (InEffectorGoal->PreviewMode == EIKRigGoalPreviewMode::Absolute)
-		{
-			Goal->Position = InEffectorGoal->CurrentTransform.GetTranslation();
-			Goal->Rotation = InEffectorGoal->CurrentTransform.Rotator();
-			Goal->PositionSpace = EIKRigGoalSpace::Component;
-			Goal->RotationSpace = EIKRigGoalSpace::Component;
-		}
-		else
-		{
-			Goal->Position = InEffectorGoal->CurrentTransform.GetTranslation() - InEffectorGoal->InitialTransform.GetTranslation();
-			const FQuat RelativeRotation = InEffectorGoal->CurrentTransform.GetRotation() * InEffectorGoal->InitialTransform.GetRotation().Inverse();
-			Goal->Rotation = RelativeRotation.Rotator();
-			Goal->PositionSpace = EIKRigGoalSpace::Additive;
-			Goal->RotationSpace = EIKRigGoalSpace::Additive;
-		}
-		
-        Goal->PositionAlpha = InEffectorGoal->PositionAlpha;
-        Goal->RotationAlpha = InEffectorGoal->RotationAlpha;
+		// container hasn't seen this goal before, create new one, copying the Effector goal
+		Goals.Emplace(InEffectorGoal);
+		return;
 	}
-	else
+
+	// goals in editor have "preview mode" which allows them to be specified relative to their
+	// initial pose to simulate an additive runtime behavior
+#if WITH_EDITOR
+	if (InEffectorGoal->PreviewMode == EIKRigGoalPreviewMode::Additive)
 	{
-		Goals.Emplace(InEffectorGoal->GoalName, InEffectorGoal);
+		// transform to be interpreted as an offset, relative to input pose
+		Goal->Position = InEffectorGoal->CurrentTransform.GetTranslation() - InEffectorGoal->InitialTransform.GetTranslation();
+		const FQuat RelativeRotation = InEffectorGoal->CurrentTransform.GetRotation() * InEffectorGoal->InitialTransform.GetRotation().Inverse();
+		Goal->Rotation = RelativeRotation.Rotator();
+		Goal->PositionSpace = EIKRigGoalSpace::Additive;
+		Goal->RotationSpace = EIKRigGoalSpace::Additive;
+	}else
+#endif
+	{
+		// transform to be interpreted as absolute and in component space
+		Goal->Position = InEffectorGoal->CurrentTransform.GetTranslation();
+		Goal->Rotation = InEffectorGoal->CurrentTransform.Rotator();
+		Goal->PositionSpace = EIKRigGoalSpace::Component;
+		Goal->RotationSpace = EIKRigGoalSpace::Component;
 	}
+	
+    Goal->PositionAlpha = InEffectorGoal->PositionAlpha;
+    Goal->RotationAlpha = InEffectorGoal->RotationAlpha;
 }
 
-bool FIKRigGoalContainer::GetGoalByName(const FName& InGoalName, FIKRigGoal& OutGoal) const
+const FIKRigGoal* FIKRigGoalContainer::FindGoalByName(const FName& GoalName) const
 {
-	if (const FIKRigGoal* Goal = Goals.Find(InGoalName))
-	{
-		OutGoal = *Goal;
-		return true;
-	}
-
-	return false;
-}
\ No newline at end of file
+	return Goals.FindByPredicate([GoalName](const FIKRigGoal& Other)	{return Other.Name == GoalName;});
+}
diff --git a/Engine/Plugins/Animation/IKRig/Source/IKRig/Private/IKRigDefinition.cpp b/Engine/Plugins/Animation/IKRig/Source/IKRig/Private/IKRigDefinition.cpp
--- a/Engine/Plugins/Animation/IKRig/Source/IKRig/Private/IKRigDefinition.cpp
+++ b/Engine/Plugins/Animation/IKRig/Source/IKRig/Private/IKRigDefinition.cpp
@@ -3,13 +3,9 @@
 #include "IKRigDefinition.h"
 #include "Engine/SkeletalMesh.h"
 
-#if WITH_EDITOR
-#include "ScopedTransaction.h"
-#endif
-
 #define LOCTEXT_NAMESPACE	"IKRigDefinition"
 
-FBoneChain* FRetargetDefinition::GetBoneChainByName(FName ChainName)
+FBoneChain* FRetargetDefinition::GetEditableBoneChainByName(FName ChainName)
 {
 	for (FBoneChain& Chain : BoneChains)
 	{
@@ -22,18 +18,21 @@ FBoneChain* FRetargetDefinition::GetBoneChainByName(FName ChainName)
 	return nullptr;
 }
 
-void UIKRigDefinition::SetPreviewMesh(USkeletalMesh* PreviewMesh, bool bMarkAsDirty)
+const FBoneChain* UIKRigDefinition::GetRetargetChainByName(FName ChainName) const
 {
-	#if WITH_EDITOR
-	if (bMarkAsDirty)
+	for (const FBoneChain& Chain : RetargetDefinition.BoneChains)
 	{
-		FScopedTransaction Transaction(LOCTEXT("SetPreviewMesh_Label", "Set Preview Mesh"));
-		Modify();
-		PreviewSkeletalMesh = PreviewMesh;
-		return;
+		if (Chain.ChainName == ChainName)
+		{
+			return &Chain;
+		}
 	}
-	#endif
 	
+	return nullptr;
+}
+
+void UIKRigDefinition::SetPreviewMesh(USkeletalMesh* PreviewMesh, bool bMarkAsDirty)
+{	
 	PreviewSkeletalMesh = PreviewMesh;
 }
 
@@ -42,65 +41,4 @@ USkeletalMesh* UIKRigDefinition::GetPreviewMesh() const
 	return PreviewSkeletalMesh.Get();
 }
 
-#if WITH_EDITOR
-bool UIKRigDefinition::Modify(bool bAlwaysMarkDirty /*=true*/)
-{
-	const bool bSavedToTransactionBuffer = Super::Modify(bAlwaysMarkDirty);
-	AssetVersion++; // inform any runtime/editor systems they should copy-up modifications	
-	return bSavedToTransactionBuffer;
-}
-#endif
-
-void UIKRigDefinition::PostLoad()
-{
-	Super::PostLoad();
-	ResetGoalTransforms();
-}
-
-void UIKRigDefinition::ResetGoalTransforms() const
-{
-	for (UIKRigEffectorGoal* Goal : Goals)
-	{
-		const FTransform InitialTransform =  GetGoalInitialTransform(Goal);
-		Goal->InitialTransform = InitialTransform;
-		Goal->CurrentTransform = InitialTransform;
-	}	
-}
-
-FTransform UIKRigDefinition::GetGoalInitialTransform(UIKRigEffectorGoal* Goal) const
-{
-	if (!Goal)
-	{
-		return FTransform::Identity; // null
-	}
-
-	const int32 BoneIndex = Skeleton.GetBoneIndexFromName(Goal->BoneName);
-	if (BoneIndex == INDEX_NONE)
-	{
-		return FTransform::Identity; // goal references unknown bone
-	}
-
-	return Skeleton.RefPoseGlobal[BoneIndex];
-}
-
-void UIKRigDefinition::SortRetargetChains()
-{
-	RetargetDefinition.BoneChains.Sort([this](const FBoneChain& A, const FBoneChain& B)
-	{
-		const int32 IndexA = this->Skeleton.GetBoneIndexFromName(A.StartBone);
-		const int32 IndexB = this->Skeleton.GetBoneIndexFromName(B.StartBone);
-		return IndexA < IndexB;
-	});
-}
-
-USkeleton* UIKRigDefinition::GetSkeletonAsset() const
-{
-	if (PreviewSkeletalMesh.IsNull())
-	{
-		return nullptr;
-	}
-
-	return PreviewSkeletalMesh->GetSkeleton();
-}
-
 #undef LOCTEXT_NAMESPACE
diff --git a/Engine/Plugins/Animation/IKRig/Source/IKRig/Private/IKRigProcessor.cpp b/Engine/Plugins/Animation/IKRig/Source/IKRig/Private/IKRigProcessor.cpp
--- a/Engine/Plugins/Animation/IKRig/Source/IKRig/Private/IKRigProcessor.cpp
+++ b/Engine/Plugins/Animation/IKRig/Source/IKRig/Private/IKRigProcessor.cpp
@@ -4,41 +4,36 @@
 #include "IKRigDefinition.h"
 #include "IKRigSolver.h"
 
-void UIKRigProcessor::Initialize(UIKRigDefinition* InRigAsset, const FReferenceSkeleton& RefSkeleton)
+void UIKRigProcessor::Initialize(UIKRigDefinition* InRigAsset, const FIKRigInputSkeleton& InputSkeleton)
 {
 	// we instantiate UObjects here which MUST be done on game thread...
 	check(IsInGameThread());
 	check(InRigAsset);
 	
 	bInitialized = false;
-	InitializedWithIKRigAssetVersion = -1;
-
-	if (!InRigAsset)
-	{
-		UE_LOG(LogTemp, Warning, TEXT("Trying to initialize IKRigProcessor with a null IKRigDefinition asset."));
-		return;
-	}
 
 	// bail out if we've already tried initializing with this exact version of the rig asset
-	if (LastVersionTried == InRigAsset->GetAssetVersion())
+	if (bTriedToInitialize)
 	{
 		return; // don't keep spamming
 	}
-	LastVersionTried = InRigAsset->GetAssetVersion();
 
+	// ok, lets try to initialize
+	bTriedToInitialize = true;
+	
 	if (InRigAsset->Skeleton.BoneNames.IsEmpty())
 	{
 		UE_LOG(LogTemp, Warning, TEXT("Trying to initialize IKRigProcessor with a IKRigDefinition that has no skeleton: %s"), *InRigAsset->GetName());
 		return;
 	}
 
-	if (InRigAsset->Solvers.IsEmpty())
+	if (InRigAsset->GetSolverArray().IsEmpty())
 	{
 		UE_LOG(LogTemp, Warning, TEXT("Trying to initialize IKRigProcessor with a IKRigDefinition that has no solvers: %s"), *InRigAsset->GetName());
 		return;
 	}
 
-	if (InRigAsset->Goals.IsEmpty())
+	if (InRigAsset->GetGoalArray().IsEmpty())
 	{
 		UE_LOG(LogTemp, Warning, TEXT("Trying to initialize IKRigProcessor with a IKRigDefinition that has no goals: %s"), *InRigAsset->GetName());
 		return;
@@ -50,17 +45,26 @@ void UIKRigProcessor::Initialize(UIKRigDefinition* InRigAsset, const FReferenceS
 	Skeleton.BoneNames = InRigAsset->Skeleton.BoneNames;
 	Skeleton.ParentIndices = InRigAsset->Skeleton.ParentIndices;
 	Skeleton.ExcludedBones = InRigAsset->Skeleton.ExcludedBones;
-	const bool bSkeletonIsCompatible = Skeleton.CopyPosesFromRefSkeleton(RefSkeleton);
+	const bool bSkeletonIsCompatible = Skeleton.CopyPosesFromInputSkeleton(InputSkeleton);
 	if (!bSkeletonIsCompatible)
 	{
 		UE_LOG(LogTemp, Warning, TEXT("IK Rig, %s trying to run on a skeleton that does not have the required bones."), *InRigAsset->GetName());
 		return;
 	}
 	
-	// initialize goal names based on solvers
-	GoalContainer.InitializeFromGoals(InRigAsset->Goals);
-	for (const UIKRigEffectorGoal* EffectorGoal : InRigAsset->Goals)
+	// initialize goals based on source asset
+	GoalContainer.Empty();
+	const TArray<UIKRigEffectorGoal*>& GoalsInAsset = InRigAsset->GetGoalArray();
+	for (const UIKRigEffectorGoal* GoalInAsset : GoalsInAsset)
 	{
+		// add a copy of the Goal to the container
+		GoalContainer.SetIKGoal(GoalInAsset);
+	}
+	
+	// initialize goal bones from asset
+	GoalBones.Reset();
+	for (const UIKRigEffectorGoal* EffectorGoal : GoalsInAsset)
+	{	
 		FGoalBone NewGoalBone;
 		NewGoalBone.BoneName = EffectorGoal->BoneName;
 		NewGoalBone.BoneIndex = Skeleton.GetBoneIndexFromName(EffectorGoal->BoneName);
@@ -89,9 +93,10 @@ void UIKRigProcessor::Initialize(UIKRigDefinition* InRigAsset, const FReferenceS
 	}
 
 	// create copies of all the solvers in the IK rig
-	Solvers.Reset(InRigAsset->Solvers.Num());
+	const TArray<UIKRigSolver*>& AssetSolvers = InRigAsset->GetSolverArray();
+	Solvers.Reset(AssetSolvers.Num());
 	int32 SolverIndex = 0;
-	for (UIKRigSolver* IKRigSolver : InRigAsset->Solvers)
+	for (const UIKRigSolver* IKRigSolver : AssetSolvers)
 	{
 		if (!IKRigSolver)
 		{
@@ -109,10 +114,16 @@ void UIKRigProcessor::Initialize(UIKRigDefinition* InRigAsset, const FReferenceS
 		Solvers.Add(Solver);
 	}
 
-	InitializedWithIKRigAssetVersion = InRigAsset->GetAssetVersion();
 	bInitialized = true;
 }
 
+void UIKRigProcessor::Initialize(UIKRigDefinition* InRigAsset, const FReferenceSkeleton& RefSkeleton)
+{
+	FIKRigInputSkeleton InputSkeleton;
+	InputSkeleton.InitializeFromRefSkeleton(RefSkeleton);
+	Initialize(InRigAsset, InputSkeleton);
+}
+
 void UIKRigProcessor::SetInputPoseGlobal(const TArray<FTransform>& InGlobalBoneTransforms) 
 {
 	check(bInitialized);
@@ -150,10 +161,14 @@ void UIKRigProcessor::Solve(const FTransform& ComponentToWorld)
 	// run all the solvers
 	for (UIKRigSolver* Solver : Solvers)
 	{
-		if (Solver->bEnabled)
+		#if WITH_EDITOR
+		if (Solver->IsEnabled())
 		{
 			Solver->Solve(Skeleton, GoalContainer);
 		}
+		#else
+		Solver->Solve(Skeleton, GoalContainer);
+		#endif
 	}
 
 	// make sure rotations are normalized coming out
@@ -165,42 +180,45 @@ void UIKRigProcessor::CopyOutputGlobalPoseToArray(TArray<FTransform>& OutputPose
 	OutputPoseGlobal = Skeleton.CurrentPoseGlobal;
 }
 
-void UIKRigProcessor::CopyAllInputsFromSourceAssetAtRuntime(UIKRigDefinition* IKRigAsset)
+void UIKRigProcessor::Reset()
+{
+	Solvers.Reset();
+	GoalContainer.Empty();
+	GoalBones.Reset();
+	Skeleton.Reset();
+	SetNeedsInitialized();
+}
+
+void UIKRigProcessor::SetNeedsInitialized()
 {
+	bInitialized = false;
+	bTriedToInitialize = false;
+};
+
+#if WITH_EDITOR
+
+void UIKRigProcessor::CopyAllInputsFromSourceAssetAtRuntime(UIKRigDefinition* SourceAsset)
+{
+	check(SourceAsset)
+	
 	// copy goal settings
-	for (const UIKRigEffectorGoal* AssetGoal : IKRigAsset->Goals)
+	const TArray<UIKRigEffectorGoal*>& AssetGoals =  SourceAsset->GetGoalArray();
+	for (const UIKRigEffectorGoal* AssetGoal : AssetGoals)
 	{
 		SetIKGoal(AssetGoal);
 	}
 
 	// copy solver settings
-	check(Solvers.Num() == IKRigAsset->Solvers.Num()); // if number of solvers has been changed, processor should have been reinitialized
+	const TArray<UIKRigSolver*>& AssetSolvers = SourceAsset->GetSolverArray();
+	check(Solvers.Num() == AssetSolvers.Num()); // if number of solvers has been changed, processor should have been reinitialized
 	for (int32 SolverIndex=0; SolverIndex<Solvers.Num(); ++SolverIndex)
 	{
-		Solvers[SolverIndex]->bEnabled = IKRigAsset->Solvers[SolverIndex]->bEnabled;
-		Solvers[SolverIndex]->UpdateSolverSettings(IKRigAsset->Solvers[SolverIndex]);
+		Solvers[SolverIndex]->SetEnabled(AssetSolvers[SolverIndex]->IsEnabled());
+		Solvers[SolverIndex]->UpdateSolverSettings(AssetSolvers[SolverIndex]);
 	}
 }
 
-bool UIKRigProcessor::NeedsInitialized(UIKRigDefinition* IKRigAsset) const
-{
-	if (!bInitialized)
-	{
-		return true; // not initialized yet at all
-	}
-
-	if (!IKRigAsset)
-	{
-		return true; // lost connection to asset
-	}
-
-	if (InitializedWithIKRigAssetVersion != IKRigAsset->GetAssetVersion())
-	{
-		return true; // IKRig asset has been modified since last initialization
-	}
-
-	return false;
-}
+#endif
 
 const FIKRigGoalContainer& UIKRigProcessor::GetGoalContainer() const
 {
@@ -216,16 +234,15 @@ FIKRigSkeleton& UIKRigProcessor::GetSkeleton()
 
 void UIKRigProcessor::ResolveFinalGoalTransforms(const FTransform& WorldToComponent)
 {
-	for (TPair<FName, FIKRigGoal>& GoalPair : GoalContainer.Goals)
+	for (FIKRigGoal& Goal : GoalContainer.Goals)
 	{
-		if (!GoalBones.Contains(GoalPair.Key))
+		if (!GoalBones.Contains(Goal.Name))
 		{
 			// user is changing goals after initialization
 			// not necessarily a bad thing, but new goal names won't work until re-init
 			continue;
 		}
-		
-		FIKRigGoal& Goal = GoalPair.Value;
+
 		const FGoalBone& GoalBone = GoalBones[Goal.Name];
 		const FTransform& InputPoseBoneTransform = Skeleton.CurrentPoseGlobal[GoalBone.BoneIndex];
 
diff --git a/Engine/Plugins/Animation/IKRig/Source/IKRig/Private/IKRigSkeleton.cpp b/Engine/Plugins/Animation/IKRig/Source/IKRig/Private/IKRigSkeleton.cpp
--- a/Engine/Plugins/Animation/IKRig/Source/IKRig/Private/IKRigSkeleton.cpp
+++ b/Engine/Plugins/Animation/IKRig/Source/IKRig/Private/IKRigSkeleton.cpp
@@ -3,22 +3,30 @@
 #include "IKRigSkeleton.h"
 
 #include "IKRigDefinition.h"
+#include "IKRigProcessor.h"
 
-void FIKRigSkeleton::Initialize(const FReferenceSkeleton& RefSkeleton, const TArray<FName>& InExcludedBones)
+void FIKRigSkeleton::Initialize(const FIKRigInputSkeleton& InputSkeleton, const TArray<FName>& InExcludedBones)
 {
+	check(InputSkeleton.BoneNames.Num() == InputSkeleton.ParentIndices.Num() &&
+		  InputSkeleton.BoneNames.Num() == InputSkeleton.LocalRefPose.Num())
+	
 	// reset all containers
 	Reset();
 	
-	// copy names and parent indices into local storage
-	TArray<FMeshBoneInfo> RawBoneInfo = RefSkeleton.GetRawRefBoneInfo();
-	for (int32 BoneIndex=0; BoneIndex<RefSkeleton.GetRawBoneNum(); ++BoneIndex)
-	{
-		BoneNames.Add(RawBoneInfo[BoneIndex].Name);
-		ParentIndices.Add(RawBoneInfo[BoneIndex].ParentIndex);	
-	}
+	// copy names, parent indices and excluded state into local storage
+	BoneNames = InputSkeleton.BoneNames;
+	ParentIndices = InputSkeleton.ParentIndices;
+	ExcludedBones = InExcludedBones;
 
 	// copy all the poses out of the ref skeleton
-	CopyPosesFromRefSkeleton(RefSkeleton);
+	CopyPosesFromInputSkeleton(InputSkeleton);
+}
+
+void FIKRigSkeleton::Initialize(const FReferenceSkeleton& RefSkeleton, const TArray<FName>& InExcludedBones)
+{
+	FIKRigInputSkeleton InputSkeleton;
+	InputSkeleton.InitializeFromRefSkeleton(RefSkeleton);
+	Initialize(InputSkeleton, InExcludedBones);
 }
 
 void FIKRigSkeleton::Reset()
@@ -80,19 +88,44 @@ int32 FIKRigSkeleton::GetParentIndexThatIsNotExcluded(const int32 BoneIndex) con
 	return ParentIndex;
 }
 
-bool FIKRigSkeleton::CopyPosesFromRefSkeleton(const FReferenceSkeleton& RefSkeleton)
+bool FIKRigSkeleton::CopyPosesFromInputSkeleton(const FIKRigInputSkeleton& InputSkeleton)
 {
 	// get a compacted local ref pose based on the stored names
 	TArray<FTransform> CompactRefPoseLocal;
-	for (const FName& BoneName : BoneNames)
+
+	// validate all the bones in the input skeleton, comparing against the stored skeleton
+	for (int32 StoredBoneIndex=0; StoredBoneIndex<BoneNames.Num(); ++StoredBoneIndex)
 	{
-		const int32 BoneIndex = RefSkeleton.FindBoneIndex(BoneName);
-		if (BoneIndex == INDEX_NONE)
+		const FName& BoneName = BoneNames[StoredBoneIndex];
+		const int32 InputBoneIndex = InputSkeleton.BoneNames.Find(BoneName);
+
+		// validate that input skeleton has required bone
+		if (InputBoneIndex == INDEX_NONE)
 		{
-			UE_LOG(LogTemp, Warning, TEXT("IK Rig is running on a skeleton that is missing bone named, '%s'."), *BoneName.ToString());
+			UE_LOG(LogTemp, Error, TEXT("IK Rig is running on a skeleton that is missing bone named, '%s'."), *BoneName.ToString());
 			return false;
 		}
-		CompactRefPoseLocal.Add(RefSkeleton.GetRefBonePose()[BoneIndex]);
+
+		// validate that input skeleton hierarchy is as expected
+		const int32 StoredParentIndex = ParentIndices[StoredBoneIndex];
+		if (StoredParentIndex != -1) // root bone has no parent
+		{
+			const FName& StoredParentName = BoneNames[StoredParentIndex];
+			const int32 InputParentIndex = InputSkeleton.ParentIndices[InputBoneIndex];
+			const FName& InputParentName = InputSkeleton.BoneNames[InputParentIndex];
+			if (StoredParentName != InputParentName)
+			{
+				UE_LOG(LogTemp, Error,
+					TEXT("IK Rig is running on a skeleton with a bone, '%s', that has a different parent '%s'. The expected parent was, '%s'."),
+					*BoneName.ToString(),
+					*InputParentName.ToString(),
+					*StoredParentName.ToString());
+				return false;
+			}	
+		}
+
+		// store ref pose of this bone
+		CompactRefPoseLocal.Add(InputSkeleton.LocalRefPose[InputBoneIndex]);
 	}
 	
 	// copy local ref pose to global
@@ -104,7 +137,7 @@ bool FIKRigSkeleton::CopyPosesFromRefSkeleton(const FReferenceSkeleton& RefSkele
 	// initialize current LOCAL pose
 	UpdateAllLocalTransformFromGlobal();
 
-	return true; // supplied ref skeleton had all the bones we are looking for
+	return true; // supplied input skeleton had all the bones we are looking for
 }
 
 void FIKRigSkeleton::ConvertLocalPoseToGlobal(
diff --git a/Engine/Plugins/Animation/IKRig/Source/IKRig/Private/Retargeter/IKRetargetProcessor.cpp b/Engine/Plugins/Animation/IKRig/Source/IKRig/Private/Retargeter/IKRetargetProcessor.cpp
new file mode 100644
--- /dev/null
+++ b/Engine/Plugins/Animation/IKRig/Source/IKRig/Private/Retargeter/IKRetargetProcessor.cpp
@@ -0,0 +1,1100 @@
+﻿// Copyright Epic Games, Inc. All Rights Reserved.
+
+#include "Retargeter/IKRetargetProcessor.h"
+
+#include "IKRigDefinition.h"
+#include "IKRigProcessor.h"
+#include "Engine/SkeletalMesh.h"
+#include "Retargeter/IKRetargeter.h"
+
+
+void FRetargetSkeleton::Initialize(
+	USkeletalMesh* InSkeletalMesh,
+	FIKRetargetPose* InRetargetPose,
+	const FName& RetargetRootBone)
+{	
+	// record which skeletal mesh this is running on
+	SkeletalMesh = InSkeletalMesh;
+	
+	// copy names and parent indices into local storage
+	const FReferenceSkeleton& RefSkeleton = SkeletalMesh->GetRefSkeleton();
+	for (int32 BoneIndex=0; BoneIndex<RefSkeleton.GetNum(); ++BoneIndex)
+	{
+		BoneNames.Add(RefSkeleton.GetBoneName(BoneIndex));
+		ParentIndices.Add(RefSkeleton.GetParentIndex(BoneIndex));	
+	}
+
+	// update retarget pose to reflect custom offsets
+	GenerateRetargetPose(InRetargetPose, RetargetRootBone);
+}
+
+void FRetargetSkeleton::Reset()
+{
+	BoneNames.Reset();
+	ParentIndices.Reset();
+	RetargetLocalPose.Reset();
+	RetargetGlobalPose.Reset();
+	SkeletalMesh = nullptr;
+}
+
+void FRetargetSkeleton::GenerateRetargetPose(FIKRetargetPose* InRetargetPose, const FName& RetargetRootBone)
+{
+	// initialize retarget pose to the skeletal mesh reference pose
+	RetargetLocalPose = SkeletalMesh->GetRefSkeleton().GetRefBonePose();
+	// copy local pose to global
+	RetargetGlobalPose = RetargetLocalPose;
+	// convert to global space
+	UpdateGlobalTransformsBelowBone(0, RetargetLocalPose, RetargetGlobalPose);
+
+	// apply retarget pose offsets (retarget pose is stored as offset relative to reference pose)
+	if (InRetargetPose && RetargetRootBone != NAME_None)
+	{
+		// apply root translation offset
+		const int32 RootBoneIndex = FindBoneIndexByName(RetargetRootBone);
+		if (RootBoneIndex != INDEX_NONE)
+		{
+			FTransform& RootTransform = RetargetGlobalPose[RootBoneIndex];
+			RootTransform.AddToTranslation(InRetargetPose->RootTranslationOffset);
+			UpdateLocalTransformOfSingleBone(RootBoneIndex, RetargetLocalPose, RetargetGlobalPose);
+			UpdateGlobalTransformsBelowBone(RootBoneIndex, RetargetLocalPose, RetargetGlobalPose);
+		}
+
+		// apply bone rotation offsets
+		for (const TTuple<FName, FQuat>& BoneRotationOffset : InRetargetPose->BoneRotationOffsets)
+		{
+			const int32 BoneIndex = FindBoneIndexByName(BoneRotationOffset.Key);
+			if (BoneIndex == INDEX_NONE)
+			{
+				continue;
+			}
+
+			FQuat BoneRotation = BoneRotationOffset.Value * RetargetGlobalPose[BoneIndex].GetRotation();
+			RetargetGlobalPose[BoneIndex].SetRotation(BoneRotation);
+			UpdateLocalTransformOfSingleBone(BoneIndex, RetargetLocalPose, RetargetGlobalPose);
+			UpdateGlobalTransformsBelowBone(BoneIndex, RetargetLocalPose, RetargetGlobalPose);
+		}
+	}
+}
+
+int32 FRetargetSkeleton::FindBoneIndexByName(const FName InName) const
+{
+	return BoneNames.IndexOfByPredicate([&InName](const FName& BoneName)
+	{
+			return BoneName == InName;
+	});
+}
+
+void FRetargetSkeleton::UpdateGlobalTransformsBelowBone(
+	const int32 StartBoneIndex,
+	TArray<FTransform>& InLocalPose,
+	TArray<FTransform>& OutGlobalPose) const
+{
+	check(BoneNames.IsValidIndex(StartBoneIndex));
+	check(BoneNames.Num() == InLocalPose.Num());
+	check(BoneNames.Num() == OutGlobalPose.Num());
+	
+	for (int32 BoneIndex=StartBoneIndex+1; BoneIndex<OutGlobalPose.Num(); ++BoneIndex)
+	{
+		UpdateGlobalTransformOfSingleBone(BoneIndex,InLocalPose,OutGlobalPose);
+	}
+}
+
+void FRetargetSkeleton::UpdateLocalTransformsBelowBone(
+	const int32 StartBoneIndex,
+	TArray<FTransform>& OutLocalPose,
+	TArray<FTransform>& InGlobalPose) const
+{
+	check(BoneNames.IsValidIndex(StartBoneIndex));
+	check(BoneNames.Num() == OutLocalPose.Num());
+	check(BoneNames.Num() == InGlobalPose.Num());
+	
+	for (int32 BoneIndex=StartBoneIndex+1; BoneIndex<InGlobalPose.Num(); ++BoneIndex)
+	{
+		UpdateLocalTransformOfSingleBone(BoneIndex, OutLocalPose, InGlobalPose);
+	}
+}
+
+void FRetargetSkeleton::UpdateGlobalTransformOfSingleBone(
+	const int32 BoneIndex,
+	const TArray<FTransform>& InLocalPose,
+	TArray<FTransform>& OutGlobalPose) const
+{
+	const int32 ParentIndex = ParentIndices[BoneIndex];
+	if (ParentIndex == INDEX_NONE)
+	{
+		return; // root always in global space
+	}
+	const FTransform ChildLocalTransform = InLocalPose[BoneIndex];
+	const FTransform ParentGlobalTransform = OutGlobalPose[ParentIndex];
+	OutGlobalPose[BoneIndex] = ChildLocalTransform * ParentGlobalTransform;
+}
+
+void FRetargetSkeleton::UpdateLocalTransformOfSingleBone(
+	const int32 BoneIndex,
+	TArray<FTransform>& OutLocalPose,
+	TArray<FTransform>& InGlobalPose) const
+{
+	const int32 ParentIndex = ParentIndices[BoneIndex];
+	if (ParentIndex == INDEX_NONE)
+	{
+		return; // root always in global space
+	}
+	const FTransform ChildGlobalTransform = InGlobalPose[BoneIndex];
+	const FTransform ParentGlobalTransform = InGlobalPose[ParentIndex];
+	OutLocalPose[BoneIndex] = ChildGlobalTransform.GetRelativeTransform(ParentGlobalTransform);
+}
+
+void FRetargetSkeleton::GetChildrenIndices(const int32 BoneIndex, TArray<int32>& OutChildren) const
+{
+	for (int32 ChildBoneIndex=0; ChildBoneIndex<ParentIndices.Num(); ++ChildBoneIndex)
+	{
+		if (ParentIndices[ChildBoneIndex] == BoneIndex)
+		{
+			OutChildren.Add(ChildBoneIndex);
+		}
+	}
+}
+
+int32 FRetargetSkeleton::GetParentIndex(const int32 BoneIndex) const
+{
+	if (BoneIndex < 0 || BoneIndex>ParentIndices.Num() || BoneIndex == INDEX_NONE)
+	{
+		return INDEX_NONE;
+	}
+
+	return ParentIndices[BoneIndex];
+}
+
+void FTargetSkeleton::Initialize(
+	USkeletalMesh* InSkeletalMesh, 
+	FIKRetargetPose* RetargetPose,
+	const FName& RetargetRootBone)
+{
+	FRetargetSkeleton::Initialize(InSkeletalMesh, RetargetPose, RetargetRootBone);
+
+	// initialize storage for output pose (the result of the retargeting)
+	OutputGlobalPose = RetargetGlobalPose;
+
+	// make storage for per-bone "Is Retargeted" flag (used for hierarchy updates)
+	IsBoneRetargeted.Init(false, OutputGlobalPose.Num());
+}
+
+void FTargetSkeleton::Reset()
+{
+	FRetargetSkeleton::Reset();
+	OutputGlobalPose.Reset();
+	IsBoneRetargeted.Reset();
+}
+
+void FTargetSkeleton::UpdateGlobalTransformsAllNonRetargetedBones(TArray<FTransform>& InOutGlobalPose)
+{
+	check(IsBoneRetargeted.Num() == InOutGlobalPose.Num());
+	
+	for (int32 BoneIndex=0; BoneIndex<InOutGlobalPose.Num(); ++BoneIndex)
+	{
+		if (!IsBoneRetargeted[BoneIndex])
+		{
+			UpdateGlobalTransformOfSingleBone(BoneIndex, RetargetLocalPose, InOutGlobalPose);
+		}
+	}
+}
+
+void FTargetSkeleton::SetBoneIsRetargeted(const int32 BoneIndex, const bool IsRetargeted)
+{
+	check(IsBoneRetargeted.IsValidIndex(BoneIndex));
+	IsBoneRetargeted[BoneIndex] = IsRetargeted;
+}
+
+bool FChainFK::Initialize(const TArray<int32>& BoneIndices, const TArray<FTransform>& InitialGlobalPose)
+{
+	check(!BoneIndices.IsEmpty());
+
+	// store all the initial bone transforms in the bone chain
+	InitialBoneTransforms.Reset();
+	for (int32 Index=0; Index < BoneIndices.Num(); ++Index)
+	{
+		const int32 BoneIndex = BoneIndices[Index];
+		if (ensure(InitialGlobalPose.IsValidIndex(BoneIndex)))
+		{
+			InitialBoneTransforms.Emplace(InitialGlobalPose[BoneIndex]);
+		}
+	}
+
+	// initialize storage for current bones
+	CurrentBoneTransforms = InitialBoneTransforms;
+
+	// calculate parameter of each bone, normalized by the length of the bone chain
+	return CalculateBoneParameters();
+}
+
+bool FChainFK::CalculateBoneParameters()
+{
+	Params.Reset();
+	
+	// special case, a single-bone chain
+	if (InitialBoneTransforms.Num() == 1)
+	{
+		Params.Add(1.0f);
+		return true;
+	}
+
+	// calculate bone lengths in chain and accumulate total length
+	TArray<float> BoneDistances;
+	float TotalChainLength = 0.0f;
+	BoneDistances.Add(0.0f);
+	for (int32 i=1; i<InitialBoneTransforms.Num(); ++i)
+	{
+		TotalChainLength += (InitialBoneTransforms[i].GetTranslation() - InitialBoneTransforms[i-1].GetTranslation()).Size();
+		BoneDistances.Add(TotalChainLength);
+	}
+
+	// cannot retarget chain if all the bones are sitting directly on each other
+	if (TotalChainLength <= KINDA_SMALL_NUMBER)
+	{
+		UE_LOG(LogTemp, Warning, TEXT("IK Retargeter bone chain length is too small to reliably retarget."));
+		return false;
+	}
+
+	// calc each bone's param along length
+	for (int32 i=0; i<InitialBoneTransforms.Num(); ++i)
+	{
+		Params.Add(BoneDistances[i] / TotalChainLength); 
+	}
+
+	return true;
+}
+
+void FChainEncoderFK::EncodePose(
+	const TArray<int32>& SourceBoneIndices,
+    const TArray<FTransform> &InputGlobalPose)
+{
+	check(SourceBoneIndices.Num() == CurrentBoneTransforms.Num());
+	
+	// copy the input pose for the chain
+	for (int32 ChainIndex=0; ChainIndex<SourceBoneIndices.Num(); ++ChainIndex)
+	{
+		const int32 BoneIndex = SourceBoneIndices[ChainIndex];
+		CurrentBoneTransforms[ChainIndex] = InputGlobalPose[BoneIndex];
+	}
+}
+
+void FChainDecoderFK::DecodePose(
+	const TArray<int32>& TargetBoneIndices,
+    const FChainEncoderFK& SourceChain,
+    const FTargetSkeleton& TargetSkeleton,
+    TArray<FTransform> &InOutGlobalPose)
+{
+	check(TargetBoneIndices.Num() == CurrentBoneTransforms.Num());
+	check(TargetBoneIndices.Num() == Params.Num());
+
+	// Before setting this chain pose, we need to ensure that any
+	// intermediate (between chains) NON-retargeted parent bones have had their
+	// global transforms updated.
+	// 
+	// For example, if this chain is retargeting a single head bone, AND the spine was
+	// retargeted in the prior step, then the neck bones will need updating first.
+	// Otherwise the neck bones will remain at their location prior to the spine update.
+	UpdateIntermediateParents(TargetSkeleton,InOutGlobalPose);
+
+	// now retarget the pose of each bone in the chain, copying from source to target
+	for (int32 ChainIndex=0; ChainIndex<TargetBoneIndices.Num(); ++ChainIndex)
+	{
+		// get the initial and current transform of source chain at param
+		// this is the interpolated transform along the chain
+		const float Param = Params[ChainIndex];
+		FTransform SourceCurrentTransform = GetTransformAtParam(
+			SourceChain.CurrentBoneTransforms,
+			SourceChain.Params,
+			Param);
+		FTransform SourceInitialTransform = GetTransformAtParam(
+			SourceChain.InitialBoneTransforms,
+			SourceChain.Params,
+			Param);
+
+		const int32 BoneIndex = TargetBoneIndices[ChainIndex];
+		const int32 ParentIndex = TargetSkeleton.ParentIndices[BoneIndex];
+		
+		// the initial transform on the target chain
+		const FTransform& TargetInitialTransform = InitialBoneTransforms[ChainIndex];
+
+		// calculate output ROTATION
+		const FQuat SourceCurrentRotation = SourceCurrentTransform.GetRotation();
+		const FQuat SourceInitialRotation = SourceInitialTransform.GetRotation();
+		const FQuat RotationDelta = SourceCurrentRotation * SourceInitialRotation.Inverse();
+		const FQuat TargetInitialRotation = TargetInitialTransform.GetRotation();
+		const FQuat OutRotation = RotationDelta * TargetInitialRotation;
+
+		// calculate output POSITION (constant for now, maybe we support retargeting translation later?)
+		const FTransform& ParentGlobalTransform = InOutGlobalPose[ParentIndex];
+		const FVector InitialLocalOffset = TargetSkeleton.RetargetLocalPose[BoneIndex].GetTranslation();
+		const FVector OutPosition = ParentGlobalTransform.TransformPosition(InitialLocalOffset);
+
+		// calculate output SCALE
+		const FVector SourceCurrentScale = SourceCurrentTransform.GetScale3D();
+		const FVector SourceInitialScale = SourceInitialTransform.GetScale3D();
+		const FVector TargetInitialScale = TargetInitialTransform.GetScale3D();
+		const FVector OutScale = SourceCurrentScale + (TargetInitialScale - SourceInitialScale);
+		
+		// apply output transform
+		FTransform OutTransform = FTransform(OutRotation, OutPosition, OutScale);
+		InOutGlobalPose[BoneIndex] = OutTransform;
+	}
+}
+
+void FChainDecoderFK::InitializeIntermediateParentIndices(
+	const int32 RetargetRootBoneIndex,
+	const int32 ChainRootBoneIndex,
+	const FTargetSkeleton& TargetSkeleton)
+{
+	IntermediateParentIndices.Reset();
+	int32 ParentBoneIndex = TargetSkeleton.ParentIndices[ChainRootBoneIndex];
+	while (true)
+	{
+		if (ParentBoneIndex < 0 || ParentBoneIndex == RetargetRootBoneIndex)
+		{
+			break; // reached root of skeleton
+		}
+
+		if (TargetSkeleton.IsBoneRetargeted[ParentBoneIndex])
+		{
+			break; // reached the start of another retargeted chain
+		}
+
+		IntermediateParentIndices.Add(ParentBoneIndex);
+		ParentBoneIndex = TargetSkeleton.ParentIndices[ParentBoneIndex];
+	}
+
+	Algo::Reverse(IntermediateParentIndices);
+}
+
+void FChainDecoderFK::UpdateIntermediateParents(
+	const FTargetSkeleton& TargetSkeleton,
+	TArray<FTransform>& InOutGlobalPose)
+{
+	for (const int32& ParentIndex : IntermediateParentIndices)
+	{
+		TargetSkeleton.UpdateGlobalTransformOfSingleBone(ParentIndex, TargetSkeleton.RetargetLocalPose, InOutGlobalPose);
+	}
+}
+
+FTransform FChainDecoderFK::GetTransformAtParam(
+	const TArray<FTransform>& Transforms,
+	const TArray<float>& InParams,
+	const float& Param) const
+{
+	if (InParams.Num() == 1)
+	{
+		return Transforms[0];
+	}
+	
+	if (Param < KINDA_SMALL_NUMBER)
+	{
+		return Transforms[0];
+	}
+
+	if (Param > 1.0f - KINDA_SMALL_NUMBER)
+	{
+		return Transforms.Last();
+	}
+
+	for (int32 ChainIndex=1; ChainIndex<InParams.Num(); ++ChainIndex)
+	{
+		const float CurrentParam = InParams[ChainIndex];
+		if (CurrentParam <= Param)
+		{
+			continue;
+		}
+		
+		const float PrevParam = InParams[ChainIndex-1];
+		const float PercentBetweenParams = (Param - PrevParam) / (CurrentParam - PrevParam);
+		const FTransform& Prev = Transforms[ChainIndex-1];
+		const FTransform& Next = Transforms[ChainIndex];
+		const FVector Position = FMath::Lerp(Prev.GetTranslation(), Next.GetTranslation(), PercentBetweenParams);
+		const FQuat Rotation = FQuat::FastLerp(Prev.GetRotation(), Next.GetRotation(), PercentBetweenParams).GetNormalized();
+		const FVector Scale = FMath::Lerp(Prev.GetScale3D(), Next.GetScale3D(), PercentBetweenParams);
+		
+		return FTransform(Rotation,Position, Scale);
+	}
+
+	checkNoEntry();
+	return FTransform::Identity;
+}
+
+bool FChainRetargeterIK::InitializeSource(
+	const TArray<int32>& BoneIndices,
+	const TArray<FTransform>& SourceInitialGlobalPose)
+{
+	if (BoneIndices.Num() < 3)
+	{
+		UE_LOG(LogTemp, Warning, TEXT("IK Retargeter trying to retarget source bone chain with IK but it has less than 3 joints."));
+		return false;
+	}
+	
+	Source.BoneIndexA = BoneIndices[0];
+	Source.BoneIndexB = BoneIndices[1];
+	Source.BoneIndexC = BoneIndices.Last();
+	
+	const FTransform& End = SourceInitialGlobalPose[Source.BoneIndexC];
+	Source.InitialEndPosition = End.GetTranslation();
+	Source.InitialEndRotation = End.GetRotation();
+
+	const FTransform& Start = SourceInitialGlobalPose[Source.BoneIndexA];
+	const float Length = (Start.GetTranslation() - Source.InitialEndPosition).Size();
+
+	if (Length <= KINDA_SMALL_NUMBER)
+	{
+		UE_LOG(LogTemp, Warning, TEXT("IK Retargeter trying to retarget source bone chain with IK, but it is zero length!"));
+    	return false;
+	}
+
+	Source.InvInitialLength = 1.0f / Length;
+	
+	return true;
+}
+
+void FChainRetargeterIK::EncodePose(const TArray<FTransform>& InputGlobalPose)
+{
+	const FVector A = InputGlobalPose[Source.BoneIndexA].GetTranslation();
+	//FVector B = InputGlobalPose[BoneIndexB].GetTranslation(); TODO use for pole vector 
+	const FVector C = InputGlobalPose[Source.BoneIndexC].GetTranslation();
+
+    // get the normalized direction / length of the IK limb (how extended it is as percentage of original length)
+    const FVector AC = C - A;
+	float ACLength;
+	FVector ACDirection;
+	AC.ToDirectionAndLength(ACDirection, ACLength);
+	const float NormalizedLimbLength = ACLength * Source.InvInitialLength;
+
+	Source.CurrentEndDirectionNormalized = ACDirection * NormalizedLimbLength;
+	Source.CurrentEndRotation = InputGlobalPose[Source.BoneIndexC].GetRotation();
+	Source.CurrentHeightFromGroundNormalized = (C.Z - Source.InitialEndPosition.Z)  * Source.InvInitialLength;
+	Source.PoleVectorDirection = FVector::OneVector; // TBD
+}
+
+bool FChainRetargeterIK::InitializeTarget(
+	const TArray<int32>& BoneIndices,
+	const TArray<FTransform> &TargetInitialGlobalPose)
+{
+	if (BoneIndices.Num() < 3)
+	{
+		UE_LOG(LogTemp, Warning, TEXT("IK Retargeter trying to retarget target bone chain with IK but it has less than 3 joints."));
+		return false;
+	}
+	
+	Target.BoneIndexA = BoneIndices[0];
+	const FTransform& Last = TargetInitialGlobalPose[BoneIndices.Last()];
+	Target.InitialEndPosition = Last.GetTranslation();
+	Target.InitialEndRotation = Last.GetRotation();
+	Target.InitialLength = (TargetInitialGlobalPose[Target.BoneIndexA].GetTranslation() - Last.GetTranslation()).Size();
+
+	if (Target.InitialLength <= KINDA_SMALL_NUMBER)
+	{
+		UE_LOG(LogTemp, Warning, TEXT("IK Retargeter trying to retarget bone chain with IK, but it is zero length!"));
+		return false;
+	}
+	
+	return true;
+}
+	
+void FChainRetargeterIK::DecodePose(
+    const TArray<FTransform>& OutGlobalPose,
+    FDecodedIKChain& OutResults) const
+{
+	// calculate end effector position
+	const FVector Start = OutGlobalPose[Target.BoneIndexA].GetTranslation();
+	OutResults.EndEffectorPosition = Start + Source.CurrentEndDirectionNormalized * Target.InitialLength; // * LimbScale (todo)
+	// optionally factor in ground height
+	//const bool IsGrounded = false;
+	//if (IsGrounded) // todo expose this
+	//{
+	//	OutResults.EndEffectorPosition.Z = (EncodedChain.HeightFromGroundNormalized * Length) + EndPositionOrig.Z;
+	//}
+
+	// calculate end effector rotation
+	const FQuat RotationDelta = Source.CurrentEndRotation * Source.InitialEndRotation.Inverse();
+	OutResults.EndEffectorRotation = RotationDelta * Target.InitialEndRotation;
+
+	// TBD calc pole vector position
+	OutResults.PoleVectorPosition = FVector::OneVector;
+}
+
+bool FRetargetChainPair::Initialize(
+    const FBoneChain& SourceBoneChain,
+    const FBoneChain& TargetBoneChain,
+    const FRetargetSkeleton& SourceSkeleton,
+    const FTargetSkeleton& TargetSkeleton)
+{
+	// validate source bone chain is compatible with source skeletal mesh
+	const bool bIsSourceValid = ValidateBoneChainWithSkeletalMesh(true, SourceBoneChain, SourceSkeleton);
+	if (!bIsSourceValid)
+	{
+		UE_LOG(LogTemp, Warning, TEXT("IK Retargeter source bone chain, '%s', is not compatible with Skeletal Mesh: '%s'"),
+			*SourceBoneChain.ChainName.ToString(), *SourceSkeleton.SkeletalMesh->GetName());
+		return false;
+	}
+
+	// validate target bone chain is compatible with target skeletal mesh
+	const bool bIsTargetValid = ValidateBoneChainWithSkeletalMesh(false, TargetBoneChain, TargetSkeleton);
+	if (!bIsTargetValid)
+    {
+		UE_LOG(LogTemp, Warning, TEXT("IK Retargeter target bone chain, '%s', is not compatible with Skeletal Mesh: '%s'"),
+            *TargetBoneChain.ChainName.ToString(), *TargetSkeleton.SkeletalMesh->GetName());
+		return false;
+    }
+
+	// store attributes of chain
+	SourceBoneChainName = SourceBoneChain.ChainName;
+	TargetBoneChainName = TargetBoneChain.ChainName;
+	
+	return true;
+}
+
+bool FRetargetChainPair::ValidateBoneChainWithSkeletalMesh(
+    const bool IsSource,
+    const FBoneChain& BoneChain,
+    const FRetargetSkeleton& RetargetSkeleton)
+{	
+	// validate start bone exists
+	const int32 StartIndex = RetargetSkeleton.FindBoneIndexByName(BoneChain.StartBone);
+	if (StartIndex == INDEX_NONE)
+	{
+		UE_LOG(LogTemp, Warning, TEXT("IK Retargeter bone chain, %s, could not find start bone, %s in mesh %s"),
+            *BoneChain.ChainName.ToString(), *BoneChain.StartBone.ToString(), *RetargetSkeleton.SkeletalMesh->GetName());
+		return false;
+	}
+	
+	// validate end bone exists
+	const int32 EndIndex = RetargetSkeleton.FindBoneIndexByName(BoneChain.EndBone);
+	if (EndIndex == INDEX_NONE)
+	{
+		UE_LOG(LogTemp, Warning, TEXT("IK Retargeter bone chain, %s, could not find end bone, %s in mesh %s"),
+            *BoneChain.ChainName.ToString(), *BoneChain.EndBone.ToString(), *RetargetSkeleton.SkeletalMesh->GetName());
+		return false;
+	}
+	
+	// validate end bone is child of start bone
+	// record all bones in chain while walking up the hierarchy (tip to root of chain)
+	TArray<int>& BoneIndices = IsSource ? SourceBoneIndices : TargetBoneIndices;
+	BoneIndices.Reset();
+	int32 NextBoneIndex = EndIndex;
+	while (true)
+	{
+		BoneIndices.Add(NextBoneIndex);
+		
+		if (NextBoneIndex == StartIndex)
+		{
+			break;
+		}
+
+		NextBoneIndex = RetargetSkeleton.GetParentIndex(NextBoneIndex);
+		
+		if (NextBoneIndex == 0)
+		{
+			UE_LOG(LogTemp, Warning, TEXT("IK Retargeter bone chain, '%s' has end bone, '%s' that is not a child of start bone, '%s' in mesh '%s'"),
+            *BoneChain.ChainName.ToString(), *BoneChain.EndBone.ToString(), *BoneChain.StartBone.ToString(), *RetargetSkeleton.SkeletalMesh->GetName());
+			return false;
+		}
+	}
+
+	// reverse the indices (we want root to tip order)
+	Algo::Reverse(BoneIndices);
+	
+	return true;
+}
+
+bool FRetargetChainPairFK::Initialize(
+	const FBoneChain& SourceBoneChain,
+	const FBoneChain& TargetBoneChain,
+	const FRetargetSkeleton& SourceSkeleton,
+	const FTargetSkeleton& TargetSkeleton)
+{
+	const bool bChainInitialized = FRetargetChainPair::Initialize(SourceBoneChain, TargetBoneChain, SourceSkeleton, TargetSkeleton);
+	if (!bChainInitialized)
+	{
+		return false;
+	}
+
+	// initialize SOURCE FK chain encoder with retarget pose
+	const bool bFKEncoderInitialized = FKEncoder.Initialize(SourceBoneIndices, SourceSkeleton.RetargetGlobalPose);
+	if (!bFKEncoderInitialized)
+	{
+		UE_LOG(LogTemp, Warning, TEXT("IK Retargeter failed to initialize FK encoder, '%s', on Skeletal Mesh: '%s'"),
+            *SourceBoneChainName.ToString(), *SourceSkeleton.SkeletalMesh->GetName());
+		return false;
+	}
+
+	// initialize TARGET FK chain decoder with retarget pose
+	const bool bFKDecoderInitialized = FKDecoder.Initialize(TargetBoneIndices, TargetSkeleton.RetargetGlobalPose);
+	if (!bFKDecoderInitialized)
+	{
+		UE_LOG(LogTemp, Warning, TEXT("IK Retargeter failed to initialize FK decoder, '%s', on Skeletal Mesh: '%s'"),
+            *TargetBoneChainName.ToString(), *TargetSkeleton.SkeletalMesh->GetName());
+		return false;
+	}
+
+	return true;
+}
+
+bool FRetargetChainPairIK::Initialize(
+	const FBoneChain& SourceBoneChain,
+	const FBoneChain& TargetBoneChain,
+	const FRetargetSkeleton& SourceSkeleton,
+	const FTargetSkeleton& TargetSkeleton)
+{
+	// validate if this chain even uses an IK Goal
+	const bool bUsingIKGoal = TargetBoneChain.IKGoalName != NAME_None && TargetBoneChain.IKGoalName != "- None -";
+	if (!bUsingIKGoal)
+	{
+		return false;
+	}
+
+	// store target IK goal
+	IKGoalName = TargetBoneChain.IKGoalName;
+
+	// initialize bone chains
+	const bool bChainInitialized = FRetargetChainPair::Initialize(SourceBoneChain, TargetBoneChain, SourceSkeleton, TargetSkeleton);
+	if (!bChainInitialized)
+	{
+		return false;
+	}
+
+	// initialize SOURCE IK chain encoder with retarget pose
+	const bool bIKEncoderInitialized = IKChainRetargeter.InitializeSource(SourceBoneIndices, SourceSkeleton.RetargetGlobalPose);
+	if (!bIKEncoderInitialized)
+	{
+		UE_LOG(LogTemp, Warning, TEXT("IK Retargeter failed to initialize IK encoder, '%s', on Skeletal Mesh: '%s'"),
+        *SourceBoneChainName.ToString(), *SourceSkeleton.SkeletalMesh->GetName());
+		return false;
+	}
+
+	// initialize TARGET IK chain decoder with retarget pose
+	const bool bIKDecoderInitialized = IKChainRetargeter.InitializeTarget(TargetBoneIndices, TargetSkeleton.RetargetGlobalPose);
+	if (!bIKDecoderInitialized)
+	{
+		UE_LOG(LogTemp, Warning, TEXT("IK Retargeter failed to initialize IK decoder, '%s', on Skeletal Mesh: '%s'"),
+        *TargetBoneChainName.ToString(), *TargetSkeleton.SkeletalMesh->GetName());
+		return false;
+	}
+
+	return true;
+}
+
+bool FRootRetargeter::InitializeSource(const FName SourceRootBoneName, const FRetargetSkeleton& SourceSkeleton)
+{
+	// validate target root bone exists
+	Source.BoneIndex = SourceSkeleton.FindBoneIndexByName(SourceRootBoneName);
+	if (Source.BoneIndex == INDEX_NONE)
+	{
+		UE_LOG(LogTemp, Warning, TEXT("IK Retargeter could not find source root bone, %s in mesh %s"),
+			*SourceRootBoneName.ToString(), *SourceSkeleton.SkeletalMesh->GetName());
+		return false;
+	}
+	
+	// record initial root data
+	const FTransform InitialTransform = SourceSkeleton.RetargetGlobalPose[Source.BoneIndex]; 
+	float InitialHeight = InitialTransform.GetTranslation().Z;
+	Source.InitialRotation = InitialTransform.GetRotation();
+
+	// ensure root height is not at origin, this happens if user sets root to ACTUAL skeleton root and not pelvis
+	if (InitialHeight < KINDA_SMALL_NUMBER)
+	{
+		// warn user and push it up slightly to avoid divide by zero
+		UE_LOG(LogTemp, Warning, TEXT("IK Retargeter root bone is very near the ground plane. This is probably not correct."));
+		InitialHeight = 1.0f;
+	}
+
+	// invert height
+	Source.InitialHeightInverse = 1.0f / InitialHeight;
+
+	return true;
+}
+
+bool FRootRetargeter::InitializeTarget(const FName TargetRootBoneName, const FTargetSkeleton& TargetSkeleton)
+{
+	// validate target root bone exists
+	Target.BoneIndex = TargetSkeleton.FindBoneIndexByName(TargetRootBoneName);
+	if (Target.BoneIndex == INDEX_NONE)
+	{
+		UE_LOG(LogTemp, Warning, TEXT("IK Retargeter could not find target root bone, %s in mesh %s"),
+            *TargetRootBoneName.ToString(), *TargetSkeleton.SkeletalMesh->GetName());
+		return false;
+	}
+
+	const FTransform TargetInitialTransform = TargetSkeleton.RetargetGlobalPose[Target.BoneIndex];
+	Target.InitialHeight = TargetInitialTransform.GetTranslation().Z;
+	Target.InitialRotation = TargetInitialTransform.GetRotation();
+
+	return true;
+}
+
+void FRootRetargeter::Reset()
+{
+	Source = FRootSource();
+	Target = FRootTarget();
+}
+
+void FRootRetargeter::EncodePose(const TArray<FTransform>& SourceGlobalPose)
+{
+	const FTransform& SourceTransform = SourceGlobalPose[Source.BoneIndex];
+	Source.CurrentPositionNormalized = SourceTransform.GetTranslation() * Source.InitialHeightInverse;
+	Source.CurrentRotation = SourceTransform.GetRotation();	
+}
+
+void FRootRetargeter::DecodePose(
+	TArray<FTransform>& OutTargetGlobalPose,
+	const float StrideScale) const
+{
+	// scale normalized position by root height
+	FVector Position = Source.CurrentPositionNormalized * Target.InitialHeight;
+	// scale horizontal displacement by stride scale
+	Position *= FVector(StrideScale, StrideScale, 1.0f);
+
+	// calc offset between initial source/target root rotations
+	const FQuat RotationDelta = Source.CurrentRotation * Source.InitialRotation.Inverse();
+	// add offset to the current source rotation
+	const FQuat Rotation = RotationDelta * Target.InitialRotation;
+
+	// apply to target
+	FTransform& TargetRootTransform = OutTargetGlobalPose[Target.BoneIndex];
+	TargetRootTransform.SetTranslation(Position);
+	TargetRootTransform.SetRotation(Rotation);
+}
+
+void UIKRetargetProcessor::Initialize(
+		USkeletalMesh* SourceSkeletalMesh,
+		USkeletalMesh* TargetSkeletalMesh,
+		UIKRetargeter* InRetargeterAsset)
+{
+	bIsInitialized = false;
+
+	// record source asset
+	RetargeterAsset = InRetargeterAsset;
+
+	// reset all data structures
+	SourceSkeleton.Reset();
+	TargetSkeleton.Reset();
+	IKRigProcessor = nullptr;
+	ChainPairsFK.Reset();
+	ChainPairsIK.Reset();
+	RootRetargeter.Reset();
+
+	// check prerequisite assets
+	if (!SourceSkeletalMesh)
+	{
+		UE_LOG(LogTemp, Error, TEXT("IK Retargeter unable to initialize. Missing source Skeletal Mesh asset."));
+		return;
+	}
+	if (!TargetSkeletalMesh)
+	{
+		UE_LOG(LogTemp, Error, TEXT("IK Retargeter unable to initialize. Missing target Skeletal Mesh asset."));
+		return;
+	}
+	if (!RetargeterAsset->SourceIKRigAsset)
+	{
+		UE_LOG(LogTemp, Error, TEXT("IK Retargeter unable to initialize. Missing source IK Rig asset."));
+		return;
+	}
+	if (!RetargeterAsset->TargetIKRigAsset)
+	{
+		UE_LOG(LogTemp, Error, TEXT("IK Retargeter unable to initialize. Missing target IK Rig asset."));
+		return;
+	}
+	if (!RetargeterAsset->RetargetPoses.Contains(RetargeterAsset->CurrentRetargetPose))
+	{
+		UE_LOG(LogTemp, Error, TEXT("IK Retargeter unable to initialize. Missing retarget pose."));
+		return;
+	}
+	
+	// initialize skeleton data for source and target
+	const TArray<FRetargetChainMap> DummySourceChainMapping;
+	SourceSkeleton.Initialize(
+		SourceSkeletalMesh,
+		nullptr,
+		NAME_None);
+	TargetSkeleton.Initialize(
+		TargetSkeletalMesh,
+		&RetargeterAsset->RetargetPoses[RetargeterAsset->CurrentRetargetPose],
+		RetargeterAsset->TargetIKRigAsset->GetRetargetRoot());
+
+	// initialize roots
+	const bool bRootsInitialized = InitializeRoots();
+	if (!bRootsInitialized)
+	{
+		// couldn't match up any BoneChain pairs, no retargeting possible
+		UE_LOG(LogTemp, Error, TEXT("IK Retargeter unable to initialize one or more root bones on source, %s and target, %s"),
+            *SourceSkeleton.SkeletalMesh->GetName(), *TargetSkeleton.SkeletalMesh->GetName());
+		return;
+	}
+
+	// initialize pairs of bone chains
+	const bool bAtLeastOneValidBoneChainPair = InitializeBoneChainPairs();
+	if (!bAtLeastOneValidBoneChainPair)
+	{
+		// couldn't match up any BoneChain pairs, no retargeting possible
+		UE_LOG(LogTemp, Error, TEXT("IK Retargeter unable to create any Bone Chain pairs between source, %s and target, %s"),
+            *SourceSkeleton.SkeletalMesh->GetName(), *TargetSkeleton.SkeletalMesh->GetName());
+		return;
+	}
+
+	// initialize the IKRigProcessor for doing IK decoding
+	const bool bIKRigInitialized = InitializeIKRig(this, TargetSkeletalMesh->GetRefSkeleton());
+	if (!bIKRigInitialized)
+	{
+		// couldn't initialize the IK Rig, we don't disable the retargeter in this case, just warn the user
+		UE_LOG(LogTemp, Warning, TEXT("IK Retargeter unable to initialize IK Rig for %s. See output for details."),
+			*TargetSkeleton.SkeletalMesh->GetName());
+	}
+
+	bIsInitialized = true;
+}
+
+bool UIKRetargetProcessor::InitializeRoots()
+{
+	// initialize root encoder
+	const FName SourceRootBoneName = RetargeterAsset->SourceIKRigAsset->GetRetargetRoot();
+	const bool bRootEncoderInit = RootRetargeter.InitializeSource(SourceRootBoneName, SourceSkeleton);
+	if (!bRootEncoderInit)
+	{
+		UE_LOG(LogTemp, Error, TEXT("IK Retargeter unable to initialize source root, '%s' on skeletal mesh: '%s'"),
+            *SourceRootBoneName.ToString(), *SourceSkeleton.SkeletalMesh->GetName());
+		return false;
+	}
+
+	// initialize root decoder
+	const FName TargetRootBoneName = RetargeterAsset->TargetIKRigAsset->GetRetargetRoot();
+	const bool bRootDecoderInit = RootRetargeter.InitializeTarget(TargetRootBoneName, TargetSkeleton);
+	if (!bRootDecoderInit)
+	{
+		UE_LOG(LogTemp, Error, TEXT("IK Retargeter unable to initialize target root, '%s' on skeletal mesh: '%s'"),
+            *TargetRootBoneName.ToString(), *TargetSkeleton.SkeletalMesh->GetName());
+		return false;
+	}
+
+	return true;
+}
+
+bool UIKRetargetProcessor::InitializeBoneChainPairs()
+{
+	check(RetargeterAsset->SourceIKRigAsset && RetargeterAsset->TargetIKRigAsset);
+	
+	for (const FRetargetChainMap& ChainSettings : RetargeterAsset->ChainMapping)
+	{
+		FName SourceChain = ChainSettings.SourceChain;
+		FName TargetChain = ChainSettings.TargetChain;
+
+		// user opted to not map this to anything, we don't need to spam a warning about it
+		if (SourceChain == NAME_None)
+		{
+			continue; 
+		}
+		
+		// get source bone chain
+		const FBoneChain* SourceBoneChain = RetargeterAsset->SourceIKRigAsset->GetRetargetChainByName(SourceChain);
+		if (!SourceBoneChain)
+		{
+			UE_LOG(LogTemp, Error, TEXT("IK Retargeter missing source bone chain: %s"), *SourceChain.ToString());
+			continue;
+		}
+
+		// get target bone chain
+		const FBoneChain* TargetBoneChain = RetargeterAsset->TargetIKRigAsset->GetRetargetChainByName(TargetChain);
+		if (!TargetBoneChain)
+		{
+			UE_LOG(LogTemp, Error, TEXT("IK Retargeter missing target bone chain: %s"), *TargetChain.ToString());
+			continue;
+		}
+
+		// all chains are loaded as FK (giving IK better starting pose)
+		FRetargetChainPairFK ChainPair;
+		const bool bChainPairValid = ChainPair.Initialize(*SourceBoneChain, *TargetBoneChain, SourceSkeleton, TargetSkeleton);
+		if (bChainPairValid)
+		{
+			ChainPairsFK.Add(ChainPair);
+		}
+		
+		// load IK chain
+		FRetargetChainPairIK ChainPairIK;
+		const bool bChainPairIKValid = ChainPairIK.Initialize(*SourceBoneChain, *TargetBoneChain, SourceSkeleton, TargetSkeleton);
+		if (bChainPairIKValid)
+		{
+			ChainPairsIK.Add(ChainPairIK);
+		}
+	}
+
+	// record which bones in the target skeleton are being retargeted
+	for (const FRetargetChainPairFK& FKChainPair : ChainPairsFK)
+	{
+		for (const int32 BoneIndex : FKChainPair.TargetBoneIndices)
+		{
+			TargetSkeleton.SetBoneIsRetargeted(BoneIndex, true);
+		}
+	}
+
+	// record intermediate bones (non-retargeted bones located BETWEEN FK chains on the target skeleton)
+	for (FRetargetChainPairFK& FKChainPair : ChainPairsFK)
+	{
+		FKChainPair.FKDecoder.InitializeIntermediateParentIndices(
+			RootRetargeter.Target.BoneIndex,
+			FKChainPair.TargetBoneIndices[0],
+			TargetSkeleton);
+	}
+
+	// root is updated before IK as well
+	TargetSkeleton.SetBoneIsRetargeted(RootRetargeter.Target.BoneIndex, true);
+
+	// return true if at least 1 pair of bone chains were initialized
+	return !(ChainPairsIK.IsEmpty() && ChainPairsFK.IsEmpty());
+}
+
+bool UIKRetargetProcessor::InitializeIKRig(UObject* Outer, const FReferenceSkeleton& InRefSkeleton)
+{	
+	// initialize IK Rig runtime processor
+	if (!IKRigProcessor)
+	{
+		IKRigProcessor = NewObject<UIKRigProcessor>(Outer);	
+	}
+	IKRigProcessor->Initialize(RetargeterAsset->TargetIKRigAsset, InRefSkeleton);
+	if (!IKRigProcessor->IsInitialized())
+	{
+		return false;
+	}
+
+	// validate that all IK bone chains have an associated Goal
+	for (FRetargetChainPairIK& ChainPair : ChainPairsIK)
+	{
+		// does the IK rig have the IK goal this bone chain requires?
+		if (!IKRigProcessor->GetGoalContainer().FindGoalByName(ChainPair.IKGoalName))
+		{
+			UE_LOG(LogTemp, Error, TEXT("IK Retargeter has target bone chain, %s that references an IK Goal, %s that is not present in any of the solvers in the IK Rig asset."),
+            *ChainPair.TargetBoneChainName.ToString(), *ChainPair.IKGoalName.ToString());
+			return false;
+		}
+	}
+	
+	return true;
+}
+
+TArray<FTransform>&  UIKRetargetProcessor::RunRetargeter(const TArray<FTransform>& InSourceGlobalPose)
+{
+	check(bIsInitialized);
+
+#if WITH_EDITOR
+	// in edit mode we just want to see the edited reference pose, not actually run the retargeting
+	// as long as the retargeter is reinitialized after every modification to the limb rotation offsets,
+	// then the TargetSkeleton.RetargetGlobalPose will contain the updated retarget pose.
+	if (RetargeterAsset->bEditRetargetPoseMode && RetargeterAsset->TargetIKRigAsset)
+	{
+		const FName RootBoneName = RetargeterAsset->TargetIKRigAsset->GetRetargetRoot();
+		TargetSkeleton.GenerateRetargetPose(&RetargeterAsset->RetargetPoses[RetargeterAsset->CurrentRetargetPose], RootBoneName);
+		return TargetSkeleton.RetargetGlobalPose; 
+	}
+#endif
+
+	// start from retarget pose
+	TargetSkeleton.OutputGlobalPose = TargetSkeleton.RetargetGlobalPose;
+
+	// ROOT retargeting
+	if (RetargeterAsset->bRetargetRoot)
+	{
+		RunRootRetarget(InSourceGlobalPose, TargetSkeleton.OutputGlobalPose);
+		// update global transforms below root
+		TargetSkeleton.UpdateGlobalTransformsBelowBone(RootRetargeter.Target.BoneIndex, TargetSkeleton.RetargetLocalPose, TargetSkeleton.OutputGlobalPose);
+	}
+	
+	// FK CHAIN retargeting
+	if (RetargeterAsset->bRetargetFK)
+	{
+		RunFKRetarget(InSourceGlobalPose, TargetSkeleton.OutputGlobalPose);
+		// update all the bones that are not controlled by FK chains or root
+		TargetSkeleton.UpdateGlobalTransformsAllNonRetargetedBones(TargetSkeleton.OutputGlobalPose);
+	}
+	
+	// IK CHAIN retargeting
+	if (RetargeterAsset->bRetargetIK)
+	{
+		RunIKRetarget(InSourceGlobalPose, TargetSkeleton.OutputGlobalPose);
+	}
+
+	return TargetSkeleton.OutputGlobalPose;
+}
+
+void UIKRetargetProcessor::RunRootRetarget(
+	const TArray<FTransform>& InGlobalTransforms,
+    TArray<FTransform>& OutGlobalTransforms)
+{
+	RootRetargeter.EncodePose(InGlobalTransforms);
+	const float StrideScale = 1.0f;
+	RootRetargeter.DecodePose(OutGlobalTransforms, StrideScale);
+}
+
+void UIKRetargetProcessor::RunFKRetarget(
+	const TArray<FTransform>& InGlobalTransforms,
+    TArray<FTransform>& OutGlobalTransforms)
+{
+	// spin through chains and encode/decode them all using the input pose
+    for (FRetargetChainPairFK& ChainPair : ChainPairsFK)
+    {
+    	ChainPair.FKEncoder.EncodePose(ChainPair.SourceBoneIndices, InGlobalTransforms);
+    	ChainPair.FKDecoder.DecodePose(ChainPair.TargetBoneIndices, ChainPair.FKEncoder, TargetSkeleton,OutGlobalTransforms);
+    }
+}
+
+void UIKRetargetProcessor::RunIKRetarget(
+	const TArray<FTransform>& InGlobalPose,
+    TArray<FTransform>& OutGlobalPose)
+{
+	if (!IKRigProcessor->IsInitialized())
+	{
+		return;
+	}
+	
+	if (ChainPairsIK.IsEmpty())
+	{
+		return; // skip IK
+	}
+	
+	// spin through IK chains
+	for (FRetargetChainPairIK& ChainPair : ChainPairsIK)
+	{
+		// encode them all using the input pose
+		ChainPair.IKChainRetargeter.EncodePose(InGlobalPose);
+		// decode the IK goal and apply to IKRig
+		FDecodedIKChain OutIKGoal;
+		ChainPair.IKChainRetargeter.DecodePose(OutGlobalPose, OutIKGoal);
+		// set the goal transform on the IK Rig
+		FIKRigGoal Goal = FIKRigGoal(
+			ChainPair.IKGoalName,
+			OutIKGoal.EndEffectorPosition,
+			OutIKGoal.EndEffectorRotation,
+			1.0f,
+			1.0f);
+		IKRigProcessor->SetIKGoal(Goal);
+	}
+
+	// copy input pose to start IK solve from
+	IKRigProcessor->SetInputPoseGlobal(OutGlobalPose);
+	// run IK solve
+	IKRigProcessor->Solve();
+	// copy results of solve
+	IKRigProcessor->CopyOutputGlobalPoseToArray(OutGlobalPose);
+}
+
+FTransform UIKRetargetProcessor::GetTargetBoneRetargetPoseGlobalTransform(const int32& TargetBoneIndex) const
+{
+	check(TargetSkeleton.BoneNames.IsValidIndex(TargetBoneIndex))
+
+	// get the current retarget pose
+	return TargetSkeleton.RetargetGlobalPose[TargetBoneIndex];
+}
+
+#if WITH_EDITOR
+void UIKRetargetProcessor::SetNeedsInitialized()
+{
+	bIsInitialized = false;
+}
+
+void UIKRetargetProcessor::CopyTargetIKRigSettingsFromAsset()
+{
+	IKRigProcessor->CopyAllInputsFromSourceAssetAtRuntime(RetargeterAsset->TargetIKRigAsset);
+}
+#endif
