IK Rig / Retargeter refactor

diff --git a/Engine/Plugins/Animation/IKRig/Source/IKRig/Private/Retargeter/IKRetargeter.cpp b/Engine/Plugins/Animation/IKRig/Source/IKRig/Private/Retargeter/IKRetargeter.cpp
--- a/Engine/Plugins/Animation/IKRig/Source/IKRig/Private/Retargeter/IKRetargeter.cpp
+++ b/Engine/Plugins/Animation/IKRig/Source/IKRig/Private/Retargeter/IKRetargeter.cpp
@@ -2,649 +2,8 @@
 
 #include "Retargeter/IKRetargeter.h"
 
-#include "Algo/LevenshteinDistance.h"
-#include "Engine/SkeletalMesh.h"
-
 const FName UIKRetargeter::DefaultPoseName = "Default Pose";
 
-void FRetargetSkeleton::Initialize(
-	USkeletalMesh* SkeletalMesh,
-	FIKRetargetPose* RetargetPose)
-{
-	// record name for debug warnings
-	SkeletalMeshName = SkeletalMesh->GetName();
-
-	// get reference skeleton
-	const FReferenceSkeleton& RefSkeleton = SkeletalMesh->GetRefSkeleton();
-	
-	// copy names and parent indices into local storage
-	for (int32 BoneIndex=0; BoneIndex<RefSkeleton.GetNum(); ++BoneIndex)
-	{
-		BoneNames.Add(RefSkeleton.GetBoneName(BoneIndex));
-		ParentIndices.Add(RefSkeleton.GetParentIndex(BoneIndex));	
-	}
-
-	// store copy of retarget pose
-	RetargetLocalPose = RefSkeleton.GetRefBonePose();
-	RetargetGlobalPose = RetargetLocalPose;
-
-	// update global retarget pose
-	UpdateGlobalTransformsBelowBone(0, RetargetLocalPose, RetargetGlobalPose);
-
-	// apply bone rotation offsets to retarget pose
-	if (RetargetPose != nullptr)
-	{
-		for (const TTuple<FName, FQuat>& BoneRotationOffset : RetargetPose->BoneRotationOffsets)
-		{
-			const int32 BoneIndex = FindBoneIndexByName(BoneRotationOffset.Key);
-			if (BoneIndex == INDEX_NONE)
-			{
-				continue;
-			}
-
-			FTransform& ChainStartTransform = RetargetGlobalPose[BoneIndex];
-			FQuat RotatedChain = BoneRotationOffset.Value * ChainStartTransform.GetRotation();
-			ChainStartTransform.SetRotation(RotatedChain);
-			UpdateLocalTransformOfSingleBone(BoneIndex, RetargetLocalPose, RetargetGlobalPose);
-			UpdateGlobalTransformsBelowBone(BoneIndex, RetargetLocalPose, RetargetGlobalPose);
-		}
-	}
-}
-
-int32 FRetargetSkeleton::FindBoneIndexByName(const FName InName) const
-{
-	for (int32 Index=0; Index<BoneNames.Num(); ++Index)
-	{
-		if (InName == BoneNames[Index])
-		{
-			return Index;
-		}
-	}
-
-	return INDEX_NONE;
-}
-
-void FRetargetSkeleton::UpdateGlobalTransformsBelowBone(
-	const int32 StartBoneIndex,
-	TArray<FTransform>& InLocalPose,
-	TArray<FTransform>& OutGlobalPose) const
-{
-	check(BoneNames.IsValidIndex(StartBoneIndex));
-	check(BoneNames.Num() == InLocalPose.Num());
-	check(BoneNames.Num() == OutGlobalPose.Num());
-	
-	for (int32 BoneIndex=StartBoneIndex+1; BoneIndex<OutGlobalPose.Num(); ++BoneIndex)
-	{
-		UpdateGlobalTransformOfSingleBone(BoneIndex,InLocalPose,OutGlobalPose);
-	}
-}
-
-void FRetargetSkeleton::UpdateLocalTransformsBelowBone(
-	const int32 StartBoneIndex,
-	TArray<FTransform>& OutLocalPose,
-	TArray<FTransform>& InGlobalPose) const
-{
-	check(BoneNames.IsValidIndex(StartBoneIndex));
-	check(BoneNames.Num() == OutLocalPose.Num());
-	check(BoneNames.Num() == InGlobalPose.Num());
-	
-	for (int32 BoneIndex=StartBoneIndex+1; BoneIndex<InGlobalPose.Num(); ++BoneIndex)
-	{
-		UpdateLocalTransformOfSingleBone(BoneIndex, OutLocalPose, InGlobalPose);
-	}
-}
-
-void FRetargetSkeleton::UpdateGlobalTransformOfSingleBone(
-	const int32 BoneIndex,
-	const TArray<FTransform>& InLocalPose,
-	TArray<FTransform>& OutGlobalPose) const
-{
-	const int32 ParentIndex = ParentIndices[BoneIndex];
-	if (ParentIndex == INDEX_NONE)
-	{
-		return; // root always in global space
-	}
-	const FTransform ChildLocalTransform = InLocalPose[BoneIndex];
-	const FTransform ParentGlobalTransform = OutGlobalPose[ParentIndex];
-	OutGlobalPose[BoneIndex] = ChildLocalTransform * ParentGlobalTransform;
-}
-
-void FRetargetSkeleton::UpdateLocalTransformOfSingleBone(
-	const int32 BoneIndex,
-	TArray<FTransform>& OutLocalPose,
-	TArray<FTransform>& InGlobalPose) const
-{
-	const int32 ParentIndex = ParentIndices[BoneIndex];
-	if (ParentIndex == INDEX_NONE)
-	{
-		return; // root always in global space
-	}
-	const FTransform ChildGlobalTransform = InGlobalPose[BoneIndex];
-	const FTransform ParentGlobalTransform = InGlobalPose[ParentIndex];
-	OutLocalPose[BoneIndex] = ChildGlobalTransform.GetRelativeTransform(ParentGlobalTransform);
-}
-
-void FRetargetSkeleton::GetChildrenIndices(const int32 BoneIndex, TArray<int32>& OutChildren) const
-{
-	for (int32 ChildBoneIndex=0; ChildBoneIndex<ParentIndices.Num(); ++ChildBoneIndex)
-	{
-		if (ParentIndices[ChildBoneIndex] == BoneIndex)
-		{
-			OutChildren.Add(ChildBoneIndex);
-		}
-	}
-}
-
-int32 FRetargetSkeleton::GetParentIndex(const int32 BoneIndex) const
-{
-	if (BoneIndex < 0 || BoneIndex>ParentIndices.Num() || BoneIndex == INDEX_NONE)
-	{
-		return INDEX_NONE;
-	}
-
-	return ParentIndices[BoneIndex];
-}
-
-void FTargetSkeleton::Initialize(
-	USkeletalMesh* SkeletalMesh, 
-	FIKRetargetPose* RetargetPose)
-{
-	FRetargetSkeleton::Initialize(SkeletalMesh, RetargetPose);
-
-	// initialize storage for output pose (the result of the retargeting)
-	OutputGlobalPose = RetargetGlobalPose;
-
-	// make storage for per-bone "Is Retargeted" flag (used for hierarchy updates)
-	IsBoneRetargeted.Init(false, OutputGlobalPose.Num());
-}
-
-void FTargetSkeleton::UpdateGlobalTransformsAllNonRetargetedBones(TArray<FTransform>& InOutGlobalPose)
-{
-	check(IsBoneRetargeted.Num() == InOutGlobalPose.Num());
-	
-	for (int32 BoneIndex=0; BoneIndex<InOutGlobalPose.Num(); ++BoneIndex)
-	{
-		if (!IsBoneRetargeted[BoneIndex])
-		{
-			UpdateGlobalTransformOfSingleBone(BoneIndex, RetargetLocalPose, InOutGlobalPose);
-		}
-	}
-}
-
-void FTargetSkeleton::SetBoneIsRetargeted(const int32 BoneIndex, const bool IsRetargeted)
-{
-	check(IsBoneRetargeted.IsValidIndex(BoneIndex));
-	IsBoneRetargeted[BoneIndex] = IsRetargeted;
-}
-
-bool FChainFK::Initialize(const TArray<int32>& BoneIndices, const TArray<FTransform>& InitialGlobalPose)
-{
-	check(!BoneIndices.IsEmpty());
-
-	// store all the initial bone transforms in the bone chain
-	for (int32 Index=0; Index < BoneIndices.Num(); ++Index)
-	{
-		const int32 BoneIndex = BoneIndices[Index];
-		check(BoneIndex < InitialGlobalPose.Num());
-		InitialBoneTransforms.Emplace(InitialGlobalPose[BoneIndex]);
-	}
-
-	// initialize storage for current bones
-	CurrentBoneTransforms = InitialBoneTransforms;
-
-	// calculate parameter of each bone, normalized by the length of the bone chain
-	return CalculateBoneParameters();
-}
-
-bool FChainFK::CalculateBoneParameters()
-{
-	// special case, a single-bone chain
-	if (InitialBoneTransforms.Num() == 1)
-	{
-		Params.Add(1.0f);
-		return true;
-	}
-
-	// calculate bone lengths in chain and accumulate total length
-	TArray<float> BoneDistances;
-	float TotalChainLength = 0.0f;
-	BoneDistances.Add(0.0f);
-	for (int32 i=1; i<InitialBoneTransforms.Num(); ++i)
-	{
-		TotalChainLength += (InitialBoneTransforms[i].GetTranslation() - InitialBoneTransforms[i-1].GetTranslation()).Size();
-		BoneDistances.Add(TotalChainLength);
-	}
-
-	// cannot retarget chain if all the bones are sitting directly on each other
-	if (TotalChainLength <= KINDA_SMALL_NUMBER)
-	{
-		UE_LOG(LogTemp, Warning, TEXT("IK Retargeter bone chain length is too small to reliably retarget."));
-		return false;
-	}
-
-	// calc each bone's param along length
-	for (int32 i=0; i<InitialBoneTransforms.Num(); ++i)
-	{
-		Params.Add(BoneDistances[i] / TotalChainLength); 
-	}
-
-	return true;
-}
-
-void FChainEncoderFK::EncodePose(
-	const TArray<int32>& SourceBoneIndices,
-    const TArray<FTransform> &InputGlobalPose)
-{
-	check(SourceBoneIndices.Num() == CurrentBoneTransforms.Num());
-	
-	// copy the input pose for the chain
-	for (int32 ChainIndex=0; ChainIndex<SourceBoneIndices.Num(); ++ChainIndex)
-	{
-		const int32 BoneIndex = SourceBoneIndices[ChainIndex];
-		CurrentBoneTransforms[ChainIndex] = InputGlobalPose[BoneIndex];
-	}
-}
-
-void FChainDecoderFK::DecodePose(
-	const TArray<int32>& TargetBoneIndices,
-    const FChainEncoderFK& SourceChain,
-    const FTargetSkeleton& TargetSkeleton,
-    TArray<FTransform> &InOutGlobalPose)
-{
-	check(TargetBoneIndices.Num() == CurrentBoneTransforms.Num());
-	check(TargetBoneIndices.Num() == Params.Num());
-
-	// Before setting this chain pose, we need to ensure that any
-	// intermediate (between chains) NON-retargeted parent bones have had their
-	// global transforms updated.
-	// 
-	// For example, if this chain is retargeting a single head bone, AND the spine was
-	// retargeted in the prior step, then the neck bones will need updating first.
-	// Otherwise the neck bones will remain at their location prior to the spine update.
-	UpdateIntermediateParents(TargetSkeleton,InOutGlobalPose);
-
-	// now retarget the pose of each bone in the chain, copying from source to target
-	for (int32 ChainIndex=0; ChainIndex<TargetBoneIndices.Num(); ++ChainIndex)
-	{
-		// get the initial and current transform of source chain at param
-		// this is the interpolated transform along the chain
-		float Param = Params[ChainIndex];
-		FTransform SourceCurrentTransform = GetTransformAtParam(
-			SourceChain.CurrentBoneTransforms,
-			SourceChain.Params,
-			Param);
-		FTransform SourceInitialTransform = GetTransformAtParam(
-			SourceChain.InitialBoneTransforms,
-			SourceChain.Params,
-			Param);
-
-		const int32 BoneIndex = TargetBoneIndices[ChainIndex];
-		const int32 ParentIndex = TargetSkeleton.ParentIndices[BoneIndex];
-		
-		// the initial transform on the target chain
-		const FTransform& TargetInitialTransform = InitialBoneTransforms[ChainIndex];
-
-		// calculate output ROTATION
-		const FQuat SourceCurrentRotation = SourceCurrentTransform.GetRotation();
-		const FQuat SourceInitialRotation = SourceInitialTransform.GetRotation();
-		const FQuat RotationDelta = SourceCurrentRotation * SourceInitialRotation.Inverse();
-		const FQuat TargetInitialRotation = TargetInitialTransform.GetRotation();
-		const FQuat OutRotation = RotationDelta * TargetInitialRotation;
-
-		// calculate output POSITION (constant for now, maybe we support retargeting translation later?)
-		const FTransform& ParentGlobalTransform = InOutGlobalPose[ParentIndex];
-		const FVector InitialLocalOffset = TargetSkeleton.RetargetLocalPose[BoneIndex].GetTranslation();
-		const FVector OutPosition = ParentGlobalTransform.TransformPosition(InitialLocalOffset);
-
-		// calculate output SCALE
-		const FVector SourceCurrentScale = SourceCurrentTransform.GetScale3D();
-		const FVector SourceInitialScale = SourceInitialTransform.GetScale3D();
-		const FVector TargetInitialScale = TargetInitialTransform.GetScale3D();
-		const FVector OutScale = SourceCurrentScale + (TargetInitialScale - SourceInitialScale);
-		
-		// apply output transform
-		FTransform OutTransform = FTransform(OutRotation, OutPosition, OutScale);
-		InOutGlobalPose[BoneIndex] = OutTransform;
-	}
-}
-
-void FChainDecoderFK::InitializeIntermediateParentIndices(
-	const int32 RetargetRootBoneIndex,
-	const int32 ChainRootBoneIndex,
-	const FTargetSkeleton& TargetSkeleton)
-{
-	IntermediateParentIndices.Reset();
-	int32 ParentBoneIndex = TargetSkeleton.ParentIndices[ChainRootBoneIndex];
-	while (true)
-	{
-		if (ParentBoneIndex < 0 || ParentBoneIndex == RetargetRootBoneIndex)
-		{
-			break; // reached root of skeleton
-		}
-
-		if (TargetSkeleton.IsBoneRetargeted[ParentBoneIndex])
-		{
-			break; // reached the start of another retargeted chain
-		}
-
-		IntermediateParentIndices.Add(ParentBoneIndex);
-		ParentBoneIndex = TargetSkeleton.ParentIndices[ParentBoneIndex];
-	}
-
-	Algo::Reverse(IntermediateParentIndices);
-}
-
-void FChainDecoderFK::UpdateIntermediateParents(
-	const FTargetSkeleton& TargetSkeleton,
-	TArray<FTransform>& InOutGlobalPose)
-{
-	for (const int32& ParentIndex : IntermediateParentIndices)
-	{
-		TargetSkeleton.UpdateGlobalTransformOfSingleBone(ParentIndex, TargetSkeleton.RetargetLocalPose, InOutGlobalPose);
-	}
-}
-
-FTransform FChainDecoderFK::GetTransformAtParam(
-	const TArray<FTransform>& Transforms,
-	const TArray<float>& InParams,
-	const float Param) const
-{
-	if (InParams.Num() == 1)
-	{
-		return Transforms[0];
-	}
-	
-	if (Param < KINDA_SMALL_NUMBER)
-	{
-		return Transforms[0];
-	}
-
-	if (Param > 1.0f - KINDA_SMALL_NUMBER)
-	{
-		return Transforms.Last();
-	}
-
-	for (int32 ChainIndex=1; ChainIndex<InParams.Num(); ++ChainIndex)
-	{
-		const float CurrentParam = InParams[ChainIndex];
-		if (CurrentParam <= Param)
-		{
-			continue;
-		}
-		
-		const float PrevParam = InParams[ChainIndex-1];
-		const float PercentBetweenParams = (Param - PrevParam) / (CurrentParam - PrevParam);
-		const FTransform& Prev = Transforms[ChainIndex-1];
-		const FTransform& Next = Transforms[ChainIndex];
-		const FVector Position = FMath::Lerp(Prev.GetTranslation(), Next.GetTranslation(), PercentBetweenParams);
-		const FQuat Rotation = FQuat::FastLerp(Prev.GetRotation(), Next.GetRotation(), PercentBetweenParams).GetNormalized();
-		const FVector Scale = FMath::Lerp(Prev.GetScale3D(), Next.GetScale3D(), PercentBetweenParams);
-		
-		return FTransform(Rotation,Position, Scale);
-	}
-
-	checkNoEntry();
-	return FTransform::Identity;
-}
-
-bool FChainRetargeterIK::InitializeSource(
-	const TArray<int32>& BoneIndices,
-	const TArray<FTransform>& SourceInitialGlobalPose)
-{
-	if (BoneIndices.Num() < 3)
-	{
-		UE_LOG(LogTemp, Warning, TEXT("IK Retargeter trying to retarget source bone chain with IK but it has less than 3 joints."));
-		return false;
-	}
-	
-	Source.BoneIndexA = BoneIndices[0];
-	Source.BoneIndexB = BoneIndices[1];
-	Source.BoneIndexC = BoneIndices.Last();
-	
-	const FTransform& End = SourceInitialGlobalPose[Source.BoneIndexC];
-	Source.InitialEndPosition = End.GetTranslation();
-	Source.InitialEndRotation = End.GetRotation();
-
-	const FTransform& Start = SourceInitialGlobalPose[Source.BoneIndexA];
-	const float Length = (Start.GetTranslation() - Source.InitialEndPosition).Size();
-
-	if (Length <= KINDA_SMALL_NUMBER)
-	{
-		UE_LOG(LogTemp, Warning, TEXT("IK Retargeter trying to retarget source bone chain with IK, but it is zero length!"));
-    	return false;
-	}
-
-	Source.InvInitialLength = 1.0f / Length;
-	
-	return true;
-}
-
-void FChainRetargeterIK::EncodePose(const TArray<FTransform>& InputGlobalPose)
-{
-	const FVector A = InputGlobalPose[Source.BoneIndexA].GetTranslation();
-	//FVector B = InputGlobalPose[BoneIndexB].GetTranslation(); TODO use for pole vector 
-	const FVector C = InputGlobalPose[Source.BoneIndexC].GetTranslation();
-
-    // get the normalized direction / length of the IK limb (how extended it is as percentage of original length)
-    const FVector AC = C - A;
-	float ACLength;
-	FVector ACDirection;
-	AC.ToDirectionAndLength(ACDirection, ACLength);
-	const float NormalizedLimbLength = ACLength * Source.InvInitialLength;
-
-	Source.CurrentEndDirectionNormalized = ACDirection * NormalizedLimbLength;
-	Source.CurrentEndRotation = InputGlobalPose[Source.BoneIndexC].GetRotation();
-	Source.CurrentHeightFromGroundNormalized = (C.Z - Source.InitialEndPosition.Z)  * Source.InvInitialLength;
-	Source.PoleVectorDirection = FVector::OneVector; // TBD
-}
-
-bool FChainRetargeterIK::InitializeTarget(
-	const TArray<int32>& BoneIndices,
-	const TArray<FTransform> &TargetInitialGlobalPose)
-{
-	if (BoneIndices.Num() < 3)
-	{
-		UE_LOG(LogTemp, Warning, TEXT("IK Retargeter trying to retarget target bone chain with IK but it has less than 3 joints."));
-		return false;
-	}
-	
-	Target.BoneIndexA = BoneIndices[0];
-	const FTransform& Last = TargetInitialGlobalPose[BoneIndices.Last()];
-	Target.InitialEndPosition = Last.GetTranslation();
-	Target.InitialEndRotation = Last.GetRotation();
-	Target.InitialLength = (TargetInitialGlobalPose[Target.BoneIndexA].GetTranslation() - Last.GetTranslation()).Size();
-
-	if (Target.InitialLength <= KINDA_SMALL_NUMBER)
-	{
-		UE_LOG(LogTemp, Warning, TEXT("IK Retargeter trying to retarget bone chain with IK, but it is zero length!"));
-		return false;
-	}
-	
-	return true;
-}
-	
-void FChainRetargeterIK::DecodePose(
-    const TArray<FTransform>& OutGlobalPose,
-    FDecodedIKChain& OutResults) const
-{
-	// calculate end effector position
-	const FVector Start = OutGlobalPose[Target.BoneIndexA].GetTranslation();
-	OutResults.EndEffectorPosition = Start + Source.CurrentEndDirectionNormalized * Target.InitialLength; // * LimbScale (todo)
-	// optionally factor in ground height
-	//const bool IsGrounded = false;
-	//if (IsGrounded) // todo expose this
-	//{
-	//	OutResults.EndEffectorPosition.Z = (EncodedChain.HeightFromGroundNormalized * Length) + EndPositionOrig.Z;
-	//}
-
-	// calculate end effector rotation
-	const FQuat RotationDelta = Source.CurrentEndRotation * Source.InitialEndRotation.Inverse();
-	OutResults.EndEffectorRotation = RotationDelta * Target.InitialEndRotation;
-
-	// TBD calc pole vector position
-	OutResults.PoleVectorPosition = FVector::OneVector;
-}
-
-bool FRetargetChainPair::Initialize(
-    const FBoneChain& SourceBoneChain,
-    const FBoneChain& TargetBoneChain,
-    const FRetargetSkeleton& SourceSkeleton,
-    const FTargetSkeleton& TargetSkeleton)
-{
-	// validate source bone chain is compatible with source skeletal mesh
-	const bool bIsSourceValid = ValidateBoneChainWithSkeletalMesh(true, SourceBoneChain, SourceSkeleton);
-	if (!bIsSourceValid)
-	{
-		UE_LOG(LogTemp, Warning, TEXT("IK Retargeter source bone chain, '%s', is not compatible with Skeletal Mesh: '%s'"),
-			*SourceBoneChain.ChainName.ToString(), *SourceSkeleton.SkeletalMeshName);
-		return false;
-	}
-
-	// validate target bone chain is compatible with target skeletal mesh
-	const bool bIsTargetValid = ValidateBoneChainWithSkeletalMesh(false, TargetBoneChain, TargetSkeleton);
-	if (!bIsTargetValid)
-    {
-		UE_LOG(LogTemp, Warning, TEXT("IK Retargeter target bone chain, '%s', is not compatible with Skeletal Mesh: '%s'"),
-            *TargetBoneChain.ChainName.ToString(), *TargetSkeleton.SkeletalMeshName);
-		return false;
-    }
-
-	// store attributes of chain
-	SourceBoneChainName = SourceBoneChain.ChainName;
-	TargetBoneChainName = TargetBoneChain.ChainName;
-	
-	return true;
-}
-
-bool FRetargetChainPair::ValidateBoneChainWithSkeletalMesh(
-    const bool IsSource,
-    const FBoneChain& BoneChain,
-    const FRetargetSkeleton& RetargetSkeleton)
-{	
-	// validate start bone exists
-	const int32 StartIndex = RetargetSkeleton.FindBoneIndexByName(BoneChain.StartBone);
-	if (StartIndex == INDEX_NONE)
-	{
-		UE_LOG(LogTemp, Warning, TEXT("IK Retargeter bone chain, %s, could not find start bone, %s in mesh %s"),
-            *BoneChain.ChainName.ToString(), *BoneChain.StartBone.ToString(), *RetargetSkeleton.SkeletalMeshName);
-		return false;
-	}
-	
-	// validate end bone exists
-	const int32 EndIndex = RetargetSkeleton.FindBoneIndexByName(BoneChain.EndBone);
-	if (EndIndex == INDEX_NONE)
-	{
-		UE_LOG(LogTemp, Warning, TEXT("IK Retargeter bone chain, %s, could not find end bone, %s in mesh %s"),
-            *BoneChain.ChainName.ToString(), *BoneChain.EndBone.ToString(), *RetargetSkeleton.SkeletalMeshName);
-		return false;
-	}
-	
-	// validate end bone is child of start bone
-	// record all bones in chain while walking up the hierarchy (tip to root of chain)
-	TArray<int>& BoneIndices = IsSource ? SourceBoneIndices : TargetBoneIndices;
-	BoneIndices.Reset();
-	int32 NextBoneIndex = EndIndex;
-	while (true)
-	{
-		BoneIndices.Add(NextBoneIndex);
-		
-		if (NextBoneIndex == StartIndex)
-		{
-			break;
-		}
-
-		NextBoneIndex = RetargetSkeleton.GetParentIndex(NextBoneIndex);
-		
-		if (NextBoneIndex == 0)
-		{
-			UE_LOG(LogTemp, Warning, TEXT("IK Retargeter bone chain, '%s' has end bone, '%s' that is not a child of start bone, '%s' in mesh '%s'"),
-            *BoneChain.ChainName.ToString(), *BoneChain.EndBone.ToString(), *BoneChain.StartBone.ToString(), *RetargetSkeleton.SkeletalMeshName);
-			return false;
-		}
-	}
-
-	// reverse the indices (we want root to tip order)
-	Algo::Reverse(BoneIndices);
-	
-	return true;
-}
-
-bool FRetargetChainPairFK::Initialize(
-	const FBoneChain& SourceBoneChain,
-	const FBoneChain& TargetBoneChain,
-	const FRetargetSkeleton& SourceSkeleton,
-	const FTargetSkeleton& TargetSkeleton)
-{
-	const bool bChainInitialized = FRetargetChainPair::Initialize(SourceBoneChain, TargetBoneChain, SourceSkeleton, TargetSkeleton);
-	if (!bChainInitialized)
-	{
-		return false;
-	}
-
-	// initialize SOURCE FK chain encoder with retarget pose
-	const bool bFKEncoderInitialized = FKEncoder.Initialize(SourceBoneIndices, SourceSkeleton.RetargetGlobalPose);
-	if (!bFKEncoderInitialized)
-	{
-		UE_LOG(LogTemp, Warning, TEXT("IK Retargeter failed to initialize FK encoder, '%s', on Skeletal Mesh: '%s'"),
-            *SourceBoneChainName.ToString(), *SourceSkeleton.SkeletalMeshName);
-		return false;
-	}
-
-	// initialize TARGET FK chain decoder with retarget pose
-	const bool bFKDecoderInitialized = FKDecoder.Initialize(TargetBoneIndices, TargetSkeleton.RetargetGlobalPose);
-	if (!bFKDecoderInitialized)
-	{
-		UE_LOG(LogTemp, Warning, TEXT("IK Retargeter failed to initialize FK decoder, '%s', on Skeletal Mesh: '%s'"),
-            *TargetBoneChainName.ToString(), *TargetSkeleton.SkeletalMeshName);
-		return false;
-	}
-
-	return true;
-}
-
-bool FRetargetChainPairIK::Initialize(
-	const FBoneChain& SourceBoneChain,
-	const FBoneChain& TargetBoneChain,
-	const FRetargetSkeleton& SourceSkeleton,
-	const FTargetSkeleton& TargetSkeleton)
-{
-	// validate if this chain even uses an IK Goal
-	const bool bUsingIKGoal = TargetBoneChain.IKGoalName != NAME_None && TargetBoneChain.IKGoalName != "- None -";
-	if (!bUsingIKGoal)
-	{
-		return false;
-	}
-
-	// store target IK goal
-	IKGoalName = TargetBoneChain.IKGoalName;
-
-	// initialize bone chains
-	const bool bChainInitialized = FRetargetChainPair::Initialize(SourceBoneChain, TargetBoneChain, SourceSkeleton, TargetSkeleton);
-	if (!bChainInitialized)
-	{
-		return false;
-	}
-
-	// initialize SOURCE IK chain encoder with retarget pose
-	const bool bIKEncoderInitialized = IKChainRetargeter.InitializeSource(SourceBoneIndices, SourceSkeleton.RetargetGlobalPose);
-	if (!bIKEncoderInitialized)
-	{
-		UE_LOG(LogTemp, Warning, TEXT("IK Retargeter failed to initialize IK encoder, '%s', on Skeletal Mesh: '%s'"),
-        *SourceBoneChainName.ToString(), *SourceSkeleton.SkeletalMeshName);
-		return false;
-	}
-
-	// initialize TARGET IK chain decoder with retarget pose
-	const bool bIKDecoderInitialized = IKChainRetargeter.InitializeTarget(TargetBoneIndices, TargetSkeleton.RetargetGlobalPose);
-	if (!bIKDecoderInitialized)
-	{
-		UE_LOG(LogTemp, Warning, TEXT("IK Retargeter failed to initialize IK decoder, '%s', on Skeletal Mesh: '%s'"),
-        *TargetBoneChainName.ToString(), *TargetSkeleton.SkeletalMeshName);
-		return false;
-	}
-
-	return true;
-}
 
 void FIKRetargetPose::AddRotationDeltaToBone(FName BoneName, FQuat RotationDelta)
 {
@@ -660,676 +19,7 @@ void FIKRetargetPose::AddRotationDeltaToBone(FName BoneName, FQuat RotationDelta
 	*RotOffset = RotationDelta * (*RotOffset);
 }
 
-bool FRootRetargeter::InitializeSource(const FName SourceRootBoneName, const FRetargetSkeleton& SourceSkeleton)
-{
-	// validate target root bone exists
-	Source.BoneIndex = SourceSkeleton.FindBoneIndexByName(SourceRootBoneName);
-	if (Source.BoneIndex == INDEX_NONE)
-	{
-		UE_LOG(LogTemp, Warning, TEXT("IK Retargeter could not find source root bone, %s in mesh %s"),
-			*SourceRootBoneName.ToString(), *SourceSkeleton.SkeletalMeshName);
-		return false;
-	}
-	
-	// record initial root data
-	const FTransform InitialTransform = SourceSkeleton.RetargetGlobalPose[Source.BoneIndex]; 
-	float InitialHeight = InitialTransform.GetTranslation().Z;
-	Source.InitialRotation = InitialTransform.GetRotation();
-
-	// ensure root height is not at origin, this happens if user sets root to ACTUAL skeleton root and not pelvis
-	if (InitialHeight < KINDA_SMALL_NUMBER)
-	{
-		// warn user and push it up slightly to avoid divide by zero
-		UE_LOG(LogTemp, Warning, TEXT("IK Retargeter root bone is very near the ground plane. This is probably not correct."));
-		InitialHeight = 1.0f;
-	}
-
-	// invert height
-	Source.InitialHeightInverse = 1.0f / InitialHeight;
-
-	return true;
-}
-
-bool FRootRetargeter::InitializeTarget(const FName TargetRootBoneName, const FTargetSkeleton& TargetSkeleton)
-{
-	// validate target root bone exists
-	Target.BoneIndex = TargetSkeleton.FindBoneIndexByName(TargetRootBoneName);
-	if (Target.BoneIndex == INDEX_NONE)
-	{
-		UE_LOG(LogTemp, Warning, TEXT("IK Retargeter could not find target root bone, %s in mesh %s"),
-            *TargetRootBoneName.ToString(), *TargetSkeleton.SkeletalMeshName);
-		return false;
-	}
-
-	const FTransform TargetInitialTransform = TargetSkeleton.RetargetGlobalPose[Target.BoneIndex];
-	Target.InitialHeight = TargetInitialTransform.GetTranslation().Z;
-	Target.InitialRotation = TargetInitialTransform.GetRotation();
-
-	return true;
-}
-
-void FRootRetargeter::EncodePose(const TArray<FTransform>& SourceGlobalPose)
-{
-	const FTransform& SourceTransform = SourceGlobalPose[Source.BoneIndex];
-	Source.CurrentPositionNormalized = SourceTransform.GetTranslation() * Source.InitialHeightInverse;
-	Source.CurrentRotation = SourceTransform.GetRotation();	
-}
-
-void FRootRetargeter::DecodePose(
-	TArray<FTransform>& OutTargetGlobalPose,
-	const float StrideScale) const
-{
-	// scale normalized position by root height
-	FVector Position = Source.CurrentPositionNormalized * Target.InitialHeight;
-	// scale horizontal displacement by stride scale
-	Position *= FVector(StrideScale, StrideScale, 1.0f);
-
-	// calc offset between initial source/target root rotations
-	const FQuat RotationDelta = Source.CurrentRotation * Source.InitialRotation.Inverse();
-	// add offset to the current source rotation
-	const FQuat Rotation = RotationDelta * Target.InitialRotation;
-
-	// apply to target
-	FTransform& TargetRootTransform = OutTargetGlobalPose[Target.BoneIndex];
-	TargetRootTransform.SetTranslation(Position);
-	TargetRootTransform.SetRotation(Rotation);
-}
-
-UIKRetargeter::UIKRetargeter()
-	: SourceIKRigAsset(nullptr),
-    TargetIKRigAsset(nullptr),
-    bIsLoadedAndValid(false)
-{
-	CleanPoseList();
-}
-
-void UIKRetargeter::Initialize(USkeletalMesh* SourceSkeletalMesh, USkeletalMesh* TargetSkeletalMesh, UObject* Outer)
-{
-	bIsLoadedAndValid = false;
-
-	check(SourceSkeletalMesh && TargetSkeletalMesh);
-	
-	// initialize skeleton data for source and target
-	const TArray<FRetargetChainMap> DummySourceChainMapping;
-	SourceSkeleton.Initialize(SourceSkeletalMesh, nullptr);
-	TargetSkeleton.Initialize(TargetSkeletalMesh, &RetargetPoses[CurrentRetargetPose]);
-
-	// initialize roots
-	const bool bRootsInitialized = InitializeRoots();
-	if (!bRootsInitialized)
-	{
-		// couldn't match up any BoneChain pairs, no retargeting possible
-		UE_LOG(LogTemp, Warning, TEXT("IK Retargeter unable to initialize one or more root bones on source, %s and target, %s"),
-            *SourceSkeleton.SkeletalMeshName, *TargetSkeleton.SkeletalMeshName);
-		return;
-	}
-
-	// initialize pairs of bone chains
-	const bool bAtLeastOneValidBoneChainPair = InitializeBoneChainPairs();
-	if (!bAtLeastOneValidBoneChainPair)
-	{
-		// couldn't match up any BoneChain pairs, no retargeting possible
-		UE_LOG(LogTemp, Warning, TEXT("IK Retargeter unable to create any Bone Chain pairs between source, %s and target, %s"),
-            *SourceSkeleton.SkeletalMeshName, *TargetSkeleton.SkeletalMeshName);
-		return;
-	}
-
-	// initialize the IKRigProcessor for doing IK decoding
-	const bool bIKRigInitialized = InitializeIKRig(Outer, TargetSkeletalMesh->GetRefSkeleton());
-	if (!bIKRigInitialized)
-	{
-		// couldn't initialize the IK Rig
-		UE_LOG(LogTemp, Warning, TEXT("IK Retargeter unable to initialize IK Rig for %s. See output for details."),
-			*TargetSkeleton.SkeletalMeshName);
-		return;
-	}
-
-	bIsLoadedAndValid = true;
-}
-
-bool UIKRetargeter::InitializeRoots()
+void FIKRetargetPose::AddTranslationDeltaToRoot(FVector TranslateDelta)
 {
-	check(SourceIKRigAsset && TargetIKRigAsset);
-
-	// initialize root encoder
-	const FName SourceRootBoneName = SourceIKRigAsset->RetargetDefinition.RootBone;
-	const bool bRootEncoderInit = RootRetargeter.InitializeSource(SourceRootBoneName, SourceSkeleton);
-	if (!bRootEncoderInit)
-	{
-		UE_LOG(LogTemp, Warning, TEXT("IK Retargeter unable to initialize source root, '%s' on skeletal mesh: '%s'"),
-            *SourceRootBoneName.ToString(), *SourceSkeleton.SkeletalMeshName);
-		return false;
-	}
-
-	// initialize root decoder
-	const FName TargetRootBoneName = TargetIKRigAsset->RetargetDefinition.RootBone;
-	const bool bRootDecoderInit = RootRetargeter.InitializeTarget(TargetRootBoneName, TargetSkeleton);
-	if (!bRootDecoderInit)
-	{
-		UE_LOG(LogTemp, Warning, TEXT("IK Retargeter unable to initialize target root, '%s' on skeletal mesh: '%s'"),
-            *TargetRootBoneName.ToString(), *TargetSkeleton.SkeletalMeshName);
-		return false;
-	}
-
-	return true;
-}
-
-bool UIKRetargeter::InitializeBoneChainPairs()
-{
-	check(SourceIKRigAsset && TargetIKRigAsset);
-	
-	FRetargetDefinition& SourceRetargetDef = SourceIKRigAsset->RetargetDefinition;
-	FRetargetDefinition& TargetRetargetDef = TargetIKRigAsset->RetargetDefinition;
-
-	for (const FRetargetChainMap& ChainSettings : ChainMapping)
-	{
-		FName SourceChain = ChainSettings.SourceChain;
-		FName TargetChain = ChainSettings.TargetChain;
-
-		// user opted to not map this to anything, we don't need to spam a warning about it
-		if (SourceChain == NAME_None)
-		{
-			continue; 
-		}
-		
-		// get source bone chain
-		FBoneChain* SourceBoneChain = SourceRetargetDef.GetBoneChainByName(SourceChain);
-		if (!SourceBoneChain)
-		{
-			UE_LOG(LogTemp, Warning, TEXT("IK Retargeter missing source bone chain: %s"), *SourceChain.ToString());
-			continue;
-		}
-
-		// get target bone chain
-		FBoneChain* TargetBoneChain = TargetRetargetDef.GetBoneChainByName(TargetChain);
-		if (!TargetBoneChain)
-		{
-			UE_LOG(LogTemp, Warning, TEXT("IK Retargeter missing target bone chain: %s"), *TargetChain.ToString());
-			continue;
-		}
-
-		// all chains are loaded as FK (giving IK better starting pose)
-		FRetargetChainPairFK ChainPair;
-		const bool bChainPairValid = ChainPair.Initialize(*SourceBoneChain, *TargetBoneChain, SourceSkeleton, TargetSkeleton);
-		if (bChainPairValid)
-		{
-			ChainPairsFK.Add(ChainPair);
-		}
-		
-		// load IK chain
-		FRetargetChainPairIK ChainPairIK;
-		const bool bChainPairIKValid = ChainPairIK.Initialize(*SourceBoneChain, *TargetBoneChain, SourceSkeleton, TargetSkeleton);
-		if (bChainPairIKValid)
-		{
-			ChainPairsIK.Add(ChainPairIK);
-		}
-	}
-
-	// record which bones in the target skeleton are being retargeted
-	for (const FRetargetChainPairFK& FKChainPair : ChainPairsFK)
-	{
-		for (const int32 BoneIndex : FKChainPair.TargetBoneIndices)
-		{
-			TargetSkeleton.SetBoneIsRetargeted(BoneIndex, true);
-		}
-	}
-
-	// record intermediate bones (non-retargeted bones located BETWEEN FK chains on the target skeleton)
-	for (FRetargetChainPairFK& FKChainPair : ChainPairsFK)
-	{
-		FKChainPair.FKDecoder.InitializeIntermediateParentIndices(
-			RootRetargeter.Target.BoneIndex,
-			FKChainPair.TargetBoneIndices[0],
-			TargetSkeleton);
-	}
-
-	// root is updated before IK as well
-	TargetSkeleton.SetBoneIsRetargeted(RootRetargeter.Target.BoneIndex, true);
-
-	// return true if at least 1 pair of bone chains were initialized
-	return !(ChainPairsIK.IsEmpty() && ChainPairsFK.IsEmpty());
-}
-
-bool UIKRetargeter::InitializeIKRig(UObject* Outer, const FReferenceSkeleton& InRefSkeleton)
-{	
-	// initialize IK Rig runtime processor
-	if (!IKRigProcessor)
-	{
-		IKRigProcessor = NewObject<UIKRigProcessor>(Outer);	
-	}
-	IKRigProcessor->Initialize(TargetIKRigAsset, InRefSkeleton);
-	if (IKRigProcessor->NeedsInitialized(TargetIKRigAsset))
-	{
-		return false;
-	}
-
-	// validate that all IK bone chains have an associated Goal
-	TArray<FName> GoalNames;
-	IKRigProcessor->GetGoalContainer().Goals.GenerateKeyArray(GoalNames);
-	for (FRetargetChainPairIK& ChainPair : ChainPairsIK)
-	{
-		// does the IK rig have the IK goal this bone chain requires?
-		if (!GoalNames.Contains(ChainPair.IKGoalName))
-		{
-			UE_LOG(LogTemp, Warning, TEXT("IK Retargeter has target bone chain, %s that references an IK Goal, %s that is not present in any of the solvers in the IK Rig asset."),
-            *ChainPair.TargetBoneChainName.ToString(), *ChainPair.IKGoalName.ToString());
-			return false;
-		}
-	}
-	
-	return true;
-}
-
-TArray<FTransform>&  UIKRetargeter::RunRetargeter(const TArray<FTransform>& InSourceGlobalPose)
-{
-	check(bIsLoadedAndValid);
-
-	// start from retarget pose
-	TargetSkeleton.OutputGlobalPose = TargetSkeleton.RetargetGlobalPose;
-	
-	// in edit mode we just want to see the edited reference pose, not actually run the retargeting
-	// as long as the retargeter is reinitialized after every modification to the limb rotation offsets,
-	// then the TargetSkeleton.RetargetGlobalPose will contain the updated retarget pose.
-	if (bEditReferencePoseMode)
-	{
-		return TargetSkeleton.OutputGlobalPose; 
-	}
-
-	// ROOT retargeting
-	if (bRetargetRoot)
-	{
-		RunRootRetarget(InSourceGlobalPose, TargetSkeleton.OutputGlobalPose);
-		// update global transforms below root
-		TargetSkeleton.UpdateGlobalTransformsBelowBone(RootRetargeter.Target.BoneIndex, TargetSkeleton.RetargetLocalPose, TargetSkeleton.OutputGlobalPose);
-	}
-	
-	// FK CHAIN retargeting
-	if (bRetargetFK)
-	{
-		RunFKRetarget(InSourceGlobalPose, TargetSkeleton.OutputGlobalPose);
-		// update all the bones that are not controlled by FK chains or root
-		TargetSkeleton.UpdateGlobalTransformsAllNonRetargetedBones(TargetSkeleton.OutputGlobalPose);
-	}
-	
-	// IK CHAIN retargeting
-	if (bRetargetIK)
-	{
-		RunIKRetarget(InSourceGlobalPose, TargetSkeleton.OutputGlobalPose);
-	}
-
-	return TargetSkeleton.OutputGlobalPose;
-}
-
-void UIKRetargeter::RunRootRetarget(
-	const TArray<FTransform>& InGlobalTransforms,
-    TArray<FTransform>& OutGlobalTransforms)
-{
-	RootRetargeter.EncodePose(InGlobalTransforms);
-	const float StrideScale = 1.0f;
-	RootRetargeter.DecodePose(OutGlobalTransforms, StrideScale);
-}
-
-void UIKRetargeter::RunFKRetarget(
-	const TArray<FTransform>& InGlobalTransforms,
-    TArray<FTransform>& OutGlobalTransforms)
-{
-	// spin through chains and encode/decode them all using the input pose
-    for (FRetargetChainPairFK& ChainPair : ChainPairsFK)
-    {
-    	ChainPair.FKEncoder.EncodePose(ChainPair.SourceBoneIndices, InGlobalTransforms);
-    	ChainPair.FKDecoder.DecodePose(ChainPair.TargetBoneIndices, ChainPair.FKEncoder, TargetSkeleton,OutGlobalTransforms);
-    }
-}
-
-void UIKRetargeter::RunIKRetarget(
-	const TArray<FTransform>& InGlobalPose,
-    TArray<FTransform>& OutGlobalPose)
-{
-	if (ChainPairsIK.IsEmpty())
-	{
-		return; // skip IK
-	}
-	
-	// spin through IK chains
-	for (FRetargetChainPairIK& ChainPair : ChainPairsIK)
-	{
-		// encode them all using the input pose
-		ChainPair.IKChainRetargeter.EncodePose(InGlobalPose);
-		// decode the IK goal and apply to IKRig
-		FDecodedIKChain OutIKGoal;
-		ChainPair.IKChainRetargeter.DecodePose(OutGlobalPose, OutIKGoal);
-		// set the goal transform on the IK Rig
-		FIKRigGoal Goal = FIKRigGoal(
-			ChainPair.IKGoalName,
-			OutIKGoal.EndEffectorPosition,
-			OutIKGoal.EndEffectorRotation,
-			1.0f,
-			1.0f);
-		IKRigProcessor->SetIKGoal(Goal);
-	}
-
-	// copy input pose to start IK solve from
-	IKRigProcessor->SetInputPoseGlobal(OutGlobalPose);
-	// run IK solve
-	IKRigProcessor->Solve();
-	// copy results of solve
-	IKRigProcessor->CopyOutputGlobalPoseToArray(OutGlobalPose);
-}
-
-void UIKRetargeter::CleanPoseList()
-{
-	// enforce the existence of a default pose
-	const bool HasDefaultPose = RetargetPoses.Contains(DefaultPoseName);
-	if (!HasDefaultPose)
-	{
-		RetargetPoses.Emplace(DefaultPoseName);
-	}
-	
-	// use default pose unless set to something else
-	if (CurrentRetargetPose == NAME_None)
-	{
-		CurrentRetargetPose = DefaultPoseName;
-	}
-
-	// remove all bone offsets that are no longer part of the target skeleton
-	if (TargetIKRigAsset)
-	{
-		const TArray<FName> AllowedBoneNames = TargetIKRigAsset->Skeleton.BoneNames;
-		for (TTuple<FName, FIKRetargetPose>& Pose : RetargetPoses)
-		{
-			// find bone offsets no longer in target skeleton
-			TArray<FName> BonesToRemove;
-			for (TTuple<FName, FQuat>& BoneOffset : Pose.Value.BoneRotationOffsets)
-			{
-				if (!AllowedBoneNames.Contains(BoneOffset.Key))
-				{
-					BonesToRemove.Add(BoneOffset.Key);
-				}
-			}
-			
-			// remove bone offsets
-			for (const FName& BoneToRemove : BonesToRemove)
-			{
-				Pose.Value.BoneRotationOffsets.Remove(BoneToRemove);
-			}
-		}
-	}
-
-	Modify();
-}
-
-#if WITH_EDITOR
-
-FName UIKRetargeter::GetSourceRootBone()
-{
-	return SourceIKRigAsset ? SourceIKRigAsset->RetargetDefinition.RootBone : FName("None");
-}
-
-FName UIKRetargeter::GetTargetRootBone()
-{
-	return TargetIKRigAsset ? TargetIKRigAsset->RetargetDefinition.RootBone : FName("None");
-}
-
-void UIKRetargeter::GetTargetChainNames(TArray<FName>& OutNames) const
-{
-	if (TargetIKRigAsset)
-	{
-		for (const FBoneChain& Chain : TargetIKRigAsset->RetargetDefinition.BoneChains)
-		{
-			OutNames.Add(Chain.ChainName);
-		}
-	}
-}
-
-void UIKRetargeter::GetSourceChainNames(TArray<FName>& OutNames) const
-{
-	if (SourceIKRigAsset)
-	{
-		for (const FBoneChain& Chain : SourceIKRigAsset->RetargetDefinition.BoneChains)
-		{
-			OutNames.Add(Chain.ChainName);
-		}
-	}
-}
-
-void UIKRetargeter::CleanChainMapping()
-{
-	if (!TargetIKRigAsset)
-	{
-		// don't clean chain mappings, in case user is replacing with IK Rig asset that has some valid mappings
-		return;
-	}
-	
-	TArray<FName> TargetChainNames;
-	GetTargetChainNames(TargetChainNames);
-
-	// remove all target chains that are no longer in the target IK rig asset
-	TArray<FName> TargetChainsToRemove;
-	for (FRetargetChainMap& ChainMap : ChainMapping)
-	{
-		if (!TargetChainNames.Contains(ChainMap.TargetChain))
-		{
-			TargetChainsToRemove.Add(ChainMap.TargetChain);
-		}
-	}
-	for (FName TargetChainToRemove : TargetChainsToRemove)
-	{
-		ChainMapping.RemoveAll([&TargetChainToRemove](FRetargetChainMap& Element)
-		{
-			return Element.TargetChain == TargetChainToRemove;
-		});
-	}
-
-	// add a mapping for each chain that is in the target IK rig (if it doesn't have one already)
-	for (FName TargetChainName : TargetChainNames)
-	{
-		const bool HasChain = ChainMapping.ContainsByPredicate([&TargetChainName](FRetargetChainMap& Element)
-		{
-			return Element.TargetChain == TargetChainName;
-		});
-		
-		if (!HasChain)
-		{
-			ChainMapping.Add(FRetargetChainMap(TargetChainName));
-		}
-	}
-
-	TArray<FName> SourceChainNames;
-	GetSourceChainNames(SourceChainNames);
-	
-	// reset any sources that are no longer present to "None"
-	for (FRetargetChainMap& ChainMap : ChainMapping)
-	{
-		if (!SourceChainNames.Contains(ChainMap.SourceChain))
-		{
-			ChainMap.SourceChain = NAME_None;
-		}
-	}
-
-	// enforce the same chain order as the target IK rig
-	ChainMapping.Sort([this](const FRetargetChainMap& A, const FRetargetChainMap& B)
-	{
-		const TArray<FBoneChain>& BoneChains = TargetIKRigAsset->RetargetDefinition.BoneChains;
-		
-		const int32 IndexA = BoneChains.IndexOfByPredicate([&A](const FBoneChain& Chain)
-		{
-			return A.TargetChain == Chain.ChainName;
-		});
-
-		const int32 IndexB = BoneChains.IndexOfByPredicate([&B](const FBoneChain& Chain)
-		{
-			return B.TargetChain == Chain.ChainName;
-		});
- 
-		return IndexA < IndexB;
-	});
-
-	// force update with latest mapping
-	Modify();
-}
-
-void UIKRetargeter::AutoMapChains()
-{
-	TArray<FName> SourceChainNames;
-	GetSourceChainNames(SourceChainNames);
-	
-	// auto-map any chains that have no value using a fuzzy string search
-	for (FRetargetChainMap& ChainMap : ChainMapping)
-	{
-		if (ChainMap.SourceChain != NAME_None)
-		{
-			continue; // already set by user
-		}
-
-		// find "best match" automatically as a convenience for the user
-		FString TargetNameLowerCase = ChainMap.TargetChain.ToString().ToLower();
-		float HighestScore = 0.2f;
-		int32 HighestScoreIndex = -1;
-		for (int32 ChainIndex=0; ChainIndex<SourceChainNames.Num(); ++ChainIndex)
-		{
-			FString SourceNameLowerCase = SourceChainNames[ChainIndex].ToString().ToLower();
-			float WorstCase = TargetNameLowerCase.Len() + SourceNameLowerCase.Len();
-			WorstCase = WorstCase < 1.0f ? 1.0f : WorstCase;
-			const float Score = 1.0f - (Algo::LevenshteinDistance(TargetNameLowerCase, SourceNameLowerCase) / WorstCase);
-			if (Score > HighestScore)
-			{
-				HighestScore = Score;
-				HighestScoreIndex = ChainIndex;
-			}
-		}
-
-		// apply source if any decent matches were found
-		if (SourceChainNames.IsValidIndex(HighestScoreIndex))
-		{
-			ChainMap.SourceChain = SourceChainNames[HighestScoreIndex];
-		}
-	}
-
-	// force update with latest mapping
-	Modify();
-}
-
-USkeleton* UIKRetargeter::GetSourceSkeletonAsset() const
-{
-	if (!SourceIKRigAsset)
-	{
-		return nullptr;
-	}
-
-	if (!SourceIKRigAsset->PreviewSkeletalMesh)
-	{
-		return nullptr;
-	}
-
-	return SourceIKRigAsset->PreviewSkeletalMesh->GetSkeleton();
-}
-
-void UIKRetargeter::AddRetargetPose(FName NewPoseName)
-{
-	if (RetargetPoses.Contains(NewPoseName))
-	{
-		return;
-	}
-
-	RetargetPoses.Add(NewPoseName);
-	CurrentRetargetPose = NewPoseName;
-
-	Modify();
-}
-
-void UIKRetargeter::RemoveRetargetPose(FName PoseToRemove)
-{
-	if (PoseToRemove == DefaultPoseName)
-	{
-		return; // cannot remove default pose
-	}
-
-	if (!RetargetPoses.Contains(PoseToRemove))
-	{
-		return; // cannot remove pose that doesn't exist
-	}
-
-	RetargetPoses.Remove(PoseToRemove);
-
-	// did we remove the currently used pose?
-	if (CurrentRetargetPose == PoseToRemove)
-	{
-		CurrentRetargetPose = UIKRetargeter::DefaultPoseName;
-	}
-
-	Modify();
-}
-
-void UIKRetargeter::ResetRetargetPose(FName PoseToReset)
-{
-	if (!RetargetPoses.Contains(PoseToReset))
-	{
-		return; // cannot reset pose that doesn't exist
-	}
-
-	RetargetPoses[PoseToReset].BoneRotationOffsets.Reset();
-	RetargetPoses[PoseToReset].RootTranslationOffset = FVector::ZeroVector;
-	RetargetPoses[PoseToReset].RootRotationOffset = FQuat::Identity;
-	
-	Modify();
-}
-
-void UIKRetargeter::SetCurrentRetargetPose(FName CurrentPose)
-{
-	check(RetargetPoses.Contains(CurrentPose));
-	CurrentRetargetPose = CurrentPose;
-	Modify();
-}
-
-void UIKRetargeter::AddRotationOffsetToRetargetPoseBone(FName BoneName, FQuat RotationOffset)
-{
-	RetargetPoses[CurrentRetargetPose].AddRotationDeltaToBone(BoneName, RotationOffset);
-	Modify();
-}
-
-bool UIKRetargeter::Modify(bool bAlwaysMarkDirty)
-{
-	const bool bSavedToTransactionBuffer = Super::Modify(bAlwaysMarkDirty);
-	AssetVersion++; // inform any runtime/editor systems they should copy-up modifications	
-	return bSavedToTransactionBuffer;
-}
-
-void UIKRetargeter::PostEditChangeProperty(FPropertyChangedEvent& PropertyChangedEvent)
-{
-	if (PropertyChangedEvent.GetPropertyName() == "TargetIKRigAsset"
-		|| PropertyChangedEvent.GetPropertyName() == "SourceIKRigAsset")
-	{
-		AutoMapChains();
-		CleanChainMapping();
-	}
-}
-
-void UIKRetargeter::PostLoad()
-{
-	Super::PostLoad();
-	CleanChainMapping();
-	CleanPoseList();
-}
-#endif
-
-FTransform UIKRetargeter::GetTargetBoneRetargetPoseGlobalTransform(const int32& TargetBoneIndex) const
-{
-	check(bIsLoadedAndValid);
-	check(TargetSkeleton.BoneNames.IsValidIndex(TargetBoneIndex))
-
-	// get the current retarget pose
-	return TargetSkeleton.RetargetGlobalPose[TargetBoneIndex];
-}
-
-FRetargetChainMap* UIKRetargeter::GetChainMap(const FName& TargetChainName)
-{
-	for (FRetargetChainMap& ChainMap : ChainMapping)
-	{
-		if (ChainMap.TargetChain == TargetChainName)
-		{
-			return &ChainMap;
-		}
-	}
-
-	return nullptr;
+	RootTranslationOffset += TranslateDelta;
 }
diff --git a/Engine/Plugins/Animation/IKRig/Source/IKRig/Private/Solvers/IKRig_BodyMover.cpp b/Engine/Plugins/Animation/IKRig/Source/IKRig/Private/Solvers/IKRig_BodyMover.cpp
--- a/Engine/Plugins/Animation/IKRig/Source/IKRig/Private/Solvers/IKRig_BodyMover.cpp
+++ b/Engine/Plugins/Animation/IKRig/Source/IKRig/Private/Solvers/IKRig_BodyMover.cpp
@@ -4,10 +4,7 @@
 #include "IKRigDataTypes.h"
 #include "IKRigSkeleton.h"
 
-
-UIKRig_BodyMover::UIKRig_BodyMover()
-{
-}
+#define LOCTEXT_NAMESPACE "UIKRig_BodyMover"
 
 void UIKRig_BodyMover::Initialize(const FIKRigSkeleton& IKRigSkeleton)
 {
@@ -39,8 +36,8 @@ void UIKRig_BodyMover::Solve(FIKRigSkeleton& IKRigSkeleton, const FIKRigGoalCont
 	FVector CurrentCentroid = FVector::ZeroVector;
 	for (UIKRig_BodyMoverEffector* Effector : Effectors)
 	{
-		FIKRigGoal Goal;
-		if (!Goals.GetGoalByName(Effector->GoalName, Goal))
+		const FIKRigGoal* Goal = Goals.FindGoalByName(Effector->GoalName);
+		if (!Goal)
 		{
 			return;
 		}
@@ -49,7 +46,7 @@ void UIKRig_BodyMover::Solve(FIKRigSkeleton& IKRigSkeleton, const FIKRigGoalCont
 		const FTransform InitialEffector = IKRigSkeleton.CurrentPoseGlobal[BoneIndex];
 
 		InitialCentroid += InitialEffector.GetTranslation();
-		CurrentCentroid += Goal.FinalBlendedPosition;
+		CurrentCentroid += Goal->FinalBlendedPosition;
 	}
 
 	// average centroids
@@ -63,8 +60,8 @@ void UIKRig_BodyMover::Solve(FIKRigSkeleton& IKRigSkeleton, const FIKRigGoalCont
 	FVector DZ = FVector::ZeroVector;
 	for (UIKRig_BodyMoverEffector* Effector : Effectors)
 	{
-		FIKRigGoal Goal;
-		if (!Goals.GetGoalByName(Effector->GoalName, Goal))
+		const FIKRigGoal* Goal = Goals.FindGoalByName(Effector->GoalName);
+		if (!Goal)
 		{
 			return;
 		}
@@ -80,7 +77,7 @@ void UIKRig_BodyMover::Solve(FIKRigSkeleton& IKRigSkeleton, const FIKRigGoalCont
 		// P is normalized vector from INITIAL centroid to INITIAL point
 		// Q is normalized vector from CURRENT centroid to CURRENT point
 		FVector P = (InitialEffector.GetTranslation() - InitialCentroid).GetSafeNormal();
-		FVector Q = (Goal.FinalBlendedPosition - CurrentCentroid).GetSafeNormal();
+		FVector Q = (Goal->FinalBlendedPosition - CurrentCentroid).GetSafeNormal();
 		// PQ^T is the outer product of P and Q which is a 3x3 matrix
 		// https://en.m.wikipedia.org/wiki/Outer_product
 		DX += FVector(P[0]*Q[0], P[0]*Q[1], P[0]*Q[2]);
@@ -128,7 +125,7 @@ void UIKRig_BodyMover::ExtractRotation(
 		FVector RCol1(R.M[1][0], R.M[1][1], R.M[1][2]);
 		FVector RCol2(R.M[2][0], R.M[2][1], R.M[2][2]);
 		FVector Omega = RCol0.Cross(DX) + RCol1.Cross(DY) + RCol2.Cross(DZ);
-		Omega *= (1.0f / fabs(RCol0.Dot(DX) + RCol1.Dot(DY) + RCol2.Dot(DZ)) + SMALL_NUMBER);
+		Omega *= 1.0f / (fabs(RCol0.Dot(DX) + RCol1.Dot(DY) + RCol2.Dot(DZ)) + SMALL_NUMBER);
 		const float W = Omega.Size();
 		if (W < SMALL_NUMBER)
 		{
@@ -139,6 +136,8 @@ void UIKRig_BodyMover::ExtractRotation(
 	}
 }
 
+#if WITH_EDITOR
+
 void UIKRig_BodyMover::UpdateSolverSettings(UIKRigSolver* InSettings)
 {
 	if(UIKRig_BodyMover* Settings = Cast<UIKRig_BodyMover>(InSettings))
@@ -171,6 +170,28 @@ void UIKRig_BodyMover::UpdateSolverSettings(UIKRigSolver* InSettings)
 	}
 }
 
+FText UIKRig_BodyMover::GetNiceName() const
+{
+	return FText(LOCTEXT("SolverName", "Body Mover"));
+}
+
+bool UIKRig_BodyMover::GetWarningMessage(FText& OutWarningMessage) const
+{
+	if (BodyBone == NAME_None)
+	{
+		OutWarningMessage = LOCTEXT("MissingRoot", "Missing root bone.");
+		return true;
+	}
+
+	if (Effectors.IsEmpty())
+	{
+		OutWarningMessage = LOCTEXT("MissingGoal", "Missing goals.");
+		return true;
+	}
+	
+	return false;
+}
+
 void UIKRig_BodyMover::AddGoal(const UIKRigEffectorGoal& NewGoal)
 {
 	UIKRig_BodyMoverEffector* NewEffector = NewObject<UIKRig_BodyMoverEffector>(this, UIKRig_BodyMoverEffector::StaticClass());
@@ -223,7 +244,7 @@ bool UIKRig_BodyMover::IsGoalConnected(const FName& GoalName) const
 	return GetIndexOfGoal(GoalName) != INDEX_NONE;
 }
 
-UObject* UIKRig_BodyMover::GetEffectorWithGoal(const FName& GoalName)
+UObject* UIKRig_BodyMover::GetEffectorWithGoal(const FName& GoalName) const
 {
 	const int32 GoalIndex = GetIndexOfGoal(GoalName);
 	if (GoalIndex == INDEX_NONE)
@@ -256,3 +277,7 @@ int32 UIKRig_BodyMover::GetIndexOfGoal(const FName& OldName) const
 
 	return INDEX_NONE;
 }
+
+#endif
+
+#undef LOCTEXT_NAMESPACE
diff --git a/Engine/Plugins/Animation/IKRig/Source/IKRig/Private/Solvers/IKRig_PBIKSolver.cpp b/Engine/Plugins/Animation/IKRig/Source/IKRig/Private/Solvers/IKRig_PBIKSolver.cpp
--- a/Engine/Plugins/Animation/IKRig/Source/IKRig/Private/Solvers/IKRig_PBIKSolver.cpp
+++ b/Engine/Plugins/Animation/IKRig/Source/IKRig/Private/Solvers/IKRig_PBIKSolver.cpp
@@ -5,6 +5,8 @@
 #include "IKRigDataTypes.h"
 #include "IKRigSkeleton.h"
 
+#define LOCTEXT_NAMESPACE "UIKRigPBIKSolver"
+
 void UIKRigPBIKSolver::Initialize(const FIKRigSkeleton& InSkeleton)
 {	
 	// check how many effectors are assigned to a bone
@@ -89,12 +91,17 @@ void UIKRigPBIKSolver::Solve(FIKRigSkeleton& IKRigSkeleton, const FIKRigGoalCont
 		{
 			continue;
 		}
-		FIKRigGoal Goal;
-		Goals.GetGoalByName(Effector->GoalName, Goal);
+		
+		const FIKRigGoal* Goal = Goals.FindGoalByName(Effector->GoalName);
+		if (!Goal)
+		{
+			return;
+		}
+		
 		Solver.SetEffectorGoal(
 			Effector->IndexInSolver,
-			Goal.FinalBlendedPosition,
-			Goal.FinalBlendedRotation,
+			Goal->FinalBlendedPosition,
+			Goal->FinalBlendedRotation,
 			OffsetAlpha,
 			Effector->StrengthAlpha,
 			Effector->PullChainAlpha,
@@ -123,6 +130,8 @@ void UIKRigPBIKSolver::Solve(FIKRigSkeleton& IKRigSkeleton, const FIKRigGoalCont
 	}
 }
 
+#if WITH_EDITOR
+
 void UIKRigPBIKSolver::UpdateSolverSettings(UIKRigSolver* InSettings)
 {
 	if(UIKRigPBIKSolver* Settings = Cast<UIKRigPBIKSolver>(InSettings))
@@ -163,6 +172,28 @@ void UIKRigPBIKSolver::UpdateSolverSettings(UIKRigSolver* InSettings)
 	}
 }
 
+FText UIKRigPBIKSolver::GetNiceName() const
+{
+	return FText(LOCTEXT("SolverName", "Full-Body IK"));
+}
+
+bool UIKRigPBIKSolver::GetWarningMessage(FText& OutWarningMessage) const
+{
+	if (RootBone == NAME_None)
+	{
+		OutWarningMessage = LOCTEXT("MissingRoot", "Missing root bone.");
+		return true;
+	}
+
+	if (Effectors.IsEmpty())
+	{
+		OutWarningMessage = LOCTEXT("MissingGoal", "Missing goals.");
+		return true;
+	}
+	
+	return false;
+}
+
 void UIKRigPBIKSolver::AddGoal(const UIKRigEffectorGoal& NewGoal)
 {
 	UIKRig_FBIKEffector* NewEffector = NewObject<UIKRig_FBIKEffector>(this, UIKRig_FBIKEffector::StaticClass());
@@ -220,7 +251,7 @@ void UIKRigPBIKSolver::SetRootBone(const FName& RootBoneName)
 	RootBone = RootBoneName;
 }
 
-UObject* UIKRigPBIKSolver::GetEffectorWithGoal(const FName& GoalName)
+UObject* UIKRigPBIKSolver::GetEffectorWithGoal(const FName& GoalName) const
 {
 	const int32 GoalIndex = GetIndexOfGoal(GoalName);
 	if (GoalIndex == INDEX_NONE)
@@ -327,3 +358,7 @@ void UIKRigPBIKSolver::PostLoad()
 		Effectors.Empty();
 	}
 }
+
+#endif
+
+#undef LOCTEXT_NAMESPACE
diff --git a/Engine/Plugins/Animation/IKRig/Source/IKRig/Private/Solvers/IKRig_SetTransform.cpp b/Engine/Plugins/Animation/IKRig/Source/IKRig/Private/Solvers/IKRig_SetTransform.cpp
--- a/Engine/Plugins/Animation/IKRig/Source/IKRig/Private/Solvers/IKRig_SetTransform.cpp
+++ b/Engine/Plugins/Animation/IKRig/Source/IKRig/Private/Solvers/IKRig_SetTransform.cpp
@@ -4,6 +4,7 @@
 #include "IKRigDataTypes.h"
 #include "IKRigSkeleton.h"
 
+#define LOCTEXT_NAMESPACE "UIKRig_SetTransform"
 
 UIKRig_SetTransform::UIKRig_SetTransform()
 {
@@ -17,8 +18,8 @@ void UIKRig_SetTransform::Initialize(const FIKRigSkeleton& IKRigSkeleton)
 
 void UIKRig_SetTransform::Solve(FIKRigSkeleton& IKRigSkeleton, const FIKRigGoalContainer& Goals)
 {
-	FIKRigGoal OutGoal;
-	if (!Goals.GetGoalByName(Goal, OutGoal))
+	const FIKRigGoal* InGoal = Goals.FindGoalByName(Goal);
+	if (!InGoal)
 	{
 		return;
 	}
@@ -35,19 +36,21 @@ void UIKRig_SetTransform::Solve(FIKRigSkeleton& IKRigSkeleton, const FIKRigGoalC
 
 	if (Effector->bEnablePosition)
 	{
-		const FVector TargetPosition = FMath::Lerp(CurrentTransform.GetTranslation(), OutGoal.FinalBlendedPosition, Effector->Alpha);
+		const FVector TargetPosition = FMath::Lerp(CurrentTransform.GetTranslation(), InGoal->FinalBlendedPosition, Effector->Alpha);
 		CurrentTransform.SetLocation(TargetPosition);
 	}
 	
 	if (Effector->bEnableRotation)
 	{
-		const FQuat TargetRotation = FMath::Lerp(CurrentTransform.GetRotation(), OutGoal.FinalBlendedRotation, Effector->Alpha);
+		const FQuat TargetRotation = FMath::Lerp(CurrentTransform.GetRotation(), InGoal->FinalBlendedRotation, Effector->Alpha);
 		CurrentTransform.SetRotation(TargetRotation);
 	}
 	
 	IKRigSkeleton.PropagateGlobalPoseBelowBone(BoneIndex);
 }
 
+#if WITH_EDITOR
+
 void UIKRig_SetTransform::UpdateSolverSettings(UIKRigSolver* InSettings)
 {
 	if (UIKRig_SetTransform* Settings = Cast<UIKRig_SetTransform>(InSettings))
@@ -58,6 +61,21 @@ void UIKRig_SetTransform::UpdateSolverSettings(UIKRigSolver* InSettings)
 	}
 }
 
+FText UIKRig_SetTransform::GetNiceName() const
+{
+	return FText(LOCTEXT("SolverName", "Set Transform"));
+}
+
+bool UIKRig_SetTransform::GetWarningMessage(FText& OutWarningMessage) const
+{
+	if (Bone == NAME_None)
+	{
+		OutWarningMessage = LOCTEXT("MissingGoal", "Missing goals.");
+		return true;
+	}
+	return false;
+}
+
 void UIKRig_SetTransform::AddGoal(const UIKRigEffectorGoal& NewGoal)
 {
 	Goal = NewGoal.GoalName;
@@ -94,7 +112,7 @@ bool UIKRig_SetTransform::IsGoalConnected(const FName& GoalName) const
 	return Goal == GoalName;
 }
 
-UObject* UIKRig_SetTransform::GetEffectorWithGoal(const FName& GoalName)
+UObject* UIKRig_SetTransform::GetEffectorWithGoal(const FName& GoalName) const
 {
 	return Goal == GoalName ? Effector : nullptr;
 }
@@ -103,3 +121,7 @@ bool UIKRig_SetTransform::IsBoneAffectedBySolver(const FName& BoneName, const FI
 {
 	return IKRigSkeleton.IsBoneInDirectLineage(BoneName, Bone);
 }
+
+#endif
+
+#undef LOCTEXT_NAMESPACE
diff --git a/Engine/Plugins/Animation/IKRig/Source/IKRig/Public/ActorComponents/IKRigComponent.h b/Engine/Plugins/Animation/IKRig/Source/IKRig/Public/ActorComponents/IKRigComponent.h
--- a/Engine/Plugins/Animation/IKRig/Source/IKRig/Public/ActorComponents/IKRigComponent.h
+++ b/Engine/Plugins/Animation/IKRig/Source/IKRig/Public/ActorComponents/IKRigComponent.h
@@ -51,11 +51,18 @@ public:
 
 	/** Remove all stored goals in this component. */
 	UFUNCTION(BlueprintCallable, Category=IKRigGoals)
-    void ClearAllGoals(){ GoalContainer.Goals.Empty(); };
+    void ClearAllGoals(){ GoalContainer.Empty(); };
 
 	// BEGIN IIKRigGoalCreator interface
 	/** Called from the IK Rig Anim node to obtain all the goals that have been set on this actor.*/
-	void AddIKGoals_Implementation(TMap<FName, FIKRigGoal>& OutGoals) override { OutGoals.Append(GoalContainer.Goals); };
+	virtual void AddIKGoals_Implementation(TMap<FName, FIKRigGoal>& OutGoals) override
+	{
+		const TArray<FIKRigGoal>& Goals = GoalContainer.GetGoalArray();
+		for (const FIKRigGoal& Goal : Goals)
+		{
+			OutGoals.Add(Goal.Name, Goal);
+		}
+	};
 	// END IIKRigGoalCreator interface
 
 private:
diff --git a/Engine/Plugins/Animation/IKRig/Source/IKRig/Public/AnimNodes/AnimNode_IKRig.h b/Engine/Plugins/Animation/IKRig/Source/IKRig/Public/AnimNodes/AnimNode_IKRig.h
--- a/Engine/Plugins/Animation/IKRig/Source/IKRig/Public/AnimNodes/AnimNode_IKRig.h
+++ b/Engine/Plugins/Animation/IKRig/Source/IKRig/Public/AnimNodes/AnimNode_IKRig.h
@@ -5,10 +5,10 @@
 #include "CoreMinimal.h"
 
 #include "IKRigDataTypes.h"
-#include "IKRigProcessor.h"
 #include "Animation/AnimNodeBase.h"
 #include "AnimNode_IKRig.generated.h"
 
+class UIKRigProcessor;
 class IIKGoalCreatorInterface;
 class UIKRigDefinition;
 
@@ -23,8 +23,8 @@ struct IKRIG_API FAnimNode_IKRig : public FAnimNode_Base
 	FPoseLink Source;
 
 	/** The IK rig to use to modify the incoming source pose. */
-	UPROPERTY(EditAnywhere, BlueprintReadWrite, Category = RigDefinition, meta = (NeverAsPin))
-	UIKRigDefinition* RigDefinitionAsset;
+	UPROPERTY(EditAnywhere, Category = RigDefinition, meta = (NeverAsPin))
+	TObjectPtr<UIKRigDefinition> RigDefinitionAsset = nullptr;
 
 	/** The input goal transforms used by the IK Rig solvers.*/
 	UPROPERTY(EditAnywhere, BlueprintReadWrite, Category = Goal, meta = (PinShownByDefault))
@@ -34,24 +34,26 @@ struct IKRIG_API FAnimNode_IKRig : public FAnimNode_Base
 	UPROPERTY(EditAnywhere, Category = Solver)
 	bool bStartFromRefPose = false;
 
+#if WITH_EDITORONLY_DATA
+
 	/** when true, goals will use the current transforms stored in the IK Rig Definition asset itself */
 	bool bDriveWithSourceAsset = false;
-
-#if WITH_EDITORONLY_DATA
+	
 	/** Toggle debug drawing of goals when node is selected.*/
 	UPROPERTY(EditAnywhere, Category = Solver)
-	bool bEnableDebugDraw;
+	bool bEnableDebugDraw = false;
 
 	/** Adjust size of debug drawing.*/
 	UPROPERTY(EditAnywhere, Category = Solver)
 	float DebugScale = 5.0f;
+	
 #endif
 
 private:
 
 	/** IK Rig runtime processor */
 	UPROPERTY(Transient)
-	TObjectPtr<UIKRigProcessor> IKRigProcessor;
+	TObjectPtr<UIKRigProcessor> IKRigProcessor = nullptr;
 
 	/** a cached list of components on the owning actor that implement the goal creator interface */
 	TArray<IIKGoalCreatorInterface*> GoalCreators;
@@ -60,8 +62,6 @@ private:
 	TMap<FCompactPoseBoneIndex, int32, FDefaultSetAllocator, TCompactPoseBoneIndexMapKeyFuncs<int32>> CompactPoseToRigIndices;
 
 public:
-	
-	FAnimNode_IKRig();
 
 	// FAnimNode_Base interface
 	virtual void GatherDebugData(FNodeDebugData& DebugData) override;
@@ -72,6 +72,9 @@ public:
 	virtual bool HasPreUpdate() const override { return true; }
 	virtual void PreUpdate(const UAnimInstance* InAnimInstance) override;
 	// End of FAnimNode_Base interface
+
+	/** force reinitialization */
+	void SetProcessorNeedsInitialized();
 	
 private:
 	
diff --git a/Engine/Plugins/Animation/IKRig/Source/IKRig/Public/AnimNodes/AnimNode_RetargetPoseFromMesh.h b/Engine/Plugins/Animation/IKRig/Source/IKRig/Public/AnimNodes/AnimNode_RetargetPoseFromMesh.h
--- a/Engine/Plugins/Animation/IKRig/Source/IKRig/Public/AnimNodes/AnimNode_RetargetPoseFromMesh.h
+++ b/Engine/Plugins/Animation/IKRig/Source/IKRig/Public/AnimNodes/AnimNode_RetargetPoseFromMesh.h
@@ -6,6 +6,8 @@
 
 #include "Animation/AnimNodeBase.h"
 #include "Retargeter/IKRetargeter.h"
+#include "Retargeter/IKRetargetProcessor.h"
+
 #include "AnimNode_RetargetPoseFromMesh.generated.h"
 
 USTRUCT(BlueprintInternalUseOnly)
@@ -23,44 +25,42 @@ struct IKRIG_API FAnimNode_RetargetPoseFromMesh : public FAnimNode_Base
 	
 	/** Retarget asset to use. Must define a Source and Target IK Rig compatible with the SourceMeshComponent and current anim instance.*/
 	UPROPERTY(EditAnywhere, Category = Settings)
-	UIKRetargeter* IKRetargeterAsset = nullptr;
+	TObjectPtr<UIKRetargeter> IKRetargeterAsset = nullptr;
+
+#if WITH_EDITOR
+	/** when true, will copy all setting from target IK Rig asset each tick (for live preview) */
+	bool bDriveTargetIKRigWithAsset = false;
+#endif
 	
 	// FAnimNode_Base interface
 	virtual void Initialize_AnyThread(const FAnimationInitializeContext& Context) override;
 	virtual void CacheBones_AnyThread(const FAnimationCacheBonesContext& Context) override;
 	virtual void Update_AnyThread(const FAnimationUpdateContext& Context) override;
 	virtual void Evaluate_AnyThread(FPoseContext& Output) override;
-	//virtual void GatherDebugData(FNodeDebugData& DebugData) override;
 	virtual bool HasPreUpdate() const override { return true; }
 	virtual void PreUpdate(const UAnimInstance* InAnimInstance) override;
 	// End of FAnimNode_Base interface
 
-	UIKRetargeter* GetCurrentlyUsedRetargeter() const;
+#if WITH_EDITOR
+	/** Force reinitialization. */
+	void SetProcessorNeedsInitialized();
+#endif
+
+	/** Read-only access to the runtime processor */
+	const UIKRetargetProcessor* GetRetargetProcessor() const;
 
 private:
 	
 	void EnsureInitialized(const UAnimInstance* InAnimInstance);
-	void InitializeRetargetData(const UAnimInstance* InAnimInstance);
 	void CopyBoneTransformsFromSource(USkeletalMeshComponent* TargetMeshComponent);
-	USkeletalMeshComponent* GetSourceMesh() const;
-
-	// indicates that all prerequisites are met and node is ready to operate
-	UPROPERTY(Transient)
-	bool bIsInitialized = false;
-	/** which version of the IK Retargeter asset was this instance last initialized with?
-	* this allows the IK Retargeter asset to undergo modifications at runtime via the editor and trigger
-	* re-initialization by incrementing this transient asset version.
-	 */
-	int32 LastVersionTried = -2;
 	
 	// source mesh references, cached during init so that we can compare and see if it has changed
-	TWeakObjectPtr<USkeletalMeshComponent>	CurrentlyUsedSourceMeshComponent;
-	TWeakObjectPtr<USkeletalMesh>			CurrentlyUsedSourceMesh;
-	TWeakObjectPtr<USkeletalMesh>			CurrentlyUsedTargetMesh;
+	TWeakObjectPtr<USkeletalMesh> CurrentlyUsedSourceMesh;
+	TWeakObjectPtr<USkeletalMesh> CurrentlyUsedTargetMesh;
+
+	/** the runtime processor used to run the retarget and generate new poses */
 	UPROPERTY(Transient)
-	TObjectPtr<UIKRetargeter>				CurrentlyUsedRetargeter = nullptr;
-	TWeakObjectPtr<UIKRigDefinition>		CurrentlyUsedSourceIKRig;
-	TWeakObjectPtr<UIKRigDefinition>		CurrentlyUsedTargetIKRig;
+	TObjectPtr<UIKRetargetProcessor> Processor = nullptr;
 
 	// cached transforms, copied on the game thread
 	TArray<FTransform> SourceMeshComponentSpaceBoneTransforms;
diff --git a/Engine/Plugins/Animation/IKRig/Source/IKRig/Public/IKRigDataTypes.h b/Engine/Plugins/Animation/IKRig/Source/IKRig/Public/IKRigDataTypes.h
--- a/Engine/Plugins/Animation/IKRig/Source/IKRig/Public/IKRigDataTypes.h
+++ b/Engine/Plugins/Animation/IKRig/Source/IKRig/Public/IKRigDataTypes.h
@@ -4,9 +4,10 @@
 
 #include "CoreMinimal.h"
 #include "IKRigDefinition.h"
+
 #include "IKRigDataTypes.generated.h"
 
-class UIKRigEffectorGoal;
+class UIKRigProcessor;
 
 UENUM(BlueprintType)
 enum class EIKRigGoalSpace : uint8
@@ -54,11 +55,11 @@ struct IKRIG_API FIKRigGoal
 	UPROPERTY(EditAnywhere, BlueprintReadWrite, Category = Rotation, meta=(EditCondition="Source == EIKRigGoalSource::Manual"))
 	FRotator Rotation;
 
-	/** Range 0-1. Smoothly blends the Goal position from the input pose (0.0) to the Goal position (1.0). */
+	/** Range 0-1, default is 1.0. Smoothly blends the Goal position from the input pose (0.0) to the Goal position (1.0). */
 	UPROPERTY(EditAnywhere, BlueprintReadWrite, Category = Position, meta = (ClampMin = "0.0", ClampMax = "1.0", UIMin = "0.0", UIMax = "1.0"))
 	float PositionAlpha;
 
-	/** Range 0-1. Smoothly blends the Goal rotation from the input pose (0.0) to the Goal rotation (1.0). */
+	/** Range 0-1, default is 1.0. Smoothly blends the Goal rotation from the input pose (0.0) to the Goal rotation (1.0). */
 	UPROPERTY(EditAnywhere, BlueprintReadWrite, Category = Rotation, meta = (ClampMin = "0.0", ClampMax = "1.0", UIMin = "0.0", UIMax = "1.0"))
 	float RotationAlpha;
 
@@ -86,8 +87,8 @@ struct IKRIG_API FIKRigGoal
     : Name(NAME_None),
 	Position(ForceInitToZero),
 	Rotation(FRotator::ZeroRotator),
-	PositionAlpha(0.0f),
-    RotationAlpha(0.0f),
+	PositionAlpha(1.0f),
+    RotationAlpha(1.0f),
 	PositionSpace(EIKRigGoalSpace::Component),
 	RotationSpace(EIKRigGoalSpace::Component),
 	FinalBlendedPosition(ForceInitToZero),
@@ -97,8 +98,8 @@ struct IKRIG_API FIKRigGoal
     : Name(GoalName),
 	Position(ForceInitToZero),
 	Rotation(FRotator::ZeroRotator),
-    PositionAlpha(0.0f),
-    RotationAlpha(0.0f),
+    PositionAlpha(1.0f),
+    RotationAlpha(1.0f),
 	PositionSpace(EIKRigGoalSpace::Component),
     RotationSpace(EIKRigGoalSpace::Component),
 	FinalBlendedPosition(ForceInitToZero),
@@ -146,22 +147,28 @@ struct IKRIG_API FIKRigGoalContainer
 	GENERATED_BODY()
 
 public:
-	
-	/** Pre-load all the names of goals (optional, you can just call SetIKGoal to add as needed) */
-	void InitializeFromGoals(const TArray<UIKRigEffectorGoal*>& InGoals);
 
-	/** Set an IK goal to go to a specific location and rotation (in component space) blended by alpha.
+	/** Set an IK goal to go to a specific location and rotation (in the specified space) blended by alpha.
 	 * Will ADD the goal if none exist with the input name. */
 	void SetIKGoal(const FIKRigGoal& InGoal);
 
-	/** Set an IK goal to go to a specific location and rotation (in component space) blended by alpha.
+	/** Set an IK goal to go to a specific location and rotation (in the specified space) blended by alpha.
 	* Will ADD the goal if none exist with the input name. */
 	void SetIKGoal(const UIKRigEffectorGoal* InEffectorGoal);
 
-	/** Get an IK goal with the given name. Returns false if no goal is found in the container with the name. */
-	bool GetGoalByName(const FName& InGoalName, FIKRigGoal& OutGoal) const;
+	/** Get an IK goal with the given name. Returns nullptr if no Goal is found in the container with the name. */
+	const FIKRigGoal* FindGoalByName(const FName& GoalName) const;
+
+	/** Clear out all Goals in container. */
+	void Empty() { Goals.Empty(); };
+
+	/** Get read-only access to array of Goals. */
+	const TArray<FIKRigGoal>& GetGoalArray() const { return Goals; };
+	
+private:
+	
+	/** Array of Goals. Cannot contain duplicates (name is key). */
+	TArray<FIKRigGoal> Goals;
 
-	/** Keys are IK Rig Goal names. Values are the Goal data structures.
-	 * These are consumed by an IKRig asset to drive effectors. */
-	TMap<FName, FIKRigGoal> Goals;
+	friend UIKRigProcessor;
 };
diff --git a/Engine/Plugins/Animation/IKRig/Source/IKRig/Public/IKRigDefinition.h b/Engine/Plugins/Animation/IKRig/Source/IKRig/Public/IKRigDefinition.h
--- a/Engine/Plugins/Animation/IKRig/Source/IKRig/Public/IKRigDefinition.h
+++ b/Engine/Plugins/Animation/IKRig/Source/IKRig/Public/IKRigDefinition.h
@@ -10,12 +10,14 @@
 
 class UIKRigSolver;
 
+#if WITH_EDITORONLY_DATA
 UENUM(BlueprintType)
 enum class EIKRigGoalPreviewMode : uint8
 {
 	Additive		UMETA(DisplayName = "Additive"),
 	Absolute		UMETA(DisplayName = "Absolute"),
 };
+#endif
 
 UCLASS()
 class IKRIG_API UIKRigEffectorGoal : public UObject
@@ -23,58 +25,62 @@ class IKRIG_API UIKRigEffectorGoal : public UObject
 	GENERATED_BODY()
 
 public:
-	
-	UIKRigEffectorGoal() : GoalName("DefaultGoal"), BoneName(NAME_None){}
-	
-	UIKRigEffectorGoal(const FName& InGoalName, const FName& InBoneName) :  GoalName(InGoalName), BoneName(InBoneName){}
 
 	/** The name used to refer to this goal from outside systems.
-	 *This is the name to use when referring to this Goal from Blueprint, Anim Graph, Control Rig or IK Retargeter.*/
+	 * This is the name to use when referring to this Goal from Blueprint, Anim Graph, Control Rig or IK Retargeter.*/
 	UPROPERTY(VisibleAnywhere, Category = "Goal Settings")
 	FName GoalName;
 
-	/**The name of the bone that this Goal is located at.*/
+	/** The name of the bone that this Goal is located at.*/
 	UPROPERTY(VisibleAnywhere, Category = "Goal Settings")
 	FName BoneName;
 
-	/**Range 0-1, default is 1. Blend between the Goal position between the input bone pose (0.0) and the current goal transform (1.0).*/
+	/** Range 0-1, default is 1. Blend between the Goal position between the input bone pose (0.0) and the current goal transform (1.0).*/
 	UPROPERTY(EditAnywhere, Category = "Goal Settings", meta = (ClampMin = "0.0", ClampMax = "1.0", UIMin = "0.0", UIMax = "1.0"))
 	float PositionAlpha = 1.0f;
 
-	/**Range 0-1, default is 1. Blend between the Goal rotation between the input bone pose (0.0) and the current goal transform (1.0).*/
+	/** Range 0-1, default is 1. Blend between the Goal rotation between the input bone pose (0.0) and the current goal transform (1.0).*/
 	UPROPERTY(EditAnywhere, Category = "Goal Settings", meta = (ClampMin = "0.0", ClampMax = "1.0", UIMin = "0.0", UIMax = "1.0"))
 	float RotationAlpha = 1.0f;
 
-	/**Effects how this Goal transform is previewed in the IK Rig editor.
-	 *"Additive" interprets the Goal transform as being relative to the input pose. Useful for previewing animations.
-	 *"Absolute" pins the Goal transform to the Gizmo in the viewport.
-	 */
+	/** The current transform of this Goal, in the Global Space of the character.*/
 	UPROPERTY(EditAnywhere, Category = "Goal Settings")
-	EIKRigGoalPreviewMode PreviewMode;
-
-	/**The current transform of this Goal, in the Global Space of the character.*/
-	UPROPERTY(EditAnywhere, Transient, Category = "Goal Settings")
 	FTransform CurrentTransform;
 
-	/**The initial transform of this Goal, as defined by the initial transform of the Goal's bone in the reference pose.*/
-	UPROPERTY(VisibleAnywhere, Transient, Category = "Goal Settings")
+	/** The initial transform of this Goal, as defined by the initial transform of the Goal's bone in the retarget pose.*/
+	UPROPERTY(VisibleAnywhere, Category = "Goal Settings")
 	FTransform InitialTransform;
+	
+	bool operator==(const UIKRigEffectorGoal& Other) const { return GoalName == Other.GoalName; }
+
+#if WITH_EDITORONLY_DATA
 
+	/** Effects how this Goal transform is previewed in the IK Rig editor.
+	* "Additive" interprets the Goal transform as being relative to the input pose. Useful for previewing animations.
+	* "Absolute" pins the Goal transform to the Gizmo in the viewport.
+	*/
+	UPROPERTY(EditAnywhere, Category = "Goal Settings")
+	EIKRigGoalPreviewMode PreviewMode = EIKRigGoalPreviewMode::Additive;
+	
 	/**The size of the Goal gizmo drawing in the editor viewport.*/
-	UPROPERTY(EditAnywhere, Category = "Goal Gizmo", meta = (ClampMin = "0.1", ClampMax = "1000.0", UIMin = "0.1", UIMax = "100.0"))
-	float GizmoSize = 7.0f;
+	UPROPERTY(EditAnywhere, Category = "Viewport Goal Settings", meta = (ClampMin = "0.1", ClampMax = "1000.0", UIMin = "0.1", UIMax = "100.0"))
+	float SizeMultiplier = 1.0f;
 
 	/**The thickness of the Goal gizmo drawing in the editor viewport.*/
-	UPROPERTY(EditAnywhere, Category = "Goal Gizmo",  meta = (ClampMin = "0.0", ClampMax = "10.0", UIMin = "0.0", UIMax = "5.0"))
-	float GizmoThickness = 0.7f;
-
-	bool operator==(const UIKRigEffectorGoal& Other) const { return GoalName == Other.GoalName; }
+	UPROPERTY(EditAnywhere, Category = "Viewport Goal Settings",  meta = (ClampMin = "0.0", ClampMax = "10.0", UIMin = "0.0", UIMax = "5.0"))
+	float ThicknessMultiplier = 0.7f;
 
 	virtual void PostLoad() override
 	{
 		Super::PostLoad();
 		SetFlags(RF_Transactional);
 	}
+	
+#endif
+
+private:
+	UPROPERTY()
+	TObjectPtr<UIKRigDefinition> ParentIKRig = nullptr;
 };
 
 USTRUCT(Blueprintable)
@@ -91,41 +97,36 @@ struct IKRIG_API FBoneChain
 	IKGoalName(NAME_None){}
 
 	UPROPERTY(EditAnywhere, Category = BoneChain)
-	FName ChainName = NAME_None;
+	FName ChainName;
 
 	UPROPERTY(EditAnywhere, Category = BoneChain)
-	FName StartBone = NAME_None;
+	FName StartBone;
 
 	UPROPERTY(EditAnywhere, Category = BoneChain)
-	FName EndBone = NAME_None;
+	FName EndBone;
 	
 	UPROPERTY(EditAnywhere, Category = IK)
-	FName IKGoalName = NAME_None;
-
-	/*
-	UPROPERTY(EditAnywhere, Category = IK)
-	FName PoleVectorGoalName;
-
-	UPROPERTY(EditAnywhere, Category = IK)
-	FName PoleVectorBoneName;
-
-	UPROPERTY(EditAnywhere, Category = IK)
-	FVector PoleVectorDirection;
-	*/
+	FName IKGoalName;
 };
 
 USTRUCT(Blueprintable)
 struct IKRIG_API FRetargetDefinition
 {
 	GENERATED_BODY()
+	
+private:
 
-	UPROPERTY(EditAnywhere, Category = BodyParts)
+	UPROPERTY()
 	FName RootBone;
 
-	UPROPERTY(EditAnywhere, Category = BodyParts)
+	UPROPERTY()
 	TArray<FBoneChain> BoneChains;
 
-	FBoneChain* GetBoneChainByName(FName ChainName);
+	/** read/write access restricted to friends only */
+	FBoneChain* GetEditableBoneChainByName(FName ChainName);
+	
+	friend class UIKRigController;
+	friend class UIKRigDefinition;
 };
 
 UCLASS(Blueprintable)
@@ -135,55 +136,60 @@ class IKRIG_API UIKRigDefinition : public UObject, public IInterface_PreviewMesh
 
 public:
 
-	/** the skeletal mesh that was used as the source of the skeleton data. also used for preview.
-	 * NOTE: the IK rig may be played back on ANY skeleton that is compatible with it's hierarchy. */
-	UPROPERTY(AssetRegistrySearchable, VisibleAnywhere, Category = "Skeleton")
-	TObjectPtr<class USkeletalMesh> PreviewSkeletalMesh;
+#if WITH_EDITORONLY_DATA
+	
+	/**The size of the Goals in the editor viewport.*/
+	UPROPERTY(EditAnywhere, Category = "Goal Drawing", meta = (ClampMin = "0.01", UIMin = "0.1", UIMax = "100.0"))
+	float GoalSize = 7.0f;
 
+	/** The thickness of the Goals in the editor viewport.*/
+	UPROPERTY(EditAnywhere, Category = "Goal Drawing",  meta = (ClampMin = "0.01", UIMin = "0.0", UIMax = "10.0"))
+	float GoalThickness = 0.7f;
+	
+#endif
+
+	/** The skeletal mesh that was used as the source of the skeleton data. Also used for preview.
+	* NOTE: the IK rig may be played back on ANY skeleton that is compatible with it's hierarchy. */
+	UPROPERTY(AssetRegistrySearchable, VisibleAnywhere, Category = "Imported Skeleton")
+	TObjectPtr<class USkeletalMesh> PreviewSkeletalMesh;
+	
 	/** hierarchy and bone-pose transforms */
-	UPROPERTY(VisibleAnywhere, Category = "Skeleton")
+	UPROPERTY()
 	FIKRigSkeleton Skeleton;
+	
+	/** read-only access to array of Goals, all modifications must go through UIKRigController */
+	const TArray<UIKRigEffectorGoal*>& GetGoalArray() const { return Goals; };
 
-	/** stack of solvers, of varying types, executed in serial fashion where
-	output of prior solve is input to the next */
-	UPROPERTY(VisibleAnywhere, instanced, Category = "IK")
-	TArray<UIKRigSolver*> Solvers;
+	/** read-only access to array of Solvers, all modifications must go through UIKRigController */
+	const TArray<UIKRigSolver*>& GetSolverArray() const { return Solvers; };
 
-	/** goals, used as effectors by solvers that support them */
-	UPROPERTY(VisibleAnywhere, Category = "IK")
-	TArray<UIKRigEffectorGoal*> Goals;
+	/** runtime, read-only access to Retarget Definition, all modifications must go through UIKRigController */
+	const TArray<FBoneChain>& GetRetargetChains() const { return RetargetDefinition.BoneChains; };
 
-	/** bone chains for animation retargeting */
-	UPROPERTY(EditAnywhere, Category = "Retargeting")
-	FRetargetDefinition RetargetDefinition;
+	/** runtime, read-only access to a FBoneChain by name. Returns nullptr if not found. */
+	const FBoneChain* GetRetargetChainByName(FName ChainName) const;
 
-	/** editors systems can use this to check if they have most up-to-date settings */
-	int32 GetAssetVersion() const {return AssetVersion;};
+	/** runtime, read-only access to Retarget Root */
+	const FName& GetRetargetRoot() const { return RetargetDefinition.RootBone; };
 	
 	/** IInterface_PreviewMeshProvider interface */
 	virtual void SetPreviewMesh(USkeletalMesh* PreviewMesh, bool bMarkAsDirty = true) override;
 	virtual USkeletalMesh* GetPreviewMesh() const override;
 	/** END IInterface_PreviewMeshProvider interface */
 
-	USkeleton* GetSkeletonAsset() const;
-
 private:
-	
-	/** UObject interface */
-#if WITH_EDITOR
-	virtual bool Modify(bool bAlwaysMarkDirty = true) override;
-#endif
-	virtual void PostLoad() override;
-	/** END UObject interface */
 
-	/** this is incremented whenever the asset is modified in such a way that would require
-	* any IKRigProcessors that are using it to reinitialize with the latest version */
-	int32 AssetVersion = 0;
+	/** goals, used as effectors by solvers that support them */
+	UPROPERTY()
+	TArray<TObjectPtr<UIKRigEffectorGoal>> Goals;
+	
+	/** polymorphic stack of solvers, executed in serial fashion where output of prior solve is input to the next */
+	UPROPERTY(instanced)
+	TArray<TObjectPtr<UIKRigSolver>> Solvers;
 
-	void ResetGoalTransforms() const;
-	FTransform GetGoalInitialTransform(UIKRigEffectorGoal* Goal) const;
+	/** bone chains for IK Retargeter */
+	UPROPERTY()
+	FRetargetDefinition RetargetDefinition;
 
-	void SortRetargetChains();
-	
 	friend class UIKRigController;
 };
diff --git a/Engine/Plugins/Animation/IKRig/Source/IKRig/Public/IKRigProcessor.h b/Engine/Plugins/Animation/IKRig/Source/IKRig/Public/IKRigProcessor.h
--- a/Engine/Plugins/Animation/IKRig/Source/IKRig/Public/IKRigProcessor.h
+++ b/Engine/Plugins/Animation/IKRig/Source/IKRig/Public/IKRigProcessor.h
@@ -39,7 +39,13 @@ public:
 	*/
 	 
 	/** setup a new processor to run the given IKRig asset
-	 *  NOTE!! this function creates new UObjects and consequently MUST be called from the main thread!!	 */
+	 *  NOTE!! this function creates new UObjects and consequently MUST be called from the main thread!!
+	 *  @param InRigAsset - the IK Rig defining the collection of solvers to execute and all the rig settings
+	 *  @param InputSkeleton - the skeleton in reference pose that you want to solve the IK on
+	 */
+	void Initialize(UIKRigDefinition* InRigAsset, const FIKRigInputSkeleton& InputSkeleton);
+
+	/** Anim Graph: convenience to initialize directly from an FReferenceSkeleton. */
 	void Initialize(UIKRigDefinition* InRigAsset, const FReferenceSkeleton& RefSkeleton);
 
 	//
@@ -54,10 +60,10 @@ public:
 	/** Optionally can be called before Solve() to use the reference pose as start pose */
 	void SetInputPoseToRefPose();
 
-	/** Set a named IK goal to go to a specific location and rotation (assumed in component space) blended by separate position/rotation alpha (0-1)*/
+	/** Set a named IK goal to go to a specific location, rotation and space, blended by separate position/rotation alpha (0-1)*/
 	void SetIKGoal(const FIKRigGoal& Goal);
 
-	/** Set a named IK goal to go to a specific location and rotation (assumed in component space) blended by separate position/rotation alpha (0-1)*/
+	/** Set a named IK goal to go to a specific location, rotation and space, blended by separate position/rotation alpha (0-1)*/
 	void SetIKGoal(const UIKRigEffectorGoal* Goal);
 
 	/** Run entire stack of solvers.
@@ -67,21 +73,27 @@ public:
 	/** Get the results after calling Solve() */
 	void CopyOutputGlobalPoseToArray(TArray<FTransform>& OutputPoseGlobal) const;
 
+	/** Reset all internal data structures. Will require re-initialization before solving again. */
+	void Reset();
+
 	//
 	// END UPDATE SEQUENCE FUNCTIONS
 	//
 
-	/** Used to propagate setting values from the source asset at runtime (settings that do not require re-initialization) */
-	void CopyAllInputsFromSourceAssetAtRuntime(UIKRigDefinition* IKRigAsset);
-
-	/** checks if the source IKRig asset has been modified in a way that would require reinitialization. */
-	bool NeedsInitialized(UIKRigDefinition* IKRigAsset) const;
-
 	/** Get access to the internal goal data (read only) */
 	const FIKRigGoalContainer& GetGoalContainer() const;
 	
 	/** Get access to the internal skeleton data */
 	FIKRigSkeleton& GetSkeleton();
+
+	bool IsInitialized() const { return bInitialized; };
+
+	void SetNeedsInitialized();
+
+#if WITH_EDITOR
+	/** Used to propagate setting values from the source asset at runtime (settings that do not require re-initialization) */
+	void CopyAllInputsFromSourceAssetAtRuntime(UIKRigDefinition* SourceAsset);
+#endif
 	
 private:
 
@@ -103,8 +115,5 @@ private:
 
 	/** solving disabled until this flag is true */
 	bool bInitialized = false;
-	/** which version of the IKRig asset was this instance last initialized with?
-	 * this allows the IKRig asset to undergo modifications at runtime via the editor*/
-	int32 InitializedWithIKRigAssetVersion = -1;
-	int32 LastVersionTried = -2;
+	bool bTriedToInitialize = false;
 };
\ No newline at end of file
diff --git a/Engine/Plugins/Animation/IKRig/Source/IKRig/Public/IKRigSkeleton.h b/Engine/Plugins/Animation/IKRig/Source/IKRig/Public/IKRigSkeleton.h
--- a/Engine/Plugins/Animation/IKRig/Source/IKRig/Public/IKRigSkeleton.h
+++ b/Engine/Plugins/Animation/IKRig/Source/IKRig/Public/IKRigSkeleton.h
@@ -21,6 +21,45 @@ struct IKRIG_API FIKRigSkeletonChain
 	FName EndBone;
 };
 
+/** Data used just to initialize an IKRigSkeleton from outside systems
+ *
+ * The input skeleton may be different than the skeleton that the IK Rig asset was created for, within some limits.
+ * 1. It must have all the Bones that the IK Rig asset referenced (must be a sub-set)
+ * 2. All the bones must have the same parents (no change in hierarchy)
+ *
+ * You can however add additional bones, change the reference pose (including proportions) and the bone indices.
+ * This allows you to run the same IK Rig asset on somewhat different skeletal meshes.
+ */
+USTRUCT()
+struct IKRIG_API FIKRigInputSkeleton
+{
+	GENERATED_BODY()
+	
+	TArray<FName> BoneNames;
+	TArray<int32> ParentIndices;
+	TArray<FTransform> LocalRefPose;
+
+	void InitializeFromRefSkeleton(const FReferenceSkeleton& RefSkeleton)
+	{
+		Reset();
+		
+		const TArray<FMeshBoneInfo>& BoneInfo = RefSkeleton.GetRefBoneInfo();
+		for (int32 BoneIndex=0; BoneIndex<BoneInfo.Num(); ++BoneIndex)
+		{
+			BoneNames.Add(BoneInfo[BoneIndex].Name);
+			ParentIndices.Add(BoneInfo[BoneIndex].ParentIndex);
+			LocalRefPose.Add(RefSkeleton.GetRefBonePose()[BoneIndex]);
+		}
+	}
+
+	void Reset()
+	{
+		BoneNames.Reset();
+		ParentIndices.Reset();
+		LocalRefPose.Reset();
+	}
+};
+
 USTRUCT()
 struct IKRIG_API FIKRigSkeleton
 {
@@ -50,7 +89,9 @@ struct IKRIG_API FIKRigSkeleton
 	UPROPERTY(VisibleAnywhere, Category = Skeleton)
 	TArray<FTransform> RefPoseGlobal;
 
-	void Initialize(const FReferenceSkeleton& RefSkeleton, const TArray<FName>& ExcludedBones);
+	void Initialize(const FIKRigInputSkeleton& InputSkeleton, const TArray<FName>& InExcludedBones);
+
+	void Initialize(const FReferenceSkeleton& RefSkeleton, const TArray<FName>& InExcludedBones);
 
 	void Reset();
 	
@@ -62,7 +103,7 @@ struct IKRIG_API FIKRigSkeleton
 
 	int32 GetParentIndexThatIsNotExcluded(const int32 BoneIndex) const;
 
-	bool CopyPosesFromRefSkeleton(const FReferenceSkeleton& RefSkeleton);
+	bool CopyPosesFromInputSkeleton(const FIKRigInputSkeleton& InputSkeleton);
 
 	static void ConvertLocalPoseToGlobal(
 		const TArray<int32>& InParentIndices,
diff --git a/Engine/Plugins/Animation/IKRig/Source/IKRig/Public/IKRigSolver.h b/Engine/Plugins/Animation/IKRig/Source/IKRig/Public/IKRigSolver.h
--- a/Engine/Plugins/Animation/IKRig/Source/IKRig/Public/IKRigSolver.h
+++ b/Engine/Plugins/Animation/IKRig/Source/IKRig/Public/IKRigSolver.h
@@ -19,9 +19,6 @@ class IKRIG_API UIKRigSolver : public UObject
 	
 public:
 
-	UPROPERTY(EditAnywhere, Category = "Solver Settings")
-	bool bEnabled = true;
-
 	//** RUNTIME */
 	/** override to setup internal data based on ref pose */
 	virtual void Initialize(const FIKRigSkeleton& IKRigSkeleton) PURE_VIRTUAL("Init");
@@ -29,12 +26,27 @@ public:
 	virtual void Solve(FIKRigSkeleton& IKRigSkeleton, const FIKRigGoalContainer& Goals) PURE_VIRTUAL("Solve");
 	//** END RUNTIME */
 
+#if WITH_EDITORONLY_DATA
+	
+	/** callback whenever this solver is edited */
+	DECLARE_EVENT_OneParam(UIKRigSolver, FIKRigSolverModified, UIKRigSolver*);
+	FIKRigSolverModified& OnSolverModified(){ return IKRigSolverModified; };
+
+	/** get if this solver is enabled */
+	bool IsEnabled() const { return bIsEnabled; };
+	/** turn solver on/off (will be skipped during execution if disabled) */
+	void SetEnabled(const bool bEnabled){ bIsEnabled = bEnabled; };
+	
 	//** SOLVER SETTINGS */
 	/** override to support RECEIVING modified solver settings from outside systems for editing/UI.
 	 * Note: you can safely cast this to your own solver type and copy any relevant settings at runtime
 	 * This is necessary because at runtime, the IKRigProcessor creates a copy of your solver class
 	 * and the copy must be notified of changes made to the class settings in the source asset.*/
 	virtual void UpdateSolverSettings(UIKRigSolver* InSettings){};
+	/** override to give your solver a nice name to display in the UI */
+	virtual FText GetNiceName() const { checkNoEntry() return FText::GetEmpty(); };
+	/** override to provide warning to user during setup of any missing components. return false if no warnings. */
+	virtual bool GetWarningMessage(FText& OutWarningMessage) const { return false; };
 	//** END SOLVER SETTINGS */
 
 	//** GOALS */
@@ -49,7 +61,7 @@ public:
 	/** override to support QUERY for a connected goal */
 	virtual bool IsGoalConnected(const FName& GoalName) const {return false;};
 	/** override to support supplying effector settings to outside systems for editing/UI */
-	virtual UObject* GetEffectorWithGoal(const FName& GoalName) {return nullptr;};
+	virtual UObject* GetEffectorWithGoal(const FName& GoalName) const {return nullptr;};
 	//** END GOALS */
 
 	//** ROOT BONE (optional, implement if your solver requires a root bone) */
@@ -64,20 +76,27 @@ public:
 	/** override to support ADDING PER-BONE settings for this solver */
 	virtual void RemoveBoneSetting(const FName& BoneName){};
 	/** override to support supplying per-bone settings to outside systems for editing/UI */
-	virtual UObject* GetBoneSetting(const FName& BoneName) const { return nullptr;};
+	virtual UObject* GetBoneSetting(const FName& BoneName) const { return nullptr; };
 	/** override to tell systems if this solver supports per-bone settings */
 	virtual bool UsesBoneSettings() const { return false;};
 	/** override to draw custom per-bone settings in the editor viewport */
 	virtual void DrawBoneSettings(const FName& BoneName, const FIKRigSkeleton& IKRigSkeleton, FPrimitiveDrawInterface* PDI) const {};
 	/** return true if the supplied Bone is affected by this solver - this provides UI feedback for user */
-	virtual bool IsBoneAffectedBySolver(const FName& BoneName, const FIKRigSkeleton& IKRigSkeleton) const {return false;};
+	virtual bool IsBoneAffectedBySolver(const FName& BoneName, const FIKRigSkeleton& IKRigSkeleton) const { return false; };
 	//** END ROOT BONE */
 
 	/** UObject interface */
-	void PostLoad() override
-	{
-		Super::PostLoad();
-		SetFlags(RF_Transactional); // patch old solvers to enable undo/redo
-	}; 
+	virtual void PostLoad() override;
+	virtual void PostTransacted(const FTransactionObjectEvent& TransactionEvent) override;
 	/** END UObject interface */
+
+private:
+
+	/** Register callbacks to update IK Rig when a solver is modified */
+	FIKRigSolverModified IKRigSolverModified;
+
+	UPROPERTY()
+	bool bIsEnabled = true;
+
+#endif
 };
diff --git a/Engine/Plugins/Animation/IKRig/Source/IKRig/Public/Retargeter/IKRetargetProcessor.h b/Engine/Plugins/Animation/IKRig/Source/IKRig/Public/Retargeter/IKRetargetProcessor.h
new file mode 100644
--- /dev/null
+++ b/Engine/Plugins/Animation/IKRig/Source/IKRig/Public/Retargeter/IKRetargetProcessor.h
@@ -0,0 +1,383 @@
+﻿// Copyright Epic Games, Inc. All Rights Reserved.
+
+#pragma once
+
+#include "CoreMinimal.h"
+#include "IKRetargetProcessor.generated.h"
+
+class UIKRigDefinition;
+class UIKRigProcessor;
+struct FReferenceSkeleton;
+struct FBoneChain;
+struct FIKRetargetPose;
+class UObject;
+class UIKRetargeter;
+class USkeletalMesh;
+
+struct IKRIG_API FRetargetSkeleton
+{
+	TArray<FName> BoneNames;
+	
+	TArray<int32> ParentIndices;
+	
+	TArray<FTransform> RetargetLocalPose;
+	
+	TArray<FTransform> RetargetGlobalPose;
+	
+	USkeletalMesh* SkeletalMesh;
+
+	void Initialize(
+		USkeletalMesh* InSkeletalMesh, 
+		FIKRetargetPose* InRetargetPose,
+		const FName& RetargetRootBone);
+
+	void Reset();
+
+	void GenerateRetargetPose(FIKRetargetPose* InRetargetPose, const FName& RetargetRootBone);
+
+	int32 FindBoneIndexByName(const FName InName) const;
+
+	int32 GetParentIndex(const int32 BoneIndex) const;
+
+	void UpdateGlobalTransformsBelowBone(
+		const int32 StartBoneIndex,
+		TArray<FTransform>& InLocalPose,
+		TArray<FTransform>& OutGlobalPose) const;
+
+	void UpdateLocalTransformsBelowBone(
+		const int32 StartBoneIndex,
+		TArray<FTransform>& OutLocalPose,
+		TArray<FTransform>& InGlobalPose) const;
+	
+	void UpdateGlobalTransformOfSingleBone(
+		const int32 BoneIndex,
+		const TArray<FTransform>& InLocalPose,
+		TArray<FTransform>& OutGlobalPose) const;
+	
+	void UpdateLocalTransformOfSingleBone(
+		const int32 BoneIndex,
+		TArray<FTransform>& OutLocalPose,
+		TArray<FTransform>& InGlobalPose) const;
+
+	void GetChildrenIndices(const int32 BoneIndex, TArray<int32>& OutChildren) const;
+};
+
+struct FTargetSkeleton : public FRetargetSkeleton
+{
+	TArray<FTransform> OutputGlobalPose;
+	TArray<bool> IsBoneRetargeted;
+
+	void Initialize(
+		USkeletalMesh* InSkeletalMesh, 
+		FIKRetargetPose* RetargetPose,
+		const FName& RetargetRootBone);
+
+	void Reset();
+
+	void SetBoneIsRetargeted(const int32 BoneIndex, const bool IsRetargeted);
+
+	void UpdateGlobalTransformsAllNonRetargetedBones(TArray<FTransform>& InOutGlobalPose);
+};
+
+struct FRootSource
+{
+	int32 BoneIndex;
+
+	FQuat InitialRotation;
+
+	float InitialHeightInverse;
+
+	FVector CurrentPositionNormalized;
+	
+	FQuat CurrentRotation;
+};
+
+struct FRootTarget
+{
+	int32 BoneIndex;
+	
+	FQuat InitialRotation;
+	
+	float InitialHeight;
+};
+
+struct FRootRetargeter
+{	
+	FRootSource Source;
+	
+	FRootTarget Target;
+
+	void Reset();
+	
+	bool InitializeSource(const FName SourceRootBoneName, const FRetargetSkeleton& SourceSkeleton);
+	
+	bool InitializeTarget(const FName TargetRootBoneName, const FTargetSkeleton& TargetSkeleton);
+
+	void EncodePose(const TArray<FTransform> &SourceGlobalPose);
+	
+	void DecodePose( TArray<FTransform> &OutTargetGlobalPose, const float StrideScale) const;
+};
+
+struct FChainFK
+{
+	TArray<FTransform> InitialBoneTransforms;
+
+	TArray<FTransform> CurrentBoneTransforms;
+
+	TArray<float> Params;
+
+	bool Initialize(
+		const TArray<int32>& BoneIndices,
+		const TArray<FTransform> &InitialGlobalPose);
+
+private:
+	
+	bool CalculateBoneParameters();	
+};
+
+struct FChainEncoderFK : public FChainFK
+{
+	void EncodePose(
+		const TArray<int32>& SourceBoneIndices,
+		const TArray<FTransform> &InputGlobalPose);
+};
+
+struct FChainDecoderFK : public FChainFK
+{
+	void InitializeIntermediateParentIndices(
+		const int32 RetargetRootBoneIndex,
+		const int32 ChainRootBoneIndex,
+		const FTargetSkeleton& TargetSkeleton);
+	
+	void DecodePose(
+		const TArray<int32>& TargetBoneIndices,
+		const FChainEncoderFK& SourceChain,
+		const FTargetSkeleton& TargetSkeleton,
+		TArray<FTransform> &InOutGlobalPose);
+
+private:
+
+	FTransform GetTransformAtParam(
+		const TArray<FTransform>& Transforms,
+		const TArray<float>& InParams,
+		const float& Param) const;
+	
+	void UpdateIntermediateParents(
+		const FTargetSkeleton& TargetSkeleton,
+		TArray<FTransform> &InOutGlobalPose);
+
+	TArray<int32> IntermediateParentIndices;
+};
+
+struct FDecodedIKChain
+{
+	FVector EndEffectorPosition;
+	
+	FQuat EndEffectorRotation;
+	
+	FVector PoleVectorPosition;
+};
+
+struct FSourceChainIK
+{
+	int32 BoneIndexA;
+	
+	int32 BoneIndexB;
+	
+	int32 BoneIndexC;
+	
+	FVector InitialEndPosition;
+	
+	FQuat InitialEndRotation;
+	
+	float InvInitialLength;
+
+	// results after encoding...
+	
+	FVector CurrentEndDirectionNormalized;
+	
+	FQuat CurrentEndRotation;
+	
+	float CurrentHeightFromGroundNormalized;
+	
+	FVector PoleVectorDirection;
+};
+
+struct FTargetChainIK
+{
+	int32 BoneIndexA;
+	
+	float InitialLength;
+	
+	FVector InitialEndPosition;
+	
+	FQuat InitialEndRotation;
+};
+
+struct FChainRetargeterIK
+{
+	FSourceChainIK Source;
+	
+	FTargetChainIK Target;
+
+	bool InitializeSource(const TArray<int32>& BoneIndices, const TArray<FTransform> &SourceInitialGlobalPose);
+	
+	bool InitializeTarget(const TArray<int32>& BoneIndices, const TArray<FTransform> &TargetInitialGlobalPose);
+
+	void EncodePose(const TArray<FTransform> &SourceInputGlobalPose);
+	
+	void DecodePose(const TArray<FTransform>& OutGlobalPose, FDecodedIKChain& OutResults) const;
+};
+
+struct FRetargetChainPair
+{
+	TArray<int32> SourceBoneIndices;
+	
+	TArray<int32> TargetBoneIndices;
+	
+	FName SourceBoneChainName;
+	
+	FName TargetBoneChainName;
+
+	virtual ~FRetargetChainPair(){};
+	
+	virtual bool Initialize(
+		const FBoneChain& SourceBoneChain,
+		const FBoneChain& TargetBoneChain,
+		const FRetargetSkeleton& SourceSkeleton,
+		const FTargetSkeleton& TargetSkeleton);
+
+private:
+
+	bool ValidateBoneChainWithSkeletalMesh(
+		const bool IsSource,
+		const FBoneChain& BoneChain,
+		const FRetargetSkeleton& RetargetSkeleton);
+};
+
+struct FRetargetChainPairFK : FRetargetChainPair
+{
+	FChainEncoderFK FKEncoder;
+
+	FChainDecoderFK FKDecoder;
+	
+	virtual bool Initialize(
+        const FBoneChain& SourceBoneChain,
+        const FBoneChain& TargetBoneChain,
+        const FRetargetSkeleton& SourceSkeleton,
+        const FTargetSkeleton& TargetSkeleton) override;
+};
+
+struct FRetargetChainPairIK : FRetargetChainPair
+{	
+	FChainRetargeterIK IKChainRetargeter;
+	
+	FName IKGoalName;
+	
+	FName PoleVectorGoalName;
+
+	virtual bool Initialize(
+        const FBoneChain& SourceBoneChain,
+        const FBoneChain& TargetBoneChain,
+        const FRetargetSkeleton& SourceSkeleton,
+        const FTargetSkeleton& TargetSkeleton) override;
+};
+
+/** The runtime processor that converts an input pose from a source skeleton into an output pose on a target skeleton.
+ * To use:
+ * 1. Initialize a processor with a Source/Target skeletal mesh and a UIKRetargeter asset.
+ * 2. Call RunRetargeter and pass in a source pose as an array of global-space transforms
+ * 3. RunRetargeter returns an array of global space transforms for the target skeleton.
+ */
+UCLASS()
+class IKRIG_API UIKRetargetProcessor : public UObject
+{
+	GENERATED_BODY()
+
+public:
+	
+	/**
+	* Initialize the retargeter to enable running it.
+	* @param SourceSkeleton - the skeletal mesh to poses FROM
+	* @param TargetSkeleton - the skeletal mesh to poses TO
+	* @param InRetargeterAsset - the source asset to use for retargeting settings
+	* @warning - Initialization does a lot of validation and can fail for many reasons. Check bIsLoadedAndValid afterwards.
+	*/
+	void Initialize(
+		USkeletalMesh *SourceSkeleton,
+		USkeletalMesh *TargetSkeleton,
+		UIKRetargeter* InRetargeterAsset);
+
+	/**
+	* Run the retarget to generate a new pose.
+	* @param InSourceGlobalPose -  is the source mesh input pose in Component/Global space
+	* @return The retargeted Component/Global space pose for the target skeleton
+	*/
+	TArray<FTransform>& RunRetargeter(const TArray<FTransform>& InSourceGlobalPose);
+
+	/**
+	* Get the Global transform, in the currently used retarget pose, for a bone in the target skeletal mesh
+	* @param TargetBoneIndex - The index of a bone in the target skeleton.
+	* @return The Global space transform for the bone or Identity if bone not found.
+	* @warning This function is only valid to call after the retargeter has been initialized (which generates global retarget pose)
+	*/
+	FTransform GetTargetBoneRetargetPoseGlobalTransform(const int32& TargetBoneIndex) const;
+
+	/** Get read-only access to the target skeleton. */
+	const FTargetSkeleton& GetTargetSkeleton() const { return TargetSkeleton; };
+
+	/** Get whether this processor is ready to call RunRetargeter() and generate new poses. */
+	bool IsInitialized() const { return bIsInitialized; };
+
+#if WITH_EDITOR
+	/** Set that this processor needs to be reinitialized. */
+	void SetNeedsInitialized();
+	/** During editor preview, drive the target IK Rig with the setting from it's source asset */
+	void CopyTargetIKRigSettingsFromAsset();
+#endif
+
+private:
+
+	/** Only true once Initialize() has successfully completed.*/
+	bool bIsInitialized = false;
+
+	/** The source asset this processor was initialized with. */
+	UIKRetargeter* RetargeterAsset = nullptr;
+
+	/** The internal data structures used to represent the SOURCE skeleton / pose during retargeter.*/
+	FRetargetSkeleton SourceSkeleton;
+
+	/** The internal data structures used to represent the TARGET skeleton / pose during retargeter.*/
+	FTargetSkeleton TargetSkeleton;
+
+	/** The IK Rig processor for running IK on the target */
+	UPROPERTY(Transient) // must be property to keep from being GC'd
+	UIKRigProcessor* IKRigProcessor = nullptr;
+
+	/** The Source/Target pairs of Bone Chains retargeted using the FK algorithm */
+	TArray<FRetargetChainPairFK> ChainPairsFK;
+
+	/** The Source/Target pairs of Bone Chains retargeted using the IK Rig */
+	TArray<FRetargetChainPairIK> ChainPairsIK;
+
+	/** The Source/Target pair of Root Bones retargeted with scaled translation */
+	FRootRetargeter RootRetargeter;
+
+	/** Initializes the FRootRetargeter */
+	bool InitializeRoots();
+
+	/** Initializes the all the chain pairs */
+	bool InitializeBoneChainPairs();
+
+	/** Initializes the IK Rig that evaluates the IK solve for the target IK chains */
+	bool InitializeIKRig(UObject* Outer, const FReferenceSkeleton& InRefSkeleton);
+	
+	/** Internal retarget phase for the root. */
+	void RunRootRetarget(const TArray<FTransform>& InGlobalTransforms, TArray<FTransform>& OutGlobalTransforms);
+
+	/** Internal retarget phase for the FK chains. */
+	void RunFKRetarget(const TArray<FTransform>& InGlobalTransforms, TArray<FTransform>& OutGlobalTransforms);
+
+	/** Internal retarget phase for the IK chains. */
+	void RunIKRetarget(const TArray<FTransform>& InGlobalPose, TArray<FTransform>& OutGlobalPose);
+};
