IK Rig / Retargeter refactor

diff --git a/Engine/Plugins/Animation/IKRig/Source/IKRig/Public/Retargeter/IKRetargeter.h b/Engine/Plugins/Animation/IKRig/Source/IKRig/Public/Retargeter/IKRetargeter.h
--- a/Engine/Plugins/Animation/IKRig/Source/IKRig/Public/Retargeter/IKRetargeter.h
+++ b/Engine/Plugins/Animation/IKRig/Source/IKRig/Public/Retargeter/IKRetargeter.h
@@ -2,10 +2,7 @@
 
 #pragma once
 
-#include "Containers/UnrealString.h"
-
 #include "IKRigDefinition.h"
-#include "IKRigProcessor.h"
 
 #include "IKRetargeter.generated.h"
 
@@ -57,265 +54,6 @@ struct IKRIG_API FRetargetChainMap
 	//FName IkRelativeTargetBone;	
 };
 
-struct IKRIG_API FRetargetSkeleton
-{
-	TArray<FName> BoneNames;
-	
-	TArray<int32> ParentIndices;
-	
-	TArray<FTransform> RetargetLocalPose;
-	
-	TArray<FTransform> RetargetGlobalPose;
-	
-	FString SkeletalMeshName;
-
-	void Initialize(
-		USkeletalMesh* SkeletalMesh, 
-		FIKRetargetPose* RetargetPose);
-
-	int32 FindBoneIndexByName(const FName InName) const;
-
-	int32 GetParentIndex(const int32 BoneIndex) const;
-
-	void UpdateGlobalTransformsBelowBone(
-		const int32 StartBoneIndex,
-		TArray<FTransform>& InLocalPose,
-		TArray<FTransform>& OutGlobalPose) const;
-
-	void UpdateLocalTransformsBelowBone(
-		const int32 StartBoneIndex,
-		TArray<FTransform>& OutLocalPose,
-		TArray<FTransform>& InGlobalPose) const;
-	
-	void UpdateGlobalTransformOfSingleBone(
-		const int32 BoneIndex,
-		const TArray<FTransform>& InLocalPose,
-		TArray<FTransform>& OutGlobalPose) const;
-	
-	void UpdateLocalTransformOfSingleBone(
-		const int32 BoneIndex,
-		TArray<FTransform>& OutLocalPose,
-		TArray<FTransform>& InGlobalPose) const;
-
-	void GetChildrenIndices(const int32 BoneIndex, TArray<int32>& OutChildren) const;
-};
-
-struct FTargetSkeleton : public FRetargetSkeleton
-{
-	TArray<FTransform> OutputGlobalPose;
-	TArray<bool> IsBoneRetargeted;
-
-	void Initialize(
-		USkeletalMesh* SkeletalMesh, 
-		FIKRetargetPose* RetargetPose);
-
-	void SetBoneIsRetargeted(const int32 BoneIndex, const bool IsRetargeted);
-
-	void UpdateGlobalTransformsAllNonRetargetedBones(TArray<FTransform>& InOutGlobalPose);
-};
-
-struct FRootSource
-{
-	int32 BoneIndex;
-
-	FQuat InitialRotation;
-
-	float InitialHeightInverse;
-
-	FVector CurrentPositionNormalized;
-	
-	FQuat CurrentRotation;
-};
-
-struct FRootTarget
-{
-	int32 BoneIndex;
-	
-	FQuat InitialRotation;
-	
-	float InitialHeight;
-};
-
-struct FRootRetargeter
-{
-	FRootSource Source;
-	
-	FRootTarget Target;
-
-	bool InitializeSource(const FName SourceRootBoneName, const FRetargetSkeleton& SourceSkeleton);
-	
-	bool InitializeTarget(const FName TargetRootBoneName, const FTargetSkeleton& TargetSkeleton);
-
-	void EncodePose(const TArray<FTransform> &SourceGlobalPose);
-	
-	void DecodePose( TArray<FTransform> &OutTargetGlobalPose, const float StrideScale) const;
-};
-
-struct FChainFK
-{
-	TArray<FTransform> InitialBoneTransforms;
-
-	TArray<FTransform> CurrentBoneTransforms;
-
-	TArray<float> Params;
-
-	bool Initialize(
-		const TArray<int32>& BoneIndices,
-		const TArray<FTransform> &InitialGlobalPose);
-
-private:
-	
-	bool CalculateBoneParameters();	
-};
-
-struct FChainEncoderFK : public FChainFK
-{
-	void EncodePose(
-		const TArray<int32>& SourceBoneIndices,
-		const TArray<FTransform> &InputGlobalPose);
-};
-
-struct FChainDecoderFK : public FChainFK
-{
-	void InitializeIntermediateParentIndices(
-		const int32 RetargetRootBoneIndex,
-		const int32 ChainRootBoneIndex,
-		const FTargetSkeleton& TargetSkeleton);
-	
-	void DecodePose(
-		const TArray<int32>& TargetBoneIndices,
-		const FChainEncoderFK& SourceChain,
-		const FTargetSkeleton& TargetSkeleton,
-		TArray<FTransform> &InOutGlobalPose);
-
-private:
-
-	FTransform GetTransformAtParam(
-		const TArray<FTransform>& Transforms,
-		const TArray<float>& InParams,
-		const float Param) const;
-	
-	void UpdateIntermediateParents(
-		const FTargetSkeleton& TargetSkeleton,
-		TArray<FTransform> &InOutGlobalPose);
-
-	TArray<int32> IntermediateParentIndices;
-};
-
-struct FDecodedIKChain
-{
-	FVector EndEffectorPosition;
-	
-	FQuat EndEffectorRotation;
-	
-	FVector PoleVectorPosition;
-};
-
-struct FSourceChainIK
-{
-	int32 BoneIndexA;
-	
-	int32 BoneIndexB;
-	
-	int32 BoneIndexC;
-	
-	FVector InitialEndPosition;
-	
-	FQuat InitialEndRotation;
-	
-	float InvInitialLength;
-
-	// results after encoding...
-	
-	FVector CurrentEndDirectionNormalized;
-	
-	FQuat CurrentEndRotation;
-	
-	float CurrentHeightFromGroundNormalized;
-	
-	FVector PoleVectorDirection;
-};
-
-struct FTargetChainIK
-{
-	int32 BoneIndexA;
-	
-	float InitialLength;
-	
-	FVector InitialEndPosition;
-	
-	FQuat InitialEndRotation;
-};
-
-struct FChainRetargeterIK
-{
-	FSourceChainIK Source;
-	
-	FTargetChainIK Target;
-
-	bool InitializeSource(const TArray<int32>& BoneIndices, const TArray<FTransform> &SourceInitialGlobalPose);
-	
-	bool InitializeTarget(const TArray<int32>& BoneIndices, const TArray<FTransform> &TargetInitialGlobalPose);
-
-	void EncodePose(const TArray<FTransform> &SourceInputGlobalPose);
-	
-	void DecodePose(const TArray<FTransform>& OutGlobalPose, FDecodedIKChain& OutResults) const;
-};
-
-struct FRetargetChainPair
-{
-	TArray<int32> SourceBoneIndices;
-	
-	TArray<int32> TargetBoneIndices;
-	
-	FName SourceBoneChainName;
-	
-	FName TargetBoneChainName;
-
-	virtual ~FRetargetChainPair(){};
-	
-	virtual bool Initialize(
-		const FBoneChain& SourceBoneChain,
-		const FBoneChain& TargetBoneChain,
-		const FRetargetSkeleton& SourceSkeleton,
-		const FTargetSkeleton& TargetSkeleton);
-
-private:
-
-	bool ValidateBoneChainWithSkeletalMesh(
-		const bool IsSource,
-		const FBoneChain& BoneChain,
-		const FRetargetSkeleton& RetargetSkeleton);
-};
-
-struct FRetargetChainPairFK : FRetargetChainPair
-{
-	FChainEncoderFK FKEncoder;
-
-	FChainDecoderFK FKDecoder;
-	
-	virtual bool Initialize(
-        const FBoneChain& SourceBoneChain,
-        const FBoneChain& TargetBoneChain,
-        const FRetargetSkeleton& SourceSkeleton,
-        const FTargetSkeleton& TargetSkeleton) override;
-};
-
-struct FRetargetChainPairIK : FRetargetChainPair
-{	
-	FChainRetargeterIK IKChainRetargeter;
-	
-	FName IKGoalName;
-	
-	FName PoleVectorGoalName;
-
-	virtual bool Initialize(
-        const FBoneChain& SourceBoneChain,
-        const FBoneChain& TargetBoneChain,
-        const FRetargetSkeleton& SourceSkeleton,
-        const FTargetSkeleton& TargetSkeleton) override;
-};
-
 USTRUCT()
 struct IKRIG_API FIKRetargetPose
 {
@@ -326,13 +64,12 @@ struct IKRIG_API FIKRetargetPose
 	UPROPERTY(EditAnywhere, Category = RetargetPose)
 	FVector RootTranslationOffset = FVector::ZeroVector;
 	
-	UPROPERTY(EditAnywhere, Category = RetargetPose)
-	FQuat RootRotationOffset = FQuat::Identity;
-	
 	UPROPERTY(EditAnywhere, Category = RetargetPose)
 	TMap<FName, FQuat> BoneRotationOffsets;
 
 	void AddRotationDeltaToBone(FName BoneName, FQuat RotationDelta);
+
+	void AddTranslationDeltaToRoot(FVector TranslateDelta);
 };
 
 UCLASS(Blueprintable)
@@ -344,153 +81,63 @@ public:
 
 	/** The rig to copy animation FROM.*/
 	UPROPERTY(VisibleAnywhere, Category = Rigs)
-	UIKRigDefinition *SourceIKRigAsset;
+	TObjectPtr<UIKRigDefinition> SourceIKRigAsset = nullptr;
 	
 	/** The rig to copy animation TO.*/
 	UPROPERTY(EditAnywhere, Category = Rigs)
-	UIKRigDefinition *TargetIKRigAsset;
-
-	/** Mapping of chains to copy animation between source and target rigs.*/
-	UPROPERTY()
-	TArray<FRetargetChainMap> ChainMapping;
-
-	/** The set of retarget poses available as options for retargeting.*/
-	UPROPERTY()
-	FName CurrentRetargetPose = DefaultPoseName;
-	static const FName DefaultPoseName;
-	
-	/** The set of retarget poses available as options for retargeting.*/
-	UPROPERTY(VisibleAnywhere, Category = RetargetPose)
-	TMap<FName, FIKRetargetPose> RetargetPoses;
+	TObjectPtr<UIKRigDefinition> TargetIKRigAsset = nullptr;
 
 	/** When false, translational motion of skeleton root is not copied. Useful for debugging.*/
-	UPROPERTY(EditAnywhere, Category = Settings)
+	UPROPERTY(EditAnywhere, Category = RetargetPhases)
 	bool bRetargetRoot = true;
 	
 	/** When false, limbs are not copied via FK. Useful for debugging limb issues suspected to be caused by FK pose.*/
-	UPROPERTY(EditAnywhere, Category = Settings)
+	UPROPERTY(EditAnywhere, Category = RetargetPhases)
 	bool bRetargetFK = true;
 	
 	/** When false, IK is not applied as part of retargeter. Useful for debugging limb issues suspected to be caused by IK.*/
-	UPROPERTY(EditAnywhere, Category = Settings)
+	UPROPERTY(EditAnywhere, Category = RetargetPhases)
 	bool bRetargetIK = true;
 
+#if WITH_EDITORONLY_DATA
 	/** Move the target actor in the viewport for easier visualization next to the source actor.*/
-	UPROPERTY(EditAnywhere, Category = DebugPreview, meta = (UIMin = "-2000.0", UIMax = "2000.0"))
+	UPROPERTY(EditAnywhere, Category = TargetActorPreview, meta = (UIMin = "-2000.0", UIMax = "2000.0"))
 	float TargetActorOffset = 150.0f;
 
 	/** Scale the target actor in the viewport for easier visualization next to the source actor.*/
-	UPROPERTY(EditAnywhere, Category = DebugPreview, meta = (UIMin = "0.01", UIMax = "10.0"))
+	UPROPERTY(EditAnywhere, Category = TargetActorPreview, meta = (UIMin = "0.01", UIMax = "10.0"))
 	float TargetActorScale = 1.0f;
 
 	/** The visual size of the bones in the viewport when editing the retarget pose.*/
-	UPROPERTY(EditAnywhere, Category = DebugPreview, meta = (ClampMin = "0.0", UIMin = "0.01", UIMax = "10.0"))
+	UPROPERTY(EditAnywhere, Category = PoseEditSettings, meta = (ClampMin = "0.0", UIMin = "0.01", UIMax = "10.0"))
 	float BoneDrawSize = 8.0f;
 
 	/** The visual thickness of the bones in the viewport when editing the retarget pose.*/
-	UPROPERTY(EditAnywhere, Category = DebugPreview, meta = (ClampMin = "0.0", UIMin = "0.01", UIMax = "10.0"))
+	UPROPERTY(EditAnywhere, Category = PoseEditSettings, meta = (ClampMin = "0.0", UIMin = "0.01", UIMax = "10.0"))
 	float BoneDrawThickness = 1.0f;
 
-	/** The internal data structures used to represent the SOURCE skeleton / pose during retargeter.*/
-	FRetargetSkeleton SourceSkeleton;
-
-	/** The internal data structures used to represent the TARGET skeleton / pose during retargeter.*/
-	FTargetSkeleton TargetSkeleton;
-
-	/** Only true once Initialize() has successfully completed.*/
-	bool bIsLoadedAndValid = false;
-
 	/** A special editor-only mode which forces the retargeter to output the current retarget reference pose,
-	 * rather than actually running the retarget and outputting the retargeted pose. Used in Edit-Pose mode.*/
-	UPROPERTY()
-	bool bEditReferencePoseMode = false;
-	/** Used in-editor to force reinitialization when the asset version of the source asset differs from the currently
-	 *running asset version.*/
+	* rather than actually running the retarget and outputting the retargeted pose. Used in Edit-Pose mode.*/
 	UPROPERTY()
-	int32 AssetVersion = 0;
-
-	UIKRetargeter();
-
-	/**
-	* Initialize the retargeter to enable running it.
-	* @param SourceSkeleton - the skeletal mesh to poses FROM
-	* @param TargetSkeleton - the skeletal mesh to poses TO
-	* @param Outer - a UObject that will live for the duration of the retargeting workload
-	* @warning - Initialization does a lot of validation and can fail for many reasons. Check bIsLoadedAndValid afterwards.
-	*/
-	void Initialize(USkeletalMesh *SourceSkeleton, USkeletalMesh *TargetSkeleton, UObject* Outer);
-
-	/**
-	 * Run the retarget to generate a new pose.
-	 * @param InSourceGlobalPose -  is the source mesh input pose in Component/Global space
-	 * @return The retargeted Component/Global space pose for the target skeleton
-	 */
-	TArray<FTransform>& RunRetargeter(const TArray<FTransform>& InSourceGlobalPose);
-
-	/**
-	* Get the Global transform, in the currently used retarget pose, for a bone in the target skeletal mesh
-	* @param TargetBoneIndex - The index of a bone in the target skeleton.
-	* @return The Global space transform for the bone or Identity if bone not found.
-	* @warning This function is only valid to call after the retargeter has been initialized (which generates global retarget pose)
-	*/
-	FTransform GetTargetBoneRetargetPoseGlobalTransform(const int32& TargetBoneIndex) const;
-
-	/**
-	* Get the mapping (between source/target) for a given target chain
-	* @param TargetChainName - The name of chain on the target IK Rig
-	* @return The mapping between the target chain and the source
-	*/
-	FRetargetChainMap* GetChainMap(const FName& TargetChainName);
-
-	void CleanPoseList();
-	
-#if WITH_EDITOR
-	
-	// UObject
-	virtual bool Modify(bool bAlwaysMarkDirty = true) override;
-	virtual void PostEditChangeProperty( struct FPropertyChangedEvent& PropertyChangedEvent) override;
-	virtual void PostLoad() override;
-	// END UObject
-
-	FName GetSourceRootBone();
-	FName GetTargetRootBone();
-
-	void GetTargetChainNames(TArray<FName>& OutNames) const;
-	void GetSourceChainNames(TArray<FName>& OutNames) const;
-
-	void CleanChainMapping();
-	void AutoMapChains();
-
-	USkeleton* GetSourceSkeletonAsset() const;
-
-	void AddRetargetPose(FName NewPoseName);
-	void RemoveRetargetPose(FName PoseToRemove);
-	void ResetRetargetPose(FName PoseToReset);
-	void SetCurrentRetargetPose(FName CurrentPose);
-	void AddRotationOffsetToRetargetPoseBone(FName BoneName, FQuat RotationOffset);
-	
-#endif
-
-private:
+	bool bEditRetargetPoseMode = false;
 
+	/** The controller responsible for managing this asset's data (all editor mutation goes through this) */
 	UPROPERTY(Transient)
-	TObjectPtr<UIKRigProcessor> IKRigProcessor;
-	
-	TArray<FRetargetChainPairFK> ChainPairsFK;
-	
-	TArray<FRetargetChainPairIK> ChainPairsIK;
+	TObjectPtr<UObject> Controller;
+#endif
 	
-	FRootRetargeter RootRetargeter;
+	/** The set of retarget poses available as options for retargeting.*/
+	UPROPERTY(VisibleAnywhere, Category = EditRetargetPoses)
+	TMap<FName, FIKRetargetPose> RetargetPoses;
 
-	bool InitializeRoots();
-		
-	bool InitializeBoneChainPairs();
-	
-	bool InitializeIKRig(UObject* Outer, const FReferenceSkeleton& InRefSkeleton);
-	
-	void RunRootRetarget(const TArray<FTransform>& InGlobalTransforms, TArray<FTransform>& OutGlobalTransforms);
+	/** Mapping of chains to copy animation between source and target rigs.*/
+	UPROPERTY()
+	TArray<FRetargetChainMap> ChainMapping;
 	
-	void RunFKRetarget(const TArray<FTransform>& InGlobalTransforms, TArray<FTransform>& OutGlobalTransforms);
-
-	void RunIKRetarget(const TArray<FTransform>& InGlobalPose, TArray<FTransform>& OutGlobalPose);
+	/** The set of retarget poses available as options for retargeting.*/
+	UPROPERTY()
+	FName CurrentRetargetPose = DefaultPoseName;
+	static const FName DefaultPoseName;
 };
+
+
diff --git a/Engine/Plugins/Animation/IKRig/Source/IKRig/Public/Solvers/IKRig_BodyMover.h b/Engine/Plugins/Animation/IKRig/Source/IKRig/Public/Solvers/IKRig_BodyMover.h
--- a/Engine/Plugins/Animation/IKRig/Source/IKRig/Public/Solvers/IKRig_BodyMover.h
+++ b/Engine/Plugins/Animation/IKRig/Source/IKRig/Public/Solvers/IKRig_BodyMover.h
@@ -31,55 +31,53 @@ class IKRIG_API UIKRig_BodyMover : public UIKRigSolver
 	GENERATED_BODY()
 
 public:
-	
-	UIKRig_BodyMover();
 
 	/** The target bone to move with the effectors. */
-	UPROPERTY(VisibleAnywhere, Category = "Body Mover Solver")
+	UPROPERTY(VisibleAnywhere, Category = "Body Mover Settings")
 	FName BodyBone;
 
 	/** Blend the translational effect of this solver on/off. Range is 0-1. Default is 1.0. */
-	UPROPERTY(EditAnywhere, Category = "Body Mover Solver", meta = (UIMin = "0.0", UIMax = "1.0"))
+	UPROPERTY(EditAnywhere, Category = "Body Mover Settings", meta = (UIMin = "0.0", UIMax = "1.0"))
 	float PositionAlpha = 1.0f;
 
 	/** Multiply the POSITIVE X translation. Range is 0-1. Default is 1.0. */
-	UPROPERTY(EditAnywhere, Category = "Body Mover Solver", meta = (UIMin = "0.0", UIMax = "1.0"))
+	UPROPERTY(EditAnywhere, Category = "Body Mover Settings", meta = (UIMin = "0.0", UIMax = "1.0"))
 	float PositionPositiveX = 1.0f;
 
 	/** Multiply the NEGATIVE X translation. Range is 0-1. Default is 1.0. */
-	UPROPERTY(EditAnywhere, Category = "Body Mover Solver", meta = (UIMin = "0.0", UIMax = "1.0"))
+	UPROPERTY(EditAnywhere, Category = "Body Mover Settings", meta = (UIMin = "0.0", UIMax = "1.0"))
 	float PositionNegativeX = 1.0f;
 
 	/** Multiply the POSITIVE Y translation. Range is 0-1. Default is 1.0. */
-	UPROPERTY(EditAnywhere, Category = "Body Mover Solver", meta = (UIMin = "0.0", UIMax = "1.0"))
+	UPROPERTY(EditAnywhere, Category = "Body Mover Settings", meta = (UIMin = "0.0", UIMax = "1.0"))
 	float PositionPositiveY = 1.0f;
 
 	/** Multiply the NEGATIVE Y translation. Range is 0-1. Default is 1.0. */
-	UPROPERTY(EditAnywhere, Category = "Body Mover Solver", meta = (UIMin = "0.0", UIMax = "1.0"))
+	UPROPERTY(EditAnywhere, Category = "Body Mover Settings", meta = (UIMin = "0.0", UIMax = "1.0"))
 	float PositionNegativeY = 1.0f;
 
 	/** Multiply the POSITIVE Z translation. Range is 0-1. Default is 1.0. */
-	UPROPERTY(EditAnywhere, Category = "Body Mover Solver", meta = (UIMin = "0.0", UIMax = "1.0"))
+	UPROPERTY(EditAnywhere, Category = "Body Mover Settings", meta = (UIMin = "0.0", UIMax = "1.0"))
 	float PositionPositiveZ = 1.0f;
 
 	/** Multiply the NEGATIVE Z translation. Range is 0-1. Default is 1.0. */
-	UPROPERTY(EditAnywhere, Category = "Body Mover Solver", meta = (UIMin = "0.0", UIMax = "1.0"))
+	UPROPERTY(EditAnywhere, Category = "Body Mover Settings", meta = (UIMin = "0.0", UIMax = "1.0"))
 	float PositionNegativeZ = 1.0f;
 
 	/** Blend the total rotational effect on/off. Range is 0-1. Default is 1.0. */
-	UPROPERTY(EditAnywhere, Category = "Body Mover Solver", meta = (UIMin = "0.0", UIMax = "1.0"))
+	UPROPERTY(EditAnywhere, Category = "Body Mover Settings", meta = (UIMin = "0.0", UIMax = "1.0"))
 	float RotationAlpha = 1.0f;
 
 	/** Blend the X-axis rotational effect on/off. Range is 0-1. Default is 1.0. */
-	UPROPERTY(EditAnywhere, Category = "Body Mover Solver", meta = (UIMin = "0.0", UIMax = "1.0"))
+	UPROPERTY(EditAnywhere, Category = "Body Mover Settings", meta = (UIMin = "0.0", UIMax = "1.0"))
 	float RotateXAlpha = 1.0f;
 
 	/** Blend the Y-axis rotational effect on/off. Range is 0-1. Default is 1.0. */
-	UPROPERTY(EditAnywhere, Category = "Body Mover Solver", meta = (UIMin = "0.0", UIMax = "1.0"))
+	UPROPERTY(EditAnywhere, Category = "Body Mover Settings", meta = (UIMin = "0.0", UIMax = "1.0"))
 	float RotateYAlpha = 1.0f;
 
 	/** Blend the Z-axis rotational effect on/off. Range is 0-1. Default is 1.0. */
-	UPROPERTY(EditAnywhere, Category = "Body Mover Solver", meta = (UIMin = "0.0", UIMax = "1.0"))
+	UPROPERTY(EditAnywhere, Category = "Body Mover Settings", meta = (UIMin = "0.0", UIMax = "1.0"))
 	float RotateZAlpha = 1.0f;
 	
 	UPROPERTY()
@@ -88,7 +86,10 @@ public:
 	/** UIKRigSolver interface */
 	virtual void Initialize(const FIKRigSkeleton& IKRigSkeleton) override;
 	virtual void Solve(FIKRigSkeleton& IKRigSkeleton, const FIKRigGoalContainer& Goals) override;
+#if WITH_EDITOR
 	virtual void UpdateSolverSettings(UIKRigSolver* InSettings) override;
+	virtual FText GetNiceName() const override;
+	virtual bool GetWarningMessage(FText& OutWarningMessage) const override;
 	virtual bool IsBoneAffectedBySolver(const FName& BoneName, const FIKRigSkeleton& IKRigSkeleton) const override;
 	// goals
 	virtual void AddGoal(const UIKRigEffectorGoal& NewGoal) override;
@@ -96,11 +97,12 @@ public:
 	virtual void RenameGoal(const FName& OldName, const FName& NewName) override;
 	virtual void SetGoalBone(const FName& GoalName, const FName& NewBoneName) override;
 	virtual bool IsGoalConnected(const FName& GoalName) const override;
-	virtual UObject* GetEffectorWithGoal(const FName& GoalName) override;
+	virtual UObject* GetEffectorWithGoal(const FName& GoalName) const override;
 	// root bone can be set on this solver
 	virtual bool CanSetRootBone() const override { return true; };
 	virtual void SetRootBone(const FName& RootBoneName) override;
 	/** END UIKRigSolver interface */
+#endif
 
 private:
 
diff --git a/Engine/Plugins/Animation/IKRig/Source/IKRig/Public/Solvers/IKRig_PBIKSolver.h b/Engine/Plugins/Animation/IKRig/Source/IKRig/Public/Solvers/IKRig_PBIKSolver.h
--- a/Engine/Plugins/Animation/IKRig/Source/IKRig/Public/Solvers/IKRig_PBIKSolver.h
+++ b/Engine/Plugins/Animation/IKRig/Source/IKRig/Public/Solvers/IKRig_PBIKSolver.h
@@ -134,55 +134,58 @@ class IKRIG_API UIKRigPBIKSolver : public UIKRigSolver
 
 public:
 
-	UPROPERTY(VisibleAnywhere, Category = "Solver Settings")
+	UPROPERTY(VisibleAnywhere, Category = "Full-Body IK Settings")
 	FName RootBone;
 
 	/** High iteration counts can help solve complex joint configurations with competing constraints, but will increase runtime cost. Default is 20. */
-	UPROPERTY(EditAnywhere, Category = SolverSettings, meta = (ClampMin = "0", UIMin = "0.0", UIMax = "100.0"))
+	UPROPERTY(EditAnywhere, Category = "Full-Body IK Settings", meta = (ClampMin = "0", UIMin = "0.0", UIMax = "100.0"))
 	int32 Iterations = 20;
 
 	/** A global mass multiplier; higher values will make the joints more stiff, but require more iterations. Typical range is 0.0 to 10.0. */
-	UPROPERTY(EditAnywhere, Category = SolverSettings, meta = (ClampMin = "0", UIMin = "0.0", UIMax = "10.0"))
+	UPROPERTY(EditAnywhere, Category = "Full-Body IK Settings", meta = (ClampMin = "0", UIMin = "0.0", UIMax = "10.0"))
 	float MassMultiplier = 1.0f;
 
 	/** Set this as low as possible while keeping the solve stable. Lower values improve convergence of effector targets. Default is 0.2. */
-	UPROPERTY(EditAnywhere, Category = SolverSettings, meta = (ClampMin = "0", UIMin = "0.0", UIMax = "10.0"))
+	UPROPERTY(EditAnywhere, Category = "Full-Body IK Settings", meta = (ClampMin = "0", UIMin = "0.0", UIMax = "10.0"))
 	float MinMassMultiplier = 0.2f;
 
 	/** If true, joints will translate to reach the effectors; causing bones to lengthen if necessary. Good for cartoon effects. Default is false. */
-	UPROPERTY(EditAnywhere, Category = SolverSettings)
+	UPROPERTY(EditAnywhere, Category = "Full-Body IK Settings")
 	bool bAllowStretch = false;
 
 	/** If true, solver will pull entire skeleton towards effectors BEFORE running constraint iterations. Default is true. */
-	UPROPERTY(EditAnywhere, Category = SolverSettings)
+	UPROPERTY(EditAnywhere, Category = "Full-Body IK Settings")
 	bool bPrePullRoot = true;
 
 	/** Lock the position and rotation of the solver root bone in-place (at animated position). Useful for partial-body solves. Default is false. */
-	UPROPERTY(EditAnywhere, Category = SolverSettings)
+	UPROPERTY(EditAnywhere, Category = "Full-Body IK Settings")
 	bool bPinRoot = false;
 
 	/** When true, the solver is reset each tick to start from the current input pose. If false, incoming animated poses are ignored and the solver starts from the results of the previous solve. Default is true. */
-	UPROPERTY(EditAnywhere, Category = SolverSettings)
+	UPROPERTY(EditAnywhere, Category = "Full-Body IK Settings")
 	bool bStartSolveFromInputPose = true;
 	
 	UPROPERTY()
-	TArray<UIKRig_FBIKEffector*> Effectors;
+	TArray<TObjectPtr<UIKRig_FBIKEffector>> Effectors;
 
 	UPROPERTY()
-	TArray<UIKRig_PBIKBoneSettings*> BoneSettings;
+	TArray<TObjectPtr<UIKRig_PBIKBoneSettings>> BoneSettings;
 
 	/** UIKRigSolver interface */
 	// runtime
 	virtual void Initialize(const FIKRigSkeleton& IKRigSkeleton) override;
 	virtual void Solve(FIKRigSkeleton& IKRigSkeleton, const FIKRigGoalContainer& Goals) override;
+#if WITH_EDITOR
 	virtual void UpdateSolverSettings(UIKRigSolver* InSettings) override;
+	virtual FText GetNiceName() const override;
+	virtual bool GetWarningMessage(FText& OutWarningMessage) const override;
 	// goals
 	virtual void AddGoal(const UIKRigEffectorGoal& NewGoal) override;
 	virtual void RemoveGoal(const FName& GoalName) override;
 	virtual void RenameGoal(const FName& OldName, const FName& NewName) override;
 	virtual void SetGoalBone(const FName& GoalName, const FName& NewBoneName) override;
 	virtual bool IsGoalConnected(const FName& GoalName) const override;
-	virtual UObject* GetEffectorWithGoal(const FName& GoalName) override;
+	virtual UObject* GetEffectorWithGoal(const FName& GoalName) const override;
 	// bone settings
 	virtual void AddBoneSetting(const FName& BoneName) override;
 	virtual void RemoveBoneSetting(const FName& BoneName) override;
@@ -195,14 +198,16 @@ public:
 	virtual bool IsBoneAffectedBySolver(const FName& BoneName, const FIKRigSkeleton& IKRigSkeleton) const override;
 	/** END UIKRigSolver interface */
 
+	//** UObject */
+	virtual void PostLoad() override;
+	//** END UObject */
+#endif
+
 private:
 
 	FPBIKSolver Solver;
 
 	int32 GetIndexOfGoal(const FName& GoalName) const;
-
-	//** UObject */
-	virtual void PostLoad() override;
 };
 
 
diff --git a/Engine/Plugins/Animation/IKRig/Source/IKRig/Public/Solvers/IKRig_SetTransform.h b/Engine/Plugins/Animation/IKRig/Source/IKRig/Public/Solvers/IKRig_SetTransform.h
--- a/Engine/Plugins/Animation/IKRig/Source/IKRig/Public/Solvers/IKRig_SetTransform.h
+++ b/Engine/Plugins/Animation/IKRig/Source/IKRig/Public/Solvers/IKRig_SetTransform.h
@@ -48,14 +48,18 @@ public:
 	/** UIKRigSolver interface */
 	virtual void Initialize(const FIKRigSkeleton& IKRigSkeleton) override;
 	virtual void Solve(FIKRigSkeleton& IKRigSkeleton, const FIKRigGoalContainer& Goals) override;
+#if WITH_EDITOR
 	virtual void UpdateSolverSettings(UIKRigSolver* InSettings) override;
+	virtual FText GetNiceName() const override;
+	virtual bool GetWarningMessage(FText& OutWarningMessage) const override;
 	virtual void AddGoal(const UIKRigEffectorGoal& NewGoal) override;
 	virtual void RemoveGoal(const FName& GoalName) override;
 	virtual void RenameGoal(const FName& OldName, const FName& NewName) override;
 	virtual void SetGoalBone(const FName& GoalName, const FName& NewBoneName) override;
 	virtual bool IsGoalConnected(const FName& GoalName) const override;
-	virtual UObject* GetEffectorWithGoal(const FName& GoalName) override;
+	virtual UObject* GetEffectorWithGoal(const FName& GoalName) const override;
 	virtual bool IsBoneAffectedBySolver(const FName& BoneName, const FIKRigSkeleton& IKRigSkeleton) const override;
+#endif
 	/** END UIKRigSolver interface */
 
 private:
diff --git a/Engine/Plugins/Animation/IKRig/Source/IKRigDeveloper/Private/AnimGraphNode_RetargetPoseFromMesh.cpp b/Engine/Plugins/Animation/IKRig/Source/IKRigDeveloper/Private/AnimGraphNode_RetargetPoseFromMesh.cpp
--- a/Engine/Plugins/Animation/IKRig/Source/IKRigDeveloper/Private/AnimGraphNode_RetargetPoseFromMesh.cpp
+++ b/Engine/Plugins/Animation/IKRig/Source/IKRigDeveloper/Private/AnimGraphNode_RetargetPoseFromMesh.cpp
@@ -80,7 +80,7 @@ void UAnimGraphNode_RetargetPoseFromMesh::ValidateAnimNodeDuringCompilation(USke
 	
 	// validate that target bone chains exist on this skeleton
 	const FReferenceSkeleton &RefSkel = ForSkeleton->GetReferenceSkeleton();
-	const TArray<FBoneChain> &TargetBoneChains = Node.IKRetargeterAsset->TargetIKRigAsset->RetargetDefinition.BoneChains;
+	const TArray<FBoneChain> &TargetBoneChains = Node.IKRetargeterAsset->TargetIKRigAsset->GetRetargetChains();
     for (const FBoneChain &Chain : TargetBoneChains)
     {
         if (RefSkel.FindBoneIndex(Chain.StartBone) == INDEX_NONE)
diff --git a/Engine/Plugins/Animation/IKRig/Source/IKRigEditor/Private/RetargetEditor/IKRetargetAnimInstance.cpp b/Engine/Plugins/Animation/IKRig/Source/IKRigEditor/Private/RetargetEditor/IKRetargetAnimInstance.cpp
--- a/Engine/Plugins/Animation/IKRig/Source/IKRigEditor/Private/RetargetEditor/IKRetargetAnimInstance.cpp
+++ b/Engine/Plugins/Animation/IKRig/Source/IKRigEditor/Private/RetargetEditor/IKRetargetAnimInstance.cpp
@@ -17,9 +17,14 @@ void UIKRetargetAnimInstance::SetRetargetAssetAndSourceComponent(
 	Proxy.SetRetargetAssetAndSourceComponent(InAsset, InSourceMeshComponent);
 }
 
-UIKRetargeter* UIKRetargetAnimInstance::GetCurrentlyUsedRetargeter() const
+const UIKRetargetProcessor* UIKRetargetAnimInstance::GetRetargetProcessor() const
 {
-	return IKRetargeterNode.GetCurrentlyUsedRetargeter();
+	return IKRetargeterNode.GetRetargetProcessor();
+}
+
+void UIKRetargetAnimInstance::SetProcessorNeedsInitialized()
+{
+	IKRetargeterNode.SetProcessorNeedsInitialized();
 }
 
 FAnimInstanceProxy* UIKRetargetAnimInstance::CreateAnimInstanceProxy()
diff --git a/Engine/Plugins/Animation/IKRig/Source/IKRigEditor/Private/RetargetEditor/IKRetargetAnimInstanceProxy.cpp b/Engine/Plugins/Animation/IKRig/Source/IKRigEditor/Private/RetargetEditor/IKRetargetAnimInstanceProxy.cpp
--- a/Engine/Plugins/Animation/IKRig/Source/IKRigEditor/Private/RetargetEditor/IKRetargetAnimInstanceProxy.cpp
+++ b/Engine/Plugins/Animation/IKRig/Source/IKRigEditor/Private/RetargetEditor/IKRetargetAnimInstanceProxy.cpp
@@ -50,5 +50,7 @@ void FIKRetargetAnimInstanceProxy::SetRetargetAssetAndSourceComponent(
 	IKRetargetNode->IKRetargeterAsset = InIKRetargetAsset;
 	IKRetargetNode->bUseAttachedParent = false;
 	IKRetargetNode->SourceMeshComponent = InSourceMeshComponent;
+	IKRetargetNode->bDriveTargetIKRigWithAsset = true;
+	IKRetargetNode->SetProcessorNeedsInitialized();
 }
 
diff --git a/Engine/Plugins/Animation/IKRig/Source/IKRigEditor/Private/RetargetEditor/IKRetargetBatchOperation.cpp b/Engine/Plugins/Animation/IKRig/Source/IKRigEditor/Private/RetargetEditor/IKRetargetBatchOperation.cpp
--- a/Engine/Plugins/Animation/IKRig/Source/IKRigEditor/Private/RetargetEditor/IKRetargetBatchOperation.cpp
+++ b/Engine/Plugins/Animation/IKRig/Source/IKRigEditor/Private/RetargetEditor/IKRetargetBatchOperation.cpp
@@ -19,7 +19,9 @@
 #include "Kismet2/BlueprintEditorUtils.h"
 #include "Kismet2/KismetEditorUtilities.h"
 #include "Misc/ScopedSlowTask.h"
+#include "RetargetEditor/SRetargetAnimAssetsWindow.h"
 #include "Retargeter/IKRetargeter.h"
+#include "Retargeter/IKRetargetProcessor.h"
 #include "Widgets/Notifications/SNotificationList.h"
 
 
@@ -200,9 +202,9 @@ void FIKRetargetBatchOperation::ConvertAnimation(
 {
 	// initialize the retargeter
 	UObject* TransientOuter = Cast<UObject>(GetTransientPackage());
-	UIKRetargeter* Retargeter = DuplicateObject(Context.IKRetargetAsset, TransientOuter);	
-	Retargeter->Initialize(Context.SourceMesh, Context.TargetMesh, TransientOuter);
-	if (!Retargeter->bIsLoadedAndValid)
+	UIKRetargetProcessor* Processor = NewObject<UIKRetargetProcessor>(TransientOuter);
+	Processor->Initialize(Context.SourceMesh, Context.TargetMesh, Context.IKRetargetAsset);
+	if (!Processor->IsInitialized())
 	{
 		UE_LOG(LogTemp, Warning, TEXT("Unable to initialize the IK Retargeter. Newly created animations were not retargeted!"));
 		return;
@@ -232,7 +234,7 @@ void FIKRetargetBatchOperation::ConvertAnimation(
 		const int32 NumFrames = SourceSequence->GetNumberOfSampledKeys();
 
 		// make space for the target keyframe data
-		const int32 NumTargetBones = Retargeter->TargetSkeleton.BoneNames.Num();
+		const int32 NumTargetBones = Processor->GetTargetSkeleton().BoneNames.Num();
 		TArray<FRawAnimSequenceTrack> BoneTracks;
 		BoneTracks.SetNumZeroed(NumTargetBones);
 
@@ -252,11 +254,11 @@ void FIKRetargetBatchOperation::ConvertAnimation(
 			}
 
 			// run the retarget
-			TArray<FTransform>& TargetComponentPose = Retargeter->RunRetargeter(SourceComponentPose);
+			TArray<FTransform>& TargetComponentPose = Processor->RunRetargeter(SourceComponentPose);
 
 			// convert to a local-space pose
 			TArray<FTransform> TargetLocalPose = TargetComponentPose;
-			Retargeter->TargetSkeleton.UpdateLocalTransformsBelowBone(0,TargetLocalPose, TargetComponentPose);
+			Processor->GetTargetSkeleton().UpdateLocalTransformsBelowBone(0,TargetLocalPose, TargetComponentPose);
 
 			// store key data for each bone
 			for (int32 TargetBoneIndex=0; TargetBoneIndex<TargetLocalPose.Num(); ++TargetBoneIndex)
@@ -271,7 +273,7 @@ void FIKRetargetBatchOperation::ConvertAnimation(
 		const bool bShouldTransact = false;
 		for (int32 TargetBoneIndex=0; TargetBoneIndex<NumTargetBones; ++TargetBoneIndex)
 		{
-			FName TargetBoneName = Retargeter->TargetSkeleton.BoneNames[TargetBoneIndex];
+			FName TargetBoneName = Processor->GetTargetSkeleton().BoneNames[TargetBoneIndex];
 			FRawAnimSequenceTrack& RawTrack = BoneTracks[TargetBoneIndex];
 			TargetSeqController.AddBoneTrack(TargetBoneName, bShouldTransact);
 			TargetSeqController.SetBoneTrackKeys(TargetBoneName, RawTrack.PosKeys, RawTrack.RotKeys, RawTrack.ScaleKeys);
@@ -339,7 +341,7 @@ void FIKRetargetBatchOperation::GetNewAssets(TArray<UObject*>& NewAssets) const
 
 
 void FIKRetargetBatchOperation::RunRetarget(FIKRetargetBatchOperationContext& Context)
-{
+{	
 	const int32 NumAssets = GenerateAssetLists(Context);
 	
 	// show progress bar
diff --git a/Engine/Plugins/Animation/IKRig/Source/IKRigEditor/Private/RetargetEditor/IKRetargetCommands.cpp b/Engine/Plugins/Animation/IKRig/Source/IKRigEditor/Private/RetargetEditor/IKRetargetCommands.cpp
--- a/Engine/Plugins/Animation/IKRig/Source/IKRigEditor/Private/RetargetEditor/IKRetargetCommands.cpp
+++ b/Engine/Plugins/Animation/IKRig/Source/IKRigEditor/Private/RetargetEditor/IKRetargetCommands.cpp
@@ -10,7 +10,6 @@ void FIKRetargetCommands::RegisterCommands()
 	UI_COMMAND(NewRetargetPose, "New Pose", "Create a new retarget pose for the target mesh.", EUserInterfaceActionType::Button, FInputChord());
 	UI_COMMAND(DeleteRetargetPose, "Delete Pose", "Delete a retarget pose.", EUserInterfaceActionType::Button, FInputChord());
 	UI_COMMAND(ResetRetargetPose, "Reset Pose", "Resets a retarget pose to the mesh reference pose.", EUserInterfaceActionType::Button, FInputChord());
-	UI_COMMAND(ExportAnimation, "Export Animation", "Export selected animation(s) on the target mesh. Creates new Animation Sequence assets.", EUserInterfaceActionType::Button, FInputChord());
 }
 
 #undef LOCTEXT_NAMESPACE
diff --git a/Engine/Plugins/Animation/IKRig/Source/IKRigEditor/Private/RetargetEditor/IKRetargetEditMode.cpp b/Engine/Plugins/Animation/IKRig/Source/IKRigEditor/Private/RetargetEditor/IKRetargetEditMode.cpp
--- a/Engine/Plugins/Animation/IKRig/Source/IKRigEditor/Private/RetargetEditor/IKRetargetEditMode.cpp
+++ b/Engine/Plugins/Animation/IKRig/Source/IKRigEditor/Private/RetargetEditor/IKRetargetEditMode.cpp
@@ -10,9 +10,12 @@
 #include "Materials/MaterialInstanceDynamic.h"
 #include "RetargetEditor/IKRetargetHitProxies.h"
 #include "RetargetEditor/IKRetargetEditor.h"
+#include "Retargeter/IKRetargetProcessor.h"
 
 FName FIKRetargetEditMode::ModeName("IKRetargetAssetEditMode");
 
+#define LOCTEXT_NAMESPACE "IKRetargeterEditMode"
+
 bool FIKRetargetEditMode::GetCameraTarget(FSphere& OutTarget) const
 {
 	// target skeletal mesh
@@ -47,8 +50,13 @@ void FIKRetargetEditMode::Render(const FSceneView* View, FViewport* Viewport, FP
 	{
 		return;
 	}
+
+	if (!Controller->GetRetargetProcessor()->IsInitialized())
+	{
+		return;
+	}
 	
-	UIKRetargeter* Retargeter = Controller->Asset;
+	const UIKRetargeter* Retargeter = Controller->AssetController->GetAsset();
 	
 	for (int32 BoneIndex = 0; BoneIndex<Retargeter->TargetIKRigAsset->Skeleton.BoneNames.Num(); ++BoneIndex)
 	{
@@ -87,22 +95,40 @@ void FIKRetargetEditMode::Render(const FSceneView* View, FViewport* Viewport, FP
 
 bool FIKRetargetEditMode::AllowWidgetMove()
 {
-	return !SelectedBones.IsEmpty();
+	return false;
 }
 
 bool FIKRetargetEditMode::ShouldDrawWidget() const
 {
-	return !SelectedBones.IsEmpty();
+	return UsesTransformWidget(CurrentWidgetMode);
 }
 
 bool FIKRetargetEditMode::UsesTransformWidget() const
 {
-	return !SelectedBones.IsEmpty();
+	return UsesTransformWidget(CurrentWidgetMode);
 }
 
 bool FIKRetargetEditMode::UsesTransformWidget(UE::Widget::EWidgetMode CheckMode) const
 {
-	return !SelectedBones.IsEmpty() && CheckMode == UE::Widget::EWidgetMode::WM_Rotate;
+	const bool bIsOnlyRootSelected = IsOnlyRootSelected();
+	const bool bIsAnyBoneSelected = !SelectedBones.IsEmpty();
+
+	if (!bIsAnyBoneSelected)
+	{
+		return false; // no bones selected, can't transform anything
+	}
+	
+	if (bIsOnlyRootSelected && CheckMode == UE::Widget::EWidgetMode::WM_Translate)
+	{
+		return true; // can translate only the root
+	}
+
+	if (bIsAnyBoneSelected && CheckMode == UE::Widget::EWidgetMode::WM_Rotate)
+	{
+		return true; // can rotate any bone
+	}
+	
+	return false;
 }
 
 FVector FIKRetargetEditMode::GetWidgetLocation() const
@@ -118,7 +144,8 @@ FVector FIKRetargetEditMode::GetWidgetLocation() const
 		return FVector::ZeroVector; 
 	}
 
-	const int32 BoneIndex = Controller->GetCurrentlyRunningRetargeter()->TargetSkeleton.FindBoneIndexByName(SelectedBones.Last());
+	const FTargetSkeleton& TargetSkeleton = Controller->GetRetargetProcessor()->GetTargetSkeleton();
+	const int32 BoneIndex = TargetSkeleton.FindBoneIndexByName(SelectedBones.Last());
 	if (BoneIndex == INDEX_NONE)
 	{
 		return FVector::ZeroVector;
@@ -132,7 +159,7 @@ bool FIKRetargetEditMode::HandleClick(FEditorViewportClient* InViewportClient, H
 	// check for selections
 	if (Click.GetKey() == EKeys::LeftMouseButton)
 	{
-		// selected goal
+		// selected bone ?
 		if (HitProxy && HitProxy->IsA(HIKRetargetEditorBoneProxy::StaticGetType()))
 		{
 			HIKRetargetEditorBoneProxy* BoneProxy = static_cast<HIKRetargetEditorBoneProxy*>(HitProxy);
@@ -142,7 +169,7 @@ bool FIKRetargetEditMode::HandleClick(FEditorViewportClient* InViewportClient, H
 		}
 		else
 		{
-			// clicking in empty space clears selected goals
+			// clicking in empty space clears selection
 			const bool bReplaceSelection = true;
 			HandleBoneSelectedInViewport(NAME_None, bReplaceSelection);
 		}
@@ -153,57 +180,105 @@ bool FIKRetargetEditMode::HandleClick(FEditorViewportClient* InViewportClient, H
 
 bool FIKRetargetEditMode::StartTracking(FEditorViewportClient* InViewportClient, FViewport* InViewport)
 {
-	if (SelectedBones.IsEmpty())
-	{
-		return false; // no goals selected to manipulate
-	}
-
+	TrackingState = None;
+	
 	const EAxisList::Type CurrentAxis = InViewportClient->GetCurrentWidgetAxis();
 	if (CurrentAxis == EAxisList::None)
 	{
 		return false; // not manipulating a required axis
 	}
-	
-	bRotatingBones = true;
-	return true;
+
+	const TSharedPtr<FIKRetargetEditorController> Controller = EditorController.Pin();
+	if (!Controller.IsValid())
+	{
+		return false; // invalid editor state
+	}
+
+	// get state of viewport
+	const bool bTranslating = InViewportClient->GetWidgetMode() == UE::Widget::EWidgetMode::WM_Translate;
+	const bool bRotating = InViewportClient->GetWidgetMode() == UE::Widget::EWidgetMode::WM_Rotate;
+	const bool bAnyBoneSelected = !SelectedBones.IsEmpty();
+	const bool bOnlyRootSelected = IsOnlyRootSelected();
+
+	// is any bone being rotated?
+	if (bRotating && bAnyBoneSelected)
+	{
+		// Start a rotation transaction
+		GEditor->BeginTransaction(LOCTEXT("RotateRetargetPoseBone", "Rotate Retarget Pose Bone"));
+		Controller->AssetController->GetAsset()->Modify();
+		TrackingState = FIKRetargetTrackingState::RotatingBone;
+		return true;
+	}
+
+	// is the root being translated?
+	if (bTranslating && bOnlyRootSelected)
+	{
+		// Start a translation transaction
+		GEditor->BeginTransaction(LOCTEXT("TranslateRetargetPoseBone", "Translate Retarget Pose Bone"));
+		Controller->AssetController->GetAsset()->Modify();
+		TrackingState = FIKRetargetTrackingState::TranslatingRoot;
+		return true;
+	}
+
+	return false;
 }
 
 bool FIKRetargetEditMode::EndTracking(FEditorViewportClient* InViewportClient, FViewport* InViewport)
-{
-	if (!bRotatingBones)
+{	
+	if (TrackingState == FIKRetargetTrackingState::None)
 	{
 		return false; // not handled
 	}
 
-	bRotatingBones = false;
+	GEditor->EndTransaction();
+	TrackingState = FIKRetargetTrackingState::None;
 	return true;
 }
 
 bool FIKRetargetEditMode::InputDelta(FEditorViewportClient* InViewportClient, FViewport* InViewport, FVector& InDrag, FRotator& InRot, FVector& InScale)
 {
-	if (!bRotatingBones)
+	if (TrackingState == FIKRetargetTrackingState::None)
 	{
 		return false; // not handled
 	}
-	
-	if(InViewportClient->GetWidgetMode() != UE::Widget::WM_Rotate)
-	{
-		return false;
-	}
 
 	const TSharedPtr<FIKRetargetEditorController> Controller = EditorController.Pin();
 	if (!Controller.IsValid())
 	{
 		return false; 
 	}
+
+	// rotating any bone
+	if (TrackingState == FIKRetargetTrackingState::RotatingBone)
+	{
+		if(InViewportClient->GetWidgetMode() != UE::Widget::WM_Rotate)
+		{
+			return false;
+		}
+
+		// apply rotation delta to all selected bones
+		for (const FName& BoneName : SelectedBones)
+		{
+			Controller->AssetController->AddRotationOffsetToRetargetPoseBone(BoneName, InRot.Quaternion());
+		}
+
+		return true;
+	}
 	
-	// rotate chains
-	for (const FName& BoneName : SelectedBones)
+	// translating root
+	if (TrackingState == FIKRetargetTrackingState::TranslatingRoot)
 	{
-		Controller->Asset->AddRotationOffsetToRetargetPoseBone(BoneName, InRot.Quaternion());
+		if(InViewportClient->GetWidgetMode() != UE::Widget::WM_Translate)
+		{
+			return false;
+		}
+
+		// apply translation delta to root
+		Controller->AssetController->AddTranslationOffsetToRetargetRootBone(InDrag);
+		return true;
 	}
 	
-	return true;
+	return false;
 }
 
 bool FIKRetargetEditMode::GetCustomDrawingCoordinateSystem(FMatrix& InMatrix, void* InData)
@@ -219,7 +294,7 @@ bool FIKRetargetEditMode::GetCustomDrawingCoordinateSystem(FMatrix& InMatrix, vo
 		return false; 
 	}
 
-	const int32 BoneIndex = Controller->Asset->TargetIKRigAsset->Skeleton.GetBoneIndexFromName(SelectedBones[0]);
+	const int32 BoneIndex = Controller->AssetController->GetAsset()->TargetIKRigAsset->Skeleton.GetBoneIndexFromName(SelectedBones[0]);
 	const FTransform CurrentTransform = Controller->GetTargetBoneTransform(BoneIndex);
 	InMatrix = CurrentTransform.ToMatrixNoScale().RemoveTranslation();
 	return true;
@@ -230,6 +305,22 @@ bool FIKRetargetEditMode::GetCustomInputCoordinateSystem(FMatrix& InMatrix, void
 	return GetCustomDrawingCoordinateSystem(InMatrix, InData);
 }
 
+bool FIKRetargetEditMode::IsOnlyRootSelected() const
+{
+	if (SelectedBones.Num() != 1)
+	{
+		return false;
+	}
+
+	const TSharedPtr<FIKRetargetEditorController> Controller = EditorController.Pin();
+	if (!Controller.IsValid())
+	{
+		return false; 
+	}
+
+	return Controller->AssetController->GetTargetRootBone() == SelectedBones[0];
+}
+
 void FIKRetargetEditMode::DrawBoneProxy(
 	FPrimitiveDrawInterface* PDI,
 	const FTransform& BoneTransform,
@@ -256,6 +347,8 @@ void FIKRetargetEditMode::DrawBoneProxy(
 void FIKRetargetEditMode::Tick(FEditorViewportClient* ViewportClient, float DeltaTime)
 {
 	FEdMode::Tick(ViewportClient, DeltaTime);
+	
+	CurrentWidgetMode = ViewportClient->GetWidgetMode();
 }
 
 void FIKRetargetEditMode::DrawHUD(FEditorViewportClient* ViewportClient, FViewport* Viewport, const FSceneView* View, FCanvas* Canvas)
@@ -294,3 +387,5 @@ void FIKRetargetEditMode::HandleBoneSelectedInViewport(const FName& BoneName, bo
 		}	
 	}
 }
+
+#undef LOCTEXT_NAMESPACE
diff --git a/Engine/Plugins/Animation/IKRig/Source/IKRigEditor/Private/RetargetEditor/IKRetargetEditor.cpp b/Engine/Plugins/Animation/IKRig/Source/IKRigEditor/Private/RetargetEditor/IKRetargetEditor.cpp
--- a/Engine/Plugins/Animation/IKRig/Source/IKRigEditor/Private/RetargetEditor/IKRetargetEditor.cpp
+++ b/Engine/Plugins/Animation/IKRig/Source/IKRigEditor/Private/RetargetEditor/IKRetargetEditor.cpp
@@ -16,6 +16,7 @@
 #include "Framework/MultiBox/MultiBoxBuilder.h"
 
 #include "Retargeter/IKRetargeter.h"
+#include "Retargeter/IKRetargetProcessor.h"
 #include "RetargetEditor/IKRetargetAnimInstance.h"
 #include "RetargetEditor/IKRetargetCommands.h"
 #include "RetargetEditor/IKRetargetEditMode.h"
@@ -44,8 +45,7 @@ void FIKRetargetEditor::InitAssetEditor(
     const TSharedPtr<IToolkitHost>& InitToolkitHost, 
     UIKRetargeter* InAsset)
 {
-	EditorController->Editor = SharedThis(this);
-	EditorController->Asset = InAsset;
+	EditorController->Initialize(SharedThis(this), InAsset);
 
 	BindCommands();
 	
@@ -53,7 +53,7 @@ void FIKRetargetEditor::InitAssetEditor(
 	PersonaToolkitArgs.OnPreviewSceneCreated = FOnPreviewSceneCreated::FDelegate::CreateSP(this, &FIKRetargetEditor::HandlePreviewSceneCreated);
 	
 	FPersonaModule& PersonaModule = FModuleManager::LoadModuleChecked<FPersonaModule>("Persona");
-	EditorController->PersonaToolkit = PersonaModule.CreatePersonaToolkit(InAsset, PersonaToolkitArgs, InAsset->GetSourceSkeletonAsset());
+	EditorController->PersonaToolkit = PersonaModule.CreatePersonaToolkit(InAsset, PersonaToolkitArgs, EditorController->AssetController->GetSourceSkeletonAsset());
 	
 	// when/if preview mesh is changed, we need to reinitialize the anim instance
     EditorController->PersonaToolkit->GetPreviewScene()->RegisterOnPreviewMeshChanged(FOnPreviewMeshChanged::CreateSP(this, &FIKRetargetEditor::HandlePreviewMeshChanged));
@@ -131,12 +131,6 @@ void FIKRetargetEditor::BindCommands()
 		FExecuteAction::CreateSP(this, &FIKRetargetEditor::HandleResetPose),
 		FCanExecuteAction(),
 		EUIActionRepeatMode::RepeatDisabled);
-		
-	ToolkitCommands->MapAction(
-		Commands.ExportAnimation,
-		FExecuteAction::CreateSP(this, &FIKRetargetEditor::ExportAnimation),
-		FCanExecuteAction(),
-		EUIActionRepeatMode::RepeatDisabled);
 }
 
 void FIKRetargetEditor::ExtendToolbar()
@@ -158,7 +152,7 @@ void FIKRetargetEditor::FillToolbar(FToolBarBuilder& ToolbarBuilder)
 	ToolbarBuilder.BeginSection("Retarget Pose");
 	{
 		PoseNames.Reset();
-		for (TTuple<FName, FIKRetargetPose>& Pose : EditorController->Asset->RetargetPoses)
+		for (const TTuple<FName, FIKRetargetPose>& Pose : EditorController->AssetController->GetRetargetPoses())
 		{
 			PoseNames.Add(MakeShareable(new FName(Pose.Key)));
 		}
@@ -237,12 +231,7 @@ FText FIKRetargetEditor::GetBaseToolkitName() const
 
 FText FIKRetargetEditor::GetToolkitName() const
 {
-	const bool bDirtyState = EditorController->Asset->GetOutermost()->IsDirty();
-
-	FFormatNamedArguments Args;
-	Args.Add(TEXT("AssetName"), FText::FromString(EditorController->Asset->GetName()));
-	Args.Add(TEXT("DirtyState"), bDirtyState ? FText::FromString(TEXT("*")) : FText::GetEmpty());
-	return FText::Format(LOCTEXT("RetargetEditorToolkitName", "{AssetName}{DirtyState}"), Args);
+	return FText::FromString(EditorController->AssetController->GetAsset()->GetName());
 }
 
 FLinearColor FIKRetargetEditor::GetWorldCentricTabColorScale() const
@@ -258,7 +247,8 @@ FString FIKRetargetEditor::GetWorldCentricTabPrefix() const
 void FIKRetargetEditor::AddReferencedObjects(FReferenceCollector& Collector)
 {
 	// hold the asset we are working on
-	Collector.AddReferencedObject(EditorController->Asset);
+	const UIKRetargeter* Retargeter = EditorController->AssetController->GetAsset();
+	Collector.AddReferencedObject(Retargeter);
 }
 
 void FIKRetargetEditor::Tick(float DeltaTime)
@@ -266,10 +256,12 @@ void FIKRetargetEditor::Tick(float DeltaTime)
 	// apply offset to the target component
 	if (EditorController->TargetSkelMeshComponent)
 	{
-		const float TargetOffset = EditorController->Asset->TargetActorOffset;
+		const UIKRetargeter* Retargeter = EditorController->AssetController->GetAsset();
+		
+		const float TargetOffset = Retargeter->TargetActorOffset;
 		EditorController->TargetSkelMeshComponent->SetRelativeLocation(FVector(TargetOffset,0,0));
 
-		const float TargetScale = EditorController->Asset->TargetActorScale;
+		const float TargetScale = Retargeter->TargetActorScale;
 		EditorController->TargetSkelMeshComponent->SetRelativeScale3D(FVector(TargetScale,TargetScale,TargetScale));
 	}
 }
@@ -315,7 +307,7 @@ void FIKRetargetEditor::HandlePreviewSceneCreated(const TSharedRef<IPersonaPrevi
 void FIKRetargetEditor::SetupAnimInstance()
 {
 	// connect the retarget asset and the source component ot the target anim instance
-	EditorController->TargetAnimInstance->SetRetargetAssetAndSourceComponent(EditorController->Asset, EditorController->SourceSkelMeshComponent);
+	EditorController->TargetAnimInstance->SetRetargetAssetAndSourceComponent(EditorController->AssetController->GetAsset(), EditorController->SourceSkelMeshComponent);
 
 	EditorController->SourceSkelMeshComponent->PreviewInstance = EditorController->SourceAnimInstance.Get();
 	EditorController->TargetSkelMeshComponent->PreviewInstance = EditorController->TargetAnimInstance.Get();
@@ -326,16 +318,12 @@ void FIKRetargetEditor::SetupAnimInstance()
 
 FText FIKRetargetEditor::GetCurrentPoseName() const
 {
-	return FText::FromName(EditorController->Asset->CurrentRetargetPose);
+	return FText::FromName(EditorController->AssetController->GetCurrentRetargetPoseName());
 }
 
 void FIKRetargetEditor::OnPoseSelected(TSharedPtr<FName> InPosePose, ESelectInfo::Type SelectInfo)
 {
-	EditorController->Asset->SetCurrentRetargetPose(*InPosePose.Get());
-}
-
-void FIKRetargetEditor::ExportAnimation() const
-{
+	EditorController->AssetController->SetCurrentRetargetPose(*InPosePose.Get());
 }
 
 void FIKRetargetEditor::HandleSourceOrTargetIKRigAssetChanged()
@@ -374,14 +362,17 @@ void FIKRetargetEditor::HandleDetailsCreated(const TSharedRef<class IDetailsView
 {
 	EditorController->DetailsView = InDetailsView;
 	EditorController->DetailsView->OnFinishedChangingProperties().AddSP(this, &FIKRetargetEditor::OnFinishedChangingDetails);
-	EditorController->DetailsView->SetObject(EditorController->Asset);
+	EditorController->DetailsView->SetObject(EditorController->AssetController->GetAsset());
 }
 
 void FIKRetargetEditor::OnFinishedChangingDetails(
     const FPropertyChangedEvent& PropertyChangedEvent)
 {
-	if (PropertyChangedEvent.GetPropertyName() == "TargetIKRigAsset"
-		|| PropertyChangedEvent.GetPropertyName() == "SourceIKRigAsset")
+	FName SourceIKRigPropertyName = GET_MEMBER_NAME_STRING_CHECKED(UIKRetargeter, SourceIKRigAsset);
+	FName TargetIKRigPropertyName = GET_MEMBER_NAME_STRING_CHECKED(UIKRetargeter, TargetIKRigAsset);
+	const bool bSourceChanged = PropertyChangedEvent.GetPropertyName() == SourceIKRigPropertyName;
+	const bool bTargetChanged = PropertyChangedEvent.GetPropertyName() == TargetIKRigPropertyName;
+	if (bSourceChanged || bTargetChanged)
 	{
 		HandleSourceOrTargetIKRigAssetChanged();
 	}
@@ -389,9 +380,9 @@ void FIKRetargetEditor::OnFinishedChangingDetails(
 
 void FIKRetargetEditor::HandleEditPose() const
 {
-	UIKRetargeter* Asset = EditorController.Get().Asset;
-	Asset->bEditReferencePoseMode = !EditorController.Get().Asset->bEditReferencePoseMode;
-	if (Asset->bEditReferencePoseMode)
+	const bool bEditPoseMode = !EditorController.Get().AssetController->GetEditRetargetPoseMode();
+	EditorController.Get().AssetController->SetEditRetargetPoseMode(bEditPoseMode);
+	if (bEditPoseMode)
 	{
 		GetEditorModeManager().ActivateMode(FIKRetargetEditMode::ModeName);
 		EditorController->SourceSkelMeshComponent->ShowReferencePose(true);
@@ -401,24 +392,22 @@ void FIKRetargetEditor::HandleEditPose() const
 		GetEditorModeManager().DeactivateMode(FIKRetargetEditMode::ModeName);
 		EditorController->PlayPreviousAnimationAsset();
 	}
-
-	Asset->Modify();
 }
 
 bool FIKRetargetEditor::CanEditPose() const
 {
-	UIKRetargeter* CurrentRetargeter = EditorController->GetCurrentlyRunningRetargeter();
-	if (!CurrentRetargeter)
+	const UIKRetargetProcessor* Processor = EditorController->GetRetargetProcessor();
+	if (!Processor)
 	{
 		return false;
 	}
 
-	return CurrentRetargeter->bIsLoadedAndValid;
+	return Processor->IsInitialized();
 }
 
 bool FIKRetargetEditor::IsEditingPose() const
 {
-	return EditorController.Get().Asset->bEditReferencePoseMode;
+	return EditorController.Get().AssetController->GetEditRetargetPoseMode();
 }
 
 void FIKRetargetEditor::HandleNewPose()
@@ -485,7 +474,7 @@ void FIKRetargetEditor::HandleNewPose()
 FReply FIKRetargetEditor::CreateNewPose()
 {
 	const FName NewPoseName = FName(NewPoseEditableText.Get()->GetText().ToString());
-	EditorController->Asset->AddRetargetPose(NewPoseName);
+	EditorController->AssetController->AddRetargetPose(NewPoseName);
 	NewPoseWindow->RequestDestroyWindow();
 	RegenerateMenusAndToolbars();
 	return FReply::Handled();
@@ -493,22 +482,21 @@ FReply FIKRetargetEditor::CreateNewPose()
 
 void FIKRetargetEditor::HandleDeletePose()
 {
-	const FName CurrentPose = EditorController->Asset->CurrentRetargetPose;
-	EditorController->Asset->RemoveRetargetPose(CurrentPose);
+	const FName CurrentPose = EditorController->AssetController->GetCurrentRetargetPoseName();
+	EditorController->AssetController->RemoveRetargetPose(CurrentPose);
 	RegenerateMenusAndToolbars();
 }
 
 bool FIKRetargetEditor::CanDeletePose() const
 {	
 	// cannot delete default pose
-	return EditorController->Asset->CurrentRetargetPose != UIKRetargeter::DefaultPoseName;
+	return EditorController->AssetController->GetCurrentRetargetPoseName() != UIKRetargeter::DefaultPoseName;
 }
 
 void FIKRetargetEditor::HandleResetPose()
 {
-	const FName CurrentPose = EditorController->Asset->CurrentRetargetPose;
-	EditorController->Asset->ResetRetargetPose(CurrentPose);
+	const FName CurrentPose = EditorController->AssetController->GetCurrentRetargetPoseName();
+	EditorController->AssetController->ResetRetargetPose(CurrentPose);
 }
 
-
 #undef LOCTEXT_NAMESPACE
diff --git a/Engine/Plugins/Animation/IKRig/Source/IKRigEditor/Private/RetargetEditor/IKRetargetEditorController.cpp b/Engine/Plugins/Animation/IKRig/Source/IKRigEditor/Private/RetargetEditor/IKRetargetEditorController.cpp
--- a/Engine/Plugins/Animation/IKRig/Source/IKRigEditor/Private/RetargetEditor/IKRetargetEditorController.cpp
+++ b/Engine/Plugins/Animation/IKRig/Source/IKRigEditor/Private/RetargetEditor/IKRetargetEditorController.cpp
@@ -5,27 +5,87 @@
 #include "RetargetEditor/IKRetargetAnimInstance.h"
 #include "RetargetEditor/SIKRetargetChainMapList.h"
 #include "Retargeter/IKRetargeter.h"
+#include "RigEditor/IKRigController.h"
 
 #define LOCTEXT_NAMESPACE "IKRetargetEditorController"
 
+void FIKRetargetEditorController::Initialize(TSharedPtr<FIKRetargetEditor> InEditor, UIKRetargeter* InAsset)
+{
+	Editor = InEditor;
+	AssetController = UIKRetargeterController::GetController(InAsset);
+
+	// bind callbacks when SOURCE or TARGET IK Rigs are modified
+	BindToIKRigAsset(AssetController->GetAsset()->SourceIKRigAsset);
+	BindToIKRigAsset(AssetController->GetAsset()->TargetIKRigAsset);
+
+	// bind callback when retargeter needs reinitialized
+	AssetController->OnRetargeterNeedsInitialized().AddSP(this, &FIKRetargetEditorController::OnRetargeterNeedsInitialized);
+}
+
+void FIKRetargetEditorController::BindToIKRigAsset(UIKRigDefinition* InIKRig)
+{
+	if (!InIKRig)
+	{
+		return;
+	}
+
+	UIKRigController* Controller = UIKRigController::GetIKRigController(InIKRig);
+	if (!Controller->OnIKRigNeedsInitialized().IsBoundToObject(this))
+	{
+		Controller->OnIKRigNeedsInitialized().AddSP(this, &FIKRetargetEditorController::OnIKRigNeedsInitialized);
+		Controller->OnRetargetChainRenamed().AddSP(this, &FIKRetargetEditorController::OnRetargetChainRenamed);
+	}
+}
+
+void FIKRetargetEditorController::OnIKRigNeedsInitialized(UIKRigDefinition* ModifiedIKRig)
+{
+	const UIKRetargeter* Retargeter = AssetController->GetAsset();
+	
+	check(ModifiedIKRig && Retargeter)
+	 
+	const bool bIsSource = ModifiedIKRig == Retargeter->SourceIKRigAsset;
+	const bool bIsTarget = ModifiedIKRig == Retargeter->TargetIKRigAsset;
+	if (!(bIsSource || bIsTarget))
+	{
+		return;
+	}
+
+	// the target anim instance has the IK RetargetPoseFromMesh node which needs reinitialized with new asset version
+	TargetAnimInstance->SetProcessorNeedsInitialized();
+	RefreshAllViews();
+}
+
+void FIKRetargetEditorController::OnRetargetChainRenamed(UIKRigDefinition* ModifiedIKRig, FName OldName, FName NewName)
+{
+	check(ModifiedIKRig)
+	
+	AssetController->OnRetargetChainRenamed(ModifiedIKRig, OldName, NewName);
+}
+
+void FIKRetargetEditorController::OnRetargeterNeedsInitialized(const UIKRetargeter* Retargeter)
+{
+	TargetAnimInstance->SetProcessorNeedsInitialized();
+	RefreshAllViews();
+}
+
 USkeletalMesh* FIKRetargetEditorController::GetSourceSkeletalMesh() const
 {
-	if (!(Asset && Asset->SourceIKRigAsset))
+	if (!(AssetController && AssetController->GetAsset()->SourceIKRigAsset))
 	{
 		return nullptr;
 	}
 
-	return Asset->SourceIKRigAsset->PreviewSkeletalMesh;
+	return AssetController->GetAsset()->SourceIKRigAsset->PreviewSkeletalMesh;
 }
 
 USkeletalMesh* FIKRetargetEditorController::GetTargetSkeletalMesh() const
 {
-	if (!(Asset && Asset->TargetIKRigAsset))
+	if (!(AssetController && AssetController->GetAsset()->TargetIKRigAsset))
 	{
 		return nullptr;
 	}
 
-	return Asset->TargetIKRigAsset->PreviewSkeletalMesh;
+	return AssetController->GetAsset()->TargetIKRigAsset->PreviewSkeletalMesh;
 }
 
 FTransform FIKRetargetEditorController::GetTargetBoneTransform(const int32& TargetBoneIndex) const
@@ -33,15 +93,15 @@ FTransform FIKRetargetEditorController::GetTargetBoneTransform(const int32& Targ
 	UIKRetargetAnimInstance* AnimInstance = TargetAnimInstance.Get();
 	check(AnimInstance);
 
-	UIKRetargeter* CurrentRetargeter = AnimInstance->GetCurrentlyUsedRetargeter();
-	check(CurrentRetargeter)
+	const UIKRetargetProcessor* RetargetProcessor = AnimInstance->GetRetargetProcessor();
+	check(RetargetProcessor)
 
 	// get transform of root of chain
-	FTransform BoneTransform = CurrentRetargeter->GetTargetBoneRetargetPoseGlobalTransform(TargetBoneIndex);
+	FTransform BoneTransform = RetargetProcessor->GetTargetBoneRetargetPoseGlobalTransform(TargetBoneIndex);
 
 	// scale and offset
-	BoneTransform.ScaleTranslation(Asset->TargetActorScale);
-	BoneTransform.AddToTranslation(FVector(Asset->TargetActorOffset, 0.f, 0.f));
+	BoneTransform.ScaleTranslation(AssetController->GetAsset()->TargetActorScale);
+	BoneTransform.AddToTranslation(FVector(AssetController->GetAsset()->TargetActorOffset, 0.f, 0.f));
 
 	return BoneTransform;
 }
@@ -51,12 +111,13 @@ bool FIKRetargetEditorController::GetTargetBoneLineSegments(
 	FVector& OutStart,
 	TArray<FVector>& OutChildren) const
 {
-	UIKRetargeter* CurrentRetargeter = GetCurrentlyRunningRetargeter();
-	check(CurrentRetargeter && CurrentRetargeter->bIsLoadedAndValid)
-	check(CurrentRetargeter->TargetSkeleton.BoneNames.IsValidIndex(TargetBoneIndex))
+	// get the runtime processor
+	const UIKRetargetProcessor* Processor = GetRetargetProcessor();
+	check(Processor && Processor->IsInitialized())
 	
 	// get the target skeleton we want to draw
-	const FTargetSkeleton& TargetSkeleton = CurrentRetargeter->TargetSkeleton;
+	const FTargetSkeleton& TargetSkeleton = Processor->GetTargetSkeleton();
+	check(TargetSkeleton.BoneNames.IsValidIndex(TargetBoneIndex))
 
 	// get the origin of the bone chain
 	OutStart = TargetSkeleton.RetargetGlobalPose[TargetBoneIndex].GetTranslation();
@@ -70,12 +131,12 @@ bool FIKRetargetEditorController::GetTargetBoneLineSegments(
 	}
 
 	// add the target translation offset and scale
-	const FVector TargetOffset(Asset->TargetActorOffset, 0.f, 0.f);
-	OutStart *= Asset->TargetActorScale;
+	const FVector TargetOffset(AssetController->GetAsset()->TargetActorOffset, 0.f, 0.f);
+	OutStart *= AssetController->GetAsset()->TargetActorScale;
 	OutStart += TargetOffset;
 	for (FVector& ChildPoint : OutChildren)
 	{
-		ChildPoint *= Asset->TargetActorScale;
+		ChildPoint *= AssetController->GetAsset()->TargetActorScale;
 		ChildPoint += TargetOffset;
 	}
 	
@@ -84,18 +145,22 @@ bool FIKRetargetEditorController::GetTargetBoneLineSegments(
 
 bool FIKRetargetEditorController::IsTargetBoneRetargeted(const int32& TargetBoneIndex)
 {
-	UIKRetargeter* CurrentRetargeter = GetCurrentlyRunningRetargeter();
-	check(CurrentRetargeter && CurrentRetargeter->bIsLoadedAndValid)
-	check(CurrentRetargeter->TargetSkeleton.BoneNames.IsValidIndex(TargetBoneIndex))
+	// get the runtime processor
+	const UIKRetargetProcessor* Processor = GetRetargetProcessor();
+	check(Processor && Processor->IsInitialized())
+	
+	// get the target skeleton
+	const FTargetSkeleton& TargetSkeleton = Processor->GetTargetSkeleton();
+	check(TargetSkeleton.BoneNames.IsValidIndex(TargetBoneIndex))
 
-	return CurrentRetargeter->TargetSkeleton.IsBoneRetargeted[TargetBoneIndex];
+	return TargetSkeleton.IsBoneRetargeted[TargetBoneIndex];
 }
 
-UIKRetargeter* FIKRetargetEditorController::GetCurrentlyRunningRetargeter() const
+const UIKRetargetProcessor* FIKRetargetEditorController::GetRetargetProcessor() const
 {	
 	if(UIKRetargetAnimInstance* AnimInstance = TargetAnimInstance.Get())
 	{
-		return AnimInstance->GetCurrentlyUsedRetargeter();
+		return AnimInstance->GetRetargetProcessor();
 	}
 
 	return nullptr;	
diff --git a/Engine/Plugins/Animation/IKRig/Source/IKRigEditor/Private/RetargetEditor/IKRetargeterController.cpp b/Engine/Plugins/Animation/IKRig/Source/IKRigEditor/Private/RetargetEditor/IKRetargeterController.cpp
new file mode 100644
--- /dev/null
+++ b/Engine/Plugins/Animation/IKRig/Source/IKRigEditor/Private/RetargetEditor/IKRetargeterController.cpp
@@ -0,0 +1,386 @@
+﻿// Copyright Epic Games, Inc. All Rights Reserved.
+
+#include "RetargetEditor/IKRetargeterController.h"
+
+#include "ScopedTransaction.h"
+#include "Algo/LevenshteinDistance.h"
+#include "Retargeter/IKRetargeter.h"
+
+#define LOCTEXT_NAMESPACE "IKRetargeterController"
+
+
+UIKRetargeterController* UIKRetargeterController::GetController(UIKRetargeter* InRetargeterAsset)
+{
+	if (!InRetargeterAsset)
+	{
+		return nullptr;
+	}
+
+	if (!InRetargeterAsset->Controller)
+	{
+		UIKRetargeterController* Controller = NewObject<UIKRetargeterController>();
+		Controller->Asset = InRetargeterAsset;
+		InRetargeterAsset->Controller = Controller;
+	}
+
+	UIKRetargeterController* Controller = Cast<UIKRetargeterController>(InRetargeterAsset->Controller);
+	// clean the asset before editing
+	Controller->CleanChainMapping();
+	Controller->CleanPoseList();
+	
+	return Controller;
+}
+
+UIKRetargeter* UIKRetargeterController::GetAsset() const
+{
+	return Asset;
+}
+
+FName UIKRetargeterController::GetSourceRootBone() const
+{
+	return Asset->SourceIKRigAsset ? Asset->SourceIKRigAsset->GetRetargetRoot() : FName("None");
+}
+
+FName UIKRetargeterController::GetTargetRootBone() const
+{
+	return Asset->TargetIKRigAsset ? Asset->TargetIKRigAsset->GetRetargetRoot() : FName("None");
+}
+
+void UIKRetargeterController::GetTargetChainNames(TArray<FName>& OutNames) const
+{
+	if (Asset->TargetIKRigAsset)
+	{
+		const TArray<FBoneChain>& Chains = Asset->TargetIKRigAsset->GetRetargetChains();
+		for (const FBoneChain& Chain : Chains)
+		{
+			OutNames.Add(Chain.ChainName);
+		}
+	}
+}
+
+void UIKRetargeterController::GetSourceChainNames(TArray<FName>& OutNames) const
+{
+	if (Asset->SourceIKRigAsset)
+	{
+		const TArray<FBoneChain>& Chains = Asset->SourceIKRigAsset->GetRetargetChains();
+		for (const FBoneChain& Chain : Chains)
+		{
+			OutNames.Add(Chain.ChainName);
+		}
+	}
+}
+
+void UIKRetargeterController::CleanChainMapping()
+{
+	if (!Asset->TargetIKRigAsset)
+	{
+		// don't clean chain mappings, in case user is replacing with IK Rig asset that has some valid mappings
+		return;
+	}
+	
+	TArray<FName> TargetChainNames;
+	GetTargetChainNames(TargetChainNames);
+
+	// remove all target chains that are no longer in the target IK rig asset
+	TArray<FName> TargetChainsToRemove;
+	for (FRetargetChainMap& ChainMap : Asset->ChainMapping)
+	{
+		if (!TargetChainNames.Contains(ChainMap.TargetChain))
+		{
+			TargetChainsToRemove.Add(ChainMap.TargetChain);
+		}
+	}
+	for (FName TargetChainToRemove : TargetChainsToRemove)
+	{
+		Asset->ChainMapping.RemoveAll([&TargetChainToRemove](FRetargetChainMap& Element)
+		{
+			return Element.TargetChain == TargetChainToRemove;
+		});
+	}
+
+	// add a mapping for each chain that is in the target IK rig (if it doesn't have one already)
+	for (FName TargetChainName : TargetChainNames)
+	{
+		const bool HasChain = Asset->ChainMapping.ContainsByPredicate([&TargetChainName](FRetargetChainMap& Element)
+		{
+			return Element.TargetChain == TargetChainName;
+		});
+		
+		if (!HasChain)
+		{
+			Asset->ChainMapping.Add(FRetargetChainMap(TargetChainName));
+		}
+	}
+
+	TArray<FName> SourceChainNames;
+	GetSourceChainNames(SourceChainNames);
+	
+	// reset any sources that are no longer present to "None"
+	for (FRetargetChainMap& ChainMap : Asset->ChainMapping)
+	{
+		if (!SourceChainNames.Contains(ChainMap.SourceChain))
+		{
+			ChainMap.SourceChain = NAME_None;
+		}
+	}
+
+	// enforce the same chain order as the target IK rig
+	Asset->ChainMapping.Sort([this](const FRetargetChainMap& A, const FRetargetChainMap& B)
+	{
+		const TArray<FBoneChain>& BoneChains = Asset->TargetIKRigAsset->GetRetargetChains();
+		
+		const int32 IndexA = BoneChains.IndexOfByPredicate([&A](const FBoneChain& Chain)
+		{
+			return A.TargetChain == Chain.ChainName;
+		});
+
+		const int32 IndexB = BoneChains.IndexOfByPredicate([&B](const FBoneChain& Chain)
+		{
+			return B.TargetChain == Chain.ChainName;
+		});
+ 
+		return IndexA < IndexB;
+	});
+
+	BroadcastNeedsReinitialized();
+}
+
+void UIKRetargeterController::CleanPoseList()
+{
+	// enforce the existence of a default pose
+	const bool HasDefaultPose = Asset->RetargetPoses.Contains(Asset->DefaultPoseName);
+	if (!HasDefaultPose)
+	{
+		Asset->RetargetPoses.Emplace(Asset->DefaultPoseName);
+	}
+	
+	// use default pose unless set to something else
+	if (Asset->CurrentRetargetPose == NAME_None)
+	{
+		Asset->CurrentRetargetPose = Asset->DefaultPoseName;
+	}
+
+	// remove all bone offsets that are no longer part of the target skeleton
+	if (Asset->TargetIKRigAsset)
+	{
+		const TArray<FName> AllowedBoneNames = Asset->TargetIKRigAsset->Skeleton.BoneNames;
+		for (TTuple<FName, FIKRetargetPose>& Pose : Asset->RetargetPoses)
+		{
+			// find bone offsets no longer in target skeleton
+			TArray<FName> BonesToRemove;
+			for (TTuple<FName, FQuat>& BoneOffset : Pose.Value.BoneRotationOffsets)
+			{
+				if (!AllowedBoneNames.Contains(BoneOffset.Key))
+				{
+					BonesToRemove.Add(BoneOffset.Key);
+				}
+			}
+			
+			// remove bone offsets
+			for (const FName& BoneToRemove : BonesToRemove)
+			{
+				Pose.Value.BoneRotationOffsets.Remove(BoneToRemove);
+			}
+		}
+	}
+
+	BroadcastNeedsReinitialized();
+}
+
+void UIKRetargeterController::AutoMapChains()
+{
+	TArray<FName> SourceChainNames;
+	GetSourceChainNames(SourceChainNames);
+	
+	// auto-map any chains that have no value using a fuzzy string search
+	for (FRetargetChainMap& ChainMap : Asset->ChainMapping)
+	{
+		if (ChainMap.SourceChain != NAME_None)
+		{
+			continue; // already set by user
+		}
+
+		// find "best match" automatically as a convenience for the user
+		FString TargetNameLowerCase = ChainMap.TargetChain.ToString().ToLower();
+		float HighestScore = 0.2f;
+		int32 HighestScoreIndex = -1;
+		for (int32 ChainIndex=0; ChainIndex<SourceChainNames.Num(); ++ChainIndex)
+		{
+			FString SourceNameLowerCase = SourceChainNames[ChainIndex].ToString().ToLower();
+			float WorstCase = TargetNameLowerCase.Len() + SourceNameLowerCase.Len();
+			WorstCase = WorstCase < 1.0f ? 1.0f : WorstCase;
+			const float Score = 1.0f - (Algo::LevenshteinDistance(TargetNameLowerCase, SourceNameLowerCase) / WorstCase);
+			if (Score > HighestScore)
+			{
+				HighestScore = Score;
+				HighestScoreIndex = ChainIndex;
+			}
+		}
+
+		// apply source if any decent matches were found
+		if (SourceChainNames.IsValidIndex(HighestScoreIndex))
+		{
+			ChainMap.SourceChain = SourceChainNames[HighestScoreIndex];
+		}
+	}
+
+	// force update with latest mapping
+	BroadcastNeedsReinitialized();
+}
+
+void UIKRetargeterController::OnRetargetChainRenamed(UIKRigDefinition* IKRig, FName OldChainName, FName NewChainName) const
+{
+	const bool bIsSourceRig = IKRig == Asset->SourceIKRigAsset;
+	check(bIsSourceRig || IKRig == Asset->TargetIKRigAsset)
+	for (FRetargetChainMap& ChainMap : Asset->ChainMapping)
+	{
+		FName& ChainNameToUpdate = bIsSourceRig ? ChainMap.SourceChain : ChainMap.TargetChain;
+		if (ChainNameToUpdate == OldChainName)
+		{
+			ChainNameToUpdate = NewChainName;
+			BroadcastNeedsReinitialized();
+			return;
+		}
+	}
+}
+
+void UIKRetargeterController::SetSourceChainForTargetChain(FName TargetChain, FName SourceChainToMapTo)
+{
+	FRetargetChainMap* ChainMap = GetChainMap(TargetChain);
+	check(ChainMap)
+	ChainMap->SourceChain = SourceChainToMapTo;
+	BroadcastNeedsReinitialized();
+}
+
+FName UIKRetargeterController::GetSourceChainForTargetChain(FName TargetChain)
+{
+	FRetargetChainMap* ChainMap = GetChainMap(TargetChain);
+	check(ChainMap)
+	return ChainMap->SourceChain;
+}
+
+const TArray<FRetargetChainMap>& UIKRetargeterController::GetChainMappings()
+{
+	return Asset->ChainMapping;
+}
+
+USkeleton* UIKRetargeterController::GetSourceSkeletonAsset() const
+{
+	if (!Asset->SourceIKRigAsset)
+	{
+		return nullptr;
+	}
+
+	if (!Asset->SourceIKRigAsset->PreviewSkeletalMesh)
+	{
+		return nullptr;
+	}
+
+	return Asset->SourceIKRigAsset->PreviewSkeletalMesh->GetSkeleton();
+}
+
+void UIKRetargeterController::AddRetargetPose(FName NewPoseName) const
+{
+	if (Asset->RetargetPoses.Contains(NewPoseName))
+	{
+		return;
+	}
+
+	Asset->RetargetPoses.Add(NewPoseName);
+	Asset->CurrentRetargetPose = NewPoseName;
+
+	BroadcastNeedsReinitialized();
+}
+
+void UIKRetargeterController::RemoveRetargetPose(FName PoseToRemove) const
+{
+	if (PoseToRemove == Asset->DefaultPoseName)
+	{
+		return; // cannot remove default pose
+	}
+
+	if (!Asset->RetargetPoses.Contains(PoseToRemove))
+	{
+		return; // cannot remove pose that doesn't exist
+	}
+
+	Asset->RetargetPoses.Remove(PoseToRemove);
+
+	// did we remove the currently used pose?
+	if (Asset->CurrentRetargetPose == PoseToRemove)
+	{
+		Asset->CurrentRetargetPose = UIKRetargeter::DefaultPoseName;
+	}
+
+	BroadcastNeedsReinitialized();
+}
+
+void UIKRetargeterController::ResetRetargetPose(FName PoseToReset) const
+{
+	if (!Asset->RetargetPoses.Contains(PoseToReset))
+	{
+		return; // cannot reset pose that doesn't exist
+	}
+
+	Asset->RetargetPoses[PoseToReset].BoneRotationOffsets.Reset();
+	Asset->RetargetPoses[PoseToReset].RootTranslationOffset = FVector::ZeroVector;
+	
+	BroadcastNeedsReinitialized();
+}
+
+FName UIKRetargeterController::GetCurrentRetargetPoseName() const
+{
+	return GetAsset()->CurrentRetargetPose;
+}
+
+void UIKRetargeterController::SetCurrentRetargetPose(FName CurrentPose) const
+{
+	check(Asset->RetargetPoses.Contains(CurrentPose));
+	Asset->CurrentRetargetPose = CurrentPose;
+	BroadcastNeedsReinitialized();
+}
+
+const TMap<FName, FIKRetargetPose>& UIKRetargeterController::GetRetargetPoses()
+{
+	return GetAsset()->RetargetPoses;
+}
+
+void UIKRetargeterController::AddRotationOffsetToRetargetPoseBone(FName BoneName, FQuat RotationOffset) const
+{
+	Asset->RetargetPoses[Asset->CurrentRetargetPose].AddRotationDeltaToBone(BoneName, RotationOffset);
+}
+
+void UIKRetargeterController::AddTranslationOffsetToRetargetRootBone(FVector TranslationOffset) const
+{
+	Asset->RetargetPoses[Asset->CurrentRetargetPose].AddTranslationDeltaToRoot(TranslationOffset);
+}
+
+void UIKRetargeterController::SetEditRetargetPoseMode(bool bEditPoseMode) const
+{
+	GetAsset()->bEditRetargetPoseMode = bEditPoseMode;
+	if (!bEditPoseMode)
+	{
+		// must reinitialize after editing the retarget pose
+		BroadcastNeedsReinitialized();
+	}
+}
+
+bool UIKRetargeterController::GetEditRetargetPoseMode() const
+{
+	return GetAsset()->bEditRetargetPoseMode;
+}
+
+FRetargetChainMap* UIKRetargeterController::GetChainMap(const FName& TargetChainName) const
+{
+	for (FRetargetChainMap& ChainMap : Asset->ChainMapping)
+	{
+		if (ChainMap.TargetChain == TargetChainName)
+		{
+			return &ChainMap;
+		}
+	}
+
+	return nullptr;
+}
+
+#undef LOCTEXT_NAMESPACE
diff --git a/Engine/Plugins/Animation/IKRig/Source/IKRigEditor/Private/RetargetEditor/SIKRetargetAssetBrowser.cpp b/Engine/Plugins/Animation/IKRig/Source/IKRigEditor/Private/RetargetEditor/SIKRetargetAssetBrowser.cpp
--- a/Engine/Plugins/Animation/IKRig/Source/IKRigEditor/Private/RetargetEditor/SIKRetargetAssetBrowser.cpp
+++ b/Engine/Plugins/Animation/IKRig/Source/IKRigEditor/Private/RetargetEditor/SIKRetargetAssetBrowser.cpp
@@ -11,7 +11,9 @@
 #include "Animation/PoseAsset.h"
 #include "RetargetEditor/IKRetargetBatchOperation.h"
 #include "RetargetEditor/IKRetargetEditorController.h"
+#include "RetargetEditor/SRetargetAnimAssetsWindow.h"
 #include "Retargeter/IKRetargeter.h"
+#include "Retargeter/IKRetargetProcessor.h"
 
 #define LOCTEXT_NAMESPACE "IKRetargeterAssetBrowser"
 
@@ -76,6 +78,11 @@ void SIKRetargetAssetBrowser::AddAssetBrowser()
 	AssetBrowserBox->SetContent(ContentBrowserModule.Get().CreateAssetPicker(AssetPickerConfig));
 }
 
+void SIKRetargetAssetBrowser::OnPathChange(const FString& NewPath)
+{
+	BatchOutputPath = NewPath;
+}
+
 FReply SIKRetargetAssetBrowser::OnExportButtonClicked() const
 {
 	FIKRetargetEditorController* Controller = EditorController.Pin().Get();
@@ -83,10 +90,22 @@ FReply SIKRetargetAssetBrowser::OnExportButtonClicked() const
 	{
 		return FReply::Handled();
 	}
-	
+
 	// assemble the data for the assets we want to batch duplicate/retarget
 	FIKRetargetBatchOperationContext BatchContext;
 
+	TSharedRef<SSelectExportPathDialog> Dialog = SNew(SSelectExportPathDialog).DefaultAssetPath(FText::FromString(BatchContext.FolderPath));
+	if(Dialog->ShowModal() != EAppReturnType::Cancel)
+	{
+		BatchContext.NameRule.FolderPath = Dialog->GetAssetPath();
+	}
+
+	// set the path to export to
+	if (!BatchOutputPath.IsEmpty())
+	{
+		BatchContext.FolderPath = BatchOutputPath;
+	}
+
 	// add selected assets to dup/retarget
 	TArray<FAssetData> SelectedAssets = GetCurrentSelectionDelegate.Execute();
 	for (const FAssetData& Asset : SelectedAssets)
@@ -98,31 +117,43 @@ FReply SIKRetargetAssetBrowser::OnExportButtonClicked() const
 
 	BatchContext.SourceMesh = Controller->GetSourceSkeletalMesh();
 	BatchContext.TargetMesh = Controller->GetTargetSkeletalMesh();
-	BatchContext.IKRetargetAsset = Controller->Asset;
+	BatchContext.IKRetargetAsset = const_cast<UIKRetargeter*>(Controller->AssetController->GetAsset());
 	BatchContext.bRemapReferencedAssets = false;
 	BatchContext.NameRule.Suffix = "_Retargeted";
 
 	// actually run the retarget
 	FIKRetargetBatchOperation BatchOperation;
 	BatchOperation.RunRetarget(BatchContext);
-	
+
 	return FReply::Handled();
+	
 }
 
 bool SIKRetargetAssetBrowser::IsExportButtonEnabled() const
 {
 	if (!EditorController.Pin().IsValid())
 	{
-		return false;
+		return false; // editor in bad state
+	}
+
+	const UIKRetargetProcessor* Processor = EditorController.Pin()->GetRetargetProcessor();
+	if (!Processor)
+	{
+		return false; // no retargeter running
 	}
 
-	UIKRetargeter* CurrentRetargeter = EditorController.Pin()->GetCurrentlyRunningRetargeter();
-	if (!IsValid(CurrentRetargeter))
+	if (!Processor->IsInitialized())
+	{
+		return false; // retargeter not loaded and valid
+	}
+
+	TArray<FAssetData> SelectedAssets = GetCurrentSelectionDelegate.Execute();
+	if (SelectedAssets.IsEmpty())
 	{
-		return false;
+		return false; // nothing selected
 	}
 
-	return CurrentRetargeter->bIsLoadedAndValid;
+	return true;
 }
 
 void SIKRetargetAssetBrowser::OnAssetDoubleClicked(const FAssetData& AssetData)
diff --git a/Engine/Plugins/Animation/IKRig/Source/IKRigEditor/Private/RetargetEditor/SIKRetargetChainMapList.cpp b/Engine/Plugins/Animation/IKRig/Source/IKRigEditor/Private/RetargetEditor/SIKRetargetChainMapList.cpp
--- a/Engine/Plugins/Animation/IKRig/Source/IKRigEditor/Private/RetargetEditor/SIKRetargetChainMapList.cpp
+++ b/Engine/Plugins/Animation/IKRig/Source/IKRigEditor/Private/RetargetEditor/SIKRetargetChainMapList.cpp
@@ -40,10 +40,11 @@ void SIKRetargetChainMapRow::Construct(
 	// NOTE: cannot just use FName because "None" is considered a null entry and removed from ComboBox.
 	SourceChainOptions.Reset();
 	SourceChainOptions.Add(MakeShareable(new FString(TEXT("None"))));
-	UIKRigDefinition* SourceIKRig = ChainMapList.Pin()->EditorController.Pin()->Asset->SourceIKRigAsset;
+	const UIKRigDefinition* SourceIKRig = ChainMapList.Pin()->EditorController.Pin()->AssetController->GetAsset()->SourceIKRigAsset;
 	if (SourceIKRig)
 	{
-		for (const FBoneChain& BoneChain : SourceIKRig->RetargetDefinition.BoneChains)
+		const TArray<FBoneChain>& Chains = SourceIKRig->GetRetargetChains();
+		for (const FBoneChain& BoneChain : Chains)
 		{
 			SourceChainOptions.Add(MakeShareable(new FString(BoneChain.ChainName.ToString())));
 		}
@@ -98,29 +99,28 @@ TSharedRef< SWidget > SIKRetargetChainMapRow::GenerateWidgetForColumn( const FNa
 
 void SIKRetargetChainMapRow::OnSourceChainComboSelectionChanged(TSharedPtr<FString> InName, ESelectInfo::Type SelectInfo)
 {
-	const TSharedPtr<FIKRetargetEditorController> Controller = ChainMapList.Pin()->EditorController.Pin();
-	if (!Controller.IsValid())
+	UIKRetargeterController* RetargeterController = ChainMapList.Pin()->GetRetargetController();
+	if (!RetargeterController)
 	{
 		return; 
 	}
 	
 	const FName TargetChainName = ChainMapElement.Pin()->TargetChainName;
-	FRetargetChainMap& ChainMap = *Controller->Asset->GetChainMap(TargetChainName);
-	ChainMap.SourceChain = FName(*InName.Get());
-	Controller->Asset->Modify();
+	const FName SourceChainName = FName(*InName.Get());
+	RetargeterController->SetSourceChainForTargetChain(TargetChainName, SourceChainName);
 }
 
 FText SIKRetargetChainMapRow::GetSourceChainName() const
 {
-	const TSharedPtr<FIKRetargetEditorController> Controller = ChainMapList.Pin()->EditorController.Pin();
-	if (!Controller.IsValid())
+	UIKRetargeterController* RetargeterController = ChainMapList.Pin()->GetRetargetController();
+	if (!RetargeterController)
 	{
 		return FText::FromName(NAME_None); 
 	}
 	
 	const FName TargetChainName = ChainMapElement.Pin()->TargetChainName;
-	const FRetargetChainMap& Mapping = *Controller->Asset->GetChainMap(TargetChainName);
-	return FText::FromName(Mapping.SourceChain);
+	const FName SourceChainName = RetargeterController->GetSourceChainForTargetChain(TargetChainName);
+	return FText::FromName(SourceChainName);
 }
 
 void SIKRetargetChainMapList::Construct(
@@ -220,39 +220,51 @@ void SIKRetargetChainMapList::Construct(
 	RefreshView();
 }
 
-FText SIKRetargetChainMapList::GetSourceRootBone() const
+UIKRetargeterController* SIKRetargetChainMapList::GetRetargetController() const
 {
 	const TSharedPtr<FIKRetargetEditorController> Controller = EditorController.Pin();
 	if (!Controller.IsValid())
+	{
+		return nullptr;
+	}
+
+	return Controller->AssetController;
+}
+
+FText SIKRetargetChainMapList::GetSourceRootBone() const
+{
+	UIKRetargeterController* RetargeterController = GetRetargetController();
+	if (!RetargeterController)
 	{
 		return FText::FromName(NAME_None); 
 	}
 	
-	return FText::FromName(Controller->Asset->GetSourceRootBone());
+	return FText::FromName(RetargeterController->GetSourceRootBone());
 }
 
 FText SIKRetargetChainMapList::GetTargetRootBone() const
 {
-	const TSharedPtr<FIKRetargetEditorController> Controller = EditorController.Pin();
-	if (!Controller.IsValid())
+	UIKRetargeterController* RetargeterController = GetRetargetController();
+	if (!RetargeterController)
 	{
 		return FText::FromName(NAME_None); 
 	}
 	
-	return FText::FromName(Controller->Asset->GetTargetRootBone());
+	return FText::FromName(RetargeterController->GetTargetRootBone());
 }
 
 bool SIKRetargetChainMapList::IsChainMapEnabled() const
 {
-	const TSharedPtr<FIKRetargetEditorController> Controller = EditorController.Pin();
-	if (!Controller.IsValid())
+	UIKRetargeterController* RetargeterController = GetRetargetController();
+	if (!RetargeterController)
 	{
 		return false; 
 	}
 	
-	if (Controller->Asset->TargetIKRigAsset)
+	if (RetargeterController->GetAsset()->TargetIKRigAsset)
 	{
-		return Controller->Asset->TargetIKRigAsset->RetargetDefinition.BoneChains.Num() > 0;
+		const TArray<FBoneChain>& Chains = RetargeterController->GetAsset()->TargetIKRigAsset->GetRetargetChains();
+		return Chains.Num() > 0;
 	}
 
 	return false;
@@ -260,18 +272,15 @@ bool SIKRetargetChainMapList::IsChainMapEnabled() const
 
 void SIKRetargetChainMapList::RefreshView()
 {
-	const TSharedPtr<FIKRetargetEditorController> Controller = EditorController.Pin();
-	if (!Controller.IsValid())
+	UIKRetargeterController* RetargeterController = GetRetargetController();
+	if (!RetargeterController)
 	{
 		return; 
 	}
 	
-	// get latest chain mapping
-	Controller->Asset->CleanChainMapping();
-	
 	// refresh list of chains
 	ListViewItems.Reset();
-	const TArray<FRetargetChainMap>& ChainMappings = Controller->Asset->ChainMapping;
+	const TArray<FRetargetChainMap>& ChainMappings = RetargeterController->GetChainMappings();
 	for (const FRetargetChainMap& ChainMap : ChainMappings)
 	{
 		TSharedPtr<FRetargetChainMapElement> ChainItem = FRetargetChainMapElement::Make(ChainMap.TargetChain);
@@ -303,14 +312,14 @@ EVisibility SIKRetargetChainMapList::IsAutoMapButtonVisible() const
 
 FReply SIKRetargetChainMapList::OnAutoMapButtonClicked() const
 {
-	const TSharedPtr<FIKRetargetEditorController> Controller = EditorController.Pin();
-	if (!Controller.IsValid())
+	UIKRetargeterController* RetargeterController = GetRetargetController();
+	if (!RetargeterController)
 	{
 		return FReply::Unhandled();
 	}
 	
-	Controller->Asset->CleanChainMapping();
-	Controller->Asset->AutoMapChains();
+	RetargeterController->CleanChainMapping();
+	RetargeterController->AutoMapChains();
 	return FReply::Handled();
 }
 
diff --git a/Engine/Plugins/Animation/IKRig/Source/IKRigEditor/Private/RetargetEditor/SRetargetAnimAssetsWindow.cpp b/Engine/Plugins/Animation/IKRig/Source/IKRigEditor/Private/RetargetEditor/SRetargetAnimAssetsWindow.cpp
--- a/Engine/Plugins/Animation/IKRig/Source/IKRigEditor/Private/RetargetEditor/SRetargetAnimAssetsWindow.cpp
+++ b/Engine/Plugins/Animation/IKRig/Source/IKRigEditor/Private/RetargetEditor/SRetargetAnimAssetsWindow.cpp
@@ -4,6 +4,7 @@
 
 #include "AnimPose.h"
 #include "AnimPreviewInstance.h"
+#include "ContentBrowserModule.h"
 #include "IContentBrowserSingleton.h"
 #include "ObjectEditorUtils.h"
 #include "SSkeletonWidget.h"
@@ -18,6 +19,111 @@
 
 #define LOCTEXT_NAMESPACE "RetargetAnimAssetWindow"
 
+void SSelectExportPathDialog::Construct(const FArguments& InArgs)
+{
+	ParentWindowPtr = InArgs._ParentWindow;
+	AssetPath = FText::FromString(FPackageName::GetLongPackagePath(InArgs._DefaultAssetPath.ToString()));
+
+	if(AssetPath.IsEmpty())
+	{
+		AssetPath = FText::FromString(TEXT("/Game"));
+	}
+
+	FPathPickerConfig PathPickerConfig;
+	PathPickerConfig.DefaultPath = AssetPath.ToString();
+	PathPickerConfig.OnPathSelected = FOnPathSelected::CreateSP(this, &SSelectExportPathDialog::OnPathChange);
+	PathPickerConfig.bAddDefaultPath = true;
+
+	FContentBrowserModule& ContentBrowserModule = FModuleManager::LoadModuleChecked<FContentBrowserModule>("ContentBrowser");
+
+	SWindow::Construct(SWindow::FArguments()
+		.Title(LOCTEXT("SSelectExportPathDialog_Title", "Select Export Path"))
+		.SupportsMinimize(false)
+		.SupportsMaximize(false)
+		//.SizingRule( ESizingRule::Autosized )
+		.ClientSize(FVector2D(450, 450))
+		[
+			SNew(SVerticalBox)
+
+			+ SVerticalBox::Slot() // Add user input block
+			.Padding(2)
+			[
+				SNew(SBorder)
+				.BorderImage(FEditorStyle::GetBrush("ToolPanel.GroupBorder"))
+				[
+					SNew(SVerticalBox)
+
+					+SVerticalBox::Slot()
+					.AutoHeight()
+					[
+						SNew(STextBlock)
+						.Text(LOCTEXT("SelectPath", "Select Path"))
+						.Font(FCoreStyle::GetDefaultFontStyle("Regular", 14))
+					]
+
+					+SVerticalBox::Slot()
+					.FillHeight(1)
+					.Padding(3)
+					[
+						ContentBrowserModule.Get().CreatePathPicker(PathPickerConfig)
+					]
+				]
+			]
+
+			+SVerticalBox::Slot()
+			.AutoHeight()
+			.HAlign(HAlign_Right)
+			.Padding(5)
+			[
+				SNew(SUniformGridPanel)
+				.SlotPadding(FEditorStyle::GetMargin("StandardDialog.SlotPadding"))
+				.MinDesiredSlotWidth(FEditorStyle::GetFloat("StandardDialog.MinDesiredSlotWidth"))
+				.MinDesiredSlotHeight(FEditorStyle::GetFloat("StandardDialog.MinDesiredSlotHeight"))
+				+SUniformGridPanel::Slot(0, 0)
+				[
+					SNew(SButton)
+					.HAlign(HAlign_Center)
+					.ContentPadding(FEditorStyle::GetMargin("StandardDialog.ContentPadding"))
+					.Text(LOCTEXT("OK", "OK"))
+					.OnClicked(this, &SSelectExportPathDialog::OnButtonClick, EAppReturnType::Ok)
+				]
+				+SUniformGridPanel::Slot(1, 0)
+				[
+					SNew(SButton)
+					.HAlign(HAlign_Center)
+					.ContentPadding(FEditorStyle::GetMargin("StandardDialog.ContentPadding"))
+					.Text(LOCTEXT("Cancel", "Cancel"))
+					.OnClicked(this, &SSelectExportPathDialog::OnButtonClick, EAppReturnType::Cancel)
+				]
+			]
+		]);
+}
+
+void SSelectExportPathDialog::OnPathChange(const FString& NewPath)
+{
+	AssetPath = FText::FromString(NewPath);
+}
+
+FReply SSelectExportPathDialog::OnButtonClick(EAppReturnType::Type ButtonID)
+{
+	UserResponse = ButtonID;
+
+	RequestDestroyWindow();
+
+	return FReply::Handled();
+}
+
+EAppReturnType::Type SSelectExportPathDialog::ShowModal()
+{
+	GEditor->EditorAddModalWindow(SharedThis(this));
+	return UserResponse;
+}
+
+FString SSelectExportPathDialog::GetAssetPath()
+{
+	return AssetPath.ToString();
+}
+
 void SRetargetPoseViewport::Construct(const FArguments& InArgs)
 {
 	SEditorViewport::Construct(SEditorViewport::FArguments());
@@ -447,7 +553,7 @@ void SRetargetAnimAssetsWindow::Construct(const FArguments& InArgs)
 						SNew(SButton)
 						.HAlign(HAlign_Center)
 						.Text(LOCTEXT("DuplicateAndRetarget_ChangeFolder", "Change..."))
-						//.OnClicked(this, &SAnimationRemapSkeleton::ShowFolderOption)
+						.OnClicked(this, &SRetargetAnimAssetsWindow::GetExportFolder)
 					]
 				]
 
@@ -683,4 +789,18 @@ FText SRetargetAnimAssetsWindow::GetFolderPath() const
 	return FText::FromString(BatchContext.NameRule.FolderPath);
 }
 
+FReply SRetargetAnimAssetsWindow::GetExportFolder()
+{
+	TSharedRef<SSelectExportPathDialog> Dialog = SNew(SSelectExportPathDialog)
+	.DefaultAssetPath(FText::FromString(BatchContext.FolderPath))
+	.ParentWindow(DialogWindow);
+	
+	if(Dialog->ShowModal() != EAppReturnType::Cancel)
+	{
+		BatchContext.NameRule.FolderPath = Dialog->GetAssetPath();
+	}
+
+	return FReply::Handled();
+}
+
 #undef LOCTEXT_NAMESPACE
diff --git a/Engine/Plugins/Animation/IKRig/Source/IKRigEditor/Private/RigEditor/IKRigAnimInstance.cpp b/Engine/Plugins/Animation/IKRig/Source/IKRigEditor/Private/RigEditor/IKRigAnimInstance.cpp
--- a/Engine/Plugins/Animation/IKRig/Source/IKRigEditor/Private/RigEditor/IKRigAnimInstance.cpp
+++ b/Engine/Plugins/Animation/IKRig/Source/IKRigEditor/Private/RigEditor/IKRigAnimInstance.cpp
@@ -15,6 +15,11 @@ void UIKRigAnimInstance::SetIKRigAsset(UIKRigDefinition* InIKRigAsset)
 	Proxy.SetIKRigAsset(InIKRigAsset);
 }
 
+void UIKRigAnimInstance::SetProcessorNeedsInitialized()
+{
+	IKRigNode.SetProcessorNeedsInitialized();
+}
+
 FAnimInstanceProxy* UIKRigAnimInstance::CreateAnimInstanceProxy()
 {
 	return new FIKRigAnimInstanceProxy(this, &IKRigNode);
diff --git a/Engine/Plugins/Animation/IKRig/Source/IKRigEditor/Private/RigEditor/IKRigAnimInstanceProxy.cpp b/Engine/Plugins/Animation/IKRig/Source/IKRigEditor/Private/RigEditor/IKRigAnimInstanceProxy.cpp
--- a/Engine/Plugins/Animation/IKRig/Source/IKRigEditor/Private/RigEditor/IKRigAnimInstanceProxy.cpp
+++ b/Engine/Plugins/Animation/IKRig/Source/IKRigEditor/Private/RigEditor/IKRigAnimInstanceProxy.cpp
@@ -3,9 +3,6 @@
 #include "RigEditor/IKRigAnimInstanceProxy.h"
 #include "RigEditor/IKRigAnimInstance.h"
 
-FIKRigAnimInstanceProxy::FIKRigAnimInstanceProxy()
-{
-}
 
 FIKRigAnimInstanceProxy::FIKRigAnimInstanceProxy(UAnimInstance* InAnimInstance, FAnimNode_IKRig* InIKRigNode)
 	: FAnimPreviewInstanceProxy(InAnimInstance),
@@ -13,10 +10,6 @@ FIKRigAnimInstanceProxy::FIKRigAnimInstanceProxy(UAnimInstance* InAnimInstance,
 {
 }
 
-FIKRigAnimInstanceProxy::~FIKRigAnimInstanceProxy()
-{
-}
-
 void FIKRigAnimInstanceProxy::Initialize(UAnimInstance* InAnimInstance)
 {
 	FAnimInstanceProxy::Initialize(InAnimInstance);
