IK Rig / Retargeter refactor

diff --git a/Engine/Plugins/Animation/IKRig/Source/IKRigEditor/Private/RigEditor/IKRigController.cpp b/Engine/Plugins/Animation/IKRig/Source/IKRigEditor/Private/RigEditor/IKRigController.cpp
--- a/Engine/Plugins/Animation/IKRig/Source/IKRigEditor/Private/RigEditor/IKRigController.cpp
+++ b/Engine/Plugins/Animation/IKRig/Source/IKRigEditor/Private/RigEditor/IKRigController.cpp
@@ -8,7 +8,7 @@
 #include "Engine/SkeletalMesh.h"
 #include "ScopedTransaction.h"
 
-#define LOCTEXT_NAMESPACE	"IKRigController"
+#define LOCTEXT_NAMESPACE "IKRigController"
 
 
 TMap<UIKRigDefinition*, UIKRigController*> UIKRigController::AssetToControllerMap;
@@ -39,27 +39,23 @@ UIKRigDefinition* UIKRigController::GetAsset() const
 
 void UIKRigController::AddBoneSetting(const FName& BoneName, int32 SolverIndex) const
 {
-	UIKRigSolver* Solver = GetSolver(SolverIndex);
-	if (!Solver)
-	{
-		return; // solver doesn't exist
-	}
+	check(IKRigAsset->Solvers.IsValidIndex(SolverIndex))
 
-	if (IKRigAsset->Skeleton.GetBoneIndexFromName(BoneName) == INDEX_NONE)
+	if (!CanAddBoneSetting(BoneName, SolverIndex))
 	{
-		return; // bone doesn't exist
+		return; // prerequisites not met
 	}
 
 	FScopedTransaction Transaction(LOCTEXT("AddBoneSetting_Label", "Add Bone Setting"));
+	UIKRigSolver* Solver = IKRigAsset->Solvers[SolverIndex];
 	Solver->Modify();
-	IKRigAsset->Modify();
-
 	Solver->AddBoneSetting(BoneName);
+	BroadcastNeedsReinitialized();
 }
 
 bool UIKRigController::CanAddBoneSetting(const FName& BoneName, int32 SolverIndex) const
 {
-	UIKRigSolver* Solver = GetSolver(SolverIndex);
+	const UIKRigSolver* Solver = GetSolver(SolverIndex);
 	if (!Solver)
 	{
 		return false; // solver doesn't exist
@@ -81,11 +77,7 @@ bool UIKRigController::CanAddBoneSetting(const FName& BoneName, int32 SolverInde
 
 void UIKRigController::RemoveBoneSetting(const FName& BoneName, int32 SolverIndex) const
 {
-	UIKRigSolver* Solver = GetSolver(SolverIndex);
-	if (!Solver)
-	{
-		return; // solver doesn't exist
-	}
+	check(IKRigAsset->Solvers.IsValidIndex(SolverIndex))
 
 	if (IKRigAsset->Skeleton.GetBoneIndexFromName(BoneName) == INDEX_NONE)
 	{
@@ -93,15 +85,15 @@ void UIKRigController::RemoveBoneSetting(const FName& BoneName, int32 SolverInde
 	}
 
 	FScopedTransaction Transaction(LOCTEXT("RemoveBoneSetting_Label", "Remove Bone Setting"));
+	UIKRigSolver* Solver = IKRigAsset->Solvers[SolverIndex];
 	Solver->Modify();
-	IKRigAsset->Modify();
-
 	Solver->RemoveBoneSetting(BoneName);
+	BroadcastNeedsReinitialized();
 }
 
 bool UIKRigController::CanRemoveBoneSetting(const FName& BoneName, int32 SolverIndex) const
 {
-	UIKRigSolver* Solver = GetSolver(SolverIndex);
+	const UIKRigSolver* Solver = GetSolver(SolverIndex);
 	if (!Solver)
 	{
 		return false; // solver doesn't exist
@@ -127,7 +119,7 @@ bool UIKRigController::CanRemoveBoneSetting(const FName& BoneName, int32 SolverI
 
 UObject* UIKRigController::GetSettingsForBone(const FName& BoneName, int32 SolverIndex) const
 {
-	UIKRigSolver* Solver = GetSolver(SolverIndex);
+	const UIKRigSolver* Solver = GetSolver(SolverIndex);
 	if (!Solver)
 	{
 		return nullptr; // solver doesn't exist
@@ -156,7 +148,7 @@ bool UIKRigController::DoesBoneHaveSettings(const FName& BoneName) const
 
 bool UIKRigController::AddRetargetChain(const FName& ChainName, const FName& StartBone, const FName& EndBone) const
 {
-	if (IKRigAsset->RetargetDefinition.GetBoneChainByName(ChainName))
+	if (IKRigAsset->GetRetargetChainByName(ChainName))
 	{
 		return false; // bone chain already exists with that name
 	}
@@ -165,7 +157,8 @@ bool UIKRigController::AddRetargetChain(const FName& ChainName, const FName& Sta
 	IKRigAsset->Modify();
 	
 	IKRigAsset->RetargetDefinition.BoneChains.Emplace(ChainName, StartBone, EndBone);
-	IKRigAsset->SortRetargetChains();
+	SortRetargetChains();
+	BroadcastNeedsReinitialized();
 	return true;
 }
 
@@ -186,22 +179,44 @@ bool UIKRigController::RemoveRetargetChain(const FName& ChainName) const
 	
 	if (IKRigAsset->RetargetDefinition.BoneChains.RemoveAll(Pred) > 0)
 	{
-		IKRigAsset->SortRetargetChains();
+		SortRetargetChains();
+		BroadcastNeedsReinitialized();
 		return true;
 	}
-	
+
 	return false;
 }
 
+FName UIKRigController::RenameRetargetChain(const FName& ChainName, const FName& NewChainName) const
+{
+	FBoneChain* Chain = IKRigAsset->RetargetDefinition.GetEditableBoneChainByName(ChainName);
+	if (!Chain)
+	{
+		return ChainName; // chain doesn't exist to rename
+	}
+
+	if (IKRigAsset->GetRetargetChainByName(NewChainName))
+	{
+		return ChainName; // bone chain already exists with the new name
+	}
+	
+	FScopedTransaction Transaction(LOCTEXT("RenameRetargetChain_Label", "Rename Retarget Chain"));
+	IKRigAsset->Modify();
+	Chain->ChainName = NewChainName;
+	RetargetChainRenamed.Broadcast(GetAsset(), ChainName, NewChainName);
+	BroadcastNeedsReinitialized();
+	return NewChainName;
+}
+
 bool UIKRigController::SetRetargetChainStartBone(const FName& ChainName, const FName& StartBoneName) const
 {
-	if (FBoneChain* BoneChain = IKRigAsset->RetargetDefinition.GetBoneChainByName(ChainName))
+	if (FBoneChain* BoneChain = IKRigAsset->RetargetDefinition.GetEditableBoneChainByName(ChainName))
 	{
 		FScopedTransaction Transaction(LOCTEXT("SetRetargetChainStartBone_Label", "Set Retarget Chain Start Bone"));
 		IKRigAsset->Modify();
-	
 		BoneChain->StartBone = StartBoneName;
-		IKRigAsset->SortRetargetChains();
+		SortRetargetChains();
+		BroadcastNeedsReinitialized();
 		return true;
 	}
 
@@ -210,35 +225,81 @@ bool UIKRigController::SetRetargetChainStartBone(const FName& ChainName, const F
 
 bool UIKRigController::SetRetargetChainEndBone(const FName& ChainName, const FName& EndBoneName) const
 {
-	if (FBoneChain* BoneChain = IKRigAsset->RetargetDefinition.GetBoneChainByName(ChainName))
+	if (FBoneChain* BoneChain = IKRigAsset->RetargetDefinition.GetEditableBoneChainByName(ChainName))
 	{
 		FScopedTransaction Transaction(LOCTEXT("SetRetargetChainEndBone_Label", "Set Retarget Chain End Bone"));
 		IKRigAsset->Modify();
-		
 		BoneChain->EndBone = EndBoneName;
+		BroadcastNeedsReinitialized();
 		return true;
 	}
 
 	return false; // no bone chain with that name
 }
 
-FBoneChain* UIKRigController::GetRetargetChain(int32 Index) const
+bool UIKRigController::SetRetargetChainGoal(const FName& ChainName, const FName& GoalName) const
 {
-	if (ensure(IKRigAsset) && IKRigAsset->RetargetDefinition.BoneChains.IsValidIndex(Index))
+	FName GoalNameToUse = GoalName;
+	if (!GetGoal(GoalName))
 	{
-		return &IKRigAsset->RetargetDefinition.BoneChains[Index];
+		GoalNameToUse = NAME_None; // no goal with that name	
 	}
 	
-	return nullptr;
+	if (FBoneChain* BoneChain = IKRigAsset->RetargetDefinition.GetEditableBoneChainByName(ChainName))
+	{
+		FScopedTransaction Transaction(LOCTEXT("SetRetargetChainGoal_Label", "Set Retarget Chain Goal"));
+		IKRigAsset->Modify();
+		BoneChain->IKGoalName = GoalNameToUse;
+		BroadcastNeedsReinitialized();
+		return true;
+	}
+
+	return false; // no bone chain with that name
 }
 
-int32 UIKRigController::GetNumRetargetChains() const
+FName UIKRigController::GetRetargetChainGoal(const FName& ChainName) const
 {
 	check(IKRigAsset)
-	return IKRigAsset->RetargetDefinition.BoneChains.Num();
+	const FBoneChain* Chain = IKRigAsset->GetRetargetChainByName(ChainName);
+	if (!Chain)
+	{
+		return NAME_None;
+	}
+	
+	return Chain->IKGoalName;
 }
 
-void UIKRigController::SetRetargetRoot(const FName& RootBoneName)
+FName UIKRigController::GetRetargetChainStartBone(const FName& ChainName) const
+{
+	check(IKRigAsset)
+	const FBoneChain* Chain = IKRigAsset->GetRetargetChainByName(ChainName);
+	if (!Chain)
+	{
+		return NAME_None;
+	}
+	
+	return Chain->StartBone;
+}
+
+FName UIKRigController::GetRetargetChainEndBone(const FName& ChainName) const
+{
+	check(IKRigAsset)
+	const FBoneChain* Chain = IKRigAsset->GetRetargetChainByName(ChainName);
+	if (!Chain)
+	{
+		return NAME_None;
+	}
+	
+	return Chain->EndBone;
+}
+
+const TArray<FBoneChain>& UIKRigController::GetRetargetChains() const
+{
+	check(IKRigAsset)
+	return IKRigAsset->GetRetargetChains();
+}
+
+void UIKRigController::SetRetargetRoot(const FName& RootBoneName) const
 {
 	check(IKRigAsset)
 
@@ -246,15 +307,27 @@ void UIKRigController::SetRetargetRoot(const FName& RootBoneName)
 	IKRigAsset->Modify();
 	
 	IKRigAsset->RetargetDefinition.RootBone = RootBoneName;
+
+	BroadcastNeedsReinitialized();
 }
 
-FName UIKRigController::GetRetargetRoot()
+FName UIKRigController::GetRetargetRoot() const
 {
 	check(IKRigAsset)
 
 	return IKRigAsset->RetargetDefinition.RootBone;
 }
 
+void UIKRigController::SortRetargetChains() const
+{
+	IKRigAsset->RetargetDefinition.BoneChains.Sort([this](const FBoneChain& A, const FBoneChain& B)
+	{
+		const int32 IndexA = IKRigAsset->Skeleton.GetBoneIndexFromName(A.StartBone);
+		const int32 IndexB = IKRigAsset->Skeleton.GetBoneIndexFromName(B.StartBone);
+		return IndexA < IndexB;
+	});
+}
+
 void UIKRigController::AddReferencedObjects(UObject* InThis, FReferenceCollector& Collector)
 {
 	// since static member, we just add only for default object
@@ -282,35 +355,45 @@ void UIKRigController::SetSourceSkeletalMesh(USkeletalMesh* SkeletalMesh, bool b
 	// update stored skeletal mesh used for previewing results
 	IKRigAsset->PreviewSkeletalMesh = SkeletalMesh;
 
-	if (!bReImportBones)
+	if (bReImportBones)
 	{
-		return;
+		// reimport the skeleton data
+		if (SkeletalMesh)
+		{
+			SetSkeleton(SkeletalMesh->GetRefSkeleton());		
+		}
+		else
+		{
+			IKRigAsset->Skeleton.Reset();
+		}
 	}
 
-	// reimport the skeleton data
-	if (SkeletalMesh)
-	{
-		SetSkeleton(SkeletalMesh->GetRefSkeleton());		
-	}
-	else
-	{
-		IKRigAsset->Skeleton.Reset();
-	}
+	BroadcastNeedsReinitialized();
 }
 
 void UIKRigController::SetSkeleton(const FReferenceSkeleton& InSkeleton) const
 {
 	FScopedTransaction Transaction(LOCTEXT("SetSkeleton_Label", "Set Skeleton"));
 	IKRigAsset->Modify();
-
 	IKRigAsset->Skeleton.Initialize(InSkeleton, IKRigAsset->Skeleton.ExcludedBones);
+	BroadcastNeedsReinitialized();
 }
 
-FIKRigSkeleton& UIKRigController::GetSkeleton() const
+const FIKRigSkeleton& UIKRigController::GetIKRigSkeleton() const
 {
 	return IKRigAsset->Skeleton;
 }
 
+USkeleton* UIKRigController::GetSkeleton() const
+{
+	if (IKRigAsset->PreviewSkeletalMesh.IsNull())
+    {
+        return nullptr;
+    }
+    
+    return IKRigAsset->PreviewSkeletalMesh->GetSkeleton();
+}
+
 void UIKRigController::SetBoneExcluded(const FName& BoneName, const bool bExclude) const
 {
 	const bool bIsExcluded = IKRigAsset->Skeleton.ExcludedBones.Contains(BoneName);
@@ -330,6 +413,8 @@ void UIKRigController::SetBoneExcluded(const FName& BoneName, const bool bExclud
 	{
 		IKRigAsset->Skeleton.ExcludedBones.Remove(BoneName);
 	}
+
+	BroadcastNeedsReinitialized();
 }
 
 bool UIKRigController::GetBoneExcluded(const FName& BoneName) const
@@ -337,12 +422,18 @@ bool UIKRigController::GetBoneExcluded(const FName& BoneName) const
 	return IKRigAsset->Skeleton.ExcludedBones.Contains(BoneName);
 }
 
+FTransform UIKRigController::GetBoneRetargetPose(const FName& BoneName) const
+{	
+	const int32 BoneIndex = IKRigAsset->Skeleton.GetBoneIndexFromName(BoneName);
+	check(BoneIndex != INDEX_NONE) // must initialize IK Rig before getting here
+	return IKRigAsset->Skeleton.RefPoseGlobal[BoneIndex];
+}
+
 USkeletalMesh* UIKRigController::GetSourceSkeletalMesh() const
 {
 	return IKRigAsset->PreviewSkeletalMesh.Get();
 }
 
-
 // -------------------------------------------------------
 // SOLVERS
 //
@@ -357,30 +448,18 @@ int32 UIKRigController::AddSolver(TSubclassOf<UIKRigSolver> InIKRigSolverClass)
 	UIKRigSolver* NewSolver = NewObject<UIKRigSolver>(IKRigAsset, InIKRigSolverClass, NAME_None, RF_Transactional);
 	check(NewSolver);
 
-	return IKRigAsset->Solvers.Add(NewSolver);
-}
+	const int32 SolverIndex = IKRigAsset->Solvers.Add(NewSolver);
 
-int32 UIKRigController::GetNumSolvers() const
-{
-	check(IKRigAsset)
-	return IKRigAsset->Solvers.Num();
+	BroadcastNeedsReinitialized();
+
+	return SolverIndex;
 }
 
-UIKRigSolver* UIKRigController::GetSolver(int32 Index) const
+void UIKRigController::RemoveSolver(const int32 SolverIndex) const
 {
 	check(IKRigAsset)
-
-	if (IKRigAsset->Solvers.IsValidIndex(Index))
-	{
-		return IKRigAsset->Solvers[Index];
-	}
 	
-	return nullptr;
-}
-
-void UIKRigController::RemoveSolver(UIKRigSolver* SolverToDelete) const
-{
-	if (!(IKRigAsset && SolverToDelete))
+	if (!IKRigAsset->Solvers.IsValidIndex(SolverIndex))
 	{
 		return;
 	}
@@ -388,7 +467,9 @@ void UIKRigController::RemoveSolver(UIKRigSolver* SolverToDelete) const
 	FScopedTransaction Transaction(LOCTEXT("RemoveSolver_Label", "Remove Solver"));
 	IKRigAsset->Modify();
 
-	IKRigAsset->Solvers.Remove(SolverToDelete);
+	IKRigAsset->Solvers.RemoveAt(SolverIndex);
+
+	BroadcastNeedsReinitialized();
 }
 
 bool UIKRigController::MoveSolverInStack(int32 SolverToMoveIndex, int32 TargetSolverIndex) const
@@ -410,6 +491,28 @@ bool UIKRigController::MoveSolverInStack(int32 SolverToMoveIndex, int32 TargetSo
 	IKRigAsset->Solvers.Insert(SolverToMove, TargetSolverIndex + 1);
 	const int32 SolverToRemove = TargetSolverIndex > SolverToMoveIndex ? SolverToMoveIndex : SolverToMoveIndex + 1;
 	IKRigAsset->Solvers.RemoveAt(SolverToRemove);
+
+	BroadcastNeedsReinitialized();
+	
+	return true;
+}
+
+bool UIKRigController::SetSolverEnabled(int32 SolverIndex, bool bIsEnabled) const
+{
+	if (!IKRigAsset->Solvers.IsValidIndex(SolverIndex))
+	{
+		return false;
+	}
+	
+	FScopedTransaction Transaction(LOCTEXT("SetSolverEndabled_Label", "Enable/Disable Solver"));
+	UIKRigSolver* Solver = IKRigAsset->Solvers[SolverIndex];
+	Solver->Modify();
+	IKRigAsset->Modify();
+
+	Solver->SetEnabled(bIsEnabled);
+
+	BroadcastNeedsReinitialized();
+	
 	return true;
 }
 
@@ -428,9 +531,10 @@ void UIKRigController::SetRootBone(const FName& RootBoneName, int32 SolverIndex)
 	FScopedTransaction Transaction(LOCTEXT("SetRootBone_Label", "Set Root Bone"));
 	UIKRigSolver* Solver = IKRigAsset->Solvers[SolverIndex];
 	Solver->Modify();
-	IKRigAsset->Modify();
 
 	Solver->SetRootBone(RootBoneName);
+
+	BroadcastNeedsReinitialized();
 }
 
 const TArray<UIKRigSolver*>& UIKRigController::GetSolverArray() const
@@ -438,6 +542,24 @@ const TArray<UIKRigSolver*>& UIKRigController::GetSolverArray() const
 	return IKRigAsset->Solvers;
 }
 
+int32 UIKRigController::GetNumSolvers() const
+{
+	check(IKRigAsset)
+	return IKRigAsset->Solvers.Num();
+}
+
+UIKRigSolver* UIKRigController::GetSolver(int32 Index) const
+{
+	check(IKRigAsset)
+
+	if (IKRigAsset->Solvers.IsValidIndex(Index))
+	{
+		return IKRigAsset->Solvers[Index];
+	}
+	
+	return nullptr;
+}
+
 // -------------------------------------------------------
 // GOALS
 //
@@ -452,14 +574,16 @@ UIKRigEffectorGoal* UIKRigController::AddNewGoal(const FName& GoalName, const FN
 	FScopedTransaction Transaction(LOCTEXT("AddSetting_Label", "Add Bone Setting"));
 	IKRigAsset->Modify();
 
-	UIKRigEffectorGoal* NewGoal = NewObject<UIKRigEffectorGoal>(GetAsset(), UIKRigEffectorGoal::StaticClass(), NAME_None, RF_Transactional);
+	UIKRigEffectorGoal* NewGoal = NewObject<UIKRigEffectorGoal>(IKRigAsset, UIKRigEffectorGoal::StaticClass(), NAME_None, RF_Transactional);
 	NewGoal->BoneName = BoneName;
 	NewGoal->GoalName = GoalName;
 	IKRigAsset->Goals.Add(NewGoal);
 
 	// set initial transform
-	NewGoal->InitialTransform = IKRigAsset->GetGoalInitialTransform(NewGoal);
+	NewGoal->InitialTransform = GetBoneRetargetPose(NewGoal->BoneName);
 	NewGoal->CurrentTransform = NewGoal->InitialTransform;
+
+	BroadcastNeedsReinitialized();
 	
 	return NewGoal;
 }
@@ -473,7 +597,6 @@ bool UIKRigController::RemoveGoal(const FName& GoalName) const
 	}
 
 	FScopedTransaction Transaction(LOCTEXT("RemoveGoal_Label", "Remove Goal"));
-	IKRigAsset->Modify();
 
 	// remove from all the solvers
 	const FName& GoalToRemove = IKRigAsset->Goals[GoalIndex]->GoalName;
@@ -484,7 +607,19 @@ bool UIKRigController::RemoveGoal(const FName& GoalName) const
 	}
 
 	// remove from core system
+	IKRigAsset->Modify();
 	IKRigAsset->Goals.RemoveAt(GoalIndex);
+
+	// clean any retarget chains that might reference the missing goal
+	for (FBoneChain& BoneChain : IKRigAsset->RetargetDefinition.BoneChains)
+	{
+		if (BoneChain.IKGoalName == GoalName)
+		{
+			BoneChain.IKGoalName = NAME_None;
+		}
+	}
+
+	BroadcastNeedsReinitialized();
 	
 	return true;
 }
@@ -515,6 +650,15 @@ FName UIKRigController::RenameGoal(const FName& OldName, const FName& PotentialN
 	IKRigAsset->Goals[GoalIndex]->Modify();
 	IKRigAsset->Goals[GoalIndex]->GoalName = NewName;
 
+	// update any retarget chains that might reference the goal
+	for (FBoneChain& BoneChain : IKRigAsset->RetargetDefinition.BoneChains)
+	{
+		if (BoneChain.IKGoalName == OldName)
+		{
+			BoneChain.IKGoalName = NewName;
+		}
+	}
+
 	// rename in solvers
 	for (UIKRigSolver* Solver : IKRigAsset->Solvers)
 	{
@@ -522,6 +666,8 @@ FName UIKRigController::RenameGoal(const FName& OldName, const FName& PotentialN
 		Solver->RenameGoal(OldName, NewName);
 	}
 
+	BroadcastNeedsReinitialized();
+
 	return NewName;
 }
 
@@ -533,7 +679,7 @@ bool UIKRigController::SetGoalBone(const FName& GoalName, const FName& NewBoneNa
 		return false; // goal doesn't exist in the rig
 	}
 
-	const int32 BoneIndex = GetSkeleton().GetBoneIndexFromName(NewBoneName);
+	const int32 BoneIndex = GetIKRigSkeleton().GetBoneIndexFromName(NewBoneName);
 	if (BoneIndex == INDEX_NONE)
 	{
 		return false; // bone does not exist in the skeleton
@@ -558,7 +704,9 @@ bool UIKRigController::SetGoalBone(const FName& GoalName, const FName& NewBoneNa
 	}
 
 	// update initial transforms
-	IKRigAsset->ResetGoalTransforms();
+	ResetGoalTransforms();
+
+	BroadcastNeedsReinitialized();
 	
 	return true;
 }
@@ -586,9 +734,10 @@ bool UIKRigController::ConnectGoalToSolver(const UIKRigEffectorGoal& Goal, int32
 	FScopedTransaction Transaction(LOCTEXT("AddGoalToSolver_Label", "Add Goal"));
 	UIKRigSolver* Solver = IKRigAsset->Solvers[SolverIndex];
 	Solver->Modify();
-	IKRigAsset->Modify();
 	
 	Solver->AddGoal(Goal);
+
+	BroadcastNeedsReinitialized();
 	return true;
 }
 
@@ -602,20 +751,21 @@ bool UIKRigController::DisconnectGoalFromSolver(const FName& GoalToRemove, int32
     FScopedTransaction Transaction(LOCTEXT("RemoveGoalFromSolver_Label", "Remove Goal"));
 	UIKRigSolver* Solver = IKRigAsset->Solvers[SolverIndex];
 	Solver->Modify();
-	IKRigAsset->Modify();
 	
 	Solver->RemoveGoal(GoalToRemove);
+
+	BroadcastNeedsReinitialized();
 	return true;
 }
 
-bool UIKRigController::IsGoalConnectedToSolver(const FName& Goal, int32 SolverIndex) const
+bool UIKRigController::IsGoalConnectedToSolver(const FName& GoalName, int32 SolverIndex) const
 {
 	if (!IKRigAsset->Solvers.IsValidIndex(SolverIndex))
 	{
 		return false;
 	}
 
-	return IKRigAsset->Solvers[SolverIndex]->IsGoalConnected(Goal);
+	return IKRigAsset->Solvers[SolverIndex]->IsGoalConnected(GoalName);
 }
 
 const TArray<UIKRigEffectorGoal*>& UIKRigController::GetAllGoals() const
@@ -623,7 +773,7 @@ const TArray<UIKRigEffectorGoal*>& UIKRigController::GetAllGoals() const
 	return IKRigAsset->Goals;
 }
 
-UIKRigEffectorGoal* UIKRigController::GetGoal(int32 GoalIndex) const
+const UIKRigEffectorGoal* UIKRigController::GetGoal(int32 GoalIndex) const
 {
 	if (!IKRigAsset->Goals.IsValidIndex(GoalIndex))
 	{
@@ -633,7 +783,7 @@ UIKRigEffectorGoal* UIKRigController::GetGoal(int32 GoalIndex) const
 	return IKRigAsset->Goals[GoalIndex];
 }
 
-UIKRigEffectorGoal* UIKRigController::GetGoal(const FName& GoalName) const
+const UIKRigEffectorGoal* UIKRigController::GetGoal(const FName& GoalName) const
 {
 	const int32 GoalIndex = GetGoalIndex(GoalName);
 	if (GoalIndex == INDEX_NONE)
@@ -660,19 +810,9 @@ UObject* UIKRigController::GetEffectorForGoal(const FName& GoalName, int32 Solve
 	return IKRigAsset->Solvers[SolverIndex]->GetEffectorWithGoal(GoalName);
 }
 
-FTransform UIKRigController::GetGoalInitialTransform(const FName& GoalName) const
-{
-	if(UIKRigEffectorGoal* Goal = GetGoal(GoalName))
-	{
-		return Goal->InitialTransform;
-	}
-	
-	return FTransform::Identity; // no goal with that name
-}
-
 FTransform UIKRigController::GetGoalCurrentTransform(const FName& GoalName) const
 {
-	if(UIKRigEffectorGoal* Goal = GetGoal(GoalName))
+	if(const UIKRigEffectorGoal* Goal = GetGoal(GoalName))
 	{
 		return Goal->CurrentTransform;
 	}
@@ -680,20 +820,22 @@ FTransform UIKRigController::GetGoalCurrentTransform(const FName& GoalName) cons
 	return FTransform::Identity; // no goal with that name
 }
 
-void UIKRigController::SetGoalInitialTransform(const FName& GoalName, const FTransform& Transform) const
+void UIKRigController::SetGoalCurrentTransform(const FName& GoalName, const FTransform& Transform) const
 {
-	UIKRigEffectorGoal* Goal = GetGoal(GoalName);
+	UIKRigEffectorGoal* Goal = const_cast<UIKRigEffectorGoal*>(GetGoal(GoalName));
 	check(Goal);
 
-	Goal->InitialTransform = Transform;
+	Goal->CurrentTransform = Transform;
 }
 
-void UIKRigController::SetGoalCurrentTransform(const FName& GoalName, const FTransform& Transform) const
+void UIKRigController::ResetGoalTransforms() const
 {
-	UIKRigEffectorGoal* Goal = GetGoal(GoalName);
-	check(Goal);
-
-	Goal->CurrentTransform = Transform;
+	for (UIKRigEffectorGoal* Goal : IKRigAsset->Goals)
+    {
+    	const FTransform InitialTransform = GetBoneRetargetPose(Goal->BoneName);
+    	Goal->InitialTransform = InitialTransform;
+    	Goal->CurrentTransform = InitialTransform;
+    }	
 }
 
 void UIKRigController::SanitizeGoalName(FString& InOutName)
diff --git a/Engine/Plugins/Animation/IKRig/Source/IKRigEditor/Private/RigEditor/IKRigEditMode.cpp b/Engine/Plugins/Animation/IKRig/Source/IKRigEditor/Private/RigEditor/IKRigEditMode.cpp
--- a/Engine/Plugins/Animation/IKRig/Source/IKRigEditor/Private/RigEditor/IKRigEditMode.cpp
+++ b/Engine/Plugins/Animation/IKRig/Source/IKRigEditor/Private/RigEditor/IKRigEditMode.cpp
@@ -31,7 +31,7 @@ bool FIKRigEditMode::GetCameraTarget(FSphere& OutTarget) const
 		TArray<FVector> GoalPoints;
 		for (const auto& GoalName : OutGoalNames)
 		{
-			FTransform GoalTransform = Controller->AssetController->GetGoalInitialTransform(GoalName);
+			FTransform GoalTransform = Controller->AssetController->GetBoneRetargetPose(Controller->AssetController->GetGoal(GoalName)->BoneName);
 			GoalPoints.Add(GoalTransform.GetLocation());
 		}
 
@@ -71,12 +71,15 @@ void FIKRigEditMode::Render(const FSceneView* View, FViewport* Viewport, FPrimit
 	}
 	
 	UIKRigController* AssetController = Controller->AssetController;
+	const UIKRigDefinition* IKRigAsset = AssetController->GetAsset();
 	TArray<UIKRigEffectorGoal*> Goals = AssetController->GetAllGoals();
 	for (const UIKRigEffectorGoal* Goal : Goals)
 	{
 		const bool bIsSelected = Controller->IsGoalSelected(Goal->GoalName);
 		PDI->SetHitProxy(new HIKRigEditorGoalProxy(Goal->GoalName));
-		GoalDrawer.DrawGoal(PDI, Goal, bIsSelected);
+		const float Size = IKRigAsset->GoalSize * Goal->SizeMultiplier;
+		const float Thickness = IKRigAsset->GoalThickness * Goal->ThicknessMultiplier;
+		GoalDrawer.DrawGoal(PDI, Goal, bIsSelected, Size, Thickness);
 		PDI->SetHitProxy(NULL);
 	}
 }
@@ -270,7 +273,7 @@ bool FIKRigEditMode::GetCustomDrawingCoordinateSystem(FMatrix& InMatrix, void* I
 		return false; // nothing selected to manipulate
 	}
 
-	if (UIKRigEffectorGoal* Goal = Controller->AssetController->GetGoal(SelectedGoals[0]))
+	if (const UIKRigEffectorGoal* Goal = Controller->AssetController->GetGoal(SelectedGoals[0]))
 	{
 		InMatrix = Goal->CurrentTransform.ToMatrixNoScale().RemoveTranslation();
 		return true;
@@ -284,6 +287,39 @@ bool FIKRigEditMode::GetCustomInputCoordinateSystem(FMatrix& InMatrix, void* InD
 	return GetCustomDrawingCoordinateSystem(InMatrix, InData);
 }
 
+bool FIKRigEditMode::InputKey(FEditorViewportClient* ViewportClient, FViewport* Viewport, FKey Key, EInputEvent Event)
+{
+	if (FEdMode::InputKey(ViewportClient, Viewport, Key, Event))
+	{
+		return false;
+	}
+		
+	if (Key == EKeys::Delete || Key == EKeys::BackSpace)
+	{
+		const TSharedPtr<FIKRigEditorController> Controller = EditorController.Pin();
+		if (!Controller.IsValid())
+		{
+			return false;
+		}
+	
+		TArray<FName> SelectedGoals = Controller->GetSelectedGoals();
+		if (SelectedGoals.IsEmpty())
+		{
+			return false; // nothing selected to manipulate
+		}
+
+		for (const FName& GoalName : SelectedGoals)
+		{
+			Controller->DeleteGoal(GoalName);
+		}
+
+		Controller->RefreshAllViews();
+		return true;
+	}
+
+	return false;
+}
+
 void FIKRigEditMode::Tick(FEditorViewportClient* ViewportClient, float DeltaTime)
 {
 	FEdMode::Tick(ViewportClient, DeltaTime);
diff --git a/Engine/Plugins/Animation/IKRig/Source/IKRigEditor/Private/RigEditor/IKRigEditorController.cpp b/Engine/Plugins/Animation/IKRig/Source/IKRigEditor/Private/RigEditor/IKRigEditorController.cpp
--- a/Engine/Plugins/Animation/IKRig/Source/IKRigEditor/Private/RigEditor/IKRigEditorController.cpp
+++ b/Engine/Plugins/Animation/IKRig/Source/IKRigEditor/Private/RigEditor/IKRigEditorController.cpp
@@ -13,15 +13,33 @@
 
 #define LOCTEXT_NAMESPACE "IKRigEditorController"
 
-void FIKRigEditorController::Reset() const
+void FIKRigEditorController::Initialize(TSharedPtr<FIKRigEditorToolkit> Toolkit, UIKRigDefinition* IKRigAsset)
 {
-	const TArray<UIKRigEffectorGoal*>& Goals = AssetController->GetAllGoals();
-	for (UIKRigEffectorGoal* Goal : Goals)
+	EditorToolkit = Toolkit;
+	AssetController = UIKRigController::GetIKRigController(IKRigAsset);
+	
+	// register callback to be informed when rig asset is modified by editor
+	if (!AssetController->OnIKRigNeedsInitialized().IsBoundToObject(this))
 	{
-		Goal->CurrentTransform = Goal->InitialTransform;
+		AssetController->OnIKRigNeedsInitialized().AddSP(this, &FIKRigEditorController::OnIKRigNeedsInitialized);
 	}
 }
 
+void FIKRigEditorController::OnIKRigNeedsInitialized(UIKRigDefinition* ModifiedIKRig)
+{
+	if (ModifiedIKRig != AssetController->GetAsset())
+	{
+		return;
+	}
+
+	AnimInstance->SetProcessorNeedsInitialized();
+}
+
+void FIKRigEditorController::Reset() const
+{
+	AssetController->ResetGoalTransforms();
+}
+
 void FIKRigEditorController::RefreshAllViews() const
 {
 	SolverStackView->RefreshStackView();
@@ -178,11 +196,11 @@ bool FIKRigEditorController::IsElementConnectedToSolver(TSharedRef<FIKRigTreeEle
 		return false; // not a valid solver index
 	}
 
-	UIKRigSolver* Solver = AssetController->GetSolver(SolverIndex);
+	const UIKRigSolver* Solver = AssetController->GetSolver(SolverIndex);
 	if (TreeElement->ElementType == IKRigTreeElementType::BONE)
 	{
 		// is this bone affected by this solver?
-		return Solver->IsBoneAffectedBySolver(TreeElement->Key, AssetController->GetSkeleton());
+		return Solver->IsBoneAffectedBySolver(TreeElement->Key, AssetController->GetIKRigSkeleton());
 	}
 
 	if (TreeElement->ElementType == IKRigTreeElementType::BONE_SETTINGS)
@@ -235,8 +253,6 @@ bool FIKRigEditorController::IsElementExcludedBone(TSharedRef<FIKRigTreeElement>
 void FIKRigEditorController::ShowDetailsForBone(const FName BoneName)
 {
 	ShowEmptyDetails();
-	// this can be useful during development/debugging, but is slow/confusing so should not be left enabled
-	//DetailsView->SetObject(AssetController->GetAsset());
 }
 
 void FIKRigEditorController::ShowDetailsForBoneSettings(const FName BoneName, int32 SolverIndex)
@@ -249,14 +265,14 @@ void FIKRigEditorController::ShowDetailsForBoneSettings(const FName BoneName, in
 
 void FIKRigEditorController::ShowDetailsForGoal(const FName GoalName)
 {
-	UIKRigEffectorGoal* Goal = AssetController->GetGoal(GoalName);
-	DetailsView->SetObject(Goal);
+	const UIKRigEffectorGoal* Goal = AssetController->GetGoal(GoalName);
+	DetailsView->SetObject(const_cast<UIKRigEffectorGoal*>(Goal));
 }
 
 void FIKRigEditorController::ShowDetailsForEffector(const FName GoalName, const int32 SolverIndex)
 {
 	// get solver that owns this effector
-	if (UIKRigSolver* SolverWithEffector = AssetController->GetSolver(SolverIndex))
+	if (const UIKRigSolver* SolverWithEffector = AssetController->GetSolver(SolverIndex))
 	{
 		if (UObject* EffectorSettings = SolverWithEffector->GetEffectorWithGoal(GoalName))
 		{
@@ -267,12 +283,12 @@ void FIKRigEditorController::ShowDetailsForEffector(const FName GoalName, const
 
 void FIKRigEditorController::ShowDetailsForSolver(const int32 SolverIndex)
 {
-	DetailsView->SetObject(AssetController->GetSolver(SolverIndex));
+	DetailsView->SetObject(const_cast<UIKRigSolver*>(AssetController->GetSolver(SolverIndex)));
 }
 
 void FIKRigEditorController::ShowEmptyDetails()
 {
-	DetailsView->SetObject(nullptr);
+	DetailsView->SetObject(const_cast<UIKRigDefinition*>(AssetController->GetAsset()));
 }
 
 void FIKRigEditorController::AddNewRetargetChain(const FName ChainName, const FName StartBone, const FName EndBone)
diff --git a/Engine/Plugins/Animation/IKRig/Source/IKRigEditor/Private/RigEditor/IKRigToolkit.cpp b/Engine/Plugins/Animation/IKRig/Source/IKRigEditor/Private/RigEditor/IKRigToolkit.cpp
--- a/Engine/Plugins/Animation/IKRig/Source/IKRigEditor/Private/RigEditor/IKRigToolkit.cpp
+++ b/Engine/Plugins/Animation/IKRig/Source/IKRigEditor/Private/RigEditor/IKRigToolkit.cpp
@@ -32,17 +32,12 @@ FIKRigEditorToolkit::FIKRigEditorToolkit()
 {
 }
 
-FIKRigEditorToolkit::~FIKRigEditorToolkit()
-{
-}
-
 void FIKRigEditorToolkit::InitAssetEditor(
 	const EToolkitMode::Type Mode,
     const TSharedPtr<IToolkitHost>& InitToolkitHost, 
     UIKRigDefinition* IKRigAsset)
 {
-	EditorController->AssetController = UIKRigController::GetIKRigController(IKRigAsset);
-	EditorController->EditorToolkit = SharedThis(this);
+	EditorController->Initialize(SharedThis(this), IKRigAsset);
 
 	BindCommands();
 	
@@ -50,7 +45,7 @@ void FIKRigEditorToolkit::InitAssetEditor(
 	PersonaToolkitArgs.OnPreviewSceneCreated = FOnPreviewSceneCreated::FDelegate::CreateSP(this, &FIKRigEditorToolkit::HandlePreviewSceneCreated);
 	
 	FPersonaModule& PersonaModule = FModuleManager::LoadModuleChecked<FPersonaModule>("Persona");
-	PersonaToolkit = PersonaModule.CreatePersonaToolkit(IKRigAsset, PersonaToolkitArgs, IKRigAsset->GetSkeletonAsset());
+	PersonaToolkit = PersonaModule.CreatePersonaToolkit(IKRigAsset, PersonaToolkitArgs, EditorController->AssetController->GetSkeleton());
 	
 	// when/if preview mesh is changed, we need to reinitialize the anim instance
     PersonaToolkit->GetPreviewScene()->RegisterOnPreviewMeshChanged(FOnPreviewMeshChanged::CreateSP(this, &FIKRigEditorToolkit::HandlePreviewMeshChanged));
@@ -146,12 +141,7 @@ FText FIKRigEditorToolkit::GetBaseToolkitName() const
 
 FText FIKRigEditorToolkit::GetToolkitName() const
 {
-	const bool bDirtyState = EditorController->AssetController->GetAsset()->GetOutermost()->IsDirty();
-
-	FFormatNamedArguments Args;
-	Args.Add(TEXT("AssetName"), FText::FromString(EditorController->AssetController->GetAsset()->GetName()));
-	Args.Add(TEXT("DirtyState"), bDirtyState ? FText::FromString(TEXT("*")) : FText::GetEmpty());
-	return FText::Format(LOCTEXT("DemoEditorToolkitName", "{AssetName}{DirtyState}"), Args);
+	return FText::FromString(EditorController->AssetController->GetAsset()->GetName());
 }
 
 FLinearColor FIKRigEditorToolkit::GetWorldCentricTabColorScale() const
@@ -167,7 +157,7 @@ FString FIKRigEditorToolkit::GetWorldCentricTabPrefix() const
 void FIKRigEditorToolkit::AddReferencedObjects(FReferenceCollector& Collector)
 {
 	// hold the asset we are working on
-	UIKRigDefinition* Asset = EditorController->AssetController->GetAsset();
+	const UIKRigDefinition* Asset = EditorController->AssetController->GetAsset();
 	Collector.AddReferencedObject(Asset);
 }
 
diff --git a/Engine/Plugins/Animation/IKRig/Source/IKRigEditor/Private/RigEditor/SIKRigRetargetChainList.cpp b/Engine/Plugins/Animation/IKRig/Source/IKRigEditor/Private/RigEditor/SIKRigRetargetChainList.cpp
--- a/Engine/Plugins/Animation/IKRig/Source/IKRigEditor/Private/RigEditor/SIKRigRetargetChainList.cpp
+++ b/Engine/Plugins/Animation/IKRig/Source/IKRigEditor/Private/RigEditor/SIKRigRetargetChainList.cpp
@@ -42,7 +42,9 @@ void SIKRigRetargetChainRow::Construct(
 	// generate list of goals
 	// NOTE: cannot just use literal "None" because Combobox considers that a "null" entry and will discard it from the list.
 	GoalOptions.Add(MakeShareable(new FString("None")));
-	for (UIKRigEffectorGoal* Goal : InChainList->EditorController.Pin()->AssetController->GetAsset()->Goals)
+	const UIKRigDefinition* IKRigAsset = InChainList->EditorController.Pin()->AssetController->GetAsset();
+	const TArray<UIKRigEffectorGoal*>& AssetGoals =  IKRigAsset->GetGoalArray();
+	for (const UIKRigEffectorGoal* Goal : AssetGoals)
 	{
 		GoalOptions.Add(MakeShareable(new FString(Goal->GoalName.ToString())));
 	}
@@ -63,7 +65,7 @@ TSharedRef<SWidget> SIKRigRetargetChainRow::GenerateWidgetForColumn(const FName&
 		.Padding(3.0f, 1.0f)
 		[
 			SNew(SEditableTextBox)
-			.Text(FText::FromName(ChainElement.Pin()->Chain->ChainName))
+			.Text(FText::FromName(ChainElement.Pin()->ChainName))
 			.Font(FEditorStyle::GetFontStyle(TEXT("BoldFont")))
 			.OnTextCommitted(this, &SIKRigRetargetChainRow::OnRenameChain)
 		];
@@ -81,7 +83,7 @@ TSharedRef<SWidget> SIKRigRetargetChainRow::GenerateWidgetForColumn(const FName&
 		.Padding(3.0f, 1.0f)
 		[
 			SNew(SComboBox<FName>)
-			.OptionsSource(&ChainList.Pin()->EditorController.Pin()->AssetController->GetSkeleton().BoneNames)
+			.OptionsSource(&ChainList.Pin()->EditorController.Pin()->AssetController->GetIKRigSkeleton().BoneNames)
 			.OnGenerateWidget(this, &SIKRigRetargetChainRow::MakeBoneComboEntryWidget)
 			.OnSelectionChanged(this, &SIKRigRetargetChainRow::OnStartBoneComboSelectionChanged)
 			[
@@ -103,7 +105,7 @@ TSharedRef<SWidget> SIKRigRetargetChainRow::GenerateWidgetForColumn(const FName&
 		.Padding(3.0f, 1.0f)
 		[
 			SNew(SComboBox<FName>)
-			.OptionsSource(&ChainList.Pin()->EditorController.Pin()->AssetController->GetSkeleton().BoneNames)
+			.OptionsSource(&ChainList.Pin()->EditorController.Pin()->AssetController->GetIKRigSkeleton().BoneNames)
 			.OnGenerateWidget(this, &SIKRigRetargetChainRow::MakeBoneComboEntryWidget)
 			.OnSelectionChanged(this, &SIKRigRetargetChainRow::OnEndBoneComboSelectionChanged)
 			[
@@ -156,52 +158,82 @@ void SIKRigRetargetChainRow::OnStartBoneComboSelectionChanged(FName InName, ESel
 		return; 
 	}
 	
-	Controller->AssetController->SetRetargetChainStartBone(ChainElement.Pin()->Chain->ChainName, InName);
+	Controller->AssetController->SetRetargetChainStartBone(ChainElement.Pin()->ChainName, InName);
 	ChainList.Pin()->RefreshView();
 }
 
 FText SIKRigRetargetChainRow::GetStartBoneName() const
 {
-	return FText::FromName(ChainElement.Pin()->Chain->StartBone);
+	const TSharedPtr<FIKRigEditorController> Controller = ChainList.Pin()->EditorController.Pin();
+	if (!Controller.IsValid())
+	{
+		return FText::GetEmpty();
+	}
+	
+	const FName StartBoneName = Controller->AssetController->GetRetargetChainStartBone(ChainElement.Pin()->ChainName);
+	return FText::FromName(StartBoneName);
 }
 
-void SIKRigRetargetChainRow::OnEndBoneComboSelectionChanged(FName InName, ESelectInfo::Type SelectInfo)
-{
+FText SIKRigRetargetChainRow::GetEndBoneName() const
+{	
 	const TSharedPtr<FIKRigEditorController> Controller = ChainList.Pin()->EditorController.Pin();
 	if (!Controller.IsValid())
 	{
-		return; 
+		return FText::GetEmpty();
 	}
 	
-	Controller->AssetController->SetRetargetChainEndBone(ChainElement.Pin()->Chain->ChainName, InName);
-	ChainList.Pin()->RefreshView();
+	const FName EndBoneName = Controller->AssetController->GetRetargetChainEndBone(ChainElement.Pin()->ChainName);
+	return FText::FromName(EndBoneName);
 }
 
-void SIKRigRetargetChainRow::OnGoalComboSelectionChanged(TSharedPtr<FString> InGoalName, ESelectInfo::Type SelectInfo)
+FText SIKRigRetargetChainRow::GetGoalName() const
 {
-	ChainElement.Pin()->Chain->IKGoalName = FName(*InGoalName.Get());
-	ChainList.Pin()->RefreshView();
+	const TSharedPtr<FIKRigEditorController> Controller = ChainList.Pin()->EditorController.Pin();
+	if (!Controller.IsValid())
+	{
+		return FText::GetEmpty();
+	}
+	
+	const FName GoalName = Controller->AssetController->GetRetargetChainGoal(ChainElement.Pin()->ChainName);
+	return FText::FromName(GoalName);
 }
 
-FText SIKRigRetargetChainRow::GetEndBoneName() const
+void SIKRigRetargetChainRow::OnEndBoneComboSelectionChanged(FName InName, ESelectInfo::Type SelectInfo)
 {
-	return FText::FromName(ChainElement.Pin()->Chain->EndBone);
+	const TSharedPtr<FIKRigEditorController> Controller = ChainList.Pin()->EditorController.Pin();
+	if (!Controller.IsValid())
+	{
+		return; 
+	}
+	
+	Controller->AssetController->SetRetargetChainEndBone(ChainElement.Pin()->ChainName, InName);
+	ChainList.Pin()->RefreshView();
 }
 
-FText SIKRigRetargetChainRow::GetGoalName() const
+void SIKRigRetargetChainRow::OnGoalComboSelectionChanged(TSharedPtr<FString> InGoalName, ESelectInfo::Type SelectInfo)
 {
-	return FText::FromName(ChainElement.Pin()->Chain->IKGoalName);
-}
+	const TSharedPtr<FIKRigEditorController> Controller = ChainList.Pin()->EditorController.Pin();
+	if (!Controller.IsValid())
+	{
+		return; 
+	}
 
-void SIKRigRetargetChainRow::OnRenameChain(const FText& InText, ETextCommit::Type) const
-{
-	ChainElement.Pin()->Chain->ChainName = FName(*InText.ToString());
+	Controller->AssetController->SetRetargetChainGoal(ChainElement.Pin()->ChainName, FName(*InGoalName.Get()));
 	ChainList.Pin()->RefreshView();
 }
 
-SIKRigRetargetChainList::~SIKRigRetargetChainList()
+void SIKRigRetargetChainRow::OnRenameChain(const FText& InText, ETextCommit::Type) const
 {
-	
+	const TSharedPtr<FIKRigEditorController> Controller = ChainList.Pin()->EditorController.Pin();
+	if (!Controller.IsValid())
+	{
+		return; 
+	}
+
+	const FName OldName = ChainElement.Pin()->ChainName;
+	const FName NewName = FName(*InText.ToString());
+	ChainElement.Pin()->ChainName = Controller->AssetController->RenameRetargetChain(OldName, NewName);
+	ChainList.Pin()->RefreshView();
 }
 
 void SIKRigRetargetChainList::Construct(const FArguments& InArgs, TSharedRef<FIKRigEditorController> InEditorController)
@@ -283,7 +315,7 @@ bool SIKRigRetargetChainList::IsAddChainEnabled() const
 	
 	if (UIKRigController* AssetController = Controller->AssetController)
 	{
-		if (AssetController->GetSkeleton().BoneNames.Num() > 0)
+		if (AssetController->GetIKRigSkeleton().BoneNames.Num() > 0)
 		{
 			return true;
 		}
@@ -305,13 +337,10 @@ void SIKRigRetargetChainList::RefreshView()
 
 	// refresh list of chains
 	ListViewItems.Reset();
-	UIKRigController* AssetController = Controller->AssetController;
-	const int32 NumChains = AssetController->GetNumRetargetChains();
-	for (int32 i=0; i<NumChains; ++i)
+	const TArray<FBoneChain>& Chains = Controller->AssetController->GetRetargetChains();
+	for (const FBoneChain& Chain : Chains)
 	{
-		FBoneChain* Chain = AssetController->GetRetargetChain(i);
-		const FText DisplayName = FText::FromName(Chain->ChainName);
-		TSharedPtr<FRetargetChainElement> ChainItem = FRetargetChainElement::Make(Chain);
+		TSharedPtr<FRetargetChainElement> ChainItem = FRetargetChainElement::Make(Chain.ChainName);
 		ListViewItems.Add(ChainItem);
 	}
 
@@ -359,8 +388,7 @@ FReply SIKRigRetargetChainList::OnKeyDown(const FGeometry& MyGeometry, const FKe
 		}
 
 		UIKRigController* AssetController = Controller->AssetController;
-		const FName ChainToRemove = SelectedItems[0]->Chain->ChainName;
-		AssetController->RemoveRetargetChain(ChainToRemove);
+		AssetController->RemoveRetargetChain(SelectedItems[0]->ChainName);
 
 		RefreshView();
 		
diff --git a/Engine/Plugins/Animation/IKRig/Source/IKRigEditor/Private/RigEditor/SIKRigSkeleton.cpp b/Engine/Plugins/Animation/IKRig/Source/IKRigEditor/Private/RigEditor/SIKRigSkeleton.cpp
--- a/Engine/Plugins/Animation/IKRig/Source/IKRigEditor/Private/RigEditor/SIKRigSkeleton.cpp
+++ b/Engine/Plugins/Animation/IKRig/Source/IKRigEditor/Private/RigEditor/SIKRigSkeleton.cpp
@@ -162,6 +162,8 @@ void SIKRigSkeletonItem::Construct(
         [
 	        SAssignNew(InlineWidget, SInlineEditableTextBlock)
 		    .Text(this, &SIKRigSkeletonItem::GetName)
+		    .Font(TextFont)
+			.ColorAndOpacity(TextColor)
 		    .OnTextCommitted(this, &SIKRigSkeletonItem::OnNameCommitted)
 		    .MultiLine(false)
         ];
@@ -247,7 +249,7 @@ void SIKRigSkeleton::Construct(const FArguments& InArgs, TSharedRef<FIKRigEditor
                 .OnMouseButtonClick(this, &SIKRigSkeleton::OnItemClicked)
                 .OnMouseButtonDoubleClick(this, &SIKRigSkeleton::OnItemDoubleClicked)
                 .OnSetExpansionRecursive(this, &SIKRigSkeleton::OnSetExpansionRecursive)
-                .HighlightParentNodesForSelection(true)
+                .HighlightParentNodesForSelection(false)
                 .ItemHeight(24)
             ]
         ]
@@ -286,7 +288,7 @@ void SIKRigSkeleton::GetSelectedBoneChains(TArray<FIKRigSkeletonChain>& OutChain
 	TArray<TSharedPtr<FIKRigTreeElement>> SelectedBoneItems;
 	GetSelectedBones(SelectedBoneItems);
 
-	FIKRigSkeleton& Skeleton = Controller->AssetController->GetSkeleton();
+	const FIKRigSkeleton& Skeleton = Controller->AssetController->GetIKRigSkeleton();
 
 	// get selected bone indices
 	TArray<int32> SelectedBones;
@@ -970,94 +972,102 @@ void SIKRigSkeleton::RefreshTreeView(bool IsInitialSetup)
 
 	// validate we have a skeleton to load
 	UIKRigController* AssetController = Controller->AssetController;
-	FIKRigSkeleton& Skeleton = AssetController->GetSkeleton();
+	const FIKRigSkeleton& Skeleton = AssetController->GetIKRigSkeleton();
 	if (Skeleton.BoneNames.IsEmpty())
 	{
 		TreeView->RequestTreeRefresh();
 		return;
 	}
 
+	// get all goals
+	TArray<UIKRigEffectorGoal*> Goals = AssetController->GetAllGoals();
+	// get all solvers
+	const TArray<UIKRigSolver*>& Solvers = AssetController->GetSolverArray();
+	// record bone element indices
+	TMap<FName, int32> BoneTreeElementIndices;
+
 	// create all bone elements
 	for (const FName BoneName : Skeleton.BoneNames)
 	{
-		TSharedPtr<FIKRigTreeElement> NewItem = MakeShared<FIKRigTreeElement>(BoneName, IKRigTreeElementType::BONE);
-		AllElements.Add(NewItem);
-	}
-
-	// store children/parent pointers on all bone elements
-	for (int32 BoneIndex=0; BoneIndex<Skeleton.BoneNames.Num(); ++BoneIndex)
-	{
-		const int32 ParentIndex = Skeleton.ParentIndices[BoneIndex];
-		if (ParentIndex < 0)
-		{
-			// store root element
-			RootElements.Add(AllElements[BoneIndex]);
-		}
-		else
-		{
-			// store pointer to child on parent
-			AllElements[ParentIndex]->Children.Add(AllElements[BoneIndex]);
-			// store pointer to parent on child
-			AllElements[BoneIndex]->Parent = AllElements[ParentIndex];
-		}
-	}
+		// create "Bone" tree element for this bone
+		TSharedPtr<FIKRigTreeElement> BoneElement = MakeShared<FIKRigTreeElement>(BoneName, IKRigTreeElementType::BONE);
+		const int32 BoneElementIndex = AllElements.Add(BoneElement);
+		BoneTreeElementIndices.Add(BoneName, BoneElementIndex);
 
-	// create all bone settings elements
-	for (int32 SolverIndex=0; SolverIndex<AssetController->GetNumSolvers(); ++SolverIndex)
-	{
-		UIKRigSolver* Solver = AssetController->GetSolver(SolverIndex);
-		if (!Solver->UsesBoneSettings())
+		// create all "Bone Setting" tree elements for this bone
+		for (int32 SolverIndex=0; SolverIndex<Solvers.Num(); ++SolverIndex)
 		{
-			continue;
-		}
-		for (int32 BoneIndex=0; BoneIndex<Skeleton.BoneNames.Num(); ++BoneIndex)
-		{
-			if (Solver->GetBoneSetting(Skeleton.BoneNames[BoneIndex]))
+			if (Solvers[SolverIndex]->GetBoneSetting(BoneName))
 			{
-				const FName DisplayName = FName("Bone Settings for: " + Solver->GetName());
+				const FName DisplayName = FName("Bone Settings for: " + Solvers[SolverIndex]->GetName());
 				TSharedPtr<FIKRigTreeElement> SettingsItem = MakeShared<FIKRigTreeElement>(DisplayName, IKRigTreeElementType::BONE_SETTINGS);
-				SettingsItem->BoneSettingBoneName = Skeleton.BoneNames[BoneIndex];
+				SettingsItem->BoneSettingBoneName = BoneName;
 				SettingsItem->BoneSettingsSolverIndex = SolverIndex;
 				AllElements.Add(SettingsItem);
 				// store hierarchy pointers for item
-				AllElements[BoneIndex]->Children.Add(SettingsItem);
-				SettingsItem->Parent = AllElements[BoneIndex];
+				BoneElement->Children.Add(SettingsItem);
+				SettingsItem->Parent = BoneElement;
 			}
 		}
-    }
-	
-	// create all goal and effector elements
-	TArray<UIKRigEffectorGoal*> Goals = AssetController->GetAllGoals();
-	for (const UIKRigEffectorGoal* Goal : Goals)
-	{
-		// make new element for goal
-		const int32 BoneIndex = Skeleton.GetBoneIndexFromName(Goal->BoneName);
-		check(BoneIndex!=INDEX_NONE);
-		TSharedPtr<FIKRigTreeElement> GoalItem = MakeShared<FIKRigTreeElement>(Goal->GoalName, IKRigTreeElementType::GOAL);
-		const int32 GoalItemIndex = AllElements.Add(GoalItem);
 
-		// store hierarchy pointers for goal
-		AllElements[BoneIndex]->Children.Add(GoalItem);
-		GoalItem->Parent = AllElements[BoneIndex];
-
-		// add all effectors connected to this goal
-		for (int32 SolverIndex=0; SolverIndex<AssetController->GetNumSolvers(); ++SolverIndex)
+		// create all "Goal" and "Effector" tree elements for this bone
+		for (const UIKRigEffectorGoal* Goal : Goals)
 		{
-			if (UObject* Effector = AssetController->GetEffectorForGoal(Goal->GoalName, SolverIndex))
+			if (Goal->BoneName != BoneName)
+			{
+				continue;
+			}
+			
+			// make new element for goal
+			TSharedPtr<FIKRigTreeElement> GoalItem = MakeShared<FIKRigTreeElement>(Goal->GoalName, IKRigTreeElementType::GOAL);
+			AllElements.Add(GoalItem);
+
+			// store hierarchy pointers for goal
+			BoneElement->Children.Add(GoalItem);
+			GoalItem->Parent = BoneElement;
+
+			// add all effectors connected to this goal
+			for (int32 SolverIndex=0; SolverIndex<Solvers.Num(); ++SolverIndex)
 			{
-				// make new element for effector
-				const UIKRigSolver* Solver = AssetController->GetSolver(SolverIndex);
-				const FName DisplayName = FName("Effector for: " + Solver->GetName());
-				TSharedPtr<FIKRigTreeElement> EffectorItem = MakeShared<FIKRigTreeElement>(DisplayName, IKRigTreeElementType::EFFECTOR);
-				EffectorItem->EffectorSolverIndex = SolverIndex;
-				EffectorItem->EffectorGoalName = Goal->GoalName;
-				AllElements.Add(EffectorItem);
-				EffectorItem->Parent = GoalItem;
-				GoalItem->Children.Add(EffectorItem);
+				if (UObject* Effector = AssetController->GetEffectorForGoal(Goal->GoalName, SolverIndex))
+				{
+					// make new element for effector
+					const FText EffectorPrefix = LOCTEXT("EffectorPrefix", "Effector for");
+					const FName DisplayName = FName(EffectorPrefix.ToString() + ": " + Solvers[SolverIndex]->GetName());
+					TSharedPtr<FIKRigTreeElement> EffectorItem = MakeShared<FIKRigTreeElement>(DisplayName, IKRigTreeElementType::EFFECTOR);
+					EffectorItem->EffectorSolverIndex = SolverIndex;
+					EffectorItem->EffectorGoalName = Goal->GoalName;
+					AllElements.Add(EffectorItem);
+					EffectorItem->Parent = GoalItem;
+					GoalItem->Children.Add(EffectorItem);
+				}
 			}
 		}
 	}
 
+	// store children/parent pointers on all bone elements
+	for (int32 BoneIndex=0; BoneIndex<Skeleton.BoneNames.Num(); ++BoneIndex)
+	{
+		const FName BoneName = Skeleton.BoneNames[BoneIndex];
+		const TSharedPtr<FIKRigTreeElement> BoneTreeElement = AllElements[BoneTreeElementIndices[BoneName]];
+		const int32 ParentIndex = Skeleton.ParentIndices[BoneIndex];
+		if (ParentIndex < 0)
+		{
+			// store the root element
+			RootElements.Add(BoneTreeElement);
+			// has no parent, so skip storing parent pointer
+			continue;
+		}
+
+		// get parent tree element
+		const FName ParentBoneName = Skeleton.BoneNames[ParentIndex];
+		const TSharedPtr<FIKRigTreeElement> ParentBoneTreeElement = AllElements[BoneTreeElementIndices[ParentBoneName]];
+		// store pointer to child on parent
+		ParentBoneTreeElement->Children.Add(BoneTreeElement);
+		// store pointer to parent on child
+		BoneTreeElement->Parent = ParentBoneTreeElement;
+	}
+
 	// restore expansion state
 	for (const TSharedPtr<FIKRigTreeElement>& Element : AllElements)
 	{
diff --git a/Engine/Plugins/Animation/IKRig/Source/IKRigEditor/Private/RigEditor/SIKRigSolverStack.cpp b/Engine/Plugins/Animation/IKRig/Source/IKRigEditor/Private/RigEditor/SIKRigSolverStack.cpp
--- a/Engine/Plugins/Animation/IKRig/Source/IKRigEditor/Private/RigEditor/SIKRigSolverStack.cpp
+++ b/Engine/Plugins/Animation/IKRig/Source/IKRigEditor/Private/RigEditor/SIKRigSolverStack.cpp
@@ -31,6 +31,9 @@ void SIKRigSolverStackItem::Construct(
 	TSharedRef<FSolverStackElement> InStackElement,
 	TSharedPtr<SIKRigSolverStack> InSolverStack)
 {
+	StackElement = InStackElement;
+	SolverStack = InSolverStack;
+	
 	STableRow<TSharedPtr<FSolverStackElement>>::Construct(
         STableRow<TSharedPtr<FSolverStackElement>>::FArguments()
         .OnDragDetected(InSolverStack.Get(), &SIKRigSolverStack::OnDragDetected)
@@ -40,27 +43,136 @@ void SIKRigSolverStackItem::Construct(
         [
             SNew(SHorizontalBox)
             + SHorizontalBox::Slot()
-            .MaxWidth(18)
-			.FillWidth(1.0)
-			.HAlign(HAlign_Left)
-			.VAlign(VAlign_Center)
+            .AutoWidth()
+            .HAlign(HAlign_Left)
+            .Padding(3)
             [
-                SNew(SImage)
-                .Image(FIKRigEditorStyle::Get().GetBrush("IKRig.DragSolver"))
+            	SNew(SHorizontalBox)
+	            + SHorizontalBox::Slot()
+	            .MaxWidth(18)
+				.FillWidth(1.0)
+				.HAlign(HAlign_Left)
+				.VAlign(VAlign_Center)
+	            [
+	                SNew(SImage)
+	                .Image(FIKRigEditorStyle::Get().GetBrush("IKRig.DragSolver"))
+	            ]
+
+	            + SHorizontalBox::Slot()
+	            .AutoWidth()
+				.HAlign(HAlign_Left)
+				.VAlign(VAlign_Center)
+				.Padding(3.0f, 1.0f)
+				[
+					SNew(SCheckBox)
+					.IsEnabled_Lambda([this]()
+					{
+						FText Warning;
+						return !GetWarningMessage(Warning);
+					})
+					.IsChecked_Lambda([InSolverStack, InStackElement]() -> ECheckBoxState
+					{
+						bool bEnabled = true;
+						if (InSolverStack.IsValid() && InSolverStack->EditorController.IsValid())
+						{
+							bEnabled = InSolverStack->EditorController.Pin()->AssetController->GetSolver(InStackElement->IndexInStack)->IsEnabled();
+						}
+						return bEnabled ? ECheckBoxState::Checked : ECheckBoxState::Unchecked;
+					})
+					.OnCheckStateChanged_Lambda([InSolverStack, InStackElement](ECheckBoxState InCheckBoxState)
+					{
+						if (InSolverStack.IsValid() &&
+							InSolverStack->EditorController.IsValid() &&
+							InSolverStack->EditorController.Pin().IsValid())
+						{
+							bool bIsChecked = InCheckBoxState == ECheckBoxState::Checked;
+							InSolverStack->EditorController.Pin()->AssetController->SetSolverEnabled(InStackElement->IndexInStack, bIsChecked);
+						}
+					})
+				]
+	     
+				+ SHorizontalBox::Slot()
+				.AutoWidth()
+				.HAlign(HAlign_Left)
+				.VAlign(VAlign_Center)
+				.Padding(3.0f, 1.0f)
+	            [
+					SNew(STextBlock)
+					.Text(InStackElement->DisplayName)
+					.IsEnabled_Lambda([this]()
+					{
+						FText Warning;
+						return !GetWarningMessage(Warning);
+					})
+					.TextStyle(FEditorStyle::Get(), "NormalText.Important")
+	            ]
+
+	            + SHorizontalBox::Slot()
+	            .AutoWidth()
+				.HAlign(HAlign_Left)
+				.VAlign(VAlign_Center)
+				.Padding(3.0f, 1.0f)
+				[
+					SNew(STextBlock)
+					.Text(InStackElement->DisplayName)
+					.Text_Lambda([this]()
+					{
+						FText Message;
+						GetWarningMessage(Message);
+						return Message;
+					})
+					.TextStyle(FEditorStyle::Get(), "NormalText.Subdued")
+				]
             ]
-     
+
 			+ SHorizontalBox::Slot()
-            .VAlign(VAlign_Fill)
-            .HAlign(HAlign_Fill)
-            .Padding(3.0f, 1.0f)
-            [
-				SNew(STextBlock)
-				.Text(InStackElement->DisplayName)
-				.TextStyle(FEditorStyle::Get(), "NormalText.Important")
-            ]
+			.FillWidth(1)
+			.HAlign(HAlign_Right)
+			.VAlign(VAlign_Center)
+			.Padding(3)
+			[
+				SNew(SButton)
+				.ToolTipText(LOCTEXT("DeleteSolver", "Delect solver and remove from stack."))
+				.OnClicked_Lambda([InSolverStack, InStackElement]() -> FReply
+				{
+					InSolverStack.Get()->DeleteSolver(InStackElement);
+					return FReply::Handled();
+				})
+				.Content()
+				[
+					SNew(SImage)
+					.Image(FAppStyle::Get().GetBrush("Icons.Delete"))
+					.ColorAndOpacity(FSlateColor::UseForeground())
+				]
+			]
         ], OwnerTable);
 }
 
+bool SIKRigSolverStackItem::GetWarningMessage(FText& Message)
+{
+	if (!(StackElement.IsValid() && SolverStack.IsValid()))
+	{
+		return false;
+	}
+	if (!SolverStack.Pin()->EditorController.IsValid())
+	{
+		return false;
+	}
+	const int32 SolverIndex = StackElement.Pin()->IndexInStack;
+	const UIKRigSolver* Solver = SolverStack.Pin()->EditorController.Pin()->AssetController->GetSolver(SolverIndex);
+	if (!Solver)
+	{
+		return false;
+	}
+
+	if (Solver->GetWarningMessage(Message))
+	{
+		return true;
+	}
+					
+	return false;
+}
+
 TSharedRef<FIKRigSolverStackDragDropOp> FIKRigSolverStackDragDropOp::New(TWeakPtr<FSolverStackElement> InElement)
 {
 	TSharedRef<FIKRigSolverStackDragDropOp> Operation = MakeShared<FIKRigSolverStackDragDropOp>();
@@ -129,14 +241,14 @@ void SIKRigSolverStack::Construct(const FArguments& InArgs, TSharedRef<FIKRigEdi
         ]
 
         +SVerticalBox::Slot()
-        .Padding(0.0f, 0.0f)
+        .Padding(0.0f)
         [
 			SAssignNew( ListView, SSolverStackListViewType )
 			.SelectionMode(ESelectionMode::Multi)
 			.IsEnabled(this, &SIKRigSolverStack::IsAddSolverEnabled)
 			.ListItemsSource( &ListViewItems )
 			.OnGenerateRow( this, &SIKRigSolverStack::MakeListRowWidget )
-			.OnMouseButtonClick(this, &SIKRigSolverStack::OnItemClicked)
+			.OnSelectionChanged(this, &SIKRigSolverStack::OnSelectionChanged)
         ]
     ];
 
@@ -165,8 +277,9 @@ void SIKRigSolverStack::BuildAddNewMenu(FMenuBuilder& MenuBuilder)
 
 		if(Class->IsChildOf(UIKRigSolver::StaticClass()) && !Class->HasAnyClassFlags(CLASS_Abstract))
 		{
+			UIKRigSolver* SolverCDO = Cast<UIKRigSolver>(Class->GetDefaultObject());
 			FUIAction Action = FUIAction( FExecuteAction::CreateSP(this, &SIKRigSolverStack::AddNewSolver, Class));
-			MenuBuilder.AddMenuEntry(FText::FromString(Class->GetName()), FText::GetEmpty(), FSlateIcon(), Action);
+			MenuBuilder.AddMenuEntry(FText::FromString(SolverCDO->GetNiceName().ToString()), FText::GetEmpty(), FSlateIcon(), Action);
 		}
 	}
 	
@@ -183,7 +296,7 @@ bool SIKRigSolverStack::IsAddSolverEnabled() const
 	
 	if (UIKRigController* AssetController = Controller->AssetController)
 	{
-		if (AssetController->GetSkeleton().BoneNames.Num() > 0)
+		if (AssetController->GetIKRigSkeleton().BoneNames.Num() > 0)
 		{
 			return true;
 		}
@@ -199,13 +312,39 @@ void SIKRigSolverStack::AddNewSolver(UClass* Class)
 	{
 		return; 
 	}
+
+	UIKRigController* AssetController = Controller->AssetController;
+	if (!AssetController)
+	{
+		return;
+	}
 	
-	if (UIKRigController* AssetController = Controller->AssetController)
+	// add the solver
+	const int32 NewSolverIndex = AssetController->AddSolver(Class);
+	// update stack view
+	RefreshStackView();
+	Controller->SkeletonView->RefreshTreeView(); // update solver indices in effector items
+	// select it
+	ListView->SetSelection(ListViewItems[NewSolverIndex]);
+	// show details for it
+	Controller->ShowDetailsForSolver(NewSolverIndex);
+}
+
+void SIKRigSolverStack::DeleteSolver(TSharedPtr<FSolverStackElement> SolverToDelete)
+{
+	const TSharedPtr<FIKRigEditorController> Controller = EditorController.Pin();
+	if (!Controller.IsValid())
+	{
+		return; 
+	}
+	
+	if (!SolverToDelete.IsValid())
 	{
-		// add the solver
-		const int32 NewSolverIndex = AssetController->AddSolver(Class);
+		return;
 	}
 
+	UIKRigController* AssetController = Controller->AssetController;
+	AssetController->RemoveSolver(SolverToDelete->IndexInStack);
 	RefreshStackView();
 	Controller->SkeletonView->RefreshTreeView(); // update solver indices in effector items
 }
@@ -217,22 +356,31 @@ void SIKRigSolverStack::RefreshStackView()
 	{
 		return; 
 	}
-	
+
+	// record/restore selection
+	int32 IndexToSelect = 0; // default to first item
+	const TArray<TSharedPtr<FSolverStackElement>> SelectedItems = ListView.Get()->GetSelectedItems();
+	if (!SelectedItems.IsEmpty())
+	{
+		IndexToSelect = SelectedItems[0]->IndexInStack;
+	}
+
+	// generate all list items
 	ListViewItems.Reset();
 	UIKRigController* AssetController = Controller->AssetController;
 	const int32 NumSolvers = AssetController->GetNumSolvers();
 	for (int32 i=0; i<NumSolvers; ++i)
 	{
-		UIKRigSolver* Solver = AssetController->GetSolver(i);
-		const FText DisplayName = Solver ? FText::FromString(Solver->GetName()) : FText::FromString("Unknown Solver");
+		const UIKRigSolver* Solver = AssetController->GetSolver(i);
+		const FText DisplayName = Solver ? FText::FromString(Solver->GetNiceName().ToString()) : FText::FromString("Unknown Solver");
 		TSharedPtr<FSolverStackElement> SolverItem = FSolverStackElement::Make(DisplayName, i);
 		ListViewItems.Add(SolverItem);
 	}
 
-	// select first item if none others selected
-	if (ListViewItems.Num() > 0 && ListView->GetNumItemsSelected() == 0)
+	// restore selection
+	if (IndexToSelect < ListViewItems.Num())
 	{
-		ListView->SetSelection(ListViewItems[0]);
+		ListView->SetSelection(ListViewItems[IndexToSelect]);
 	}
 
 	ListView->RequestListRefresh();
@@ -268,16 +416,24 @@ FReply SIKRigSolverStack::OnDragDetected(
 	return FReply::Unhandled();
 }
 
-void SIKRigSolverStack::OnItemClicked(TSharedPtr<FSolverStackElement> InItem)
+void SIKRigSolverStack::OnSelectionChanged(TSharedPtr<FSolverStackElement> InItem, ESelectInfo::Type SelectInfo)
 {
 	const TSharedPtr<FIKRigEditorController> Controller = EditorController.Pin();
 	if (!Controller.IsValid())
 	{
 		return; 
 	}
-	
-	Controller->ShowDetailsForSolver(InItem.Get()->IndexInStack);
-	Controller->SkeletonView->RefreshTreeView(); // update filter showing which bones are affected
+
+	if (!InItem.IsValid())
+	{
+		Controller->ShowEmptyDetails();	
+	}else
+	{
+		Controller->ShowDetailsForSolver(InItem.Get()->IndexInStack);
+	}
+
+	// update bones greyed out when not affected
+	Controller->SkeletonView->RefreshTreeView(); 
 }
 
 TOptional<EItemDropZone> SIKRigSolverStack::OnCanAcceptDrop(
@@ -326,30 +482,17 @@ FReply SIKRigSolverStack::OnAcceptDrop(
 }
 
 FReply SIKRigSolverStack::OnKeyDown(const FGeometry& MyGeometry, const FKeyEvent& InKeyEvent)
-{
-	const FKey Key = InKeyEvent.GetKey();
-	
-	const TSharedPtr<FIKRigEditorController> Controller = EditorController.Pin();
-	if (!Controller.IsValid())
-	{
-		return FReply::Handled();
-	}
-
+{	
 	// handle deleting selected solver
-	if (Key == EKeys::Delete)
+	FKey KeyPressed = InKeyEvent.GetKey();
+	if (KeyPressed == EKeys::Delete || KeyPressed == EKeys::BackSpace)
 	{
 		TArray<TSharedPtr<FSolverStackElement>> SelectedItems = ListView->GetSelectedItems();
-		if (SelectedItems.IsEmpty())
+		if (!SelectedItems.IsEmpty())
 		{
-			return FReply::Unhandled();
+			// only delete 1 at a time to avoid messing up indices
+			DeleteSolver(SelectedItems[0]);
 		}
-
-		UIKRigController* AssetController = Controller->AssetController;
-		UIKRigSolver* SolverToRemove = AssetController->GetSolver(SelectedItems[0]->IndexInStack);
-		AssetController->RemoveSolver(SolverToRemove);
-
-		RefreshStackView();
-		Controller->SkeletonView->RefreshTreeView(); // update solver indices in effector items
 		
 		return FReply::Handled();
 	}
diff --git a/Engine/Plugins/Animation/IKRig/Source/IKRigEditor/Public/RetargetEditor/AssetTypeActions_IKRetargeter.h b/Engine/Plugins/Animation/IKRig/Source/IKRigEditor/Public/RetargetEditor/AssetTypeActions_IKRetargeter.h
--- a/Engine/Plugins/Animation/IKRig/Source/IKRigEditor/Public/RetargetEditor/AssetTypeActions_IKRetargeter.h
+++ b/Engine/Plugins/Animation/IKRig/Source/IKRigEditor/Public/RetargetEditor/AssetTypeActions_IKRetargeter.h
@@ -15,7 +15,7 @@ public:
 	// IAssetTypeActions Implementation
 	virtual FText GetName() const override { return NSLOCTEXT("AssetTypeActions", "AssetTypeActions_IKRetargeter", "IK Retargeter"); }
 	virtual void GetActions(const TArray<UObject*>& InObjects, struct FToolMenuSection& Section) override;
-	virtual FColor GetTypeColor() const override { return FColor(255, 192, 128); }
+	virtual FColor GetTypeColor() const override { return FColor(0, 40, 255); }
 	virtual UClass* GetSupportedClass() const override;
 	virtual void OpenAssetEditor(const TArray<UObject*>& InObjects, TSharedPtr<class IToolkitHost> EditWithinLevelEditor = TSharedPtr<IToolkitHost>()) override;
 	virtual uint32 GetCategories() override { return EAssetTypeCategories::Animation; }
diff --git a/Engine/Plugins/Animation/IKRig/Source/IKRigEditor/Public/RetargetEditor/IKRetargetAnimInstance.h b/Engine/Plugins/Animation/IKRig/Source/IKRigEditor/Public/RetargetEditor/IKRetargetAnimInstance.h
--- a/Engine/Plugins/Animation/IKRig/Source/IKRigEditor/Public/RetargetEditor/IKRetargetAnimInstance.h
+++ b/Engine/Plugins/Animation/IKRig/Source/IKRigEditor/Public/RetargetEditor/IKRetargetAnimInstance.h
@@ -20,7 +20,9 @@ public:
 		UIKRetargeter* InIKRetargetAsset,
 		TWeakObjectPtr<USkeletalMeshComponent> InSourceMeshComponent);
 
-	UIKRetargeter* GetCurrentlyUsedRetargeter() const;
+	const UIKRetargetProcessor* GetRetargetProcessor() const;
+
+	void SetProcessorNeedsInitialized();
 
 protected:
 	virtual FAnimInstanceProxy* CreateAnimInstanceProxy() override;
diff --git a/Engine/Plugins/Animation/IKRig/Source/IKRigEditor/Public/RetargetEditor/IKRetargetEditMode.h b/Engine/Plugins/Animation/IKRig/Source/IKRigEditor/Public/RetargetEditor/IKRetargetEditMode.h
--- a/Engine/Plugins/Animation/IKRig/Source/IKRigEditor/Public/RetargetEditor/IKRetargetEditMode.h
+++ b/Engine/Plugins/Animation/IKRig/Source/IKRigEditor/Public/RetargetEditor/IKRetargetEditMode.h
@@ -11,6 +11,13 @@ class FIKRetargetEditorController;
 class FIKRetargetEditor;
 class FIKRetargetPreviewScene;
 
+enum FIKRetargetTrackingState : int8
+{
+	None,
+	RotatingBone,
+	TranslatingRoot,
+};
+
 class FIKRetargetEditMode : public IPersonaEditMode
 {
 public:
@@ -45,11 +52,16 @@ public:
 	virtual bool GetCustomInputCoordinateSystem(FMatrix& InMatrix, void* InData) override;
 	/** END FEdMode interface */
 
-	bool IsBoneSelected(const FName& BoneName) const;
+private:
 
+	UE::Widget::EWidgetMode CurrentWidgetMode;
+	
+	bool IsOnlyRootSelected() const;
+	
+	bool IsBoneSelected(const FName& BoneName) const;
+	
 	void HandleBoneSelectedInViewport(const FName& BoneName, bool bReplace);
-
-private:
+	
 	/** The hosting app */
 	TWeakPtr<FIKRetargetEditorController> EditorController;
 
@@ -62,7 +74,7 @@ private:
 		float Thickness,
 		bool bIsSelected) const;
 
-	/** chain viewport selection state */
+	/** viewport selection/editing state */
 	TArray<FName> SelectedBones;
-	bool bRotatingBones = false;
+	FIKRetargetTrackingState TrackingState;
 };
diff --git a/Engine/Plugins/Animation/IKRig/Source/IKRigEditor/Public/RetargetEditor/IKRetargetEditor.h b/Engine/Plugins/Animation/IKRig/Source/IKRigEditor/Public/RetargetEditor/IKRetargetEditor.h
--- a/Engine/Plugins/Animation/IKRig/Source/IKRigEditor/Public/RetargetEditor/IKRetargetEditor.h
+++ b/Engine/Plugins/Animation/IKRig/Source/IKRigEditor/Public/RetargetEditor/IKRetargetEditor.h
@@ -56,6 +56,10 @@ public:
 	
 	/** FGCObject interface */
 	virtual void AddReferencedObjects(FReferenceCollector& Collector) override;
+	virtual FString GetReferencerName() const override
+	{
+		return TEXT("FIKRetargetEditor");
+	}
 	/** END FGCObject interface */
 
 	//** FTickableEditorObject Interface
@@ -107,10 +111,6 @@ private:
 	void OnPoseSelected(TSharedPtr<FName> InPoseName, ESelectInfo::Type SelectInfo);
 	/* END edit reference pose */
 
-	/* export animation*/
-	void ExportAnimation() const;
-	/* END edit reference pose */
-
 	void HandleSourceOrTargetIKRigAssetChanged();
 	
 	/** centralized management of across all views */
diff --git a/Engine/Plugins/Animation/IKRig/Source/IKRigEditor/Public/RetargetEditor/IKRetargetEditorController.h b/Engine/Plugins/Animation/IKRig/Source/IKRigEditor/Public/RetargetEditor/IKRetargetEditorController.h
--- a/Engine/Plugins/Animation/IKRig/Source/IKRigEditor/Public/RetargetEditor/IKRetargetEditorController.h
+++ b/Engine/Plugins/Animation/IKRig/Source/IKRigEditor/Public/RetargetEditor/IKRetargetEditorController.h
@@ -1,22 +1,35 @@
 // Copyright Epic Games, Inc. All Rights Reserved.
 #pragma once
+
+#include "IKRetargeterController.h"
 #include "IPersonaToolkit.h"
 #include "UObject/ObjectMacros.h"
 #include "UObject/ObjectPtr.h"
-
+class UIKRetargetProcessor;
 class SIKRetargetChainMapList;
 class UIKRetargetAnimInstance;
 class FIKRetargetEditor;
-class UIKRetargeter;
 class UDebugSkelMeshComponent;
+class UIKRigDefinition;
 
 /** a home for cross-widget communication to synchronize state across all tabs and viewport */
-class FIKRetargetEditorController
+class FIKRetargetEditorController : public TSharedFromThis<FIKRetargetEditorController>
 {
 public:
+
+	/** Initialize the editor */
+	void Initialize(TSharedPtr<FIKRetargetEditor> InEditor, UIKRetargeter* InAsset);
+	/** Bind callbacks to this IK Rig */
+	void BindToIKRigAsset(UIKRigDefinition* InIKRig);
+	/** callback when IK Rig asset requires reinitialization */
+	void OnIKRigNeedsInitialized(UIKRigDefinition* ModifiedIKRig);
+	/** callback when IK Rig asset's retarget chain has been renamed */
+	void OnRetargetChainRenamed(UIKRigDefinition* ModifiedIKRig, FName OldName, FName NewName);
+	/** callback when IK Retargeter asset requires reinitialization */
+	void OnRetargeterNeedsInitialized(const UIKRetargeter* Retargeter);
 	
-	/** the data model */
-	UIKRetargeter* Asset;
+	/** all modifications to the data model should go through this controller */
+	UIKRetargeterController* AssetController;
 
 	/** viewport skeletal mesh */
 	UDebugSkelMeshComponent* SourceSkelMeshComponent;
@@ -52,7 +65,7 @@ public:
 	bool IsTargetBoneRetargeted(const int32& TargetBoneIndex);
 
 	/** get the retargeter that is running in the viewport (which is a duplicate of the source asset) */
-	UIKRetargeter* GetCurrentlyRunningRetargeter() const;
+	const UIKRetargetProcessor* GetRetargetProcessor() const;
 
 	/** Sequence Browser and Edit Pose mode **/
 	void PlayAnimationAsset(UAnimationAsset* AssetToPlay);
diff --git a/Engine/Plugins/Animation/IKRig/Source/IKRigEditor/Public/RetargetEditor/IKRetargeterController.h b/Engine/Plugins/Animation/IKRig/Source/IKRigEditor/Public/RetargetEditor/IKRetargeterController.h
new file mode 100644
--- /dev/null
+++ b/Engine/Plugins/Animation/IKRig/Source/IKRigEditor/Public/RetargetEditor/IKRetargeterController.h
@@ -0,0 +1,107 @@
+﻿// Copyright Epic Games, Inc. All Rights Reserved.
+
+#pragma once
+
+#include "CoreMinimal.h"
+
+#include "Retargeter/IKRetargeter.h"
+#include "UObject/Object.h"
+
+#include "IKRetargeterController.generated.h"
+
+struct FRetargetChainMap;
+class UIKRigDefinition;
+class UIKRetargeter;
+
+/** A singleton (per-asset) class used to make modifications to a UIKRetargeter asset.
+* Call the static UIKRetargeterController() function to get the controller for the asset you want to modify. */ 
+UCLASS(config = Engine, hidecategories = UObject)
+class IKRIGEDITOR_API UIKRetargeterController : public UObject
+{
+	GENERATED_BODY()
+
+public:
+
+	/** Use this to get the controller for the given retargeter asset */
+	static UIKRetargeterController* GetController(UIKRetargeter* InRetargeterAsset);
+	/** Get access to the retargeter asset.
+	 *@warning Do not make modifications to the asset directly. Use this API instead. */
+	UIKRetargeter* GetAsset() const;
+	/** Get the USkeleton used on the Source asset */
+	USkeleton* GetSourceSkeletonAsset() const;
+
+	/** Get name of the Root bone used for retargeting the Source skeleton. */
+	FName GetSourceRootBone() const;
+	/** Get name of the Root bone used for retargeting the Target skeleton. */
+	FName GetTargetRootBone() const;
+
+	/** RETARGET CHAIN MAPPING
+	* 
+	*/
+	/** Get names of all the target bone chains. */
+	void GetTargetChainNames(TArray<FName>& OutNames) const;
+	/** Get names of all the source bone chains. */
+	void GetSourceChainNames(TArray<FName>& OutNames) const;
+	/** Remove invalid chain mappings (no longer existing in currently referenced source/target IK Rig assets) */
+	void CleanChainMapping();
+	/** Use fuzzy string search to find "best" Source chain to map to each Target chain */
+	void AutoMapChains();
+	/** Callback when IK Rig chain is renamed. Retains existing mappings using the new name */
+	void OnRetargetChainRenamed(UIKRigDefinition* IKRig, FName OldChainName, FName NewChainName) const;
+	/** Set the source chain to map to a given target chain */
+	void SetSourceChainForTargetChain(FName TargetChain, FName SourceChainToMapTo);
+	/** Get the source chain mapped to a given target chain */
+	FName GetSourceChainForTargetChain(FName TargetChain);
+	/** Get read-only access to the list of chain mappings */
+	const TArray<FRetargetChainMap>& GetChainMappings();
+	/** END RETARGET CHAIN MAPPING */
+
+	/** RETARGET POSE EDITING
+	 * 
+	 */
+	/** Remove bones from retarget poses that are no longer in skeleton */
+	void CleanPoseList();
+	/** Add new retarget pose. */
+	void AddRetargetPose(FName NewPoseName) const;
+	/** Remove a retarget pose. */
+	void RemoveRetargetPose(FName PoseToRemove) const;
+	/** Reset a retarget pose (removes all stored deltas, returning pose to reference pose */
+	void ResetRetargetPose(FName PoseToReset) const;
+	/** Get the current retarget pose */
+    FName GetCurrentRetargetPoseName() const;
+	/** Change which retarget pose is used by the retargeter at runtime */
+	void SetCurrentRetargetPose(FName CurrentPose) const;
+	/** Get read-only access to list of retarget poses */
+	const TMap<FName, FIKRetargetPose>& GetRetargetPoses();
+	/** Add a delta rotation to a given bone for the current retarget pose (used in Edit Mode in the retarget editor) */
+	void AddRotationOffsetToRetargetPoseBone(FName BoneName, FQuat RotationOffset) const;
+	/** Add a delta translation to the root bone (used in Edit Mode in the retarget editor) */
+	void AddTranslationOffsetToRetargetRootBone(FVector TranslationOffset) const;
+	/** Set whether to output retarget pose. Will output current retarget pose if true, or run retarget otherwise. */
+	void SetEditRetargetPoseMode(bool bOutputRetargetPose) const;
+	/** Get whether in mode to output retarget pose (true) or run retarget (false). */
+	bool GetEditRetargetPoseMode() const;
+	/** END RETARGET POSE EDITING */
+
+private:
+
+	/** Called whenever the rig is modified in such a way that would require re-initialization by dependent systems.*/
+	DECLARE_MULTICAST_DELEGATE_OneParam(FOnRetargeterNeedsInitialized, const UIKRetargeter*);
+	FOnRetargeterNeedsInitialized RetargeterNeedsInitialized;
+
+	void BroadcastNeedsReinitialized() const
+	{
+		RetargeterNeedsInitialized.Broadcast(GetAsset());
+	}
+	
+public:
+	
+	FOnRetargeterNeedsInitialized& OnRetargeterNeedsInitialized(){ return RetargeterNeedsInitialized; };
+	
+private:
+
+	FRetargetChainMap* GetChainMap(const FName& TargetChainName) const;
+
+	/** The actual asset that this Controller modifies. */
+	UIKRetargeter* Asset = nullptr;
+};
diff --git a/Engine/Plugins/Animation/IKRig/Source/IKRigEditor/Public/RetargetEditor/SIKRetargetAssetBrowser.h b/Engine/Plugins/Animation/IKRig/Source/IKRigEditor/Public/RetargetEditor/SIKRetargetAssetBrowser.h
--- a/Engine/Plugins/Animation/IKRig/Source/IKRigEditor/Public/RetargetEditor/SIKRetargetAssetBrowser.h
+++ b/Engine/Plugins/Animation/IKRig/Source/IKRigEditor/Public/RetargetEditor/SIKRetargetAssetBrowser.h
@@ -22,6 +22,8 @@ private:
 
 	void AddAssetBrowser();
 
+	void OnPathChange(const FString& NewPath);
+
 	FReply OnExportButtonClicked() const;
 	bool IsExportButtonEnabled() const;
 	
@@ -37,5 +39,8 @@ private:
 	/** the animation asset browser */
 	TSharedPtr<SBox> AssetBrowserBox;
 
+	/** the path to export assets to */
+	FString BatchOutputPath = "";
+
 	friend FIKRetargetEditorController;
 };
diff --git a/Engine/Plugins/Animation/IKRig/Source/IKRigEditor/Public/RetargetEditor/SIKRetargetChainMapList.h b/Engine/Plugins/Animation/IKRig/Source/IKRigEditor/Public/RetargetEditor/SIKRetargetChainMapList.h
--- a/Engine/Plugins/Animation/IKRig/Source/IKRigEditor/Public/RetargetEditor/SIKRetargetChainMapList.h
+++ b/Engine/Plugins/Animation/IKRig/Source/IKRigEditor/Public/RetargetEditor/SIKRetargetChainMapList.h
@@ -94,6 +94,8 @@ private:
 	TArray< TSharedPtr<FRetargetChainMapElement> > ListViewItems;
 	/** END list view */
 
+	UIKRetargeterController* GetRetargetController() const;
+
 	/** callbacks */
 	FText GetSourceRootBone() const;
 	FText GetTargetRootBone() const;
diff --git a/Engine/Plugins/Animation/IKRig/Source/IKRigEditor/Public/RetargetEditor/SRetargetAnimAssetsWindow.h b/Engine/Plugins/Animation/IKRig/Source/IKRigEditor/Public/RetargetEditor/SRetargetAnimAssetsWindow.h
--- a/Engine/Plugins/Animation/IKRig/Source/IKRigEditor/Public/RetargetEditor/SRetargetAnimAssetsWindow.h
+++ b/Engine/Plugins/Animation/IKRig/Source/IKRigEditor/Public/RetargetEditor/SRetargetAnimAssetsWindow.h
@@ -10,6 +10,42 @@
 
 class UIKRetargeter;
 
+//** Dialog to select path to export to */
+class SSelectExportPathDialog: public SWindow
+{
+public:
+	SLATE_BEGIN_ARGS(SSelectExportPathDialog)
+	{
+	}
+
+	SLATE_ARGUMENT(FText, DefaultAssetPath)
+
+	/** A reference to the parent window */
+	SLATE_ARGUMENT(TSharedPtr<SWindow>, ParentWindow)
+	
+	SLATE_END_ARGS()
+
+	SSelectExportPathDialog() : UserResponse(EAppReturnType::Cancel)
+	{
+	}
+
+	void Construct(const FArguments& InArgs);
+
+public:
+	/** Displays the dialog in a blocking fashion */
+	EAppReturnType::Type ShowModal();
+
+	/** Gets the resulting asset path */
+	FString GetAssetPath();
+
+protected:
+	void OnPathChange(const FString& NewPath);
+	FReply OnButtonClick(EAppReturnType::Type ButtonID);
+
+	EAppReturnType::Type UserResponse;
+	FText AssetPath;
+};
+
 //** Viewport in retarget dialog that shows source / target retarget pose */
 class SRetargetPoseViewport: public SEditorViewport
 {
@@ -178,6 +214,7 @@ private:
 
 	/** Modify folder output path */
 	FText GetFolderPath() const;
+	FReply GetExportFolder();
 
 	/** Necessary data collected from UI to run retarget. */
 	FIKRetargetBatchOperationContext BatchContext;
diff --git a/Engine/Plugins/Animation/IKRig/Source/IKRigEditor/Public/RigEditor/AssetTypeActions_IKRigDefinition.h b/Engine/Plugins/Animation/IKRig/Source/IKRigEditor/Public/RigEditor/AssetTypeActions_IKRigDefinition.h
--- a/Engine/Plugins/Animation/IKRig/Source/IKRigEditor/Public/RigEditor/AssetTypeActions_IKRigDefinition.h
+++ b/Engine/Plugins/Animation/IKRig/Source/IKRigEditor/Public/RigEditor/AssetTypeActions_IKRigDefinition.h
@@ -12,7 +12,7 @@ public:
 	// IAssetTypeActions Implementation
 	virtual FText GetName() const override { return NSLOCTEXT("AssetTypeActions", "AssetTypeActions_IKRigDefinition", "IK Rig"); }
 	virtual void GetActions(const TArray<UObject*>& InObjects, struct FToolMenuSection& Section) override;
-	virtual FColor GetTypeColor() const override { return FColor(255, 192, 128); }
+	virtual FColor GetTypeColor() const override { return FColor(255, 215, 0); }
 	virtual UClass* GetSupportedClass() const override;
 	virtual void OpenAssetEditor(const TArray<UObject*>& InObjects, TSharedPtr<class IToolkitHost> EditWithinLevelEditor = TSharedPtr<IToolkitHost>()) override;
 	virtual uint32 GetCategories() override { return EAssetTypeCategories::Animation; }
diff --git a/Engine/Plugins/Animation/IKRig/Source/IKRigEditor/Public/RigEditor/IKRigAnimInstance.h b/Engine/Plugins/Animation/IKRig/Source/IKRigEditor/Public/RigEditor/IKRigAnimInstance.h
--- a/Engine/Plugins/Animation/IKRig/Source/IKRigEditor/Public/RigEditor/IKRigAnimInstance.h
+++ b/Engine/Plugins/Animation/IKRig/Source/IKRigEditor/Public/RigEditor/IKRigAnimInstance.h
@@ -19,6 +19,8 @@ public:
 
 	void SetIKRigAsset(UIKRigDefinition* IKRigAsset);
 
+	void SetProcessorNeedsInitialized();
+
 protected:
 	virtual FAnimInstanceProxy* CreateAnimInstanceProxy() override;
 
diff --git a/Engine/Plugins/Animation/IKRig/Source/IKRigEditor/Public/RigEditor/IKRigAnimInstanceProxy.h b/Engine/Plugins/Animation/IKRig/Source/IKRigEditor/Public/RigEditor/IKRigAnimInstanceProxy.h
--- a/Engine/Plugins/Animation/IKRig/Source/IKRigEditor/Public/RigEditor/IKRigAnimInstanceProxy.h
+++ b/Engine/Plugins/Animation/IKRig/Source/IKRigEditor/Public/RigEditor/IKRigAnimInstanceProxy.h
@@ -16,9 +16,9 @@ struct FIKRigAnimInstanceProxy : public FAnimPreviewInstanceProxy
 
 public:
 	
-	FIKRigAnimInstanceProxy();
+	FIKRigAnimInstanceProxy() = default;
 	FIKRigAnimInstanceProxy(UAnimInstance* InAnimInstance, FAnimNode_IKRig* IKRigNode);
-	virtual ~FIKRigAnimInstanceProxy();
+	virtual ~FIKRigAnimInstanceProxy() override {};
 
 	/** FAnimPreviewInstanceProxy interface */
 	virtual void Initialize(UAnimInstance* InAnimInstance) override;
diff --git a/Engine/Plugins/Animation/IKRig/Source/IKRigEditor/Public/RigEditor/IKRigController.h b/Engine/Plugins/Animation/IKRig/Source/IKRigEditor/Public/RigEditor/IKRigController.h
--- a/Engine/Plugins/Animation/IKRig/Source/IKRigEditor/Public/RigEditor/IKRigController.h
+++ b/Engine/Plugins/Animation/IKRig/Source/IKRigEditor/Public/RigEditor/IKRigController.h
@@ -13,82 +13,155 @@ class UIKRigSolver;
 class UIKRigDefinition;
 class UIKRigEffectorGoal;
 class USkeletalMesh;
+class USkeleton;
 struct FBoneChain;
 
-UCLASS(config = Engine, hidecategories = UObject, BlueprintType)
+/** A singleton (per-asset) class used to make modifications to a UIKRigDefinition asset.
+ * Call the static GetIKRigController() function to get the controller for the asset you want to modify. */ 
+UCLASS(config = Engine, hidecategories = UObject)
 class IKRIGEDITOR_API UIKRigController : public UObject
 {
 	GENERATED_BODY()
 
 public:
 
-	//** use this to get a handle to a controller for the given IKRig */
+	/** Use this to get the controller for the given IKRig */
 	static UIKRigController* GetIKRigController(UIKRigDefinition* InIKRigDefinition);
-	//** get the asset this controller controls */
+	
+	/** Get the asset this controller controls. */
 	UIKRigDefinition* GetAsset() const;
 
-	//** skeleton */
+	/** SKELETON
+	 * 
+	 */
+	/** Sets the preview mesh to use, can optionally reinitialize the skeleton with bReImportBone=true. */
 	void SetSourceSkeletalMesh(USkeletalMesh* SkeletalMesh, bool bReImportBones) const;
+	/** Get the skeletal mesh asset this IK Rig was initialized with */
 	USkeletalMesh* GetSourceSkeletalMesh() const;
+	/** Initializes the skeleton to use the supplied bones / hierarchy */
 	void SetSkeleton(const FReferenceSkeleton& InSkeleton) const;
-	FIKRigSkeleton& GetSkeleton() const;
+	/** Get read-access to the IKRig skeleton representation */
+	const FIKRigSkeleton& GetIKRigSkeleton() const;
+	/** Get the USkeleton asset this rig was initialized with */
+	USkeleton* GetSkeleton() const;
+	/** Include/exclude a bone from all the solvers. All bones are included by default. */
 	void SetBoneExcluded(const FName& BoneName, const bool bExclude) const;
+	/** Returns true if the given Bone is excluded, false otherwise. */
 	bool GetBoneExcluded(const FName& BoneName) const;
-	//** END skeleton */
-
-	//** solvers */
+	/** Get the global-space retarget pose transform of the given Bone. Can only be called AFTER skeleton is initialized. */
+	FTransform GetBoneRetargetPose(const FName& BoneName) const;
+	/** END SKELETON */
+
+	/** SOLVERS
+	 * 
+	 */
+	/** Add a new solver of the given type to the bottom of the stack. Returns the stack index. */
 	int32 AddSolver(TSubclassOf<UIKRigSolver> InSolverClass) const;
-	void RemoveSolver(UIKRigSolver* SolverToDelete) const;
+	/** Remove the solver at the given stack index. */
+	void RemoveSolver(const int32 SolverIndex) const;
+	/** Move the solver at the given index to the target index. */
 	bool MoveSolverInStack(int32 SolverToMoveIndex, int32 TargetSolverIndex) const;
+	/** Enabled/disable the given solver. */
+	bool SetSolverEnabled(int32 SolverIndex, bool bIsEnabled) const;
+	/** Get access to the given solver. */
 	UIKRigSolver* GetSolver(int32 Index) const;
+	/** Get the number of solvers in the stack. */
 	int32 GetNumSolvers() const;
+	/** Set the root bone on a given solver. (not all solvers support root bones, checks CanSetRootBone() first) */
 	void SetRootBone(const FName& RootBoneName, int32 SolverIndex) const;
+	/** Get read-only access to the array of solvers. */
 	const TArray<UIKRigSolver*>& GetSolverArray() const;
-	//** END solvers */
+	/** END SOLVERS */
 
-	//** goals */
+	/** GOALS
+	 * 
+	 */
+	/** Add a new Goal associated with the given Bone. GoalName must be unique. Bones can have multiple Goals (rare). */
 	UIKRigEffectorGoal* AddNewGoal(const FName& GoalName, const FName& BoneName) const;
+	/** Remove the Goal by name. */
 	bool RemoveGoal(const FName& GoalName) const;
+	/** Rename a Goal. Returns new name, which may be different after being sanitized. Returns NAME_None if this fails.*/
 	FName RenameGoal(const FName& OldName, const FName& PotentialNewName) const;
+	/** The the Bone that the given Goal should be parented to / associated with. */
 	bool SetGoalBone(const FName& GoalName, const FName& NewBoneName) const;
+	/** The the Bone associated with the given Goal. */
 	FName GetBoneForGoal(const FName& GoalName) const;
+	/** Connect the given Goal to the given Solver. This creates an "Effector" with settings specific to this Solver.*/
 	bool ConnectGoalToSolver(const UIKRigEffectorGoal& Goal, int32 SolverIndex) const;
+	/** Disconnect the given Goal from the given Solver. This removes the Effector that associates the Goal with the Solver.*/
 	bool DisconnectGoalFromSolver(const FName& GoalToRemove, int32 SolverIndex) const;
-	bool IsGoalConnectedToSolver(const FName& Goal, int32 SolverIndex) const;
+	/** Returns true if the given Goal is connected to the given Solver. False otherwise. */
+	bool IsGoalConnectedToSolver(const FName& GoalName, int32 SolverIndex) const;
+	/** Get the index of the given Goal in the list of Goals. */
 	int32 GetGoalIndex(const FName& GoalName) const;
+	/** Get the name of Goal at the given index. */
 	FName GetGoalName(const int32& GoalIndex) const;
+	/** Get read-only access to the list of Goals. */
 	const TArray<UIKRigEffectorGoal*>& GetAllGoals() const;
-	UIKRigEffectorGoal* GetGoal(int32 GoalIndex) const;
-	UIKRigEffectorGoal* GetGoal(const FName& GoalName) const;
+	/** Get read-only access to the Goal at the given index. */
+	const UIKRigEffectorGoal* GetGoal(int32 GoalIndex) const;
+	/** Get read-only access to the Goal with the given name. */
+	const UIKRigEffectorGoal* GetGoal(const FName& GoalName) const;
+	/** Get the UObject for the settings associated with the given Goal in the given Solver.
+	 ** Solvers can define their own per-Goal settings depending on their needs. These are termed "Effectors". */
 	UObject* GetEffectorForGoal(const FName& GoalName, int32 SolverIndex) const;
-	FTransform GetGoalInitialTransform(const FName& GoalName) const;
+	/** Get the global-space transform of the given Goal. This may be set by the user in the editor, or at runtime. */
 	FTransform GetGoalCurrentTransform(const FName& GoalName) const;
-	void SetGoalInitialTransform(const FName& GoalName, const FTransform& Transform) const;
+	/** Set the Goal to the given transform. */
 	void SetGoalCurrentTransform(const FName& GoalName, const FTransform& Transform) const;
+	/** Reset all Goals back to their initial transforms. */
+	void ResetGoalTransforms() const;
+	/** Ensure that the given name adheres to required standards for Goal names (no special characters etc..)*/
 	static void SanitizeGoalName(FString& InOutName);
-	DECLARE_MULTICAST_DELEGATE(FGoalModified);
-	FGoalModified OnGoalModified;
-	//** END goals */
+	/** END Goals */
 
-	//** bone settings */
+	/** Bone Settings
+	 * 
+	 */
+	/** Add settings to the given Bone/Solver. Does nothing if Bone already has settings in this Solver.*/
 	void AddBoneSetting(const FName& BoneName, int32 SolverIndex) const;
+	/** Remove settings for the given Bone/Solver. Does nothing if Bone doesn't have setting in this Solver.*/
 	void RemoveBoneSetting(const FName& BoneName, int32 SolverIndex) const;
+	/** Returns true if this Bone can have settings in the given Solver. */
 	bool CanAddBoneSetting(const FName& BoneName, int32 SolverIndex) const;
+	/** Returns true if settings for this Bone can be removed from the given Solver. */
 	bool CanRemoveBoneSetting(const FName& BoneName, int32 SolverIndex) const;
+	/** Get the generic (Solver-specific) Bone settings UObject for this Bone in the given Solver.*/
 	UObject* GetSettingsForBone(const FName& BoneName, int32 SolverIndex) const;
+	/** Returns true if the given Bone has any settings in any Solver. */
 	bool DoesBoneHaveSettings(const FName& BoneName) const;
-	//** END bone settings */
+	/** END Bone Settings */
 
-	//** retargeting and retarget chains */
+	/** Retargeting Options and Retarget Bone Chains
+	 * 
+	 */
+	/** Add a Chain with the given Name and Start/End bones. Returns true if a new Chain was created. */
 	bool AddRetargetChain(const FName& ChainName, const FName& StartBone, const FName& EndBone) const;
+	/** Remove a Chain with the given name. Returns true if a Chain was removed. */
 	bool RemoveRetargetChain(const FName& ChainName) const;
+	/** Renamed the given Chain. Returns the new name (same as old if unsuccessful). */
+	FName RenameRetargetChain(const FName& ChainName, const FName& NewChainName) const;
+	/** Set the Start Bone for the given Chain. Returns true if operation was successful. */
 	bool SetRetargetChainStartBone(const FName& ChainName, const FName& StartBoneName) const;
+	/** Set the End Bone for the given Chain. Returns true if operation was successful. */
 	bool SetRetargetChainEndBone(const FName& ChainName, const FName& EndBoneName) const;
-	FBoneChain* GetRetargetChain(int32 Index) const;
-	int32 GetNumRetargetChains() const;
-	void SetRetargetRoot(const FName& RootBoneName);
-	FName GetRetargetRoot();
-	//** END retarget chains */
+	/** Set the Goal for the given Chain. Returns true if operation was successful. */
+	bool SetRetargetChainGoal(const FName& ChainName, const FName& GoalName) const;
+	/** Get the Goal name for the given Chain. */
+	FName GetRetargetChainGoal(const FName& ChainName) const;
+	/** Get the End Bone name for the given Chain. */
+    FName GetRetargetChainStartBone(const FName& ChainName) const;
+    /** Get the Start Bone name for the given Chain. */
+    FName GetRetargetChainEndBone(const FName& ChainName) const;
+	/** Get read-only access to the list of Chains. */
+	const TArray<FBoneChain>& GetRetargetChains() const;
+	/** Set the Root Bone of the retargeting (can only be one). */
+	void SetRetargetRoot(const FName& RootBoneName) const;
+	/** Get the name of the Root Bone of the retargeting (can only be one). */
+	FName GetRetargetRoot() const;
+	/** Sorts the Chains from Root to tip based on the Start Bone of each Chain. */
+	void SortRetargetChains() const;
+	/** END retarget chains */
 	
 	// BEGIN UObject
 	static void AddReferencedObjects(UObject* InThis, FReferenceCollector& Collector);
@@ -96,10 +169,32 @@ public:
 
 private:
 
+	/** Called whenever the rig is modified in such a way that would require re-initialization by dependent systems.*/
+	DECLARE_MULTICAST_DELEGATE_OneParam(FOnIKRigNeedsInitialized, UIKRigDefinition*);
+	FOnIKRigNeedsInitialized IKRigNeedsInitialized;
+
+	/** Called whenever a retarget chain is renamed.*/
+	DECLARE_MULTICAST_DELEGATE_ThreeParams(FOnRetargetChainRenamed, UIKRigDefinition*, FName /*old name*/, FName /*new name*/);
+	FOnRetargetChainRenamed RetargetChainRenamed;
+
+	void BroadcastNeedsReinitialized() const
+	{
+		IKRigNeedsInitialized.Broadcast(GetAsset());
+	}
+	
+public:
+	
+	FOnIKRigNeedsInitialized& OnIKRigNeedsInitialized(){ return IKRigNeedsInitialized; };
+
+	FOnRetargetChainRenamed& OnRetargetChainRenamed(){ return RetargetChainRenamed; };
+
+private:
+
+	/** The actual IKRigDefinition asset that this Controller modifies. */
 	UPROPERTY(transient)
-	UIKRigDefinition* IKRigAsset = nullptr;
+	TObjectPtr<UIKRigDefinition> IKRigAsset = nullptr;
 
-	// todo, I'm not convinced Controller functions shouldn't just live on the asset itself...
+	/** Lazy-generated map of Controllers to IK Rig Assets. Avoids duplicate controllers. */
 	static TMap<UIKRigDefinition*, UIKRigController*> AssetToControllerMap;
 
 	friend class UIKRigDefinition;
diff --git a/Engine/Plugins/Animation/IKRig/Source/IKRigEditor/Public/RigEditor/IKRigEditMode.h b/Engine/Plugins/Animation/IKRig/Source/IKRigEditor/Public/RigEditor/IKRigEditMode.h
--- a/Engine/Plugins/Animation/IKRig/Source/IKRigEditor/Public/RigEditor/IKRigEditMode.h
+++ b/Engine/Plugins/Animation/IKRig/Source/IKRigEditor/Public/RigEditor/IKRigEditMode.h
@@ -45,11 +45,15 @@ struct GoalGizmo
 		BoxPoints.Add(FVector(-0.5f, 0.5f, -0.5f));
 	}
 
-	void DrawGoal(FPrimitiveDrawInterface* PDI, const UIKRigEffectorGoal* Goal, bool bIsSelected) const
+	void DrawGoal(
+		FPrimitiveDrawInterface* PDI,
+		const UIKRigEffectorGoal* Goal,
+		bool bIsSelected,
+		float Size,
+		float Thickness) const
 	{
 		const FLinearColor Color = bIsSelected ? FLinearColor::Green : FLinearColor::Yellow;
-		const float Thickness = Goal->GizmoThickness;
-		const float Scale = FMath::Clamp(Goal->GizmoSize, 0.1f, 1000.0f);
+		const float Scale = FMath::Clamp(Size, 0.1f, 1000.0f);
 		const FTransform Transform = Goal->CurrentTransform;
 		for (int32 PointIndex = 0; PointIndex < BoxPoints.Num() - 1; PointIndex += 2)
 		{
@@ -92,6 +96,7 @@ public:
 	virtual bool InputDelta(FEditorViewportClient* InViewportClient, FViewport* InViewport, FVector& InDrag, FRotator& InRot, FVector& InScale) override;
 	virtual bool GetCustomDrawingCoordinateSystem(FMatrix& InMatrix, void* InData) override;
 	virtual bool GetCustomInputCoordinateSystem(FMatrix& InMatrix, void* InData) override;
+	virtual bool InputKey(FEditorViewportClient* ViewportClient,FViewport* Viewport, FKey Key,EInputEvent Event);
 	/** END FEdMode interface */
 
 private:
diff --git a/Engine/Plugins/Animation/IKRig/Source/IKRigEditor/Public/RigEditor/IKRigEditorController.h b/Engine/Plugins/Animation/IKRig/Source/IKRigEditor/Public/RigEditor/IKRigEditorController.h
--- a/Engine/Plugins/Animation/IKRig/Source/IKRigEditor/Public/RigEditor/IKRigEditorController.h
+++ b/Engine/Plugins/Animation/IKRig/Source/IKRigEditor/Public/RigEditor/IKRigEditorController.h
@@ -18,10 +18,16 @@ class FIKRigTreeElement;
 class UDebugSkelMeshComponent;
 
 /** a home for cross-widget communication to synchronize state across all tabs and viewport */
-class FIKRigEditorController
+class FIKRigEditorController : public TSharedFromThis<FIKRigEditorController>
 {
 public:
 
+	/** initialize the editor controller to an instance of the IK Rig editor */
+	void Initialize(TSharedPtr<FIKRigEditorToolkit> Toolkit, UIKRigDefinition* Asset);
+
+	/** callback when IK Rig requires re-initialization */
+	void OnIKRigNeedsInitialized(UIKRigDefinition* ModifiedIKRig);
+
 	/** create goals */
 	void AddNewGoals(const TArray<FName>& GoalNames, const TArray<FName>& BoneNames);
 	/** delete goals */
diff --git a/Engine/Plugins/Animation/IKRig/Source/IKRigEditor/Public/RigEditor/IKRigToolkit.h b/Engine/Plugins/Animation/IKRig/Source/IKRigEditor/Public/RigEditor/IKRigToolkit.h
--- a/Engine/Plugins/Animation/IKRig/Source/IKRigEditor/Public/RigEditor/IKRigToolkit.h
+++ b/Engine/Plugins/Animation/IKRig/Source/IKRigEditor/Public/RigEditor/IKRigToolkit.h
@@ -42,7 +42,7 @@ class FIKRigEditorToolkit :
 public:
 
 	FIKRigEditorToolkit();
-	virtual ~FIKRigEditorToolkit();
+	virtual ~FIKRigEditorToolkit() override {};
 
 	void InitAssetEditor(
 		const EToolkitMode::Type Mode,
diff --git a/Engine/Plugins/Animation/IKRig/Source/IKRigEditor/Public/RigEditor/SIKRigRetargetChainList.h b/Engine/Plugins/Animation/IKRig/Source/IKRigEditor/Public/RigEditor/SIKRigRetargetChainList.h
--- a/Engine/Plugins/Animation/IKRig/Source/IKRigEditor/Public/RigEditor/SIKRigRetargetChainList.h
+++ b/Engine/Plugins/Animation/IKRig/Source/IKRigEditor/Public/RigEditor/SIKRigRetargetChainList.h
@@ -4,6 +4,7 @@
 
 #include "CoreMinimal.h"
 #include "EditorUndoClient.h"
+#include "IKRigDefinition.h"
 #include "DragAndDrop/DecoratedDragDropOp.h"
 #include "Widgets/SCompoundWidget.h"
 #include "Widgets/Views/STableViewBase.h"
@@ -30,17 +31,17 @@ public:
         TSharedRef<FRetargetChainElement> InStackElement,
         TSharedPtr<SIKRigRetargetChainList> InChainList);
 
-	static TSharedRef<FRetargetChainElement> Make(FBoneChain* InChain)
+	static TSharedRef<FRetargetChainElement> Make(const FName& InChainName)
 	{
-		return MakeShareable(new FRetargetChainElement(InChain));
+		return MakeShareable(new FRetargetChainElement(InChainName));
 	}
 
-	FBoneChain* Chain;
+	FName ChainName;
 
 private:
 	
 	/** Hidden constructor, always use Make above */
-	FRetargetChainElement(FBoneChain* InChain) : Chain(InChain) {}
+	FRetargetChainElement(const FName& InChainName) : ChainName(InChainName) {}
 
 	/** Hidden constructor, always use Make above */
 	FRetargetChainElement() {}
@@ -112,8 +113,6 @@ public:
 	SLATE_BEGIN_ARGS(SIKRigRetargetChainList) {}
 	SLATE_END_ARGS()
 
-    ~SIKRigRetargetChainList();
-
 	void Construct(const FArguments& InArgs, TSharedRef<FIKRigEditorController> InEditorController);
 
 private:
diff --git a/Engine/Plugins/Animation/IKRig/Source/IKRigEditor/Public/RigEditor/SIKRigSolverStack.h b/Engine/Plugins/Animation/IKRig/Source/IKRigEditor/Public/RigEditor/SIKRigSolverStack.h
--- a/Engine/Plugins/Animation/IKRig/Source/IKRigEditor/Public/RigEditor/SIKRigSolverStack.h
+++ b/Engine/Plugins/Animation/IKRig/Source/IKRigEditor/Public/RigEditor/SIKRigSolverStack.h
@@ -53,6 +53,12 @@ public:
         const TSharedRef<STableViewBase>& OwnerTable,
         TSharedRef<FSolverStackElement> InStackElement,
         TSharedPtr<SIKRigSolverStack> InSolverStack);
+
+	bool GetWarningMessage(FText& Message);
+
+private:
+	TWeakPtr<FSolverStackElement> StackElement;
+	TWeakPtr<SIKRigSolverStack> SolverStack;
 };
 
 class FIKRigSolverStackDragDropOp : public FDecoratedDragDropOp
@@ -96,8 +102,10 @@ private:
 	
 	/** menu command callback for adding a new solver */
 	void AddNewSolver(UClass* Class);
-	/** when a solver is clicked on in the stack view */
-	void OnItemClicked(TSharedPtr<FSolverStackElement> InItem);
+	/** delete solver from stack */
+	void DeleteSolver(TSharedPtr<FSolverStackElement> SolverToDelete);
+	/** when a solver is selected on in the stack view */
+	void OnSelectionChanged(TSharedPtr<FSolverStackElement> InItem, ESelectInfo::Type SelectInfo);
 
 	/** list view generate row callback */
 	TSharedRef<ITableRow> MakeListRowWidget(TSharedPtr<FSolverStackElement> InElement, const TSharedRef<STableViewBase>& OwnerTable);
