Unity build fix.

#rnx

diff --git a/Engine/Plugins/Animation/IKRig/Source/IKRigEditor/Private/RetargetEditor/IKRetargeterController.cpp b/Engine/Plugins/Animation/IKRig/Source/IKRigEditor/Private/RetargetEditor/IKRetargeterController.cpp
--- a/Engine/Plugins/Animation/IKRig/Source/IKRigEditor/Private/RetargetEditor/IKRetargeterController.cpp
+++ b/Engine/Plugins/Animation/IKRig/Source/IKRigEditor/Private/RetargetEditor/IKRetargeterController.cpp
@@ -2,8 +2,8 @@
 
 #include "RetargetEditor/IKRetargeterController.h"
 
-#include "ScopedTransaction.h"
 #include "Algo/LevenshteinDistance.h"
+#include "Engine/SkeletalMesh.h"
 #include "Retargeter/IKRetargeter.h"
 
 #define LOCTEXT_NAMESPACE "IKRetargeterController"
diff --git a/Engine/Plugins/Animation/IKRig/Source/IKRigEditor/Private/RigEditor/SIKRigSolverStack.cpp b/Engine/Plugins/Animation/IKRig/Source/IKRigEditor/Private/RigEditor/SIKRigSolverStack.cpp
--- a/Engine/Plugins/Animation/IKRig/Source/IKRigEditor/Private/RigEditor/SIKRigSolverStack.cpp
+++ b/Engine/Plugins/Animation/IKRig/Source/IKRigEditor/Private/RigEditor/SIKRigSolverStack.cpp
@@ -9,9 +9,10 @@
 
 #include "Framework/Commands/UICommandList.h"
 #include "Framework/MultiBox/MultiBoxBuilder.h"
-#include "IPersonaToolkit.h"
 #include "SKismetInspector.h"
 #include "Dialogs/Dialogs.h"
+#include "Widgets/Input/SButton.h"
+#include "Widgets/Input/SCheckBox.h"
 #include "Widgets/Text/SInlineEditableTextBlock.h"
 #include "SEditorHeaderButton.h"
 
