IK Rig/Retargeter bug fixing and polish.

diff --git a/Engine/Plugins/Animation/IKRig/Source/IKRig/Private/AnimNodes/AnimNode_IKRig.cpp b/Engine/Plugins/Animation/IKRig/Source/IKRig/Private/AnimNodes/AnimNode_IKRig.cpp
--- a/Engine/Plugins/Animation/IKRig/Source/IKRig/Private/AnimNodes/AnimNode_IKRig.cpp
+++ b/Engine/Plugins/Animation/IKRig/Source/IKRig/Private/AnimNodes/AnimNode_IKRig.cpp
@@ -230,7 +230,17 @@ void FAnimNode_IKRig::CacheBones_AnyThread(const FAnimationCacheBonesContext& Co
 		return;
 	}
 
-	// fill up node names
+	if (!IKRigProcessor)
+	{
+		return;
+	}
+
+	if (!IKRigProcessor->IsInitialized())
+	{
+		return;
+	}
+
+	// fill up node names, mapping the anim graph bone indices to the IK Rig bones
 	CompactPoseToRigIndices.Reset();
 	const TArray<FBoneIndexType>& RequiredBonesArray = RequiredBones.GetBoneIndicesArray();
 	const FReferenceSkeleton& RefSkeleton = RequiredBones.GetReferenceSkeleton();
@@ -245,7 +255,7 @@ void FAnimNode_IKRig::CacheBones_AnyThread(const FAnimationCacheBonesContext& Co
 		
 		FCompactPoseBoneIndex CPIndex = RequiredBones.MakeCompactPoseIndex(FMeshPoseBoneIndex(MeshBone));
 		const FName Name = RefSkeleton.GetBoneName(MeshBone);
-		CompactPoseToRigIndices.Add(CPIndex) = RigDefinitionAsset->Skeleton.GetBoneIndexFromName(Name);
+		CompactPoseToRigIndices.Add(CPIndex) = IKRigProcessor->GetSkeleton().GetBoneIndexFromName(Name);
 	}
 }
 
diff --git a/Engine/Plugins/Animation/IKRig/Source/IKRig/Private/IKRigDataTypes.cpp b/Engine/Plugins/Animation/IKRig/Source/IKRig/Private/IKRigDataTypes.cpp
--- a/Engine/Plugins/Animation/IKRig/Source/IKRig/Private/IKRigDataTypes.cpp
+++ b/Engine/Plugins/Animation/IKRig/Source/IKRig/Private/IKRigDataTypes.cpp
@@ -5,7 +5,7 @@
 
 void FIKRigGoalContainer::SetIKGoal(const FIKRigGoal& InGoal)
 {
-	FIKRigGoal* Goal = const_cast<FIKRigGoal*>(FindGoalByName(InGoal.Name));
+	FIKRigGoal* Goal = FindGoalWriteable(InGoal.Name);
 	if (!Goal)
 	{
 		// container hasn't seen this goal before, create new one, copying the input goal
@@ -19,7 +19,7 @@ void FIKRigGoalContainer::SetIKGoal(const FIKRigGoal& InGoal)
 
 void FIKRigGoalContainer::SetIKGoal(const UIKRigEffectorGoal* InEffectorGoal)
 {
-	FIKRigGoal* Goal = const_cast<FIKRigGoal*>(FindGoalByName(InEffectorGoal->GoalName));
+	FIKRigGoal* Goal = FindGoalWriteable(InEffectorGoal->GoalName);
 	if (!Goal)
 	{
 		// container hasn't seen this goal before, create new one, copying the Effector goal
@@ -56,3 +56,8 @@ const FIKRigGoal* FIKRigGoalContainer::FindGoalByName(const FName& GoalName) con
 {
 	return Goals.FindByPredicate([GoalName](const FIKRigGoal& Other)	{return Other.Name == GoalName;});
 }
+
+FIKRigGoal* FIKRigGoalContainer::FindGoalWriteable(const FName& GoalName) const
+{
+	return const_cast<FIKRigGoal*>(FindGoalByName(GoalName));
+}
diff --git a/Engine/Plugins/Animation/IKRig/Source/IKRig/Private/IKRigDefinition.cpp b/Engine/Plugins/Animation/IKRig/Source/IKRig/Private/IKRigDefinition.cpp
--- a/Engine/Plugins/Animation/IKRig/Source/IKRig/Private/IKRigDefinition.cpp
+++ b/Engine/Plugins/Animation/IKRig/Source/IKRig/Private/IKRigDefinition.cpp
@@ -18,6 +18,15 @@ FBoneChain* FRetargetDefinition::GetEditableBoneChainByName(FName ChainName)
 	return nullptr;
 }
 
+
+#if WITH_EDITOR
+void UIKRigDefinition::PostEditUndo()
+{
+	Super::PostEditUndo();
+	IKRigEditUndo.Broadcast();
+}
+#endif
+
 const FBoneChain* UIKRigDefinition::GetRetargetChainByName(FName ChainName) const
 {
 	for (const FBoneChain& Chain : RetargetDefinition.BoneChains)
diff --git a/Engine/Plugins/Animation/IKRig/Source/IKRig/Private/IKRigProcessor.cpp b/Engine/Plugins/Animation/IKRig/Source/IKRig/Private/IKRigProcessor.cpp
--- a/Engine/Plugins/Animation/IKRig/Source/IKRig/Private/IKRigProcessor.cpp
+++ b/Engine/Plugins/Animation/IKRig/Source/IKRig/Private/IKRigProcessor.cpp
@@ -4,7 +4,7 @@
 #include "IKRigDefinition.h"
 #include "IKRigSolver.h"
 
-void UIKRigProcessor::Initialize(UIKRigDefinition* InRigAsset, const FIKRigInputSkeleton& InputSkeleton)
+void UIKRigProcessor::Initialize(const UIKRigDefinition* InRigAsset, const FIKRigInputSkeleton& InputSkeleton)
 {
 	// we instantiate UObjects here which MUST be done on game thread...
 	check(IsInGameThread());
@@ -23,34 +23,18 @@ void UIKRigProcessor::Initialize(UIKRigDefinition* InRigAsset, const FIKRigInput
 	
 	if (InRigAsset->Skeleton.BoneNames.IsEmpty())
 	{
-		UE_LOG(LogTemp, Warning, TEXT("Trying to initialize IKRigProcessor with a IKRigDefinition that has no skeleton: %s"), *InRigAsset->GetName());
+		UE_LOG(LogTemp, Error, TEXT("Trying to initialize IKRig that has no skeleton: %s"), *InRigAsset->GetName());
 		return;
 	}
 
-	if (InRigAsset->GetSolverArray().IsEmpty())
+	if (!UIKRigProcessor::IsIKRigCompatibleWithSkeleton(InRigAsset, InputSkeleton))
 	{
-		UE_LOG(LogTemp, Warning, TEXT("Trying to initialize IKRigProcessor with a IKRigDefinition that has no solvers: %s"), *InRigAsset->GetName());
+		UE_LOG(LogTemp, Error, TEXT("Trying to initialize IKRig with a Skeleton that is missing required bones. See output log. %s"), *InRigAsset->GetName());
 		return;
 	}
 
-	if (InRigAsset->GetGoalArray().IsEmpty())
-	{
-		UE_LOG(LogTemp, Warning, TEXT("Trying to initialize IKRigProcessor with a IKRigDefinition that has no goals: %s"), *InRigAsset->GetName());
-		return;
-	}
-
-	// copy skeleton data from IKRigDefinition
-	// we use the serialized bone names and parent indices (from when asset was initialized)
-	// but we use the CURRENT ref pose from the currently running skeletal mesh (RefSkeleton)
-	Skeleton.BoneNames = InRigAsset->Skeleton.BoneNames;
-	Skeleton.ParentIndices = InRigAsset->Skeleton.ParentIndices;
-	Skeleton.ExcludedBones = InRigAsset->Skeleton.ExcludedBones;
-	const bool bSkeletonIsCompatible = Skeleton.CopyPosesFromInputSkeleton(InputSkeleton);
-	if (!bSkeletonIsCompatible)
-	{
-		UE_LOG(LogTemp, Warning, TEXT("IK Rig, %s trying to run on a skeleton that does not have the required bones."), *InRigAsset->GetName());
-		return;
-	}
+	// copy skeleton data from the actual skeleton we want to run on
+	Skeleton.SetInputSkeleton(InputSkeleton, InRigAsset->Skeleton.ExcludedBones);
 	
 	// initialize goals based on source asset
 	GoalContainer.Empty();
@@ -117,13 +101,93 @@ void UIKRigProcessor::Initialize(UIKRigDefinition* InRigAsset, const FIKRigInput
 	bInitialized = true;
 }
 
-void UIKRigProcessor::Initialize(UIKRigDefinition* InRigAsset, const FReferenceSkeleton& RefSkeleton)
+void UIKRigProcessor::Initialize(const UIKRigDefinition* InRigAsset, const FReferenceSkeleton& RefSkeleton)
 {
-	FIKRigInputSkeleton InputSkeleton;
-	InputSkeleton.InitializeFromRefSkeleton(RefSkeleton);
+	const FIKRigInputSkeleton InputSkeleton = FIKRigInputSkeleton(RefSkeleton);
 	Initialize(InRigAsset, InputSkeleton);
 }
 
+bool UIKRigProcessor::IsIKRigCompatibleWithSkeleton(const UIKRigDefinition* InRigAsset, const FIKRigInputSkeleton& InputSkeleton)
+{
+	// first we validate that all the required bones are in the input skeleton...
+	
+	TSet<FName> RequiredBones;
+	const TArray<UIKRigSolver*>& AssetSolvers = InRigAsset->GetSolverArray();
+	for (const UIKRigSolver* Solver : AssetSolvers)
+	{
+		const FName RootBone = Solver->GetRootBone();
+		if (RootBone != NAME_None)
+		{
+			RequiredBones.Add(RootBone);
+		}
+
+		Solver->GetBonesWithSettings(RequiredBones);
+	}
+
+	const TArray<UIKRigEffectorGoal*>& Goals = InRigAsset->GetGoalArray();
+	for (const UIKRigEffectorGoal* Goal : Goals)
+	{
+		RequiredBones.Add(Goal->BoneName);
+	}
+
+	bool bAllRequiredBonesFound = true;
+	for (const FName& RequiredBone : RequiredBones)
+	{
+		if (!InputSkeleton.BoneNames.Contains(RequiredBone))
+		{
+			UE_LOG(LogTemp, Warning, TEXT("IK Rig, '%s' is missing a required bone in Skeletal Mesh: '%s'."), *InRigAsset->GetName(), *RequiredBone.ToString());
+			bAllRequiredBonesFound = false;
+		}
+	}
+
+	if (!bAllRequiredBonesFound)
+	{
+		return false;
+	}
+
+	// now we validate that hierarchy matches for all required bones...
+	bool bAllParentsValid = true;
+	
+	for (const FName& RequiredBone : RequiredBones)
+	{
+		const int32 InputBoneIndex = InputSkeleton.BoneNames.Find(RequiredBone);
+		const int32 AssetBoneIndex = InRigAsset->Skeleton.BoneNames.Find(RequiredBone);
+
+		// we shouldn't get this far otherwise due to early return above...
+		check(InputBoneIndex != INDEX_NONE && AssetBoneIndex != INDEX_NONE)
+
+		// validate that input skeleton hierarchy is as expected
+		const int32 AssetParentIndex = InRigAsset->Skeleton.ParentIndices[AssetBoneIndex];
+		if (InRigAsset->Skeleton.BoneNames.IsValidIndex(AssetParentIndex)) // root bone has no parent
+		{
+			const FName& AssetParentName = InRigAsset->Skeleton.BoneNames[AssetParentIndex];
+			const int32 InputParentIndex = InputSkeleton.ParentIndices[InputBoneIndex];
+			if (!InputSkeleton.BoneNames.IsValidIndex(InputParentIndex))
+			{
+				bAllParentsValid = false;
+				UE_LOG(LogTemp, Error,
+					TEXT("IK Rig is running on a skeleton with a required bone, '%s', that expected to have a valid parent. The expected parent was, '%s'."),
+					*RequiredBone.ToString(),
+					*AssetParentName.ToString());
+				continue;
+			}
+			const FName& InputParentName = InputSkeleton.BoneNames[InputParentIndex];
+			if (AssetParentName != InputParentName)
+			{
+				// we only warn about this, because it may be nice not to have the exact same hierarchy
+				UE_LOG(LogTemp, Warning,
+					TEXT("IK Rig is running on a skeleton with a required bone, '%s', that has a different parent '%s'. The expected parent was, '%s'."),
+					*InputParentName.ToString(),
+					*InputParentName.ToString(),
+					*AssetParentName.ToString());
+				continue;
+			}
+		}
+	}
+
+	return bAllParentsValid;
+}
+
 void UIKRigProcessor::SetInputPoseGlobal(const TArray<FTransform>& InGlobalBoneTransforms) 
 {
 	check(bInitialized);
@@ -197,7 +261,7 @@ void UIKRigProcessor::SetNeedsInitialized()
 
 #if WITH_EDITOR
 
-void UIKRigProcessor::CopyAllInputsFromSourceAssetAtRuntime(UIKRigDefinition* SourceAsset)
+void UIKRigProcessor::CopyAllInputsFromSourceAssetAtRuntime(const UIKRigDefinition* SourceAsset)
 {
 	check(SourceAsset)
 	
diff --git a/Engine/Plugins/Animation/IKRig/Source/IKRig/Private/IKRigSkeleton.cpp b/Engine/Plugins/Animation/IKRig/Source/IKRig/Private/IKRigSkeleton.cpp
--- a/Engine/Plugins/Animation/IKRig/Source/IKRig/Private/IKRigSkeleton.cpp
+++ b/Engine/Plugins/Animation/IKRig/Source/IKRig/Private/IKRigSkeleton.cpp
@@ -2,31 +2,44 @@
 
 #include "IKRigSkeleton.h"
 
-#include "IKRigDefinition.h"
-#include "IKRigProcessor.h"
+void FIKRigSkeleton::SetInputSkeleton(const FReferenceSkeleton& RefSkeleton, const TArray<FName>& InExcludedBones)
+{
+	const FIKRigInputSkeleton InputSkeleton = FIKRigInputSkeleton(RefSkeleton);
+	SetInputSkeleton(InputSkeleton, InExcludedBones);
+}
 
-void FIKRigSkeleton::Initialize(const FIKRigInputSkeleton& InputSkeleton, const TArray<FName>& InExcludedBones)
+void FIKRigSkeleton::SetInputSkeleton(const FIKRigInputSkeleton& InputSkeleton, const TArray<FName>& InExcludedBones)
 {
 	check(InputSkeleton.BoneNames.Num() == InputSkeleton.ParentIndices.Num() &&
-		  InputSkeleton.BoneNames.Num() == InputSkeleton.LocalRefPose.Num())
+		InputSkeleton.BoneNames.Num() == InputSkeleton.LocalRefPose.Num())
 	
 	// reset all containers
 	Reset();
 	
-	// copy names, parent indices and excluded state into local storage
+	// we use the bone names and parent indices from the input skeleton
 	BoneNames = InputSkeleton.BoneNames;
 	ParentIndices = InputSkeleton.ParentIndices;
+	
+	// bones are excluded at the skeleton level, instead of per-solver
 	ExcludedBones = InExcludedBones;
+	
+	// get a compacted local ref pose based on the stored names
+	TArray<FTransform> CompactRefPoseLocal;
+	// copy a compacted reference pose using only the bones in the input skeleton
+	for (int32 BoneIndex=0; BoneIndex<InputSkeleton.BoneNames.Num(); ++BoneIndex)
+	{
+		// store ref pose of this bone
+		CompactRefPoseLocal.Add(InputSkeleton.LocalRefPose[BoneIndex]);
+	}
+	
+	// copy local ref pose to global
+	ConvertLocalPoseToGlobal(ParentIndices, CompactRefPoseLocal, RefPoseGlobal);
 
-	// copy all the poses out of the ref skeleton
-	CopyPosesFromInputSkeleton(InputSkeleton);
-}
+	// initialize CURRENT GLOBAL pose to reference pose
+	CurrentPoseGlobal = RefPoseGlobal;
 
-void FIKRigSkeleton::Initialize(const FReferenceSkeleton& RefSkeleton, const TArray<FName>& InExcludedBones)
-{
-	FIKRigInputSkeleton InputSkeleton;
-	InputSkeleton.InitializeFromRefSkeleton(RefSkeleton);
-	Initialize(InputSkeleton, InExcludedBones);
+	// initialize CURRENT LOCAL pose from global pose
+	UpdateAllLocalTransformFromGlobal();
 }
 
 void FIKRigSkeleton::Reset()
@@ -88,58 +101,6 @@ int32 FIKRigSkeleton::GetParentIndexThatIsNotExcluded(const int32 BoneIndex) con
 	return ParentIndex;
 }
 
-bool FIKRigSkeleton::CopyPosesFromInputSkeleton(const FIKRigInputSkeleton& InputSkeleton)
-{
-	// get a compacted local ref pose based on the stored names
-	TArray<FTransform> CompactRefPoseLocal;
-
-	// validate all the bones in the input skeleton, comparing against the stored skeleton
-	for (int32 StoredBoneIndex=0; StoredBoneIndex<BoneNames.Num(); ++StoredBoneIndex)
-	{
-		const FName& BoneName = BoneNames[StoredBoneIndex];
-		const int32 InputBoneIndex = InputSkeleton.BoneNames.Find(BoneName);
-
-		// validate that input skeleton has required bone
-		if (InputBoneIndex == INDEX_NONE)
-		{
-			UE_LOG(LogTemp, Error, TEXT("IK Rig is running on a skeleton that is missing bone named, '%s'."), *BoneName.ToString());
-			return false;
-		}
-
-		// validate that input skeleton hierarchy is as expected
-		const int32 StoredParentIndex = ParentIndices[StoredBoneIndex];
-		if (StoredParentIndex != -1) // root bone has no parent
-		{
-			const FName& StoredParentName = BoneNames[StoredParentIndex];
-			const int32 InputParentIndex = InputSkeleton.ParentIndices[InputBoneIndex];
-			const FName& InputParentName = InputSkeleton.BoneNames[InputParentIndex];
-			if (StoredParentName != InputParentName)
-			{
-				UE_LOG(LogTemp, Error,
-					TEXT("IK Rig is running on a skeleton with a bone, '%s', that has a different parent '%s'. The expected parent was, '%s'."),
-					*BoneName.ToString(),
-					*InputParentName.ToString(),
-					*StoredParentName.ToString());
-				return false;
-			}	
-		}
-
-		// store ref pose of this bone
-		CompactRefPoseLocal.Add(InputSkeleton.LocalRefPose[InputBoneIndex]);
-	}
-	
-	// copy local ref pose to global
-	ConvertLocalPoseToGlobal(ParentIndices, CompactRefPoseLocal, RefPoseGlobal);
-
-	// initialize current GLOBAL pose to reference pose
-	CurrentPoseGlobal = RefPoseGlobal;
-
-	// initialize current LOCAL pose
-	UpdateAllLocalTransformFromGlobal();
-
-	return true; // supplied input skeleton had all the bones we are looking for
-}
-
 void FIKRigSkeleton::ConvertLocalPoseToGlobal(
 	const TArray<int32>& InParentIndices,
 	const TArray<FTransform>& InLocalPose,
diff --git a/Engine/Plugins/Animation/IKRig/Source/IKRig/Private/Retargeter/IKRetargetProcessor.cpp b/Engine/Plugins/Animation/IKRig/Source/IKRig/Private/Retargeter/IKRetargetProcessor.cpp
--- a/Engine/Plugins/Animation/IKRig/Source/IKRig/Private/Retargeter/IKRetargetProcessor.cpp
+++ b/Engine/Plugins/Animation/IKRig/Source/IKRig/Private/Retargeter/IKRetargetProcessor.cpp
@@ -8,10 +8,7 @@
 #include "Retargeter/IKRetargeter.h"
 
 
-void FRetargetSkeleton::Initialize(
-	USkeletalMesh* InSkeletalMesh,
-	FIKRetargetPose* InRetargetPose,
-	const FName& RetargetRootBone)
+void FRetargetSkeleton::Initialize(USkeletalMesh* InSkeletalMesh)
 {	
 	// record which skeletal mesh this is running on
 	SkeletalMesh = InSkeletalMesh;
@@ -25,7 +22,7 @@ void FRetargetSkeleton::Initialize(
 	}
 
 	// update retarget pose to reflect custom offsets
-	GenerateRetargetPose(InRetargetPose, RetargetRootBone);
+	GenerateRetargetPose();
 }
 
 void FRetargetSkeleton::Reset()
@@ -37,7 +34,7 @@ void FRetargetSkeleton::Reset()
 	SkeletalMesh = nullptr;
 }
 
-void FRetargetSkeleton::GenerateRetargetPose(FIKRetargetPose* InRetargetPose, const FName& RetargetRootBone)
+void FRetargetSkeleton::GenerateRetargetPose()
 {
 	// initialize retarget pose to the skeletal mesh reference pose
 	RetargetLocalPose = SkeletalMesh->GetRefSkeleton().GetRefBonePose();
@@ -45,35 +42,6 @@ void FRetargetSkeleton::GenerateRetargetPose(FIKRetargetPose* InRetargetPose, co
 	RetargetGlobalPose = RetargetLocalPose;
 	// convert to global space
 	UpdateGlobalTransformsBelowBone(0, RetargetLocalPose, RetargetGlobalPose);
-
-	// apply retarget pose offsets (retarget pose is stored as offset relative to reference pose)
-	if (InRetargetPose && RetargetRootBone != NAME_None)
-	{
-		// apply root translation offset
-		const int32 RootBoneIndex = FindBoneIndexByName(RetargetRootBone);
-		if (RootBoneIndex != INDEX_NONE)
-		{
-			FTransform& RootTransform = RetargetGlobalPose[RootBoneIndex];
-			RootTransform.AddToTranslation(InRetargetPose->RootTranslationOffset);
-			UpdateLocalTransformOfSingleBone(RootBoneIndex, RetargetLocalPose, RetargetGlobalPose);
-			UpdateGlobalTransformsBelowBone(RootBoneIndex, RetargetLocalPose, RetargetGlobalPose);
-		}
-
-		// apply bone rotation offsets
-		for (const TTuple<FName, FQuat>& BoneRotationOffset : InRetargetPose->BoneRotationOffsets)
-		{
-			const int32 BoneIndex = FindBoneIndexByName(BoneRotationOffset.Key);
-			if (BoneIndex == INDEX_NONE)
-			{
-				continue;
-			}
-
-			FQuat BoneRotation = BoneRotationOffset.Value * RetargetGlobalPose[BoneIndex].GetRotation();
-			RetargetGlobalPose[BoneIndex].SetRotation(BoneRotation);
-			UpdateLocalTransformOfSingleBone(BoneIndex, RetargetLocalPose, RetargetGlobalPose);
-			UpdateGlobalTransformsBelowBone(BoneIndex, RetargetLocalPose, RetargetGlobalPose);
-		}
-	}
 }
 
 int32 FRetargetSkeleton::FindBoneIndexByName(const FName InName) const
@@ -167,16 +135,82 @@ int32 FRetargetSkeleton::GetParentIndex(const int32 BoneIndex) const
 
 void FTargetSkeleton::Initialize(
 	USkeletalMesh* InSkeletalMesh, 
-	FIKRetargetPose* RetargetPose,
-	const FName& RetargetRootBone)
+	const FIKRetargetPose* RetargetPose,
+	const FName& RetargetRootBone,
+	const TArray<FBoneChain>& TargetChains)
 {
-	FRetargetSkeleton::Initialize(InSkeletalMesh, RetargetPose, RetargetRootBone);
+	FRetargetSkeleton::Initialize(InSkeletalMesh);
 
 	// initialize storage for output pose (the result of the retargeting)
 	OutputGlobalPose = RetargetGlobalPose;
 
 	// make storage for per-bone "Is Retargeted" flag (used for hierarchy updates)
+	// these are bones that are in a target chain that is mapped to a source chain (ie, will actually be retargeted)
+	// these flags are actually set later in init phase when bone chains are mapped together
 	IsBoneRetargeted.Init(false, OutputGlobalPose.Num());
+
+	// determine set of bones referenced by one of the target bone chains to be retargeted
+	// this is the set of bones that will be affected by the retarget pose
+	IsBoneInAnyTargetChain.Init(false, OutputGlobalPose.Num());
+	for (const FBoneChain& TargetChain : TargetChains)
+	{
+		TArray<int32> BonesInChain;
+		if (FResolvedBoneChain(TargetChain, *this, BonesInChain).IsValid())
+		{
+			for (int32 BoneInChain : BonesInChain)
+			{
+				IsBoneInAnyTargetChain[BoneInChain] = true;
+			}
+		}
+	}
+
+	// generate the retarget pose (applies stored offsets)
+	// NOTE: this must be done AFTER generating IsBoneInAnyTargetChain array above
+	GenerateRetargetPose(RetargetPose, RetargetRootBone);
+}
+
+void FTargetSkeleton::GenerateRetargetPose(const FIKRetargetPose* InRetargetPose, const FName& RetargetRootBone)
+{
+	// create a retarget pose by copying the ref pose
+	FRetargetSkeleton::GenerateRetargetPose();
+	
+	// apply retarget pose offsets (retarget pose is stored as offset relative to reference pose)
+	if (InRetargetPose && RetargetRootBone != NAME_None)
+	{
+		// apply root translation offset
+		const int32 RootBoneIndex = FindBoneIndexByName(RetargetRootBone);
+		if (RootBoneIndex != INDEX_NONE)
+		{
+			FTransform& RootTransform = RetargetGlobalPose[RootBoneIndex];
+			RootTransform.AddToTranslation(InRetargetPose->RootTranslationOffset);
+			UpdateLocalTransformOfSingleBone(RootBoneIndex, RetargetLocalPose, RetargetGlobalPose);
+			UpdateGlobalTransformsBelowBone(RootBoneIndex, RetargetLocalPose, RetargetGlobalPose);
+		}
+
+		// apply bone rotation offsets
+		for (const TTuple<FName, FQuat>& BoneRotationOffset : InRetargetPose->BoneRotationOffsets)
+		{
+			const int32 BoneIndex = FindBoneIndexByName(BoneRotationOffset.Key);
+			if (BoneIndex == INDEX_NONE)
+			{
+				// this can happen if a retarget pose recorded a bone offset for a bone that is not present in the
+				// target skeleton; ie, the retarget pose was generated from a different Skeletal Mesh with extra bones
+				continue;
+			}
+
+			if (!IsBoneInAnyTargetChain[BoneIndex])
+			{
+				// this can happen if a retarget pose includes bone edits from a bone chain that was subsequently removed,
+				// and the asset has not run through the "CleanChainMapping" operation yet (happens on load)
+				continue;
+			}
+
+			FQuat BoneRotation = BoneRotationOffset.Value * RetargetGlobalPose[BoneIndex].GetRotation();
+			RetargetGlobalPose[BoneIndex].SetRotation(BoneRotation);
+			UpdateLocalTransformOfSingleBone(BoneIndex, RetargetLocalPose, RetargetGlobalPose);
+			UpdateGlobalTransformsBelowBone(BoneIndex, RetargetLocalPose, RetargetGlobalPose);
+		}
+	}
 }
 
 void FTargetSkeleton::Reset()
@@ -199,6 +233,46 @@ void FTargetSkeleton::UpdateGlobalTransformsAllNonRetargetedBones(TArray<FTransf
 	}
 }
 
+FResolvedBoneChain::FResolvedBoneChain(
+	const FBoneChain& BoneChain,
+	const FRetargetSkeleton& Skeleton,
+	TArray<int32> &OutBoneIndices)
+{
+	// validate start and end bones exist
+	const int32 StartIndex = Skeleton.FindBoneIndexByName(BoneChain.StartBone);
+	const int32 EndIndex = Skeleton.FindBoneIndexByName(BoneChain.EndBone);
+	bFoundStartBone = StartIndex != INDEX_NONE;
+	bFoundEndBone = EndIndex != INDEX_NONE;
+	
+	// validate end bone is child of start bone
+	bEndIsChildOfStart = true;
+	// record all bones in chain while walking up the hierarchy (tip to root of chain)
+	OutBoneIndices.Reset();
+	int32 NextBoneIndex = EndIndex;
+	while (true)
+	{
+		OutBoneIndices.Add(NextBoneIndex);
+		
+		if (NextBoneIndex == StartIndex)
+		{
+			break;
+		}
+
+		NextBoneIndex = Skeleton.GetParentIndex(NextBoneIndex);
+		
+		if (NextBoneIndex == 0)
+		{
+			// oops, we walked all the way up to the root without finding the start bone
+			bEndIsChildOfStart = false;
+			OutBoneIndices.Reset();
+			return;
+		}
+	}
+
+	// reverse the indices (we want root to tip order)
+	Algo::Reverse(OutBoneIndices);
+}
+
 void FTargetSkeleton::SetBoneIsRetargeted(const int32 BoneIndex, const bool IsRetargeted)
 {
 	check(IsBoneRetargeted.IsValidIndex(BoneIndex));
@@ -553,53 +627,35 @@ bool FRetargetChainPair::ValidateBoneChainWithSkeletalMesh(
     const bool IsSource,
     const FBoneChain& BoneChain,
     const FRetargetSkeleton& RetargetSkeleton)
-{	
-	// validate start bone exists
-	const int32 StartIndex = RetargetSkeleton.FindBoneIndexByName(BoneChain.StartBone);
-	if (StartIndex == INDEX_NONE)
+{
+	// record the chain indices
+	TArray<int32>& BoneIndices = IsSource ? SourceBoneIndices : TargetBoneIndices;
+	
+	// resolve the bone bone to the skeleton
+	const FResolvedBoneChain ResolvedChain = FResolvedBoneChain(BoneChain, RetargetSkeleton, BoneIndices);
+	
+	// warn if START bone not found
+	if (!ResolvedChain.bFoundStartBone)
 	{
 		UE_LOG(LogTemp, Warning, TEXT("IK Retargeter bone chain, %s, could not find start bone, %s in mesh %s"),
             *BoneChain.ChainName.ToString(), *BoneChain.StartBone.ToString(), *RetargetSkeleton.SkeletalMesh->GetName());
-		return false;
 	}
 	
-	// validate end bone exists
-	const int32 EndIndex = RetargetSkeleton.FindBoneIndexByName(BoneChain.EndBone);
-	if (EndIndex == INDEX_NONE)
+	// warn if END bone not found
+	if (!ResolvedChain.bFoundEndBone)
 	{
 		UE_LOG(LogTemp, Warning, TEXT("IK Retargeter bone chain, %s, could not find end bone, %s in mesh %s"),
             *BoneChain.ChainName.ToString(), *BoneChain.EndBone.ToString(), *RetargetSkeleton.SkeletalMesh->GetName());
-		return false;
 	}
-	
-	// validate end bone is child of start bone
-	// record all bones in chain while walking up the hierarchy (tip to root of chain)
-	TArray<int>& BoneIndices = IsSource ? SourceBoneIndices : TargetBoneIndices;
-	BoneIndices.Reset();
-	int32 NextBoneIndex = EndIndex;
-	while (true)
-	{
-		BoneIndices.Add(NextBoneIndex);
-		
-		if (NextBoneIndex == StartIndex)
-		{
-			break;
-		}
 
-		NextBoneIndex = RetargetSkeleton.GetParentIndex(NextBoneIndex);
-		
-		if (NextBoneIndex == 0)
-		{
-			UE_LOG(LogTemp, Warning, TEXT("IK Retargeter bone chain, '%s' has end bone, '%s' that is not a child of start bone, '%s' in mesh '%s'"),
-            *BoneChain.ChainName.ToString(), *BoneChain.EndBone.ToString(), *BoneChain.StartBone.ToString(), *RetargetSkeleton.SkeletalMesh->GetName());
-			return false;
-		}
+	// warn if END bone was not a child of START bone
+	if (!ResolvedChain.bEndIsChildOfStart)
+	{
+		UE_LOG(LogTemp, Warning, TEXT("IK Retargeter bone chain, %s, end bone, '%s' was not a child of the start bone '%s'."),
+			*BoneChain.ChainName.ToString(), *BoneChain.EndBone.ToString(), *BoneChain.StartBone.ToString());
 	}
-
-	// reverse the indices (we want root to tip order)
-	Algo::Reverse(BoneIndices);
 	
-	return true;
+	return ResolvedChain.IsValid();
 }
 
 bool FRetargetChainPairFK::Initialize(
@@ -789,55 +845,50 @@ void UIKRetargetProcessor::Initialize(
 		UE_LOG(LogTemp, Error, TEXT("IK Retargeter unable to initialize. Missing target Skeletal Mesh asset."));
 		return;
 	}
-	if (!RetargeterAsset->SourceIKRigAsset)
+	if (!RetargeterAsset->GetSourceIKRig())
 	{
 		UE_LOG(LogTemp, Error, TEXT("IK Retargeter unable to initialize. Missing source IK Rig asset."));
 		return;
 	}
-	if (!RetargeterAsset->TargetIKRigAsset)
+	if (!RetargeterAsset->GetTargetIKRig())
 	{
 		UE_LOG(LogTemp, Error, TEXT("IK Retargeter unable to initialize. Missing target IK Rig asset."));
 		return;
 	}
-	if (!RetargeterAsset->RetargetPoses.Contains(RetargeterAsset->CurrentRetargetPose))
+	if (!RetargeterAsset->GetCurrentRetargetPose())
 	{
 		UE_LOG(LogTemp, Error, TEXT("IK Retargeter unable to initialize. Missing retarget pose."));
 		return;
 	}
 	
 	// initialize skeleton data for source and target
-	const TArray<FRetargetChainMap> DummySourceChainMapping;
-	SourceSkeleton.Initialize(
-		SourceSkeletalMesh,
-		nullptr,
-		NAME_None);
+	SourceSkeleton.Initialize(SourceSkeletalMesh);
 	TargetSkeleton.Initialize(
 		TargetSkeletalMesh,
-		&RetargeterAsset->RetargetPoses[RetargeterAsset->CurrentRetargetPose],
-		RetargeterAsset->TargetIKRigAsset->GetRetargetRoot());
+		RetargeterAsset->GetCurrentRetargetPose(),
+		RetargeterAsset->GetTargetIKRig()->GetRetargetRoot(),
+		RetargeterAsset->GetTargetIKRig()->GetRetargetChains());
 
 	// initialize roots
-	const bool bRootsInitialized = InitializeRoots();
+	bRootsInitialized = InitializeRoots();
 	if (!bRootsInitialized)
 	{
 		// couldn't match up any BoneChain pairs, no retargeting possible
 		UE_LOG(LogTemp, Error, TEXT("IK Retargeter unable to initialize one or more root bones on source, %s and target, %s"),
             *SourceSkeleton.SkeletalMesh->GetName(), *TargetSkeleton.SkeletalMesh->GetName());
-		return;
 	}
 
 	// initialize pairs of bone chains
-	const bool bAtLeastOneValidBoneChainPair = InitializeBoneChainPairs();
+	bAtLeastOneValidBoneChainPair = InitializeBoneChainPairs();
 	if (!bAtLeastOneValidBoneChainPair)
 	{
-		// couldn't match up any BoneChain pairs, no retargeting possible
-		UE_LOG(LogTemp, Error, TEXT("IK Retargeter unable to create any Bone Chain pairs between source, %s and target, %s"),
+		// couldn't match up any BoneChain pairs, no limb retargeting possible
+		UE_LOG(LogTemp, Warning, TEXT("IK Retargeter unable to create any Bone Chain pairs between source, %s and target, %s"),
             *SourceSkeleton.SkeletalMesh->GetName(), *TargetSkeleton.SkeletalMesh->GetName());
-		return;
 	}
 
 	// initialize the IKRigProcessor for doing IK decoding
-	const bool bIKRigInitialized = InitializeIKRig(this, TargetSkeletalMesh->GetRefSkeleton());
+	bIKRigInitialized = InitializeIKRig(this, TargetSkeletalMesh->GetRefSkeleton());
 	if (!bIKRigInitialized)
 	{
 		// couldn't initialize the IK Rig, we don't disable the retargeter in this case, just warn the user
@@ -851,23 +902,21 @@ void UIKRetargetProcessor::Initialize(
 bool UIKRetargetProcessor::InitializeRoots()
 {
 	// initialize root encoder
-	const FName SourceRootBoneName = RetargeterAsset->SourceIKRigAsset->GetRetargetRoot();
+	const FName SourceRootBoneName = RetargeterAsset->GetSourceIKRig()->GetRetargetRoot();
 	const bool bRootEncoderInit = RootRetargeter.InitializeSource(SourceRootBoneName, SourceSkeleton);
 	if (!bRootEncoderInit)
 	{
-		UE_LOG(LogTemp, Error, TEXT("IK Retargeter unable to initialize source root, '%s' on skeletal mesh: '%s'"),
+		UE_LOG(LogTemp, Warning, TEXT("IK Retargeter unable to initialize source root, '%s' on skeletal mesh: '%s'"),
             *SourceRootBoneName.ToString(), *SourceSkeleton.SkeletalMesh->GetName());
-		return false;
 	}
 
 	// initialize root decoder
-	const FName TargetRootBoneName = RetargeterAsset->TargetIKRigAsset->GetRetargetRoot();
+	const FName TargetRootBoneName = RetargeterAsset->GetTargetIKRig()->GetRetargetRoot();
 	const bool bRootDecoderInit = RootRetargeter.InitializeTarget(TargetRootBoneName, TargetSkeleton);
 	if (!bRootDecoderInit)
 	{
-		UE_LOG(LogTemp, Error, TEXT("IK Retargeter unable to initialize target root, '%s' on skeletal mesh: '%s'"),
+		UE_LOG(LogTemp, Warning, TEXT("IK Retargeter unable to initialize target root, '%s' on skeletal mesh: '%s'"),
             *TargetRootBoneName.ToString(), *TargetSkeleton.SkeletalMesh->GetName());
-		return false;
 	}
 
 	return true;
@@ -875,32 +924,26 @@ bool UIKRetargetProcessor::InitializeRoots()
 
 bool UIKRetargetProcessor::InitializeBoneChainPairs()
 {
-	check(RetargeterAsset->SourceIKRigAsset && RetargeterAsset->TargetIKRigAsset);
-	
-	for (const FRetargetChainMap& ChainSettings : RetargeterAsset->ChainMapping)
+	check(RetargeterAsset->GetSourceIKRig() && RetargeterAsset->GetTargetIKRig());
+
+	const TArray<FRetargetChainMap>& ChainMapping = RetargeterAsset->GetChainMapping();
+	for (const FRetargetChainMap& ChainMap : ChainMapping)
 	{
-		FName SourceChain = ChainSettings.SourceChain;
-		FName TargetChain = ChainSettings.TargetChain;
+		// get target bone chain
+		const FBoneChain* TargetBoneChain = RetargeterAsset->GetTargetIKRig()->GetRetargetChainByName(ChainMap.TargetChain);
+		ensureMsgf(TargetBoneChain, TEXT("IK Retargeter has mapping for a target bone chain that doesn't exist: %s"), *ChainMap.TargetChain.ToString());
 
 		// user opted to not map this to anything, we don't need to spam a warning about it
-		if (SourceChain == NAME_None)
+		if (ChainMap.SourceChain == NAME_None)
 		{
 			continue; 
 		}
 		
 		// get source bone chain
-		const FBoneChain* SourceBoneChain = RetargeterAsset->SourceIKRigAsset->GetRetargetChainByName(SourceChain);
+		const FBoneChain* SourceBoneChain = RetargeterAsset->GetSourceIKRig()->GetRetargetChainByName(ChainMap.SourceChain);
 		if (!SourceBoneChain)
 		{
-			UE_LOG(LogTemp, Error, TEXT("IK Retargeter missing source bone chain: %s"), *SourceChain.ToString());
-			continue;
-		}
-
-		// get target bone chain
-		const FBoneChain* TargetBoneChain = RetargeterAsset->TargetIKRigAsset->GetRetargetChainByName(TargetChain);
-		if (!TargetBoneChain)
-		{
-			UE_LOG(LogTemp, Error, TEXT("IK Retargeter missing target bone chain: %s"), *TargetChain.ToString());
+			UE_LOG(LogTemp, Error, TEXT("IK Retargeter missing source bone chain: %s"), *ChainMap.SourceChain.ToString());
 			continue;
 		}
 
@@ -953,7 +996,7 @@ bool UIKRetargetProcessor::InitializeIKRig(UObject* Outer, const FReferenceSkele
 	{
 		IKRigProcessor = NewObject<UIKRigProcessor>(Outer);	
 	}
-	IKRigProcessor->Initialize(RetargeterAsset->TargetIKRigAsset, InRefSkeleton);
+	IKRigProcessor->Initialize(RetargeterAsset->GetTargetIKRig(), InRefSkeleton);
 	if (!IKRigProcessor->IsInitialized())
 	{
 		return false;
@@ -982,10 +1025,10 @@ TArray<FTransform>&  UIKRetargetProcessor::RunRetargeter(const TArray<FTransform
 	// in edit mode we just want to see the edited reference pose, not actually run the retargeting
 	// as long as the retargeter is reinitialized after every modification to the limb rotation offsets,
 	// then the TargetSkeleton.RetargetGlobalPose will contain the updated retarget pose.
-	if (RetargeterAsset->bEditRetargetPoseMode && RetargeterAsset->TargetIKRigAsset)
+	if (RetargeterAsset->IsInEditRetargetPoseMode() && RetargeterAsset->GetTargetIKRig())
 	{
-		const FName RootBoneName = RetargeterAsset->TargetIKRigAsset->GetRetargetRoot();
-		TargetSkeleton.GenerateRetargetPose(&RetargeterAsset->RetargetPoses[RetargeterAsset->CurrentRetargetPose], RootBoneName);
+		const FName RootBoneName = RetargeterAsset->GetTargetIKRig()->GetRetargetRoot();
+		TargetSkeleton.GenerateRetargetPose(RetargeterAsset->GetCurrentRetargetPose(), RootBoneName);
 		return TargetSkeleton.RetargetGlobalPose; 
 	}
 #endif
@@ -994,7 +1037,7 @@ TArray<FTransform>&  UIKRetargetProcessor::RunRetargeter(const TArray<FTransform
 	TargetSkeleton.OutputGlobalPose = TargetSkeleton.RetargetGlobalPose;
 
 	// ROOT retargeting
-	if (RetargeterAsset->bRetargetRoot)
+	if (RetargeterAsset->bRetargetRoot && bRootsInitialized)
 	{
 		RunRootRetarget(InSourceGlobalPose, TargetSkeleton.OutputGlobalPose);
 		// update global transforms below root
@@ -1002,7 +1045,7 @@ TArray<FTransform>&  UIKRetargetProcessor::RunRetargeter(const TArray<FTransform
 	}
 	
 	// FK CHAIN retargeting
-	if (RetargeterAsset->bRetargetFK)
+	if (RetargeterAsset->bRetargetFK && bAtLeastOneValidBoneChainPair)
 	{
 		RunFKRetarget(InSourceGlobalPose, TargetSkeleton.OutputGlobalPose);
 		// update all the bones that are not controlled by FK chains or root
@@ -1010,7 +1053,7 @@ TArray<FTransform>&  UIKRetargetProcessor::RunRetargeter(const TArray<FTransform
 	}
 	
 	// IK CHAIN retargeting
-	if (RetargeterAsset->bRetargetIK)
+	if (RetargeterAsset->bRetargetIK && bAtLeastOneValidBoneChainPair && bIKRigInitialized)
 	{
 		RunIKRetarget(InSourceGlobalPose, TargetSkeleton.OutputGlobalPose);
 	}
@@ -1067,7 +1110,9 @@ void UIKRetargetProcessor::RunIKRetarget(
 			OutIKGoal.EndEffectorPosition,
 			OutIKGoal.EndEffectorRotation,
 			1.0f,
-			1.0f);
+			1.0f,
+			EIKRigGoalSpace::Component,
+			EIKRigGoalSpace::Component);
 		IKRigProcessor->SetIKGoal(Goal);
 	}
 
@@ -1091,10 +1136,15 @@ FTransform UIKRetargetProcessor::GetTargetBoneRetargetPoseGlobalTransform(const
 void UIKRetargetProcessor::SetNeedsInitialized()
 {
 	bIsInitialized = false;
+	if (IKRigProcessor)
+	{
+		// may not be initialized yet (during setup as prerequisites are being created)
+		IKRigProcessor->SetNeedsInitialized();
+	}
 }
 
 void UIKRetargetProcessor::CopyTargetIKRigSettingsFromAsset()
 {
-	IKRigProcessor->CopyAllInputsFromSourceAssetAtRuntime(RetargeterAsset->TargetIKRigAsset);
+	IKRigProcessor->CopyAllInputsFromSourceAssetAtRuntime(RetargeterAsset->GetTargetIKRig());
 }
 #endif
diff --git a/Engine/Plugins/Animation/IKRig/Source/IKRig/Private/Retargeter/IKRetargeter.cpp b/Engine/Plugins/Animation/IKRig/Source/IKRig/Private/Retargeter/IKRetargeter.cpp
--- a/Engine/Plugins/Animation/IKRig/Source/IKRig/Private/Retargeter/IKRetargeter.cpp
+++ b/Engine/Plugins/Animation/IKRig/Source/IKRig/Private/Retargeter/IKRetargeter.cpp
@@ -3,7 +3,20 @@
 #include "Retargeter/IKRetargeter.h"
 
 const FName UIKRetargeter::DefaultPoseName = "Default Pose";
+const FName UIKRetargeter::GetSourceIKRigPropertyName() { return GET_MEMBER_NAME_STRING_CHECKED(UIKRetargeter, SourceIKRigAsset); };
+const FName UIKRetargeter::GetTargetIKRigPropertyName() { return GET_MEMBER_NAME_STRING_CHECKED(UIKRetargeter, TargetIKRigAsset); };
+#if WITH_EDITOR
+const FName UIKRetargeter::GetTargetPreviewMeshPropertyName() { return GET_MEMBER_NAME_STRING_CHECKED(UIKRetargeter, TargetPreviewMesh); };
+#endif
+const FName UIKRetargeter::GetDefaultPoseName() { return DefaultPoseName; }
 
+#if WITH_EDITOR
+void UIKRetargeter::PostEditUndo()
+{
+	Super::PostEditUndo();
+	IKRigEditUndo.Broadcast();
+};
+#endif
 
 void FIKRetargetPose::AddRotationDeltaToBone(FName BoneName, FQuat RotationDelta)
 {
diff --git a/Engine/Plugins/Animation/IKRig/Source/IKRig/Private/Solvers/IKRig_BodyMover.cpp b/Engine/Plugins/Animation/IKRig/Source/IKRig/Private/Solvers/IKRig_BodyMover.cpp
--- a/Engine/Plugins/Animation/IKRig/Source/IKRig/Private/Solvers/IKRig_BodyMover.cpp
+++ b/Engine/Plugins/Animation/IKRig/Source/IKRig/Private/Solvers/IKRig_BodyMover.cpp
@@ -244,7 +244,7 @@ bool UIKRig_BodyMover::IsGoalConnected(const FName& GoalName) const
 	return GetIndexOfGoal(GoalName) != INDEX_NONE;
 }
 
-UObject* UIKRig_BodyMover::GetEffectorWithGoal(const FName& GoalName) const
+UObject* UIKRig_BodyMover::GetGoalSettings(const FName& GoalName) const
 {
 	const int32 GoalIndex = GetIndexOfGoal(GoalName);
 	if (GoalIndex == INDEX_NONE)
diff --git a/Engine/Plugins/Animation/IKRig/Source/IKRig/Private/Solvers/IKRig_PBIKSolver.cpp b/Engine/Plugins/Animation/IKRig/Source/IKRig/Private/Solvers/IKRig_PBIKSolver.cpp
--- a/Engine/Plugins/Animation/IKRig/Source/IKRig/Private/Solvers/IKRig_PBIKSolver.cpp
+++ b/Engine/Plugins/Animation/IKRig/Source/IKRig/Private/Solvers/IKRig_PBIKSolver.cpp
@@ -130,6 +130,17 @@ void UIKRigPBIKSolver::Solve(FIKRigSkeleton& IKRigSkeleton, const FIKRigGoalCont
 	}
 }
 
+void UIKRigPBIKSolver::GetBonesWithSettings(TSet<FName>& OutBonesWithSettings) const
+{
+	for (UIKRig_PBIKBoneSettings* BoneSetting : BoneSettings)
+	{
+		if (BoneSetting)
+		{
+			OutBonesWithSettings.Add(BoneSetting->Bone);
+		}
+	}
+}
+
 #if WITH_EDITOR
 
 void UIKRigPBIKSolver::UpdateSolverSettings(UIKRigSolver* InSettings)
@@ -174,7 +185,7 @@ void UIKRigPBIKSolver::UpdateSolverSettings(UIKRigSolver* InSettings)
 
 FText UIKRigPBIKSolver::GetNiceName() const
 {
-	return FText(LOCTEXT("SolverName", "Full-Body IK"));
+	return FText(LOCTEXT("SolverName", "Full Body IK"));
 }
 
 bool UIKRigPBIKSolver::GetWarningMessage(FText& OutWarningMessage) const
@@ -251,7 +262,7 @@ void UIKRigPBIKSolver::SetRootBone(const FName& RootBoneName)
 	RootBone = RootBoneName;
 }
 
-UObject* UIKRigPBIKSolver::GetEffectorWithGoal(const FName& GoalName) const
+UObject* UIKRigPBIKSolver::GetGoalSettings(const FName& GoalName) const
 {
 	const int32 GoalIndex = GetIndexOfGoal(GoalName);
 	if (GoalIndex == INDEX_NONE)
diff --git a/Engine/Plugins/Animation/IKRig/Source/IKRig/Private/Solvers/IKRig_SetTransform.cpp b/Engine/Plugins/Animation/IKRig/Source/IKRig/Private/Solvers/IKRig_SetTransform.cpp
--- a/Engine/Plugins/Animation/IKRig/Source/IKRig/Private/Solvers/IKRig_SetTransform.cpp
+++ b/Engine/Plugins/Animation/IKRig/Source/IKRig/Private/Solvers/IKRig_SetTransform.cpp
@@ -112,7 +112,7 @@ bool UIKRig_SetTransform::IsGoalConnected(const FName& GoalName) const
 	return Goal == GoalName;
 }
 
-UObject* UIKRig_SetTransform::GetEffectorWithGoal(const FName& GoalName) const
+UObject* UIKRig_SetTransform::GetGoalSettings(const FName& GoalName) const
 {
 	return Goal == GoalName ? Effector : nullptr;
 }
diff --git a/Engine/Plugins/Animation/IKRig/Source/IKRig/Public/ActorComponents/IKRigComponent.h b/Engine/Plugins/Animation/IKRig/Source/IKRig/Public/ActorComponents/IKRigComponent.h
--- a/Engine/Plugins/Animation/IKRig/Source/IKRig/Public/ActorComponents/IKRigComponent.h
+++ b/Engine/Plugins/Animation/IKRig/Source/IKRig/Public/ActorComponents/IKRigComponent.h
@@ -26,7 +26,15 @@ public:
         const float PositionAlpha,
         const float RotationAlpha)
 	{
-		GoalContainer.SetIKGoal(FIKRigGoal(GoalName, Position, Rotation, PositionAlpha, RotationAlpha));
+		GoalContainer.SetIKGoal(
+			FIKRigGoal(
+				GoalName,
+				Position,
+				Rotation,
+				PositionAlpha,
+				RotationAlpha,
+				EIKRigGoalSpace::Component,
+				EIKRigGoalSpace::Component));
 	};
 
 	/** Set an IK Rig Goal transform (assumed in Component Space of Skeletal Mesh) with separate alpha values. */
@@ -39,7 +47,15 @@ public:
 	{
 		const FVector Position = Transform.GetTranslation();
 		const FQuat Rotation = Transform.GetRotation();
-		GoalContainer.SetIKGoal(FIKRigGoal(GoalName, Position, Rotation, PositionAlpha, RotationAlpha));
+		GoalContainer.SetIKGoal(
+			FIKRigGoal(
+				GoalName,
+				Position,
+				Rotation,
+				PositionAlpha,
+				RotationAlpha,
+				EIKRigGoalSpace::Component,
+				EIKRigGoalSpace::Component));
 	};
 	
 	/** Apply a IKRigGoal and store it on this rig. Goal transform assumed in Component Space of Skeletal Mesh. */
diff --git a/Engine/Plugins/Animation/IKRig/Source/IKRig/Public/IKRigDataTypes.h b/Engine/Plugins/Animation/IKRig/Source/IKRig/Public/IKRigDataTypes.h
--- a/Engine/Plugins/Animation/IKRig/Source/IKRig/Public/IKRigDataTypes.h
+++ b/Engine/Plugins/Animation/IKRig/Source/IKRig/Public/IKRigDataTypes.h
@@ -89,8 +89,8 @@ struct IKRIG_API FIKRigGoal
 	Rotation(FRotator::ZeroRotator),
 	PositionAlpha(1.0f),
     RotationAlpha(1.0f),
-	PositionSpace(EIKRigGoalSpace::Component),
-	RotationSpace(EIKRigGoalSpace::Component),
+	PositionSpace(EIKRigGoalSpace::Additive),
+	RotationSpace(EIKRigGoalSpace::Additive),
 	FinalBlendedPosition(ForceInitToZero),
 	FinalBlendedRotation(FQuat::Identity){}
 	
@@ -100,26 +100,28 @@ struct IKRIG_API FIKRigGoal
 	Rotation(FRotator::ZeroRotator),
     PositionAlpha(1.0f),
     RotationAlpha(1.0f),
-	PositionSpace(EIKRigGoalSpace::Component),
-    RotationSpace(EIKRigGoalSpace::Component),
+	PositionSpace(EIKRigGoalSpace::Additive),
+    RotationSpace(EIKRigGoalSpace::Additive),
 	FinalBlendedPosition(ForceInitToZero),
     FinalBlendedRotation(FQuat::Identity){}
 
 	FIKRigGoal(
-        const FName& Name,
-        const FVector& Position,
-        const FQuat& Rotation,
-        const float PositionAlpha,
-        const float RotationAlpha)
-        : Name(Name),
-        Position(Position),
-        Rotation(Rotation.Rotator()),
-        PositionAlpha(PositionAlpha),
-        RotationAlpha(RotationAlpha),
-		PositionSpace(EIKRigGoalSpace::Component),
-		RotationSpace(EIKRigGoalSpace::Component),
-		FinalBlendedPosition(Position),
-		FinalBlendedRotation(Rotation){}
+        const FName& InName,
+        const FVector& InPosition,
+        const FQuat& InRotation,
+        const float InPositionAlpha,
+        const float InRotationAlpha,
+        EIKRigGoalSpace InPositionSpace,
+        EIKRigGoalSpace InRotationSpace)
+        : Name(InName),
+        Position(InPosition),
+        Rotation(InRotation.Rotator()),
+        PositionAlpha(InPositionAlpha),
+        RotationAlpha(InRotationAlpha),
+		PositionSpace(InPositionSpace),
+		RotationSpace(InRotationSpace),
+		FinalBlendedPosition(InPosition),
+		FinalBlendedRotation(InRotation){}
 
 	FIKRigGoal(const UIKRigEffectorGoal* InGoal)
 		: Name(InGoal->GoalName),
@@ -166,6 +168,8 @@ public:
 	const TArray<FIKRigGoal>& GetGoalArray() const { return Goals; };
 	
 private:
+
+	FIKRigGoal* FindGoalWriteable(const FName& GoalName) const;
 	
 	/** Array of Goals. Cannot contain duplicates (name is key). */
 	TArray<FIKRigGoal> Goals;
diff --git a/Engine/Plugins/Animation/IKRig/Source/IKRig/Public/IKRigDefinition.h b/Engine/Plugins/Animation/IKRig/Source/IKRig/Public/IKRigDefinition.h
--- a/Engine/Plugins/Animation/IKRig/Source/IKRig/Public/IKRigDefinition.h
+++ b/Engine/Plugins/Animation/IKRig/Source/IKRig/Public/IKRigDefinition.h
@@ -145,7 +145,19 @@ public:
 	/** The thickness of the Goals in the editor viewport.*/
 	UPROPERTY(EditAnywhere, Category = "Goal Drawing",  meta = (ClampMin = "0.01", UIMin = "0.0", UIMax = "10.0"))
 	float GoalThickness = 0.7f;
+
+	/** The controller responsible for managing this asset's data (all editor mutation goes through this) */
+	UPROPERTY(Transient)
+	TObjectPtr<UObject> Controller;
+	
+#endif
+
+#if WITH_EDITOR
+	// This delegate is called when an edit operation is undone on the rig asset.
+	DECLARE_MULTICAST_DELEGATE(OnIKRigEditUndo);
+	OnIKRigEditUndo IKRigEditUndo;
 	
+	virtual void PostEditUndo() override;
 #endif
 
 	/** The skeletal mesh that was used as the source of the skeleton data. Also used for preview.
diff --git a/Engine/Plugins/Animation/IKRig/Source/IKRig/Public/IKRigProcessor.h b/Engine/Plugins/Animation/IKRig/Source/IKRig/Public/IKRigProcessor.h
--- a/Engine/Plugins/Animation/IKRig/Source/IKRig/Public/IKRigProcessor.h
+++ b/Engine/Plugins/Animation/IKRig/Source/IKRig/Public/IKRigProcessor.h
@@ -43,10 +43,10 @@ public:
 	 *  @param InRigAsset - the IK Rig defining the collection of solvers to execute and all the rig settings
 	 *  @param InputSkeleton - the skeleton in reference pose that you want to solve the IK on
 	 */
-	void Initialize(UIKRigDefinition* InRigAsset, const FIKRigInputSkeleton& InputSkeleton);
+	void Initialize(const UIKRigDefinition* InRigAsset, const FIKRigInputSkeleton& InputSkeleton);
 
 	/** Anim Graph: convenience to initialize directly from an FReferenceSkeleton. */
-	void Initialize(UIKRigDefinition* InRigAsset, const FReferenceSkeleton& RefSkeleton);
+	void Initialize(const UIKRigDefinition* InRigAsset, const FReferenceSkeleton& RefSkeleton);
 
 	//
 	// BEGIN UPDATE SEQUENCE FUNCTIONS
@@ -86,13 +86,17 @@ public:
 	/** Get access to the internal skeleton data */
 	FIKRigSkeleton& GetSkeleton();
 
+
+	/** Used to determine if the IK Rig asset is compatible with a given skeleton. */
+	static bool IsIKRigCompatibleWithSkeleton(const UIKRigDefinition* InRigAsset, const FIKRigInputSkeleton& InputSkeleton);
+
 	bool IsInitialized() const { return bInitialized; };
 
 	void SetNeedsInitialized();
 
 #if WITH_EDITOR
 	/** Used to propagate setting values from the source asset at runtime (settings that do not require re-initialization) */
-	void CopyAllInputsFromSourceAssetAtRuntime(UIKRigDefinition* SourceAsset);
+	void CopyAllInputsFromSourceAssetAtRuntime(const UIKRigDefinition* SourceAsset);
 #endif
 	
 private:
diff --git a/Engine/Plugins/Animation/IKRig/Source/IKRig/Public/IKRigSkeleton.h b/Engine/Plugins/Animation/IKRig/Source/IKRig/Public/IKRigSkeleton.h
--- a/Engine/Plugins/Animation/IKRig/Source/IKRig/Public/IKRigSkeleton.h
+++ b/Engine/Plugins/Animation/IKRig/Source/IKRig/Public/IKRigSkeleton.h
@@ -3,11 +3,13 @@
 #pragma once
 
 #include "CoreMinimal.h"
+
 #include "ReferenceSkeleton.h"
 
 #include "IKRigSkeleton.generated.h"
 
 class UIKRigDefinition;
+class UIKRigSolver;
 
 USTRUCT()
 struct IKRIG_API FIKRigSkeletonChain
@@ -29,6 +31,8 @@ struct IKRIG_API FIKRigSkeletonChain
  *
  * You can however add additional bones, change the reference pose (including proportions) and the bone indices.
  * This allows you to run the same IK Rig asset on somewhat different skeletal meshes.
+ *
+ * To validate compatibility use UIKRigProcess::IsIKRigCompatibleWithSkeleton()
  */
 USTRUCT()
 struct IKRIG_API FIKRigInputSkeleton
@@ -39,7 +43,14 @@ struct IKRIG_API FIKRigInputSkeleton
 	TArray<int32> ParentIndices;
 	TArray<FTransform> LocalRefPose;
 
-	void InitializeFromRefSkeleton(const FReferenceSkeleton& RefSkeleton)
+	FIKRigInputSkeleton() = default;
+	
+	FIKRigInputSkeleton(const FReferenceSkeleton& RefSkeleton)
+	{
+		Initialize(RefSkeleton);
+	}
+
+	void Initialize(const FReferenceSkeleton& RefSkeleton)
 	{
 		Reset();
 		
@@ -89,9 +100,9 @@ struct IKRIG_API FIKRigSkeleton
 	UPROPERTY(VisibleAnywhere, Category = Skeleton)
 	TArray<FTransform> RefPoseGlobal;
 
-	void Initialize(const FIKRigInputSkeleton& InputSkeleton, const TArray<FName>& InExcludedBones);
+	void SetInputSkeleton(const FReferenceSkeleton& RefSkeleton, const TArray<FName>& InExcludedBones);
 
-	void Initialize(const FReferenceSkeleton& RefSkeleton, const TArray<FName>& InExcludedBones);
+	void SetInputSkeleton(const FIKRigInputSkeleton& InputSkeleton, const TArray<FName>& InExcludedBones);
 
 	void Reset();
 	
@@ -103,8 +114,6 @@ struct IKRIG_API FIKRigSkeleton
 
 	int32 GetParentIndexThatIsNotExcluded(const int32 BoneIndex) const;
 
-	bool CopyPosesFromInputSkeleton(const FIKRigInputSkeleton& InputSkeleton);
-
 	static void ConvertLocalPoseToGlobal(
 		const TArray<int32>& InParentIndices,
 		const TArray<FTransform>& InLocalPose,
diff --git a/Engine/Plugins/Animation/IKRig/Source/IKRig/Public/IKRigSolver.h b/Engine/Plugins/Animation/IKRig/Source/IKRig/Public/IKRigSolver.h
--- a/Engine/Plugins/Animation/IKRig/Source/IKRig/Public/IKRigSolver.h
+++ b/Engine/Plugins/Animation/IKRig/Source/IKRig/Public/IKRigSolver.h
@@ -26,8 +26,17 @@ public:
 	virtual void Solve(FIKRigSkeleton& IKRigSkeleton, const FIKRigGoalContainer& Goals) PURE_VIRTUAL("Solve");
 	//** END RUNTIME */
 
-#if WITH_EDITORONLY_DATA
+	//** ROOT BONE (optional, implement if your solver requires a root bone) */
+	/** if solver requires a root bone, then override this to return it. */
+	virtual FName GetRootBone() const { return NAME_None; };
+	/** override to support telling outside systems which bones this solver has setting for.
+	* NOTE: This must be overriden on solvers that use bone settings.
+	* NOTE: Only ADD to the incoming set, do not remove from it. */
+	virtual void GetBonesWithSettings(TSet<FName>& OutBonesWithSettings) const {};
 	
+	//** END ROOT BONE */
+
+#if WITH_EDITORONLY_DATA
 	/** callback whenever this solver is edited */
 	DECLARE_EVENT_OneParam(UIKRigSolver, FIKRigSolverModified, UIKRigSolver*);
 	FIKRigSolverModified& OnSolverModified(){ return IKRigSolverModified; };
@@ -60,14 +69,14 @@ public:
 	virtual void SetGoalBone(const FName& GoalName, const FName& NewBoneName) PURE_VIRTUAL("SetGoalBone");
 	/** override to support QUERY for a connected goal */
 	virtual bool IsGoalConnected(const FName& GoalName) const {return false;};
-	/** override to support supplying effector settings to outside systems for editing/UI */
-	virtual UObject* GetEffectorWithGoal(const FName& GoalName) const {return nullptr;};
+	/** override to support supplying goals settings specific to this solver to outside systems for editing/UI */
+	virtual UObject* GetGoalSettings(const FName& GoalName) const {return nullptr;};
 	//** END GOALS */
 
 	//** ROOT BONE (optional, implement if your solver requires a root bone) */
 	/** override to support SETTING ROOT BONE for the solver */
 	virtual void SetRootBone(const FName& RootBoneName){};
-	virtual bool CanSetRootBone() const { return false; };
+	virtual bool RequiresRootBone() const { return false; };
 	//** END ROOT BONE */
 
 	//** BONE SETTINGS (optional, implement if your solver supports per-bone settings) */
@@ -75,11 +84,13 @@ public:
 	virtual void AddBoneSetting(const FName& BoneName){};
 	/** override to support ADDING PER-BONE settings for this solver */
 	virtual void RemoveBoneSetting(const FName& BoneName){};
-	/** override to support supplying per-bone settings to outside systems for editing/UI */
-	virtual UObject* GetBoneSetting(const FName& BoneName) const { return nullptr; };
+	/** override to support supplying per-bone settings to outside systems for editing/UI
+	 ** NOTE: This must be overriden on solvers that use bone settings.*/
+	virtual UObject* GetBoneSetting(const FName& BoneName) const { ensure(!UsesBoneSettings()); return nullptr; };
+	
 	/** override to tell systems if this solver supports per-bone settings */
 	virtual bool UsesBoneSettings() const { return false;};
-	/** override to draw custom per-bone settings in the editor viewport */
+	/** todo override to draw custom per-bone settings in the editor viewport */
 	virtual void DrawBoneSettings(const FName& BoneName, const FIKRigSkeleton& IKRigSkeleton, FPrimitiveDrawInterface* PDI) const {};
 	/** return true if the supplied Bone is affected by this solver - this provides UI feedback for user */
 	virtual bool IsBoneAffectedBySolver(const FName& BoneName, const FIKRigSkeleton& IKRigSkeleton) const { return false; };
@@ -97,6 +108,6 @@ private:
 
 	UPROPERTY()
 	bool bIsEnabled = true;
-
 #endif
+
 };
diff --git a/Engine/Plugins/Animation/IKRig/Source/IKRig/Public/Retargeter/IKRetargetProcessor.h b/Engine/Plugins/Animation/IKRig/Source/IKRig/Public/Retargeter/IKRetargetProcessor.h
--- a/Engine/Plugins/Animation/IKRig/Source/IKRig/Public/Retargeter/IKRetargetProcessor.h
+++ b/Engine/Plugins/Animation/IKRig/Source/IKRig/Public/Retargeter/IKRetargetProcessor.h
@@ -26,14 +26,11 @@ struct IKRIG_API FRetargetSkeleton
 	
 	USkeletalMesh* SkeletalMesh;
 
-	void Initialize(
-		USkeletalMesh* InSkeletalMesh, 
-		FIKRetargetPose* InRetargetPose,
-		const FName& RetargetRootBone);
+	void Initialize(USkeletalMesh* InSkeletalMesh);
 
 	void Reset();
 
-	void GenerateRetargetPose(FIKRetargetPose* InRetargetPose, const FName& RetargetRootBone);
+	void GenerateRetargetPose();
 
 	int32 FindBoneIndexByName(const FName InName) const;
 
@@ -66,11 +63,15 @@ struct FTargetSkeleton : public FRetargetSkeleton
 {
 	TArray<FTransform> OutputGlobalPose;
 	TArray<bool> IsBoneRetargeted;
+	TArray<bool> IsBoneInAnyTargetChain;
 
 	void Initialize(
 		USkeletalMesh* InSkeletalMesh, 
-		FIKRetargetPose* RetargetPose,
-		const FName& RetargetRootBone);
+		const FIKRetargetPose* RetargetPose,
+		const FName& RetargetRootBone,
+		const TArray<FBoneChain>& TargetChains);
+
+	void GenerateRetargetPose(const FIKRetargetPose* InRetargetPose, const FName& RetargetRootBone);
 
 	void Reset();
 
@@ -79,6 +80,24 @@ struct FTargetSkeleton : public FRetargetSkeleton
 	void UpdateGlobalTransformsAllNonRetargetedBones(TArray<FTransform>& InOutGlobalPose);
 };
 
+/** resolving an FBoneChain to an actual skeleton, used to validate compatibility and get all chain indices */
+struct IKRIG_API FResolvedBoneChain
+{
+	FResolvedBoneChain(const FBoneChain& BoneChain, const FRetargetSkeleton& Skeleton, TArray<int32> &OutBoneIndices);
+
+	/* Does the START bone exist in the skeleton? */
+	bool bFoundStartBone = false;
+	/* Does the END bone exist in the skeleton? */
+	bool bFoundEndBone = false;
+	/* Is the END bone a child of the START bone? */
+	bool bEndIsChildOfStart = false;
+
+	bool IsValid() const
+	{
+		return bFoundStartBone && bFoundEndBone && bEndIsChildOfStart;
+	}
+};
+
 struct FRootSource
 {
 	int32 BoneIndex;
@@ -341,6 +360,15 @@ private:
 	/** Only true once Initialize() has successfully completed.*/
 	bool bIsInitialized = false;
 
+	/** true when roots are able to be retargeted */
+	bool bRootsInitialized = false;
+
+	/** true when at least one pair of bone chains is able to be retargeted */
+	bool bAtLeastOneValidBoneChainPair = false;
+
+	/** true when roots are able to be retargeted */
+	bool bIKRigInitialized = false;
+
 	/** The source asset this processor was initialized with. */
 	UIKRetargeter* RetargeterAsset = nullptr;
 
@@ -362,7 +390,7 @@ private:
 
 	/** The Source/Target pair of Root Bones retargeted with scaled translation */
 	FRootRetargeter RootRetargeter;
-
+	
 	/** Initializes the FRootRetargeter */
 	bool InitializeRoots();
 
diff --git a/Engine/Plugins/Animation/IKRig/Source/IKRig/Public/Retargeter/IKRetargeter.h b/Engine/Plugins/Animation/IKRig/Source/IKRig/Public/Retargeter/IKRetargeter.h
--- a/Engine/Plugins/Animation/IKRig/Source/IKRig/Public/Retargeter/IKRetargeter.h
+++ b/Engine/Plugins/Animation/IKRig/Source/IKRig/Public/Retargeter/IKRetargeter.h
@@ -72,13 +72,51 @@ struct IKRIG_API FIKRetargetPose
 	void AddTranslationDeltaToRoot(FVector TranslateDelta);
 };
 
+
 UCLASS(Blueprintable)
 class IKRIG_API UIKRetargeter : public UObject
 {
 	GENERATED_BODY()
-	
 public:
 
+	/** Get read-only access to the source IK Rig asset */
+	const UIKRigDefinition* GetSourceIKRig() const { return SourceIKRigAsset; };
+	/** Get read-only access to the target IK Rig asset */
+	const UIKRigDefinition* GetTargetIKRig() const { return TargetIKRigAsset; };
+	/** Get read-write access to the source IK Rig asset.
+	 * WARNING: do not use for editing the data model. Use Controller class instead. */
+	 UIKRigDefinition* GetSourceIKRigWriteable() const { return SourceIKRigAsset; };
+	/** Get read-write access to the target IK Rig asset.
+	 * WARNING: do not use for editing the data model. Use Controller class instead. */
+	UIKRigDefinition* GetTargetIKRigWriteable() const { return TargetIKRigAsset; };
+	/** Get read-only access to the chain mapping */
+	const TArray<FRetargetChainMap>& GetChainMapping() const { return ChainMapping; };
+	/** Get read-only access to a retarget pose */
+	const FIKRetargetPose* GetCurrentRetargetPose() const { return &RetargetPoses[CurrentRetargetPose]; };
+
+	/* Get name of Source IK Rig property */
+	static const FName GetSourceIKRigPropertyName();
+	/* Get name of Target IK Rig property */
+	static const FName GetTargetIKRigPropertyName();
+#if WITH_EDITOR
+	/* Get name of Target Preview Mesh property */
+	static const FName GetTargetPreviewMeshPropertyName();
+#endif
+	/* Get name of default pose */
+	static const FName GetDefaultPoseName();
+
+#if WITH_EDITOR
+	bool IsInEditRetargetPoseMode() const { return bEditRetargetPoseMode; };
+
+	// This delegate is called when an edit operation is undone on the rig asset.
+	DECLARE_MULTICAST_DELEGATE(OnIKRigEditUndo);
+	OnIKRigEditUndo IKRigEditUndo;
+	
+	virtual void PostEditUndo() override;
+#endif
+
+private:
+
 	/** The rig to copy animation FROM.*/
 	UPROPERTY(VisibleAnywhere, Category = Rigs)
 	TObjectPtr<UIKRigDefinition> SourceIKRigAsset = nullptr;
@@ -87,6 +125,14 @@ public:
 	UPROPERTY(EditAnywhere, Category = Rigs)
 	TObjectPtr<UIKRigDefinition> TargetIKRigAsset = nullptr;
 
+public:
+
+#if WITH_EDITORONLY_DATA
+	/** The Skeletal Mesh to preview the retarget on.*/
+	UPROPERTY(EditAnywhere, Category = Rigs)
+	TObjectPtr<USkeletalMesh> TargetPreviewMesh = nullptr;
+#endif
+	
 	/** When false, translational motion of skeleton root is not copied. Useful for debugging.*/
 	UPROPERTY(EditAnywhere, Category = RetargetPhases)
 	bool bRetargetRoot = true;
@@ -116,6 +162,7 @@ public:
 	UPROPERTY(EditAnywhere, Category = PoseEditSettings, meta = (ClampMin = "0.0", UIMin = "0.01", UIMax = "10.0"))
 	float BoneDrawThickness = 1.0f;
 
+private:
 	/** A special editor-only mode which forces the retargeter to output the current retarget reference pose,
 	* rather than actually running the retarget and outputting the retargeted pose. Used in Edit-Pose mode.*/
 	UPROPERTY()
@@ -125,7 +172,8 @@ public:
 	UPROPERTY(Transient)
 	TObjectPtr<UObject> Controller;
 #endif
-	
+
+private:
 	/** The set of retarget poses available as options for retargeting.*/
 	UPROPERTY(VisibleAnywhere, Category = EditRetargetPoses)
 	TMap<FName, FIKRetargetPose> RetargetPoses;
@@ -138,6 +186,8 @@ public:
 	UPROPERTY()
 	FName CurrentRetargetPose = DefaultPoseName;
 	static const FName DefaultPoseName;
+
+	friend class UIKRetargeterController;
 };
 
 
diff --git a/Engine/Plugins/Animation/IKRig/Source/IKRig/Public/Solvers/IKRig_BodyMover.h b/Engine/Plugins/Animation/IKRig/Source/IKRig/Public/Solvers/IKRig_BodyMover.h
--- a/Engine/Plugins/Animation/IKRig/Source/IKRig/Public/Solvers/IKRig_BodyMover.h
+++ b/Engine/Plugins/Animation/IKRig/Source/IKRig/Public/Solvers/IKRig_BodyMover.h
@@ -86,6 +86,8 @@ public:
 	/** UIKRigSolver interface */
 	virtual void Initialize(const FIKRigSkeleton& IKRigSkeleton) override;
 	virtual void Solve(FIKRigSkeleton& IKRigSkeleton, const FIKRigGoalContainer& Goals) override;
+
+	virtual FName GetRootBone() const override { return BodyBone; };
 #if WITH_EDITOR
 	virtual void UpdateSolverSettings(UIKRigSolver* InSettings) override;
 	virtual FText GetNiceName() const override;
@@ -97,9 +99,9 @@ public:
 	virtual void RenameGoal(const FName& OldName, const FName& NewName) override;
 	virtual void SetGoalBone(const FName& GoalName, const FName& NewBoneName) override;
 	virtual bool IsGoalConnected(const FName& GoalName) const override;
-	virtual UObject* GetEffectorWithGoal(const FName& GoalName) const override;
+	virtual UObject* GetGoalSettings(const FName& GoalName) const override;
 	// root bone can be set on this solver
-	virtual bool CanSetRootBone() const override { return true; };
+	virtual bool RequiresRootBone() const override { return true; };
 	virtual void SetRootBone(const FName& RootBoneName) override;
 	/** END UIKRigSolver interface */
 #endif
diff --git a/Engine/Plugins/Animation/IKRig/Source/IKRig/Public/Solvers/IKRig_PBIKSolver.h b/Engine/Plugins/Animation/IKRig/Source/IKRig/Public/Solvers/IKRig_PBIKSolver.h
--- a/Engine/Plugins/Animation/IKRig/Source/IKRig/Public/Solvers/IKRig_PBIKSolver.h
+++ b/Engine/Plugins/Animation/IKRig/Source/IKRig/Public/Solvers/IKRig_PBIKSolver.h
@@ -15,19 +15,31 @@ class IKRIG_API UIKRig_FBIKEffector : public UObject
 	GENERATED_BODY()
 
 public:
-
+	/** The Goal that is driving this effector's transform. */
 	UPROPERTY(VisibleAnywhere, Category = "Full Body IK Effector")
 	FName GoalName;
-
+	
+	/** The bone that this effector will pull on. */
 	UPROPERTY(VisibleAnywhere, Category = "Full Body IK Effector")
 	FName BoneName;
-	
+
+	/** Range 0-1 (default is 1.0). The strength of the effector when pulling the bone towards it's target location.
+	* At 0.0, the effector does not pull at all, but the bones between the effector and the root will still slightly resist motion from other effectors.
+	* This can thus act as a "stabilizer" for parts of the body that you do not want to behave in a pure FK fashion.
+	*/
 	UPROPERTY(EditAnywhere, Category = "Full Body IK Effector", meta = (ClampMin = "0", ClampMax = "1", UIMin = "0.0", UIMax = "1.0"))
 	float StrengthAlpha = 1.0f;
-	
+
+	/** Range 0-1 (default is 1.0). When enabled (greater than 0.0), the solver internally partitions the skeleton into 'chains' which extend from the effector to the nearest fork in the skeleton.
+	*These chains are pre-rotated and translated, as a whole, towards the effector targets.
+	*This can improve the results for sparse bone chains, and significantly improve convergence on dense bone chains.
+	*But it may cause undesirable results in highly constrained bone chains (like robot arms).
+	*/
 	UPROPERTY(EditAnywhere, Category = "Full Body IK Effector", meta = (ClampMin = "0", ClampMax = "1", UIMin = "0.0", UIMax = "1.0"))
 	float PullChainAlpha = 1.0f;
 
+	/** Range 0-1 (default is 1.0).
+	*Blends the effector bone rotation between the rotation of the effector transform (1.0) and the rotation of the input bone (0.0).*/
 	UPROPERTY(EditAnywhere, Category = "Full Body IK Effector", meta = (ClampMin = "0", ClampMax = "1", UIMin = "0.0", UIMax = "1.0"))
 	float PinRotation = 1.0f;
 
@@ -56,37 +68,65 @@ public:
 		Z(EPBIKLimitType::Free),
 		PreferredAngles(FVector::ZeroVector){}
 
+	/** The bone these settings are applied to. */
 	UPROPERTY(VisibleAnywhere, Category = Bone, meta = (Constant, CustomWidget = "BoneName"))
 	FName Bone;
-
+	
+	/** Range is 0 to 1 (Default is 0). At higher values, the bone will resist rotating (forcing other bones to compensate). */
 	UPROPERTY(EditAnywhere, Category = Stiffness, meta = (ClampMin = "0", ClampMax = "1", UIMin = "0.0", UIMax = "1.0"))
 	float RotationStiffness = 0.0f;
+
+	/** Range is 0 to 1 (Default is 0). At higher values, the bone will resist translational motion (forcing other bones to compensate). */
 	UPROPERTY(EditAnywhere, Category = Stiffness, meta = (ClampMin = "0", ClampMax = "1", UIMin = "0.0", UIMax = "1.0"))
 	float PositionStiffness = 0.0f;
 
+	/** Limit the rotation angle of the bone on the X axis.
+	 *Free: can rotate freely in this axis.
+	 *Limited: rotation is clamped between the min/max angles relative to the Skeletal Mesh reference pose.
+	 *Locked: no rotation is allowed in this axis (will remain at reference pose angle). */
 	UPROPERTY(EditAnywhere, Category = Limits)
 	EPBIKLimitType X;
-	UPROPERTY(EditAnywhere, Category = Limits, meta = (ClampMin = "-180", ClampMax = "180", UIMin = "-180.0", UIMax = "180.0"))
+	/**Range is -180 to 0 (Default is 0). Degrees of rotation in the negative X direction to allow when joint is in "Limited" mode. */
+	UPROPERTY(EditAnywhere, Category = Limits, meta = (ClampMin = "-180", ClampMax = "0", UIMin = "-180.0", UIMax = "0.0"))
 	float MinX = 0.0f;
+	/**Range is 0 to 180 (Default is 0). Degrees of rotation in the positive X direction to allow when joint is in "Limited" mode. */
 	UPROPERTY(EditAnywhere, Category = Limits, meta = (ClampMin = "0", ClampMax = "180", UIMin = "0.0", UIMax = "180.0"))
 	float MaxX = 0.0f;
 
+	/** Limit the rotation angle of the bone on the Y axis.
+	*Free: can rotate freely in this axis.
+	*Limited: rotation is clamped between the min/max angles relative to the Skeletal Mesh reference pose.
+	*Locked: no rotation is allowed in this axis (will remain at input pose angle). */
 	UPROPERTY(EditAnywhere, Category = Limits)
 	EPBIKLimitType Y;
-	UPROPERTY(EditAnywhere, Category = Limits, meta = (ClampMin = "-180", ClampMax = "180", UIMin = "-180.0", UIMax = "180.0"))
+	/**Range is -180 to 0 (Default is 0). Degrees of rotation in the negative Y direction to allow when joint is in "Limited" mode. */
+	UPROPERTY(EditAnywhere, Category = Limits, meta = (ClampMin = "-180", ClampMax = "0", UIMin = "-180.0", UIMax = "0.0"))
 	float MinY = 0.0f;
+	/**Range is 0 to 180 (Default is 0). Degrees of rotation in the positive Y direction to allow when joint is in "Limited" mode. */
 	UPROPERTY(EditAnywhere, Category = Limits, meta = (ClampMin = "0", ClampMax = "180", UIMin = "0.0", UIMax = "180.0"))
 	float MaxY = 0.0f;
 
+	/** Limit the rotation angle of the bone on the Z axis.
+	*Free: can rotate freely in this axis.
+	*Limited: rotation is clamped between the min/max angles relative to the Skeletal Mesh reference pose.
+	*Locked: no rotation is allowed in this axis (will remain at input pose angle). */
 	UPROPERTY(EditAnywhere, Category = Limits)
 	EPBIKLimitType Z;
-	UPROPERTY(EditAnywhere, Category = Limits, meta = (ClampMin = "-180", ClampMax = "180", UIMin = "-180.0", UIMax = "180.0"))
+	/**Range is -180 to 0 (Default is 0). Degrees of rotation in the negative Z direction to allow when joint is in "Limited" mode. */
+	UPROPERTY(EditAnywhere, Category = Limits, meta = (ClampMin = "-180", ClampMax = "0", UIMin = "-180.0", UIMax = "0.0"))
 	float MinZ = 0.0f;
+	/**Range is 0 to 180 (Default is 0). Degrees of rotation in the positive Z direction to allow when joint is in "Limited" mode. */
 	UPROPERTY(EditAnywhere, Category = Limits, meta = (ClampMin = "0", ClampMax = "180", UIMin = "0.0", UIMax = "180.0"))
 	float MaxZ = 0.0f;
 
+	/**When true, this bone will "prefer" to rotate in the direction specified by the Preferred Angles when the chain it belongs to is compressed.
+	 * Preferred Angles should be the first method used to fix bones that bend in the wrong direction (rather than limits).
+	 * The resulting angles can be visualized on their own by temporarily setting the Solver iterations to 0 and moving the effectors.*/
 	UPROPERTY(EditAnywhere, Category = PreferredAngles)
 	bool bUsePreferredAngles = false;
+	/**The local Euler angles (in degrees) used to rotate this bone when the chain it belongs to is squashed.
+	 * This happens by moving the effector at the tip of the chain towards the root of the chain.
+	 * This can be used to coerce knees and elbows to bend in the anatomically "correct" direction without resorting to limits (which may require more iterations to converge).*/
 	UPROPERTY(EditAnywhere, Category = PreferredAngles)
 	FVector PreferredAngles;
 
@@ -134,35 +174,40 @@ class IKRIG_API UIKRigPBIKSolver : public UIKRigSolver
 
 public:
 
-	UPROPERTY(VisibleAnywhere, Category = "Full-Body IK Settings")
+	/** All bones above this bone in the hierarchy will be completely ignored by the solver. Typically this is set to
+	 * the top-most skinned bone in the Skeletal Mesh (ie Pelvis, Hips etc), NOT the actual root of the entire skeleton.
+	 *
+	 * If you want to use the solver on a single chain of bones, and NOT translate the chain, ensure that "PinRoot" is
+	 * checked on to disable the root from translating to reach the effector goals.*/
+	UPROPERTY(VisibleAnywhere, Category = "Full Body IK Settings")
 	FName RootBone;
 
 	/** High iteration counts can help solve complex joint configurations with competing constraints, but will increase runtime cost. Default is 20. */
-	UPROPERTY(EditAnywhere, Category = "Full-Body IK Settings", meta = (ClampMin = "0", UIMin = "0.0", UIMax = "100.0"))
+	UPROPERTY(EditAnywhere, Category = "Full Body IK Settings", meta = (ClampMin = "0", ClampMax = "1000", UIMin = "0.0", UIMax = "100.0"))
 	int32 Iterations = 20;
 
 	/** A global mass multiplier; higher values will make the joints more stiff, but require more iterations. Typical range is 0.0 to 10.0. */
-	UPROPERTY(EditAnywhere, Category = "Full-Body IK Settings", meta = (ClampMin = "0", UIMin = "0.0", UIMax = "10.0"))
+	UPROPERTY(EditAnywhere, Category = "Full Body IK Settings", meta = (ClampMin = "0", UIMin = "0.0", UIMax = "10.0"))
 	float MassMultiplier = 1.0f;
 
 	/** Set this as low as possible while keeping the solve stable. Lower values improve convergence of effector targets. Default is 0.2. */
-	UPROPERTY(EditAnywhere, Category = "Full-Body IK Settings", meta = (ClampMin = "0", UIMin = "0.0", UIMax = "10.0"))
+	UPROPERTY(EditAnywhere, Category = "Full Body IK Settings", meta = (ClampMin = "0", UIMin = "0.0", UIMax = "10.0"))
 	float MinMassMultiplier = 0.2f;
 
 	/** If true, joints will translate to reach the effectors; causing bones to lengthen if necessary. Good for cartoon effects. Default is false. */
-	UPROPERTY(EditAnywhere, Category = "Full-Body IK Settings")
+	UPROPERTY(EditAnywhere, Category = "Full Body IK Settings")
 	bool bAllowStretch = false;
 
 	/** If true, solver will pull entire skeleton towards effectors BEFORE running constraint iterations. Default is true. */
-	UPROPERTY(EditAnywhere, Category = "Full-Body IK Settings")
+	UPROPERTY(EditAnywhere, Category = "Full Body IK Settings")
 	bool bPrePullRoot = true;
 
 	/** Lock the position and rotation of the solver root bone in-place (at animated position). Useful for partial-body solves. Default is false. */
-	UPROPERTY(EditAnywhere, Category = "Full-Body IK Settings")
+	UPROPERTY(EditAnywhere, Category = "Full Body IK Settings")
 	bool bPinRoot = false;
 
 	/** When true, the solver is reset each tick to start from the current input pose. If false, incoming animated poses are ignored and the solver starts from the results of the previous solve. Default is true. */
-	UPROPERTY(EditAnywhere, Category = "Full-Body IK Settings")
+	UPROPERTY(EditAnywhere, Category = "Full Body IK Settings")
 	bool bStartSolveFromInputPose = true;
 	
 	UPROPERTY()
@@ -175,6 +220,10 @@ public:
 	// runtime
 	virtual void Initialize(const FIKRigSkeleton& IKRigSkeleton) override;
 	virtual void Solve(FIKRigSkeleton& IKRigSkeleton, const FIKRigGoalContainer& Goals) override;
+
+	virtual FName GetRootBone() const override { return RootBone; };
+	virtual void GetBonesWithSettings(TSet<FName>& OutBonesWithSettings) const override;
+	
 #if WITH_EDITOR
 	virtual void UpdateSolverSettings(UIKRigSolver* InSettings) override;
 	virtual FText GetNiceName() const override;
@@ -185,7 +234,7 @@ public:
 	virtual void RenameGoal(const FName& OldName, const FName& NewName) override;
 	virtual void SetGoalBone(const FName& GoalName, const FName& NewBoneName) override;
 	virtual bool IsGoalConnected(const FName& GoalName) const override;
-	virtual UObject* GetEffectorWithGoal(const FName& GoalName) const override;
+	virtual UObject* GetGoalSettings(const FName& GoalName) const override;
 	// bone settings
 	virtual void AddBoneSetting(const FName& BoneName) override;
 	virtual void RemoveBoneSetting(const FName& BoneName) override;
@@ -193,7 +242,7 @@ public:
 	virtual bool UsesBoneSettings() const override{ return true;};
 	virtual void DrawBoneSettings(const FName& BoneName, const FIKRigSkeleton& IKRigSkeleton, FPrimitiveDrawInterface* PDI) const override;
 	// root bone can be set on this solver
-	virtual bool CanSetRootBone() const override { return true; };
+	virtual bool RequiresRootBone() const override { return true; };
 	virtual void SetRootBone(const FName& RootBoneName) override;
 	virtual bool IsBoneAffectedBySolver(const FName& BoneName, const FIKRigSkeleton& IKRigSkeleton) const override;
 	/** END UIKRigSolver interface */
diff --git a/Engine/Plugins/Animation/IKRig/Source/IKRig/Public/Solvers/IKRig_SetTransform.h b/Engine/Plugins/Animation/IKRig/Source/IKRig/Public/Solvers/IKRig_SetTransform.h
--- a/Engine/Plugins/Animation/IKRig/Source/IKRig/Public/Solvers/IKRig_SetTransform.h
+++ b/Engine/Plugins/Animation/IKRig/Source/IKRig/Public/Solvers/IKRig_SetTransform.h
@@ -57,7 +57,7 @@ public:
 	virtual void RenameGoal(const FName& OldName, const FName& NewName) override;
 	virtual void SetGoalBone(const FName& GoalName, const FName& NewBoneName) override;
 	virtual bool IsGoalConnected(const FName& GoalName) const override;
-	virtual UObject* GetEffectorWithGoal(const FName& GoalName) const override;
+	virtual UObject* GetGoalSettings(const FName& GoalName) const override;
 	virtual bool IsBoneAffectedBySolver(const FName& BoneName, const FIKRigSkeleton& IKRigSkeleton) const override;
 #endif
 	/** END UIKRigSolver interface */
diff --git a/Engine/Plugins/Animation/IKRig/Source/IKRigDeveloper/Private/AnimGraphNode_RetargetPoseFromMesh.cpp b/Engine/Plugins/Animation/IKRig/Source/IKRigDeveloper/Private/AnimGraphNode_RetargetPoseFromMesh.cpp
--- a/Engine/Plugins/Animation/IKRig/Source/IKRigDeveloper/Private/AnimGraphNode_RetargetPoseFromMesh.cpp
+++ b/Engine/Plugins/Animation/IKRig/Source/IKRigDeveloper/Private/AnimGraphNode_RetargetPoseFromMesh.cpp
@@ -62,25 +62,25 @@ void UAnimGraphNode_RetargetPoseFromMesh::ValidateAnimNodeDuringCompilation(USke
 	}
 
 	// validate SOURCE IK Rig asset has been assigned
-	if (!Node.IKRetargeterAsset->SourceIKRigAsset)
+	if (!Node.IKRetargeterAsset->GetSourceIKRig())
 	{
 		MessageLog.Warning(TEXT("@@ has IK Retargeter that is missing a source IK Rig asset."), this);
 	}
 
 	// validate TARGET IK Rig asset has been assigned
-	if (!Node.IKRetargeterAsset->TargetIKRigAsset)
+	if (!Node.IKRetargeterAsset->GetTargetIKRig())
 	{
 		MessageLog.Warning(TEXT("@@ has IK Retargeter that is missing a target IK Rig asset."), this);
 	}
 
-	if (!(Node.IKRetargeterAsset->SourceIKRigAsset && Node.IKRetargeterAsset->TargetIKRigAsset))
+	if (!(Node.IKRetargeterAsset->GetSourceIKRig() && Node.IKRetargeterAsset->GetTargetIKRig()))
 	{
 		return;
 	}
 	
 	// validate that target bone chains exist on this skeleton
 	const FReferenceSkeleton &RefSkel = ForSkeleton->GetReferenceSkeleton();
-	const TArray<FBoneChain> &TargetBoneChains = Node.IKRetargeterAsset->TargetIKRigAsset->GetRetargetChains();
+	const TArray<FBoneChain> &TargetBoneChains = Node.IKRetargeterAsset->GetTargetIKRig()->GetRetargetChains();
     for (const FBoneChain &Chain : TargetBoneChains)
     {
         if (RefSkel.FindBoneIndex(Chain.StartBone) == INDEX_NONE)
@@ -102,8 +102,8 @@ void UAnimGraphNode_RetargetPoseFromMesh::PreloadRequiredAssets()
 	if (Node.IKRetargeterAsset)
 	{
 		PreloadObject(Node.IKRetargeterAsset);
-		PreloadObject(Node.IKRetargeterAsset->SourceIKRigAsset);
-		PreloadObject(Node.IKRetargeterAsset->TargetIKRigAsset);
+		PreloadObject(Node.IKRetargeterAsset->GetSourceIKRigWriteable());
+		PreloadObject(Node.IKRetargeterAsset->GetTargetIKRigWriteable());
 	}
 }
 
diff --git a/Engine/Plugins/Animation/IKRig/Source/IKRigEditor/Private/RetargetEditor/IKRetargetEditMode.cpp b/Engine/Plugins/Animation/IKRig/Source/IKRigEditor/Private/RetargetEditor/IKRetargetEditMode.cpp
--- a/Engine/Plugins/Animation/IKRig/Source/IKRigEditor/Private/RetargetEditor/IKRetargetEditMode.cpp
+++ b/Engine/Plugins/Animation/IKRig/Source/IKRigEditor/Private/RetargetEditor/IKRetargetEditMode.cpp
@@ -55,18 +55,20 @@ void FIKRetargetEditMode::Render(const FSceneView* View, FViewport* Viewport, FP
 	{
 		return;
 	}
-	
+
 	const UIKRetargeter* Retargeter = Controller->AssetController->GetAsset();
-	
-	for (int32 BoneIndex = 0; BoneIndex<Retargeter->TargetIKRigAsset->Skeleton.BoneNames.Num(); ++BoneIndex)
-	{
+	const UIKRetargetProcessor* RetargetProcessor = Controller->GetRetargetProcessor();
+	const FTargetSkeleton& TargetSkeleton = RetargetProcessor->GetTargetSkeleton();
+
+	for (int32 BoneIndex = 0; BoneIndex<TargetSkeleton.BoneNames.Num(); ++BoneIndex)
+	{		
 		// filter out bones that are not part of a retargeted chain
-		if (!Controller->IsTargetBoneRetargeted(BoneIndex))
+		if (!TargetSkeleton.IsBoneInAnyTargetChain[BoneIndex])
 		{
 			continue;
 		}
 		
-		const FName BoneName = Retargeter->TargetIKRigAsset->Skeleton.BoneNames[BoneIndex];
+		const FName& BoneName = TargetSkeleton.BoneNames[BoneIndex];
 
 		// is this chain currently selected?
 		const bool bIsSelected = IsBoneSelected(BoneName);
@@ -294,7 +296,7 @@ bool FIKRetargetEditMode::GetCustomDrawingCoordinateSystem(FMatrix& InMatrix, vo
 		return false; 
 	}
 
-	const int32 BoneIndex = Controller->AssetController->GetAsset()->TargetIKRigAsset->Skeleton.GetBoneIndexFromName(SelectedBones[0]);
+	const int32 BoneIndex = Controller->AssetController->GetAsset()->GetTargetIKRig()->Skeleton.GetBoneIndexFromName(SelectedBones[0]);
 	const FTransform CurrentTransform = Controller->GetTargetBoneTransform(BoneIndex);
 	InMatrix = CurrentTransform.ToMatrixNoScale().RemoveTranslation();
 	return true;
diff --git a/Engine/Plugins/Animation/IKRig/Source/IKRigEditor/Private/RetargetEditor/IKRetargetEditor.cpp b/Engine/Plugins/Animation/IKRig/Source/IKRigEditor/Private/RetargetEditor/IKRetargetEditor.cpp
--- a/Engine/Plugins/Animation/IKRig/Source/IKRigEditor/Private/RetargetEditor/IKRetargetEditor.cpp
+++ b/Engine/Plugins/Animation/IKRig/Source/IKRigEditor/Private/RetargetEditor/IKRetargetEditor.cpp
@@ -16,16 +16,12 @@
 #include "Framework/MultiBox/MultiBoxBuilder.h"
 
 #include "Retargeter/IKRetargeter.h"
-#include "Retargeter/IKRetargetProcessor.h"
 #include "RetargetEditor/IKRetargetAnimInstance.h"
 #include "RetargetEditor/IKRetargetCommands.h"
 #include "RetargetEditor/IKRetargetEditMode.h"
 #include "RetargetEditor/IKRetargetMode.h"
 #include "RetargetEditor/IKRetargetEditorController.h"
 
-#include "Widgets/Input/SButton.h"
-#include "Widgets/Input/SEditableTextBox.h"
-
 #define LOCTEXT_NAMESPACE "IKRetargeterEditor"
 
 const FName IKRetargetEditorModes::IKRetargetEditorMode("IKRetargetEditorMode");
@@ -53,11 +49,8 @@ void FIKRetargetEditor::InitAssetEditor(
 	PersonaToolkitArgs.OnPreviewSceneCreated = FOnPreviewSceneCreated::FDelegate::CreateSP(this, &FIKRetargetEditor::HandlePreviewSceneCreated);
 	
 	FPersonaModule& PersonaModule = FModuleManager::LoadModuleChecked<FPersonaModule>("Persona");
-	EditorController->PersonaToolkit = PersonaModule.CreatePersonaToolkit(InAsset, PersonaToolkitArgs, EditorController->AssetController->GetSourceSkeletonAsset());
+	EditorController->PersonaToolkit = PersonaModule.CreatePersonaToolkit(InAsset, PersonaToolkitArgs);
 	
-	// when/if preview mesh is changed, we need to reinitialize the anim instance
-    EditorController->PersonaToolkit->GetPreviewScene()->RegisterOnPreviewMeshChanged(FOnPreviewMeshChanged::CreateSP(this, &FIKRetargetEditor::HandlePreviewMeshChanged));
-
 	TSharedRef<IAssetFamily> AssetFamily = PersonaModule.CreatePersonaAssetFamily(InAsset);
 	AssetFamily->RecordAssetOpened(FAssetData(InAsset));
 
@@ -107,28 +100,28 @@ void FIKRetargetEditor::BindCommands()
 
 	ToolkitCommands->MapAction(
         Commands.EditRetargetPose,
-        FExecuteAction::CreateSP(this, &FIKRetargetEditor::HandleEditPose),
-        FCanExecuteAction::CreateSP(this, &FIKRetargetEditor::CanEditPose),
-        FIsActionChecked::CreateSP(this,  &FIKRetargetEditor::IsEditingPose),
+        FExecuteAction::CreateSP(EditorController, &FIKRetargetEditorController::HandleEditPose),
+        FCanExecuteAction::CreateSP(EditorController, &FIKRetargetEditorController::CanEditPose),
+        FIsActionChecked::CreateSP(EditorController,  &FIKRetargetEditorController::IsEditingPose),
 		EUIActionRepeatMode::RepeatDisabled);
 
 	ToolkitCommands->MapAction(
 		Commands.NewRetargetPose,
-		FExecuteAction::CreateSP(this, &FIKRetargetEditor::HandleNewPose),
+		FExecuteAction::CreateSP(EditorController, &FIKRetargetEditorController::HandleNewPose),
 		FCanExecuteAction(),
 		FCanExecuteAction(),
 		EUIActionRepeatMode::RepeatDisabled);
 
 	ToolkitCommands->MapAction(
 		Commands.DeleteRetargetPose,
-		FExecuteAction::CreateSP(this, &FIKRetargetEditor::HandleDeletePose),
-		FCanExecuteAction::CreateSP(this, &FIKRetargetEditor::CanDeletePose),
+		FExecuteAction::CreateSP(EditorController, &FIKRetargetEditorController::HandleDeletePose),
+		FCanExecuteAction::CreateSP(EditorController, &FIKRetargetEditorController::CanDeletePose),
 		FCanExecuteAction(),
 		EUIActionRepeatMode::RepeatDisabled);
 
 	ToolkitCommands->MapAction(
 		Commands.ResetRetargetPose,
-		FExecuteAction::CreateSP(this, &FIKRetargetEditor::HandleResetPose),
+		FExecuteAction::CreateSP(EditorController, &FIKRetargetEditorController::HandleResetPose),
 		FCanExecuteAction(),
 		EUIActionRepeatMode::RepeatDisabled);
 }
@@ -151,10 +144,10 @@ void FIKRetargetEditor::FillToolbar(FToolBarBuilder& ToolbarBuilder)
 {
 	ToolbarBuilder.BeginSection("Retarget Pose");
 	{
-		PoseNames.Reset();
+		EditorController->PoseNames.Reset();
 		for (const TTuple<FName, FIKRetargetPose>& Pose : EditorController->AssetController->GetRetargetPoses())
 		{
-			PoseNames.Add(MakeShareable(new FName(Pose.Key)));
+			EditorController->PoseNames.Add(MakeShareable(new FName(Pose.Key)));
 		}
 
 		TSharedRef<SWidget> PoseListWidget = SNew(SHorizontalBox)
@@ -176,14 +169,14 @@ void FIKRetargetEditor::FillToolbar(FToolBarBuilder& ToolbarBuilder)
 		.Padding(3.0f, 1.0f)
 		[
 			SNew(SComboBox<TSharedPtr<FName>>)
-			.OptionsSource(&PoseNames)
+			.OptionsSource(&EditorController->PoseNames)
 			.OnGenerateWidget_Lambda([](TSharedPtr<FName> InItem)
 			{
 				return SNew(STextBlock).Text(FText::FromName(*InItem.Get()));
 			})
-			.OnSelectionChanged(this, &FIKRetargetEditor::OnPoseSelected)
+			.OnSelectionChanged(EditorController, &FIKRetargetEditorController::OnPoseSelected)
 			[
-				SNew(STextBlock).Text(this, &FIKRetargetEditor::GetCurrentPoseName)
+				SNew(STextBlock).Text(EditorController, &FIKRetargetEditorController::GetCurrentPoseName)
 			]
 		];
 		ToolbarBuilder.AddWidget(PoseListWidget);
@@ -316,16 +309,6 @@ void FIKRetargetEditor::SetupAnimInstance()
 	EditorController->TargetAnimInstance->InitializeAnimation();
 }
 
-FText FIKRetargetEditor::GetCurrentPoseName() const
-{
-	return FText::FromName(EditorController->AssetController->GetCurrentRetargetPoseName());
-}
-
-void FIKRetargetEditor::OnPoseSelected(TSharedPtr<FName> InPosePose, ESelectInfo::Type SelectInfo)
-{
-	EditorController->AssetController->SetCurrentRetargetPose(*InPosePose.Get());
-}
-
 void FIKRetargetEditor::HandleSourceOrTargetIKRigAssetChanged()
 {
 	// set the source and target skeletal meshes on the component
@@ -344,18 +327,8 @@ void FIKRetargetEditor::HandleSourceOrTargetIKRigAssetChanged()
 	}
 	
 	SetupAnimInstance();
-	
-	EditorController->RefreshAllViews();
-}
 
-void FIKRetargetEditor::HandlePreviewMeshChanged(USkeletalMesh* InOldSkeletalMesh, USkeletalMesh* InNewSkeletalMesh)
-{
-	if (!InOldSkeletalMesh)
-	{
-		return; // first time setup
-	}
-	
-	HandleSourceOrTargetIKRigAssetChanged();
+	EditorController->RefreshAllViews();
 }
 
 void FIKRetargetEditor::HandleDetailsCreated(const TSharedRef<class IDetailsView>& InDetailsView)
@@ -368,135 +341,20 @@ void FIKRetargetEditor::HandleDetailsCreated(const TSharedRef<class IDetailsView
 void FIKRetargetEditor::OnFinishedChangingDetails(
     const FPropertyChangedEvent& PropertyChangedEvent)
 {
-	FName SourceIKRigPropertyName = GET_MEMBER_NAME_STRING_CHECKED(UIKRetargeter, SourceIKRigAsset);
-	FName TargetIKRigPropertyName = GET_MEMBER_NAME_STRING_CHECKED(UIKRetargeter, TargetIKRigAsset);
-	const bool bSourceChanged = PropertyChangedEvent.GetPropertyName() == SourceIKRigPropertyName;
-	const bool bTargetChanged = PropertyChangedEvent.GetPropertyName() == TargetIKRigPropertyName;
-	if (bSourceChanged || bTargetChanged)
-	{
-		HandleSourceOrTargetIKRigAssetChanged();
-	}
-}
+	const bool bSourceChanged = PropertyChangedEvent.GetPropertyName() == UIKRetargeter::GetSourceIKRigPropertyName();
+	const bool bTargetChanged = PropertyChangedEvent.GetPropertyName() == UIKRetargeter::GetTargetIKRigPropertyName();
+	const bool bPreviewChanged = PropertyChangedEvent.GetPropertyName() == UIKRetargeter::GetTargetPreviewMeshPropertyName();
 
-void FIKRetargetEditor::HandleEditPose() const
-{
-	const bool bEditPoseMode = !EditorController.Get().AssetController->GetEditRetargetPoseMode();
-	EditorController.Get().AssetController->SetEditRetargetPoseMode(bEditPoseMode);
-	if (bEditPoseMode)
+	if (bTargetChanged)
 	{
-		GetEditorModeManager().ActivateMode(FIKRetargetEditMode::ModeName);
-		EditorController->SourceSkelMeshComponent->ShowReferencePose(true);
+		EditorController->AssetController->CleanChainMapping();
+		EditorController->AssetController->AutoMapChains();
 	}
-	else
-	{
-		GetEditorModeManager().DeactivateMode(FIKRetargetEditMode::ModeName);
-		EditorController->PlayPreviousAnimationAsset();
-	}
-}
-
-bool FIKRetargetEditor::CanEditPose() const
-{
-	const UIKRetargetProcessor* Processor = EditorController->GetRetargetProcessor();
-	if (!Processor)
+	
+	if (bSourceChanged || bTargetChanged || bPreviewChanged)
 	{
-		return false;
+		HandleSourceOrTargetIKRigAssetChanged();
 	}
-
-	return Processor->IsInitialized();
-}
-
-bool FIKRetargetEditor::IsEditingPose() const
-{
-	return EditorController.Get().AssetController->GetEditRetargetPoseMode();
-}
-
-void FIKRetargetEditor::HandleNewPose()
-{
-	SAssignNew(NewPoseWindow, SWindow)
-	.Title(LOCTEXT("NewRetargetPoseOptions", "Create New Retarget Pose"))
-	.ClientSize(FVector2D(250, 80))
-	.HasCloseButton(true)
-	.SupportsMinimize(false) .SupportsMaximize(false)
-	[
-		SNew(SBorder)
-		.BorderImage( FEditorStyle::GetBrush("Menu.Background") )
-		[
-			SNew(SVerticalBox)
-
-			+ SVerticalBox::Slot()
-			.Padding(4)
-			.AutoHeight()
-			[
-				SAssignNew(NewPoseEditableText, SEditableTextBox)
-				.Text(LOCTEXT("NewRetargetPoseName", "NewPose"))
-			]
-
-			+ SVerticalBox::Slot()
-			.Padding(4)
-			.AutoHeight()
-			[
-				SNew(SHorizontalBox)
-				+ SHorizontalBox::Slot()
-				.HAlign(HAlign_Center)
-				[
-					SNew(SButton)
-					.ButtonStyle(FAppStyle::Get(), "Button")
-					.TextStyle( FAppStyle::Get(), "DialogButtonText" )
-					.HAlign(HAlign_Center)
-					.VAlign(VAlign_Center)
-					.Text(LOCTEXT("OkButtonLabel", "Ok"))
-					.OnClicked(this, &FIKRetargetEditor::CreateNewPose)
-				]
-				
-				+ SHorizontalBox::Slot()
-				.HAlign(HAlign_Center)
-				[
-					SNew(SButton)
-					.ButtonStyle(FAppStyle::Get(), "Button")
-					.TextStyle( FAppStyle::Get(), "DialogButtonText" )
-					.HAlign(HAlign_Center)
-					.VAlign(VAlign_Center)
-					.Text(LOCTEXT("CancelButtonLabel", "Cancel"))
-					.OnClicked_Lambda( [=]()
-					{
-						NewPoseWindow->RequestDestroyWindow();
-						return FReply::Handled();
-					})
-				]
-			]	
-		]
-	];
-
-	GEditor->EditorAddModalWindow(NewPoseWindow.ToSharedRef());
-	NewPoseWindow.Reset();
-}
-
-FReply FIKRetargetEditor::CreateNewPose()
-{
-	const FName NewPoseName = FName(NewPoseEditableText.Get()->GetText().ToString());
-	EditorController->AssetController->AddRetargetPose(NewPoseName);
-	NewPoseWindow->RequestDestroyWindow();
-	RegenerateMenusAndToolbars();
-	return FReply::Handled();
-}
-
-void FIKRetargetEditor::HandleDeletePose()
-{
-	const FName CurrentPose = EditorController->AssetController->GetCurrentRetargetPoseName();
-	EditorController->AssetController->RemoveRetargetPose(CurrentPose);
-	RegenerateMenusAndToolbars();
-}
-
-bool FIKRetargetEditor::CanDeletePose() const
-{	
-	// cannot delete default pose
-	return EditorController->AssetController->GetCurrentRetargetPoseName() != UIKRetargeter::DefaultPoseName;
-}
-
-void FIKRetargetEditor::HandleResetPose()
-{
-	const FName CurrentPose = EditorController->AssetController->GetCurrentRetargetPoseName();
-	EditorController->AssetController->ResetRetargetPose(CurrentPose);
 }
 
 #undef LOCTEXT_NAMESPACE
diff --git a/Engine/Plugins/Animation/IKRig/Source/IKRigEditor/Private/RetargetEditor/IKRetargetEditorController.cpp b/Engine/Plugins/Animation/IKRig/Source/IKRigEditor/Private/RetargetEditor/IKRetargetEditorController.cpp
--- a/Engine/Plugins/Animation/IKRig/Source/IKRigEditor/Private/RetargetEditor/IKRetargetEditorController.cpp
+++ b/Engine/Plugins/Animation/IKRig/Source/IKRigEditor/Private/RetargetEditor/IKRetargetEditorController.cpp
@@ -2,7 +2,11 @@
 
 #include "RetargetEditor/IKRetargetEditorController.h"
 
+#include "EditorModeManager.h"
+#include "Animation/DebugSkelMeshComponent.h"
 #include "RetargetEditor/IKRetargetAnimInstance.h"
+#include "RetargetEditor/IKRetargetEditMode.h"
+#include "RetargetEditor/IKRetargetEditor.h"
 #include "RetargetEditor/SIKRetargetChainMapList.h"
 #include "Retargeter/IKRetargeter.h"
 #include "RigEditor/IKRigController.h"
@@ -15,14 +19,14 @@ void FIKRetargetEditorController::Initialize(TSharedPtr<FIKRetargetEditor> InEdi
 	AssetController = UIKRetargeterController::GetController(InAsset);
 
 	// bind callbacks when SOURCE or TARGET IK Rigs are modified
-	BindToIKRigAsset(AssetController->GetAsset()->SourceIKRigAsset);
-	BindToIKRigAsset(AssetController->GetAsset()->TargetIKRigAsset);
+	BindToIKRigAsset(AssetController->GetAsset()->GetSourceIKRigWriteable());
+	BindToIKRigAsset(AssetController->GetAsset()->GetTargetIKRigWriteable());
 
 	// bind callback when retargeter needs reinitialized
 	AssetController->OnRetargeterNeedsInitialized().AddSP(this, &FIKRetargetEditorController::OnRetargeterNeedsInitialized);
 }
 
-void FIKRetargetEditorController::BindToIKRigAsset(UIKRigDefinition* InIKRig)
+void FIKRetargetEditorController::BindToIKRigAsset(UIKRigDefinition* InIKRig) const
 {
 	if (!InIKRig)
 	{
@@ -37,14 +41,14 @@ void FIKRetargetEditorController::BindToIKRigAsset(UIKRigDefinition* InIKRig)
 	}
 }
 
-void FIKRetargetEditorController::OnIKRigNeedsInitialized(UIKRigDefinition* ModifiedIKRig)
+void FIKRetargetEditorController::OnIKRigNeedsInitialized(UIKRigDefinition* ModifiedIKRig) const
 {
 	const UIKRetargeter* Retargeter = AssetController->GetAsset();
 	
 	check(ModifiedIKRig && Retargeter)
 	 
-	const bool bIsSource = ModifiedIKRig == Retargeter->SourceIKRigAsset;
-	const bool bIsTarget = ModifiedIKRig == Retargeter->TargetIKRigAsset;
+	const bool bIsSource = ModifiedIKRig == Retargeter->GetSourceIKRig();
+	const bool bIsTarget = ModifiedIKRig == Retargeter->GetTargetIKRig();
 	if (!(bIsSource || bIsTarget))
 	{
 		return;
@@ -55,37 +59,44 @@ void FIKRetargetEditorController::OnIKRigNeedsInitialized(UIKRigDefinition* Modi
 	RefreshAllViews();
 }
 
-void FIKRetargetEditorController::OnRetargetChainRenamed(UIKRigDefinition* ModifiedIKRig, FName OldName, FName NewName)
+void FIKRetargetEditorController::OnRetargetChainRenamed(UIKRigDefinition* ModifiedIKRig, FName OldName, FName NewName) const
 {
 	check(ModifiedIKRig)
 	
 	AssetController->OnRetargetChainRenamed(ModifiedIKRig, OldName, NewName);
 }
 
-void FIKRetargetEditorController::OnRetargeterNeedsInitialized(const UIKRetargeter* Retargeter)
+void FIKRetargetEditorController::OnRetargeterNeedsInitialized(const UIKRetargeter* Retargeter) const
 {
+	// force edit pose mode to be off
+	AssetController->SetEditRetargetPoseMode(false, false); // turn off and don't reinitialize (avoid infinite loop)
+	Editor.Pin()->GetEditorModeManager().DeactivateMode(FIKRetargetEditMode::ModeName);
+	// force reinit the runtime retarget processor
 	TargetAnimInstance->SetProcessorNeedsInitialized();
+	// refresh all the UI views
 	RefreshAllViews();
+	// keep playing animation
+	PlayPreviousAnimationAsset();
 }
 
 USkeletalMesh* FIKRetargetEditorController::GetSourceSkeletalMesh() const
 {
-	if (!(AssetController && AssetController->GetAsset()->SourceIKRigAsset))
+	if (!(AssetController && AssetController->GetAsset()->GetSourceIKRig()))
 	{
 		return nullptr;
 	}
 
-	return AssetController->GetAsset()->SourceIKRigAsset->PreviewSkeletalMesh;
+	return AssetController->GetAsset()->GetSourceIKRig()->PreviewSkeletalMesh;
 }
 
 USkeletalMesh* FIKRetargetEditorController::GetTargetSkeletalMesh() const
 {
-	if (!(AssetController && AssetController->GetAsset()->TargetIKRigAsset))
+	if (!(AssetController && AssetController->GetAsset()->GetTargetIKRig()))
 	{
 		return nullptr;
 	}
 
-	return AssetController->GetAsset()->TargetIKRigAsset->PreviewSkeletalMesh;
+	return AssetController->GetTargetPreviewMesh();
 }
 
 FTransform FIKRetargetEditorController::GetTargetBoneTransform(const int32& TargetBoneIndex) const
@@ -143,19 +154,6 @@ bool FIKRetargetEditorController::GetTargetBoneLineSegments(
 	return true;
 }
 
-bool FIKRetargetEditorController::IsTargetBoneRetargeted(const int32& TargetBoneIndex)
-{
-	// get the runtime processor
-	const UIKRetargetProcessor* Processor = GetRetargetProcessor();
-	check(Processor && Processor->IsInitialized())
-	
-	// get the target skeleton
-	const FTargetSkeleton& TargetSkeleton = Processor->GetTargetSkeleton();
-	check(TargetSkeleton.BoneNames.IsValidIndex(TargetBoneIndex))
-
-	return TargetSkeleton.IsBoneRetargeted[TargetBoneIndex];
-}
-
 const UIKRetargetProcessor* FIKRetargetEditorController::GetRetargetProcessor() const
 {	
 	if(UIKRetargetAnimInstance* AnimInstance = TargetAnimInstance.Get())
@@ -166,8 +164,9 @@ const UIKRetargetProcessor* FIKRetargetEditorController::GetRetargetProcessor()
 	return nullptr;	
 }
 
-void FIKRetargetEditorController::RefreshAllViews()
+void FIKRetargetEditorController::RefreshAllViews() const
 {
+	Editor.Pin()->RegenerateMenusAndToolbars();
 	ChainsView.Get()->RefreshView();
 }
 
@@ -180,7 +179,7 @@ void FIKRetargetEditorController::PlayAnimationAsset(UAnimationAsset* AssetToPla
 	}
 }
 
-void FIKRetargetEditorController::PlayPreviousAnimationAsset()
+void FIKRetargetEditorController::PlayPreviousAnimationAsset() const
 {
 	if (PreviousAsset)
 	{
@@ -188,4 +187,139 @@ void FIKRetargetEditorController::PlayPreviousAnimationAsset()
 	}
 }
 
+void FIKRetargetEditorController::HandleEditPose() const
+{
+	const bool bEditPoseMode = !AssetController->GetEditRetargetPoseMode();
+	AssetController->SetEditRetargetPoseMode(bEditPoseMode);
+	if (bEditPoseMode)
+	{
+		Editor.Pin()->GetEditorModeManager().ActivateMode(FIKRetargetEditMode::ModeName);
+		SourceSkelMeshComponent->ShowReferencePose(true);
+	}
+	else
+	{
+		Editor.Pin()->GetEditorModeManager().DeactivateMode(FIKRetargetEditMode::ModeName);
+		PlayPreviousAnimationAsset();
+	}
+}
+
+bool FIKRetargetEditorController::CanEditPose() const
+{
+	const UIKRetargetProcessor* Processor = GetRetargetProcessor();
+	if (!Processor)
+	{
+		return false;
+	}
+
+	return Processor->IsInitialized();
+}
+
+bool FIKRetargetEditorController::IsEditingPose() const
+{
+	return AssetController->GetEditRetargetPoseMode();
+}
+
+void FIKRetargetEditorController::HandleNewPose()
+{
+	// get a unique pose name to use as suggestion
+	const FName DefaultNewPoseName = FName(LOCTEXT("NewRetargetPoseName", "CustomRetargetPose").ToString());
+	const FName UniqueNewPoseName = AssetController->MakePoseNameUnique(DefaultNewPoseName);
+	
+	SAssignNew(NewPoseWindow, SWindow)
+	.Title(LOCTEXT("NewRetargetPoseOptions", "Create New Retarget Pose"))
+	.ClientSize(FVector2D(250, 80))
+	.HasCloseButton(true)
+	.SupportsMinimize(false) .SupportsMaximize(false)
+	[
+		SNew(SBorder)
+		.BorderImage( FEditorStyle::GetBrush("Menu.Background") )
+		[
+			SNew(SVerticalBox)
+
+			+ SVerticalBox::Slot()
+			.Padding(4)
+			.AutoHeight()
+			[
+				SAssignNew(NewPoseEditableText, SEditableTextBox)
+				.Text(FText::FromName(UniqueNewPoseName))
+			]
+
+			+ SVerticalBox::Slot()
+			.Padding(4)
+			.AutoHeight()
+			[
+				SNew(SHorizontalBox)
+				+ SHorizontalBox::Slot()
+				.HAlign(HAlign_Center)
+				[
+					SNew(SButton)
+					.ButtonStyle(FAppStyle::Get(), "Button")
+					.TextStyle( FAppStyle::Get(), "DialogButtonText" )
+					.HAlign(HAlign_Center)
+					.VAlign(VAlign_Center)
+					.Text(LOCTEXT("OkButtonLabel", "Ok"))
+					.OnClicked(this, &FIKRetargetEditorController::CreateNewPose)
+				]
+				
+				+ SHorizontalBox::Slot()
+				.HAlign(HAlign_Center)
+				[
+					SNew(SButton)
+					.ButtonStyle(FAppStyle::Get(), "Button")
+					.TextStyle( FAppStyle::Get(), "DialogButtonText" )
+					.HAlign(HAlign_Center)
+					.VAlign(VAlign_Center)
+					.Text(LOCTEXT("CancelButtonLabel", "Cancel"))
+					.OnClicked_Lambda( [=]()
+					{
+						NewPoseWindow->RequestDestroyWindow();
+						return FReply::Handled();
+					})
+				]
+			]	
+		]
+	];
+
+	GEditor->EditorAddModalWindow(NewPoseWindow.ToSharedRef());
+	NewPoseWindow.Reset();
+}
+
+FReply FIKRetargetEditorController::CreateNewPose() const
+{
+	const FName NewPoseName = FName(NewPoseEditableText.Get()->GetText().ToString());
+	AssetController->AddRetargetPose(NewPoseName);
+	NewPoseWindow->RequestDestroyWindow();
+	Editor.Pin()->RegenerateMenusAndToolbars();
+	return FReply::Handled();
+}
+
+void FIKRetargetEditorController::HandleDeletePose() const
+{
+	const FName CurrentPose = AssetController->GetCurrentRetargetPoseName();
+	AssetController->RemoveRetargetPose(CurrentPose);
+	Editor.Pin()->RegenerateMenusAndToolbars();
+}
+
+bool FIKRetargetEditorController::CanDeletePose() const
+{	
+	// cannot delete default pose
+	return AssetController->GetCurrentRetargetPoseName() != UIKRetargeter::GetDefaultPoseName();
+}
+
+void FIKRetargetEditorController::HandleResetPose() const
+{
+	const FName CurrentPose = AssetController->GetCurrentRetargetPoseName();
+	AssetController->ResetRetargetPose(CurrentPose);
+}
+
+FText FIKRetargetEditorController::GetCurrentPoseName() const
+{
+	return FText::FromName(AssetController->GetCurrentRetargetPoseName());
+}
+
+void FIKRetargetEditorController::OnPoseSelected(TSharedPtr<FName> InPosePose, ESelectInfo::Type SelectInfo) const
+{
+	AssetController->SetCurrentRetargetPose(*InPosePose.Get());
+}
+
 #undef LOCTEXT_NAMESPACE
diff --git a/Engine/Plugins/Animation/IKRig/Source/IKRigEditor/Private/RetargetEditor/IKRetargetFactory.cpp b/Engine/Plugins/Animation/IKRig/Source/IKRigEditor/Private/RetargetEditor/IKRetargetFactory.cpp
--- a/Engine/Plugins/Animation/IKRig/Source/IKRigEditor/Private/RetargetEditor/IKRetargetFactory.cpp
+++ b/Engine/Plugins/Animation/IKRig/Source/IKRigEditor/Private/RetargetEditor/IKRetargetFactory.cpp
@@ -7,6 +7,7 @@
 #include "IContentBrowserSingleton.h"
 #include "Editor.h"
 #include "EditorStyleSet.h"
+#include "RetargetEditor/IKRetargeterController.h"
 #include "Widgets/SWindow.h"
 #include "Widgets/Layout/SBorder.h"
 
@@ -41,7 +42,8 @@ UObject* UIKRetargetFactory::FactoryCreateNew(
 	}
 	
 	UIKRetargeter* Retargeter = NewObject<UIKRetargeter>(InParent, Class, Name, Flags);
-	Retargeter->SourceIKRigAsset = SourceIKRig;
+	UIKRetargeterController* Controller = UIKRetargeterController::GetController(Retargeter);
+	Controller->SetSourceIKRig(SourceIKRig);
 	return Retargeter;
 }
 
diff --git a/Engine/Plugins/Animation/IKRig/Source/IKRigEditor/Private/RetargetEditor/IKRetargetMode.cpp b/Engine/Plugins/Animation/IKRig/Source/IKRigEditor/Private/RetargetEditor/IKRetargetMode.cpp
--- a/Engine/Plugins/Animation/IKRig/Source/IKRigEditor/Private/RetargetEditor/IKRetargetMode.cpp
+++ b/Engine/Plugins/Animation/IKRig/Source/IKRigEditor/Private/RetargetEditor/IKRetargetMode.cpp
@@ -31,7 +31,6 @@ FIKRetargetMode::FIKRetargetMode(
 	// register Persona tabs
 	FPersonaModule& PersonaModule = FModuleManager::LoadModuleChecked<FPersonaModule>("Persona");
 	TabFactories.RegisterFactory(PersonaModule.CreatePersonaViewportTabFactory(InHostingApp, ViewportArgs));
-	TabFactories.RegisterFactory(PersonaModule.CreateAdvancedPreviewSceneTabFactory(InHostingApp, InPreviewScene));
 	TabFactories.RegisterFactory(PersonaModule.CreateDetailsTabFactory(InHostingApp, FOnDetailsCreated::CreateSP(&IKRetargetEditor.Get(), &FIKRetargetEditor::HandleDetailsCreated)));
 
 	// register custom tabs
@@ -39,7 +38,7 @@ FIKRetargetMode::FIKRetargetMode(
 	TabFactories.RegisterFactory(MakeShared<FIKRetargetAssetBrowserTabSummoner>(IKRetargetEditor));
 
 	// create tab layout
-	TabLayout = FTabManager::NewLayout("Standalone_IKRetargetEditor_Layout_v1.006")
+	TabLayout = FTabManager::NewLayout("Standalone_IKRetargetEditor_Layout_v1.007")
 		->AddArea
 		(
 			FTabManager::NewPrimaryArea()
@@ -66,7 +65,6 @@ FIKRetargetMode::FIKRetargetMode(
 					FTabManager::NewStack()
 						->SetSizeCoefficient(0.6f)
 						->AddTab(FPersonaTabs::DetailsID, ETabState::OpenedTab)
-						->AddTab(FPersonaTabs::AdvancedPreviewSceneSettingsID, ETabState::OpenedTab)
 						->SetForegroundTab(FPersonaTabs::DetailsID)
 					)
 					->Split
diff --git a/Engine/Plugins/Animation/IKRig/Source/IKRigEditor/Private/RetargetEditor/IKRetargeterController.cpp b/Engine/Plugins/Animation/IKRig/Source/IKRigEditor/Private/RetargetEditor/IKRetargeterController.cpp
--- a/Engine/Plugins/Animation/IKRig/Source/IKRigEditor/Private/RetargetEditor/IKRetargeterController.cpp
+++ b/Engine/Plugins/Animation/IKRig/Source/IKRigEditor/Private/RetargetEditor/IKRetargeterController.cpp
@@ -27,6 +27,9 @@ UIKRetargeterController* UIKRetargeterController::GetController(UIKRetargeter* I
 	// clean the asset before editing
 	Controller->CleanChainMapping();
 	Controller->CleanPoseList();
+
+	// bind callback to reinitialize processor on the asset undo
+	InRetargeterAsset->IKRigEditUndo.AddUObject(Controller, &UIKRetargeterController::BroadcastNeedsReinitialized);
 	
 	return Controller;
 }
@@ -36,6 +39,41 @@ UIKRetargeter* UIKRetargeterController::GetAsset() const
 	return Asset;
 }
 
+void UIKRetargeterController::SetSourceIKRig(UIKRigDefinition* SourceIKRig)
+{
+	Asset->SourceIKRigAsset = SourceIKRig;
+}
+
+void UIKRetargeterController::SetTargetIKRig(UIKRigDefinition* TargetIKRig)
+{
+	const bool bTargetWasNull = Asset->TargetIKRigAsset == nullptr;
+	Asset->TargetIKRigAsset = TargetIKRig;
+
+	if (bTargetWasNull)
+	{
+		CleanChainMapping();
+		AutoMapChains();
+	}
+}
+
+USkeletalMesh* UIKRetargeterController::GetTargetPreviewMesh()
+{
+	// can't preview anything if target IK Rig is null
+	if (!Asset->GetTargetIKRig())
+	{
+		return nullptr;
+	}
+
+	// optionally prefer override if one is provided
+	if (!Asset->TargetPreviewMesh.IsNull())
+	{
+		return Asset->TargetPreviewMesh;
+	}
+
+	// fallback to preview mesh from IK Rig asset
+	return Asset->GetTargetIKRig()->PreviewSkeletalMesh;
+}
+
 FName UIKRetargeterController::GetSourceRootBone() const
 {
 	return Asset->SourceIKRigAsset ? Asset->SourceIKRigAsset->GetRetargetRoot() : FName("None");
@@ -246,6 +284,9 @@ void UIKRetargeterController::OnRetargetChainRenamed(UIKRigDefinition* IKRig, FN
 
 void UIKRetargeterController::SetSourceChainForTargetChain(FName TargetChain, FName SourceChainToMapTo)
 {
+	FScopedTransaction Transaction(LOCTEXT("SetRetargetChainSource", "Set Retarget Chain Source"));
+	Asset->Modify();
+	
 	FRetargetChainMap* ChainMap = GetChainMap(TargetChain);
 	check(ChainMap)
 	ChainMap->SourceChain = SourceChainToMapTo;
@@ -281,11 +322,10 @@ USkeleton* UIKRetargeterController::GetSourceSkeletonAsset() const
 
 void UIKRetargeterController::AddRetargetPose(FName NewPoseName) const
 {
-	if (Asset->RetargetPoses.Contains(NewPoseName))
-	{
-		return;
-	}
-
+	FScopedTransaction Transaction(LOCTEXT("AddRetargetPose", "Add Retarget Pose"));
+	Asset->Modify();
+	
+	NewPoseName = MakePoseNameUnique(NewPoseName);
 	Asset->RetargetPoses.Add(NewPoseName);
 	Asset->CurrentRetargetPose = NewPoseName;
 
@@ -304,6 +344,9 @@ void UIKRetargeterController::RemoveRetargetPose(FName PoseToRemove) const
 		return; // cannot remove pose that doesn't exist
 	}
 
+	FScopedTransaction Transaction(LOCTEXT("RemoveRetargetPose", "Remove Retarget Pose"));
+	Asset->Modify();
+
 	Asset->RetargetPoses.Remove(PoseToRemove);
 
 	// did we remove the currently used pose?
@@ -322,6 +365,9 @@ void UIKRetargeterController::ResetRetargetPose(FName PoseToReset) const
 		return; // cannot reset pose that doesn't exist
 	}
 
+	FScopedTransaction Transaction(LOCTEXT("ResetRetargetPose", "Reset Retarget Pose"));
+	Asset->Modify();
+	
 	Asset->RetargetPoses[PoseToReset].BoneRotationOffsets.Reset();
 	Asset->RetargetPoses[PoseToReset].RootTranslationOffset = FVector::ZeroVector;
 	
@@ -336,7 +382,11 @@ FName UIKRetargeterController::GetCurrentRetargetPoseName() const
 void UIKRetargeterController::SetCurrentRetargetPose(FName CurrentPose) const
 {
 	check(Asset->RetargetPoses.Contains(CurrentPose));
+
+	FScopedTransaction Transaction(LOCTEXT("SetCurrentPose", "Set Current Pose"));
+	Asset->Modify();
 	Asset->CurrentRetargetPose = CurrentPose;
+	
 	BroadcastNeedsReinitialized();
 }
 
@@ -355,10 +405,10 @@ void UIKRetargeterController::AddTranslationOffsetToRetargetRootBone(FVector Tra
 	Asset->RetargetPoses[Asset->CurrentRetargetPose].AddTranslationDeltaToRoot(TranslationOffset);
 }
 
-void UIKRetargeterController::SetEditRetargetPoseMode(bool bEditPoseMode) const
+void UIKRetargeterController::SetEditRetargetPoseMode(bool bEditPoseMode, bool bReinitializeAfter) const
 {
 	GetAsset()->bEditRetargetPoseMode = bEditPoseMode;
-	if (!bEditPoseMode)
+	if (!bEditPoseMode && bReinitializeAfter)
 	{
 		// must reinitialize after editing the retarget pose
 		BroadcastNeedsReinitialized();
@@ -370,6 +420,18 @@ bool UIKRetargeterController::GetEditRetargetPoseMode() const
 	return GetAsset()->bEditRetargetPoseMode;
 }
 
+FName UIKRetargeterController::MakePoseNameUnique(FName PoseName) const
+{
+	FName UniqueName = PoseName;
+	int32 Suffix = 1;
+	while (Asset->RetargetPoses.Contains(UniqueName))
+	{
+		UniqueName = FName(PoseName.ToString() + "_" + FString::FromInt(Suffix));
+		++Suffix;
+	}
+	return UniqueName;
+}
+
 FRetargetChainMap* UIKRetargeterController::GetChainMap(const FName& TargetChainName) const
 {
 	for (FRetargetChainMap& ChainMap : Asset->ChainMapping)
diff --git a/Engine/Plugins/Animation/IKRig/Source/IKRigEditor/Private/RetargetEditor/SIKRetargetAssetBrowser.cpp b/Engine/Plugins/Animation/IKRig/Source/IKRigEditor/Private/RetargetEditor/SIKRetargetAssetBrowser.cpp
--- a/Engine/Plugins/Animation/IKRig/Source/IKRigEditor/Private/RetargetEditor/SIKRetargetAssetBrowser.cpp
+++ b/Engine/Plugins/Animation/IKRig/Source/IKRigEditor/Private/RetargetEditor/SIKRetargetAssetBrowser.cpp
@@ -117,7 +117,7 @@ FReply SIKRetargetAssetBrowser::OnExportButtonClicked() const
 
 	BatchContext.SourceMesh = Controller->GetSourceSkeletalMesh();
 	BatchContext.TargetMesh = Controller->GetTargetSkeletalMesh();
-	BatchContext.IKRetargetAsset = const_cast<UIKRetargeter*>(Controller->AssetController->GetAsset());
+	BatchContext.IKRetargetAsset = Controller->AssetController->GetAsset();
 	BatchContext.bRemapReferencedAssets = false;
 	BatchContext.NameRule.Suffix = "_Retargeted";
 
diff --git a/Engine/Plugins/Animation/IKRig/Source/IKRigEditor/Private/RetargetEditor/SIKRetargetChainMapList.cpp b/Engine/Plugins/Animation/IKRig/Source/IKRigEditor/Private/RetargetEditor/SIKRetargetChainMapList.cpp
--- a/Engine/Plugins/Animation/IKRig/Source/IKRigEditor/Private/RetargetEditor/SIKRetargetChainMapList.cpp
+++ b/Engine/Plugins/Animation/IKRig/Source/IKRigEditor/Private/RetargetEditor/SIKRetargetChainMapList.cpp
@@ -40,7 +40,7 @@ void SIKRetargetChainMapRow::Construct(
 	// NOTE: cannot just use FName because "None" is considered a null entry and removed from ComboBox.
 	SourceChainOptions.Reset();
 	SourceChainOptions.Add(MakeShareable(new FString(TEXT("None"))));
-	const UIKRigDefinition* SourceIKRig = ChainMapList.Pin()->EditorController.Pin()->AssetController->GetAsset()->SourceIKRigAsset;
+	const UIKRigDefinition* SourceIKRig = ChainMapList.Pin()->EditorController.Pin()->AssetController->GetAsset()->GetSourceIKRig();
 	if (SourceIKRig)
 	{
 		const TArray<FBoneChain>& Chains = SourceIKRig->GetRetargetChains();
@@ -154,7 +154,7 @@ void SIKRetargetChainMapList::Construct(
 			.VAlign(VAlign_Center)
 			[
 				SNew(STextBlock)
-				.Text(this, &SIKRetargetChainMapList::GetSourceRootBone)
+				.Text(this, &SIKRetargetChainMapList::GetTargetRootBone)
 				.IsEnabled(false)
 			]
 				
@@ -175,7 +175,7 @@ void SIKRetargetChainMapList::Construct(
 			.VAlign(VAlign_Center)
 			[
 				SNew(STextBlock)
-				.Text(this, &SIKRetargetChainMapList::GetTargetRootBone)
+				.Text(this, &SIKRetargetChainMapList::GetSourceRootBone)
 				.IsEnabled(false)
 			]
 		]
@@ -261,9 +261,9 @@ bool SIKRetargetChainMapList::IsChainMapEnabled() const
 		return false; 
 	}
 	
-	if (RetargeterController->GetAsset()->TargetIKRigAsset)
+	if (RetargeterController->GetAsset()->GetTargetIKRig())
 	{
-		const TArray<FBoneChain>& Chains = RetargeterController->GetAsset()->TargetIKRigAsset->GetRetargetChains();
+		const TArray<FBoneChain>& Chains = RetargeterController->GetAsset()->GetTargetIKRig()->GetRetargetChains();
 		return Chains.Num() > 0;
 	}
 
diff --git a/Engine/Plugins/Animation/IKRig/Source/IKRigEditor/Private/RetargetEditor/SRetargetAnimAssetsWindow.cpp b/Engine/Plugins/Animation/IKRig/Source/IKRigEditor/Private/RetargetEditor/SRetargetAnimAssetsWindow.cpp
--- a/Engine/Plugins/Animation/IKRig/Source/IKRigEditor/Private/RetargetEditor/SRetargetAnimAssetsWindow.cpp
+++ b/Engine/Plugins/Animation/IKRig/Source/IKRigEditor/Private/RetargetEditor/SRetargetAnimAssetsWindow.cpp
@@ -257,7 +257,7 @@ void SRetargetAnimAssetsWindow::Construct(const FArguments& InArgs)
 								return false;
 							}
 							
-							return BatchContext.IKRetargetAsset->SourceIKRigAsset != nullptr;
+							return BatchContext.IKRetargetAsset->GetSourceIKRig() != nullptr;
 						})
 						.ObjectPath(this, &SRetargetAnimAssetsWindow::GetCurrentSourceMeshPath)
 						.OnObjectChanged(this, &SRetargetAnimAssetsWindow::SourceMeshAssigned)
@@ -274,7 +274,7 @@ void SRetargetAnimAssetsWindow::Construct(const FArguments& InArgs)
 								return true;
 							}
 							
-							USkeletalMesh* PreviewMesh = BatchContext.IKRetargetAsset->SourceIKRigAsset->GetPreviewMesh();
+							USkeletalMesh* PreviewMesh = BatchContext.IKRetargetAsset->GetSourceIKRig()->GetPreviewMesh();
 							if (!PreviewMesh)
 							{
 								return true;
@@ -332,7 +332,7 @@ void SRetargetAnimAssetsWindow::Construct(const FArguments& InArgs)
 								return false;
 							}
 							
-							return BatchContext.IKRetargetAsset->TargetIKRigAsset != nullptr;
+							return BatchContext.IKRetargetAsset->GetTargetIKRig() != nullptr;
 						})
 						.ObjectPath(this, &SRetargetAnimAssetsWindow::GetCurrentTargetMeshPath)
 						.OnObjectChanged(this, &SRetargetAnimAssetsWindow::TargetMeshAssigned)
@@ -709,8 +709,8 @@ void SRetargetAnimAssetsWindow::RetargeterAssigned(const FAssetData& InAssetData
 {
 	UIKRetargeter* InRetargeter = Cast<UIKRetargeter>(InAssetData.GetAsset());
 	BatchContext.IKRetargetAsset = InRetargeter;
-	const UIKRigDefinition* SourceIKRig = InRetargeter ? InRetargeter->SourceIKRigAsset : nullptr;
-	const UIKRigDefinition* TargetIKRig = InRetargeter ? InRetargeter->TargetIKRigAsset : nullptr;
+	const UIKRigDefinition* SourceIKRig = InRetargeter ? InRetargeter->GetSourceIKRig() : nullptr;
+	const UIKRigDefinition* TargetIKRig = InRetargeter ? InRetargeter->GetTargetIKRig() : nullptr;
 	USkeletalMesh* SourceMesh =  SourceIKRig ? SourceIKRig->GetPreviewMesh() : nullptr;
 	USkeletalMesh* TargetMesh =  TargetIKRig ? TargetIKRig->GetPreviewMesh() : nullptr;
 	SourceMeshAssigned(FAssetData(SourceMesh));
diff --git a/Engine/Plugins/Animation/IKRig/Source/IKRigEditor/Private/RigEditor/IKRigAssetBrowserTabSummoner.cpp b/Engine/Plugins/Animation/IKRig/Source/IKRigEditor/Private/RigEditor/IKRigAssetBrowserTabSummoner.cpp
new file mode 100644
--- /dev/null
+++ b/Engine/Plugins/Animation/IKRig/Source/IKRigEditor/Private/RigEditor/IKRigAssetBrowserTabSummoner.cpp
@@ -0,0 +1,36 @@
+﻿// Copyright Epic Games, Inc. All Rights Reserved.
+#include "RigEditor/IKRigAssetBrowserTabSummoner.h"
+
+#include "IDocumentation.h"
+#include "RigEditor/IKRigToolkit.h"
+#include "RigEditor/IKRigEditorStyle.h"
+#include "RigEditor/SIKRigAssetBrowser.h"
+
+#define LOCTEXT_NAMESPACE "RetargetAssetBrowserTabSummoner"
+
+const FName FIKRigAssetBrowserTabSummoner::TabID(TEXT("AssetBrowser"));
+
+FIKRigAssetBrowserTabSummoner::FIKRigAssetBrowserTabSummoner(const TSharedRef<FIKRigEditorToolkit>& InRetargetEditor)
+	: FWorkflowTabFactory(TabID, InRetargetEditor),
+	IKRigEditor(InRetargetEditor)
+{
+	bIsSingleton = true; // only allow a single instance of this tab
+	
+	TabLabel = LOCTEXT("IKRigAssetBrowserTabLabel", "Asset Browser");
+	TabIcon = FSlateIcon(FIKRigEditorStyle::Get().GetStyleSetName(), "IKRig.TabIcon");
+
+	ViewMenuDescription = LOCTEXT("IKRigAssetBrowser_ViewMenu_Desc", "Asset Browser");
+	ViewMenuTooltip = LOCTEXT("IKRigAssetBrowser_ViewMenu_ToolTip", "Show the Asset Browser Tab");
+}
+
+TSharedPtr<SToolTip> FIKRigAssetBrowserTabSummoner::CreateTabToolTipWidget(const FWorkflowTabSpawnInfo& Info) const
+{
+	return  IDocumentation::Get()->CreateToolTip(LOCTEXT("IKRigAssetBrowserTooltip", "Select an animation asset to preview the IK results."), NULL, TEXT("Shared/Editors/Persona"), TEXT("IKRigAssetBrowser_Window"));
+}
+
+TSharedRef<SWidget> FIKRigAssetBrowserTabSummoner::CreateTabBody(const FWorkflowTabSpawnInfo& Info) const
+{
+	return SNew(SIKRigAssetBrowser, IKRigEditor.Pin()->GetController());
+}
+
+#undef LOCTEXT_NAMESPACE 
diff --git a/Engine/Plugins/Animation/IKRig/Source/IKRigEditor/Private/RigEditor/IKRigController.cpp b/Engine/Plugins/Animation/IKRig/Source/IKRigEditor/Private/RigEditor/IKRigController.cpp
--- a/Engine/Plugins/Animation/IKRig/Source/IKRigEditor/Private/RigEditor/IKRigController.cpp
+++ b/Engine/Plugins/Animation/IKRig/Source/IKRigEditor/Private/RigEditor/IKRigController.cpp
@@ -3,6 +3,7 @@
 #include "RigEditor/IKRigController.h"
 
 #include "IKRigDefinition.h"
+#include "IKRigProcessor.h"
 #include "IKRigSolver.h"
 
 #include "Engine/SkeletalMesh.h"
@@ -11,35 +12,32 @@
 #define LOCTEXT_NAMESPACE "IKRigController"
 
 
-TMap<UIKRigDefinition*, UIKRigController*> UIKRigController::AssetToControllerMap;
-
 UIKRigController* UIKRigController::GetIKRigController(UIKRigDefinition* InIKRigDefinition)
 {
 	if (!InIKRigDefinition)
 	{
 		return nullptr;
 	}
-	
-	UIKRigController** Controller = AssetToControllerMap.Find(InIKRigDefinition);
-	if (Controller)
+
+	if (!InIKRigDefinition->Controller)
 	{
-		return *Controller;
+		UIKRigController* Controller = NewObject<UIKRigController>();
+		Controller->Asset = InIKRigDefinition;
+		InIKRigDefinition->Controller = Controller;
+		InIKRigDefinition->IKRigEditUndo.AddUObject(Controller, &UIKRigController::BroadcastNeedsReinitialized);
 	}
 
-	UIKRigController* NewController = NewObject<UIKRigController>();
-	AssetToControllerMap.Add(InIKRigDefinition) = NewController;
-	NewController->IKRigAsset = InIKRigDefinition;
-	return NewController;
+	return Cast<UIKRigController>(InIKRigDefinition->Controller);
 }
 
 UIKRigDefinition* UIKRigController::GetAsset() const
 {
-	return IKRigAsset;
+	return Asset;
 }
 
 void UIKRigController::AddBoneSetting(const FName& BoneName, int32 SolverIndex) const
 {
-	check(IKRigAsset->Solvers.IsValidIndex(SolverIndex))
+	check(Asset->Solvers.IsValidIndex(SolverIndex))
 
 	if (!CanAddBoneSetting(BoneName, SolverIndex))
 	{
@@ -47,7 +45,7 @@ void UIKRigController::AddBoneSetting(const FName& BoneName, int32 SolverIndex)
 	}
 
 	FScopedTransaction Transaction(LOCTEXT("AddBoneSetting_Label", "Add Bone Setting"));
-	UIKRigSolver* Solver = IKRigAsset->Solvers[SolverIndex];
+	UIKRigSolver* Solver = Asset->Solvers[SolverIndex];
 	Solver->Modify();
 	Solver->AddBoneSetting(BoneName);
 	BroadcastNeedsReinitialized();
@@ -61,31 +59,31 @@ bool UIKRigController::CanAddBoneSetting(const FName& BoneName, int32 SolverInde
 		return false; // solver doesn't exist
 	}
 
-	if (IKRigAsset->Skeleton.GetBoneIndexFromName(BoneName) == INDEX_NONE)
+	if (Asset->Skeleton.GetBoneIndexFromName(BoneName) == INDEX_NONE)
 	{
 		return false; // bone doesn't exist
 	}
 
-	if (!IKRigAsset->Solvers[SolverIndex]->UsesBoneSettings())
+	if (!Asset->Solvers[SolverIndex]->UsesBoneSettings())
 	{
 		return false; // solver doesn't support per-bone settings
 	}
 
 	// returns true if the solver in question does NOT already have a settings object for this bone 
-	return !static_cast<bool>(IKRigAsset->Solvers[SolverIndex]->GetBoneSetting(BoneName));
+	return !static_cast<bool>(Asset->Solvers[SolverIndex]->GetBoneSetting(BoneName));
 }
 
 void UIKRigController::RemoveBoneSetting(const FName& BoneName, int32 SolverIndex) const
 {
-	check(IKRigAsset->Solvers.IsValidIndex(SolverIndex))
+	check(Asset->Solvers.IsValidIndex(SolverIndex))
 
-	if (IKRigAsset->Skeleton.GetBoneIndexFromName(BoneName) == INDEX_NONE)
+	if (Asset->Skeleton.GetBoneIndexFromName(BoneName) == INDEX_NONE)
 	{
 		return; // bone doesn't exist
 	}
 
 	FScopedTransaction Transaction(LOCTEXT("RemoveBoneSetting_Label", "Remove Bone Setting"));
-	UIKRigSolver* Solver = IKRigAsset->Solvers[SolverIndex];
+	UIKRigSolver* Solver = Asset->Solvers[SolverIndex];
 	Solver->Modify();
 	Solver->RemoveBoneSetting(BoneName);
 	BroadcastNeedsReinitialized();
@@ -104,7 +102,7 @@ bool UIKRigController::CanRemoveBoneSetting(const FName& BoneName, int32 SolverI
 		return false; // solver doesn't use bone settings
 	}
 
-	if (IKRigAsset->Skeleton.GetBoneIndexFromName(BoneName) == INDEX_NONE)
+	if (Asset->Skeleton.GetBoneIndexFromName(BoneName) == INDEX_NONE)
 	{
 		return false; // bone doesn't exist
 	}
@@ -130,12 +128,12 @@ UObject* UIKRigController::GetSettingsForBone(const FName& BoneName, int32 Solve
 
 bool UIKRigController::DoesBoneHaveSettings(const FName& BoneName) const
 {
-	if (IKRigAsset->Skeleton.GetBoneIndexFromName(BoneName) == INDEX_NONE)
+	if (Asset->Skeleton.GetBoneIndexFromName(BoneName) == INDEX_NONE)
 	{
 		return false; // bone doesn't exist
 	}
 	
-	for (UIKRigSolver* Solver : IKRigAsset->Solvers)
+	for (UIKRigSolver* Solver : Asset->Solvers)
 	{
 		if (UObject* BoneSetting = Solver->GetBoneSetting(BoneName))
 		{
@@ -148,15 +146,15 @@ bool UIKRigController::DoesBoneHaveSettings(const FName& BoneName) const
 
 bool UIKRigController::AddRetargetChain(const FName& ChainName, const FName& StartBone, const FName& EndBone) const
 {
-	if (IKRigAsset->GetRetargetChainByName(ChainName))
+	if (Asset->GetRetargetChainByName(ChainName))
 	{
 		return false; // bone chain already exists with that name
 	}
 
 	FScopedTransaction Transaction(LOCTEXT("AddRetargetChain_Label", "Add Retarget Chain"));
-	IKRigAsset->Modify();
+	Asset->Modify();
 	
-	IKRigAsset->RetargetDefinition.BoneChains.Emplace(ChainName, StartBone, EndBone);
+	Asset->RetargetDefinition.BoneChains.Emplace(ChainName, StartBone, EndBone);
 	SortRetargetChains();
 	BroadcastNeedsReinitialized();
 	return true;
@@ -165,7 +163,7 @@ bool UIKRigController::AddRetargetChain(const FName& ChainName, const FName& Sta
 bool UIKRigController::RemoveRetargetChain(const FName& ChainName) const
 {
 	FScopedTransaction Transaction(LOCTEXT("RemoveRetargetChain_Label", "Remove Retarget Chain"));
-	IKRigAsset->Modify();
+	Asset->Modify();
 	
 	auto Pred = [ChainName](const FBoneChain& Element)
 	{
@@ -177,7 +175,7 @@ bool UIKRigController::RemoveRetargetChain(const FName& ChainName) const
 		return false;
 	};
 	
-	if (IKRigAsset->RetargetDefinition.BoneChains.RemoveAll(Pred) > 0)
+	if (Asset->RetargetDefinition.BoneChains.RemoveAll(Pred) > 0)
 	{
 		SortRetargetChains();
 		BroadcastNeedsReinitialized();
@@ -189,19 +187,19 @@ bool UIKRigController::RemoveRetargetChain(const FName& ChainName) const
 
 FName UIKRigController::RenameRetargetChain(const FName& ChainName, const FName& NewChainName) const
 {
-	FBoneChain* Chain = IKRigAsset->RetargetDefinition.GetEditableBoneChainByName(ChainName);
+	FBoneChain* Chain = Asset->RetargetDefinition.GetEditableBoneChainByName(ChainName);
 	if (!Chain)
 	{
 		return ChainName; // chain doesn't exist to rename
 	}
 
-	if (IKRigAsset->GetRetargetChainByName(NewChainName))
+	if (Asset->GetRetargetChainByName(NewChainName))
 	{
 		return ChainName; // bone chain already exists with the new name
 	}
 	
 	FScopedTransaction Transaction(LOCTEXT("RenameRetargetChain_Label", "Rename Retarget Chain"));
-	IKRigAsset->Modify();
+	Asset->Modify();
 	Chain->ChainName = NewChainName;
 	RetargetChainRenamed.Broadcast(GetAsset(), ChainName, NewChainName);
 	BroadcastNeedsReinitialized();
@@ -210,10 +208,10 @@ FName UIKRigController::RenameRetargetChain(const FName& ChainName, const FName&
 
 bool UIKRigController::SetRetargetChainStartBone(const FName& ChainName, const FName& StartBoneName) const
 {
-	if (FBoneChain* BoneChain = IKRigAsset->RetargetDefinition.GetEditableBoneChainByName(ChainName))
+	if (FBoneChain* BoneChain = Asset->RetargetDefinition.GetEditableBoneChainByName(ChainName))
 	{
 		FScopedTransaction Transaction(LOCTEXT("SetRetargetChainStartBone_Label", "Set Retarget Chain Start Bone"));
-		IKRigAsset->Modify();
+		Asset->Modify();
 		BoneChain->StartBone = StartBoneName;
 		SortRetargetChains();
 		BroadcastNeedsReinitialized();
@@ -225,10 +223,10 @@ bool UIKRigController::SetRetargetChainStartBone(const FName& ChainName, const F
 
 bool UIKRigController::SetRetargetChainEndBone(const FName& ChainName, const FName& EndBoneName) const
 {
-	if (FBoneChain* BoneChain = IKRigAsset->RetargetDefinition.GetEditableBoneChainByName(ChainName))
+	if (FBoneChain* BoneChain = Asset->RetargetDefinition.GetEditableBoneChainByName(ChainName))
 	{
 		FScopedTransaction Transaction(LOCTEXT("SetRetargetChainEndBone_Label", "Set Retarget Chain End Bone"));
-		IKRigAsset->Modify();
+		Asset->Modify();
 		BoneChain->EndBone = EndBoneName;
 		BroadcastNeedsReinitialized();
 		return true;
@@ -245,10 +243,10 @@ bool UIKRigController::SetRetargetChainGoal(const FName& ChainName, const FName&
 		GoalNameToUse = NAME_None; // no goal with that name	
 	}
 	
-	if (FBoneChain* BoneChain = IKRigAsset->RetargetDefinition.GetEditableBoneChainByName(ChainName))
+	if (FBoneChain* BoneChain = Asset->RetargetDefinition.GetEditableBoneChainByName(ChainName))
 	{
 		FScopedTransaction Transaction(LOCTEXT("SetRetargetChainGoal_Label", "Set Retarget Chain Goal"));
-		IKRigAsset->Modify();
+		Asset->Modify();
 		BoneChain->IKGoalName = GoalNameToUse;
 		BroadcastNeedsReinitialized();
 		return true;
@@ -259,8 +257,8 @@ bool UIKRigController::SetRetargetChainGoal(const FName& ChainName, const FName&
 
 FName UIKRigController::GetRetargetChainGoal(const FName& ChainName) const
 {
-	check(IKRigAsset)
-	const FBoneChain* Chain = IKRigAsset->GetRetargetChainByName(ChainName);
+	check(Asset)
+	const FBoneChain* Chain = Asset->GetRetargetChainByName(ChainName);
 	if (!Chain)
 	{
 		return NAME_None;
@@ -271,8 +269,8 @@ FName UIKRigController::GetRetargetChainGoal(const FName& ChainName) const
 
 FName UIKRigController::GetRetargetChainStartBone(const FName& ChainName) const
 {
-	check(IKRigAsset)
-	const FBoneChain* Chain = IKRigAsset->GetRetargetChainByName(ChainName);
+	check(Asset)
+	const FBoneChain* Chain = Asset->GetRetargetChainByName(ChainName);
 	if (!Chain)
 	{
 		return NAME_None;
@@ -283,8 +281,8 @@ FName UIKRigController::GetRetargetChainStartBone(const FName& ChainName) const
 
 FName UIKRigController::GetRetargetChainEndBone(const FName& ChainName) const
 {
-	check(IKRigAsset)
-	const FBoneChain* Chain = IKRigAsset->GetRetargetChainByName(ChainName);
+	check(Asset)
+	const FBoneChain* Chain = Asset->GetRetargetChainByName(ChainName);
 	if (!Chain)
 	{
 		return NAME_None;
@@ -295,123 +293,112 @@ FName UIKRigController::GetRetargetChainEndBone(const FName& ChainName) const
 
 const TArray<FBoneChain>& UIKRigController::GetRetargetChains() const
 {
-	check(IKRigAsset)
-	return IKRigAsset->GetRetargetChains();
+	check(Asset)
+	return Asset->GetRetargetChains();
 }
 
 void UIKRigController::SetRetargetRoot(const FName& RootBoneName) const
 {
-	check(IKRigAsset)
+	check(Asset)
 
 	FScopedTransaction Transaction(LOCTEXT("SetRetargetRootBone_Label", "Set Retarget Root Bone"));
-	IKRigAsset->Modify();
+	Asset->Modify();
 	
-	IKRigAsset->RetargetDefinition.RootBone = RootBoneName;
+	Asset->RetargetDefinition.RootBone = RootBoneName;
 
 	BroadcastNeedsReinitialized();
 }
 
 FName UIKRigController::GetRetargetRoot() const
 {
-	check(IKRigAsset)
+	check(Asset)
 
-	return IKRigAsset->RetargetDefinition.RootBone;
+	return Asset->RetargetDefinition.RootBone;
 }
 
 void UIKRigController::SortRetargetChains() const
 {
-	IKRigAsset->RetargetDefinition.BoneChains.Sort([this](const FBoneChain& A, const FBoneChain& B)
+	Asset->RetargetDefinition.BoneChains.Sort([this](const FBoneChain& A, const FBoneChain& B)
 	{
-		const int32 IndexA = IKRigAsset->Skeleton.GetBoneIndexFromName(A.StartBone);
-		const int32 IndexB = IKRigAsset->Skeleton.GetBoneIndexFromName(B.StartBone);
+		const int32 IndexA = Asset->Skeleton.GetBoneIndexFromName(A.StartBone);
+		const int32 IndexB = Asset->Skeleton.GetBoneIndexFromName(B.StartBone);
 		return IndexA < IndexB;
 	});
 }
 
-void UIKRigController::AddReferencedObjects(UObject* InThis, FReferenceCollector& Collector)
-{
-	// since static member, we just add only for default object
-	if (InThis && InThis->IsTemplate())
-	{
-		for (auto Iter = AssetToControllerMap.CreateIterator(); Iter; ++Iter)
-		{
-			// we add controllers to the list. so that it doesn't GCed
-			// when do we delete from the list?
-			UIKRigController* Obj = Iter.Value();
-			Collector.AddReferencedObject(Obj);
-		}
-	}
-}
-
 // -------------------------------------------------------
 // SKELETON
 //
 
-void UIKRigController::SetSourceSkeletalMesh(USkeletalMesh* SkeletalMesh, bool bReImportBones) const
+bool UIKRigController::SetSkeletalMesh(USkeletalMesh* SkeletalMesh) const
 {
+	if (!SkeletalMesh)
+	{
+		return false;
+	}
+	
+	// first determine runtime compatibility between the IK Rig asset and the skeleton we're trying to run it on
+	const FReferenceSkeleton& RefSkeleton = SkeletalMesh->GetRefSkeleton();
+	const FIKRigInputSkeleton InputSkeleton = FIKRigInputSkeleton(RefSkeleton);
+	if (!UIKRigProcessor::IsIKRigCompatibleWithSkeleton(GetAsset(), InputSkeleton))
+	{
+		UE_LOG(LogTemp, Warning, TEXT("Trying to initialize IKRig with a Skeleton that is missing required bones. See output log. %s"), *GetAsset()->GetName());
+		return false;
+	}
+	
 	FScopedTransaction Transaction(LOCTEXT("SetSkeletalMesh_Label", "Set Skeletal Mesh"));
-	IKRigAsset->Modify();
+	Asset->Modify();
 
 	// update stored skeletal mesh used for previewing results
-	IKRigAsset->PreviewSkeletalMesh = SkeletalMesh;
-
-	if (bReImportBones)
+	Asset->PreviewSkeletalMesh = SkeletalMesh;
+	// copy skeleton data from the actual skeleton we want to run on
+	Asset->Skeleton.SetInputSkeleton(InputSkeleton, GetAsset()->Skeleton.ExcludedBones);
+	// update goal's initial transforms to reflect new
+	for (UIKRigEffectorGoal* Goal : Asset->Goals)
 	{
-		// reimport the skeleton data
-		if (SkeletalMesh)
-		{
-			SetSkeleton(SkeletalMesh->GetRefSkeleton());		
-		}
-		else
-		{
-			IKRigAsset->Skeleton.Reset();
-		}
+		Goal->Modify();
+		const FTransform InitialTransform = GetRefPoseTransformOfBone(Goal->BoneName);
+		Goal->InitialTransform = InitialTransform;
 	}
 
 	BroadcastNeedsReinitialized();
-}
 
-void UIKRigController::SetSkeleton(const FReferenceSkeleton& InSkeleton) const
-{
-	FScopedTransaction Transaction(LOCTEXT("SetSkeleton_Label", "Set Skeleton"));
-	IKRigAsset->Modify();
-	IKRigAsset->Skeleton.Initialize(InSkeleton, IKRigAsset->Skeleton.ExcludedBones);
-	BroadcastNeedsReinitialized();
+	return true;
 }
 
 const FIKRigSkeleton& UIKRigController::GetIKRigSkeleton() const
 {
-	return IKRigAsset->Skeleton;
+	return Asset->Skeleton;
 }
 
 USkeleton* UIKRigController::GetSkeleton() const
 {
-	if (IKRigAsset->PreviewSkeletalMesh.IsNull())
+	if (Asset->PreviewSkeletalMesh.IsNull())
     {
         return nullptr;
     }
     
-    return IKRigAsset->PreviewSkeletalMesh->GetSkeleton();
+    return Asset->PreviewSkeletalMesh->GetSkeleton();
 }
 
 void UIKRigController::SetBoneExcluded(const FName& BoneName, const bool bExclude) const
 {
-	const bool bIsExcluded = IKRigAsset->Skeleton.ExcludedBones.Contains(BoneName);
+	const bool bIsExcluded = Asset->Skeleton.ExcludedBones.Contains(BoneName);
 	if (bIsExcluded == bExclude)
 	{
 		return; // already in the requested state of exclusion
 	}
 
 	FScopedTransaction Transaction(LOCTEXT("SetBoneExcluded_Label", "Set Bone Excluded"));
-	IKRigAsset->Modify();
+	Asset->Modify();
 
 	if (bExclude)
 	{
-		IKRigAsset->Skeleton.ExcludedBones.Add(BoneName);
+		Asset->Skeleton.ExcludedBones.Add(BoneName);
 	}
 	else
 	{
-		IKRigAsset->Skeleton.ExcludedBones.Remove(BoneName);
+		Asset->Skeleton.ExcludedBones.Remove(BoneName);
 	}
 
 	BroadcastNeedsReinitialized();
@@ -419,19 +406,19 @@ void UIKRigController::SetBoneExcluded(const FName& BoneName, const bool bExclud
 
 bool UIKRigController::GetBoneExcluded(const FName& BoneName) const
 {
-	return IKRigAsset->Skeleton.ExcludedBones.Contains(BoneName);
+	return Asset->Skeleton.ExcludedBones.Contains(BoneName);
 }
 
-FTransform UIKRigController::GetBoneRetargetPose(const FName& BoneName) const
+FTransform UIKRigController::GetRefPoseTransformOfBone(const FName& BoneName) const
 {	
-	const int32 BoneIndex = IKRigAsset->Skeleton.GetBoneIndexFromName(BoneName);
+	const int32 BoneIndex = Asset->Skeleton.GetBoneIndexFromName(BoneName);
 	check(BoneIndex != INDEX_NONE) // must initialize IK Rig before getting here
-	return IKRigAsset->Skeleton.RefPoseGlobal[BoneIndex];
+	return Asset->Skeleton.RefPoseGlobal[BoneIndex];
 }
 
-USkeletalMesh* UIKRigController::GetSourceSkeletalMesh() const
+USkeletalMesh* UIKRigController::GetSkeletalMesh() const
 {
-	return IKRigAsset->PreviewSkeletalMesh.Get();
+	return Asset->PreviewSkeletalMesh.Get();
 }
 
 // -------------------------------------------------------
@@ -440,15 +427,15 @@ USkeletalMesh* UIKRigController::GetSourceSkeletalMesh() const
 
 int32 UIKRigController::AddSolver(TSubclassOf<UIKRigSolver> InIKRigSolverClass) const
 {
-	check(IKRigAsset)
+	check(Asset)
 
 	FScopedTransaction Transaction(LOCTEXT("AddSolver_Label", "Add Solver"));
-	IKRigAsset->Modify();
+	Asset->Modify();
 
-	UIKRigSolver* NewSolver = NewObject<UIKRigSolver>(IKRigAsset, InIKRigSolverClass, NAME_None, RF_Transactional);
+	UIKRigSolver* NewSolver = NewObject<UIKRigSolver>(Asset, InIKRigSolverClass, NAME_None, RF_Transactional);
 	check(NewSolver);
 
-	const int32 SolverIndex = IKRigAsset->Solvers.Add(NewSolver);
+	const int32 SolverIndex = Asset->Solvers.Add(NewSolver);
 
 	BroadcastNeedsReinitialized();
 
@@ -457,40 +444,40 @@ int32 UIKRigController::AddSolver(TSubclassOf<UIKRigSolver> InIKRigSolverClass)
 
 void UIKRigController::RemoveSolver(const int32 SolverIndex) const
 {
-	check(IKRigAsset)
+	check(Asset)
 	
-	if (!IKRigAsset->Solvers.IsValidIndex(SolverIndex))
+	if (!Asset->Solvers.IsValidIndex(SolverIndex))
 	{
 		return;
 	}
 
 	FScopedTransaction Transaction(LOCTEXT("RemoveSolver_Label", "Remove Solver"));
-	IKRigAsset->Modify();
+	Asset->Modify();
 
-	IKRigAsset->Solvers.RemoveAt(SolverIndex);
+	Asset->Solvers.RemoveAt(SolverIndex);
 
 	BroadcastNeedsReinitialized();
 }
 
 bool UIKRigController::MoveSolverInStack(int32 SolverToMoveIndex, int32 TargetSolverIndex) const
 {
-	if (!IKRigAsset->Solvers.IsValidIndex(SolverToMoveIndex))
+	if (!Asset->Solvers.IsValidIndex(SolverToMoveIndex))
 	{
 		return false;
 	}
 
-	if (!IKRigAsset->Solvers.IsValidIndex(TargetSolverIndex))
+	if (!Asset->Solvers.IsValidIndex(TargetSolverIndex))
 	{
 		return false;
 	}
 
 	FScopedTransaction Transaction(LOCTEXT("ReorderSolver_Label", "Reorder Solvers"));
-	IKRigAsset->Modify();
+	Asset->Modify();
 
-	UIKRigSolver* SolverToMove = IKRigAsset->Solvers[SolverToMoveIndex];
-	IKRigAsset->Solvers.Insert(SolverToMove, TargetSolverIndex + 1);
+	UIKRigSolver* SolverToMove = Asset->Solvers[SolverToMoveIndex];
+	Asset->Solvers.Insert(SolverToMove, TargetSolverIndex + 1);
 	const int32 SolverToRemove = TargetSolverIndex > SolverToMoveIndex ? SolverToMoveIndex : SolverToMoveIndex + 1;
-	IKRigAsset->Solvers.RemoveAt(SolverToRemove);
+	Asset->Solvers.RemoveAt(SolverToRemove);
 
 	BroadcastNeedsReinitialized();
 	
@@ -499,15 +486,15 @@ bool UIKRigController::MoveSolverInStack(int32 SolverToMoveIndex, int32 TargetSo
 
 bool UIKRigController::SetSolverEnabled(int32 SolverIndex, bool bIsEnabled) const
 {
-	if (!IKRigAsset->Solvers.IsValidIndex(SolverIndex))
+	if (!Asset->Solvers.IsValidIndex(SolverIndex))
 	{
 		return false;
 	}
 	
 	FScopedTransaction Transaction(LOCTEXT("SetSolverEndabled_Label", "Enable/Disable Solver"));
-	UIKRigSolver* Solver = IKRigAsset->Solvers[SolverIndex];
+	UIKRigSolver* Solver = Asset->Solvers[SolverIndex];
 	Solver->Modify();
-	IKRigAsset->Modify();
+	Asset->Modify();
 
 	Solver->SetEnabled(bIsEnabled);
 
@@ -518,18 +505,18 @@ bool UIKRigController::SetSolverEnabled(int32 SolverIndex, bool bIsEnabled) cons
 
 void UIKRigController::SetRootBone(const FName& RootBoneName, int32 SolverIndex) const
 {
-	if (!IKRigAsset->Solvers.IsValidIndex(SolverIndex))
+	if (!Asset->Solvers.IsValidIndex(SolverIndex))
 	{
 		return; // solver doesn't exist
 	}
 
-	if (IKRigAsset->Skeleton.GetBoneIndexFromName(RootBoneName) == INDEX_NONE)
+	if (Asset->Skeleton.GetBoneIndexFromName(RootBoneName) == INDEX_NONE)
 	{
 		return; // bone doesn't exist
 	}
 
 	FScopedTransaction Transaction(LOCTEXT("SetRootBone_Label", "Set Root Bone"));
-	UIKRigSolver* Solver = IKRigAsset->Solvers[SolverIndex];
+	UIKRigSolver* Solver = Asset->Solvers[SolverIndex];
 	Solver->Modify();
 
 	Solver->SetRootBone(RootBoneName);
@@ -539,22 +526,35 @@ void UIKRigController::SetRootBone(const FName& RootBoneName, int32 SolverIndex)
 
 const TArray<UIKRigSolver*>& UIKRigController::GetSolverArray() const
 {
-	return IKRigAsset->Solvers;
+	return Asset->Solvers;
+}
+
+FString UIKRigController::GetSolverUniqueName(int32 SolverIndex)
+{
+	check(Asset)
+	if (!Asset->Solvers.IsValidIndex(SolverIndex))
+	{
+		checkNoEntry()
+		return "";
+	}
+
+	return FString::FromInt(SolverIndex+1) + " - " + Asset->Solvers[SolverIndex]->GetNiceName().ToString();
+		
 }
 
 int32 UIKRigController::GetNumSolvers() const
 {
-	check(IKRigAsset)
-	return IKRigAsset->Solvers.Num();
+	check(Asset)
+	return Asset->Solvers.Num();
 }
 
 UIKRigSolver* UIKRigController::GetSolver(int32 Index) const
 {
-	check(IKRigAsset)
+	check(Asset)
 
-	if (IKRigAsset->Solvers.IsValidIndex(Index))
+	if (Asset->Solvers.IsValidIndex(Index))
 	{
-		return IKRigAsset->Solvers[Index];
+		return Asset->Solvers[Index];
 	}
 	
 	return nullptr;
@@ -572,15 +572,15 @@ UIKRigEffectorGoal* UIKRigController::AddNewGoal(const FName& GoalName, const FN
 	}
 	
 	FScopedTransaction Transaction(LOCTEXT("AddSetting_Label", "Add Bone Setting"));
-	IKRigAsset->Modify();
+	Asset->Modify();
 
-	UIKRigEffectorGoal* NewGoal = NewObject<UIKRigEffectorGoal>(IKRigAsset, UIKRigEffectorGoal::StaticClass(), NAME_None, RF_Transactional);
+	UIKRigEffectorGoal* NewGoal = NewObject<UIKRigEffectorGoal>(Asset, UIKRigEffectorGoal::StaticClass(), NAME_None, RF_Transactional);
 	NewGoal->BoneName = BoneName;
 	NewGoal->GoalName = GoalName;
-	IKRigAsset->Goals.Add(NewGoal);
+	Asset->Goals.Add(NewGoal);
 
 	// set initial transform
-	NewGoal->InitialTransform = GetBoneRetargetPose(NewGoal->BoneName);
+	NewGoal->InitialTransform = GetRefPoseTransformOfBone(NewGoal->BoneName);
 	NewGoal->CurrentTransform = NewGoal->InitialTransform;
 
 	BroadcastNeedsReinitialized();
@@ -599,19 +599,19 @@ bool UIKRigController::RemoveGoal(const FName& GoalName) const
 	FScopedTransaction Transaction(LOCTEXT("RemoveGoal_Label", "Remove Goal"));
 
 	// remove from all the solvers
-	const FName& GoalToRemove = IKRigAsset->Goals[GoalIndex]->GoalName;
-	for (UIKRigSolver* Solver : IKRigAsset->Solvers)
+	const FName& GoalToRemove = Asset->Goals[GoalIndex]->GoalName;
+	for (UIKRigSolver* Solver : Asset->Solvers)
 	{
 		Solver->Modify();
 		Solver->RemoveGoal(GoalToRemove);
 	}
 
 	// remove from core system
-	IKRigAsset->Modify();
-	IKRigAsset->Goals.RemoveAt(GoalIndex);
+	Asset->Modify();
+	Asset->Goals.RemoveAt(GoalIndex);
 
 	// clean any retarget chains that might reference the missing goal
-	for (FBoneChain& BoneChain : IKRigAsset->RetargetDefinition.BoneChains)
+	for (FBoneChain& BoneChain : Asset->RetargetDefinition.BoneChains)
 	{
 		if (BoneChain.IKGoalName == GoalName)
 		{
@@ -644,14 +644,14 @@ FName UIKRigController::RenameGoal(const FName& OldName, const FName& PotentialN
 	}
 	
 	FScopedTransaction Transaction(LOCTEXT("RenameGoal_Label", "Rename Goal"));
-	IKRigAsset->Modify();
+	Asset->Modify();
 
 	// rename in core
-	IKRigAsset->Goals[GoalIndex]->Modify();
-	IKRigAsset->Goals[GoalIndex]->GoalName = NewName;
+	Asset->Goals[GoalIndex]->Modify();
+	Asset->Goals[GoalIndex]->GoalName = NewName;
 
 	// update any retarget chains that might reference the goal
-	for (FBoneChain& BoneChain : IKRigAsset->RetargetDefinition.BoneChains)
+	for (FBoneChain& BoneChain : Asset->RetargetDefinition.BoneChains)
 	{
 		if (BoneChain.IKGoalName == OldName)
 		{
@@ -660,7 +660,7 @@ FName UIKRigController::RenameGoal(const FName& OldName, const FName& PotentialN
 	}
 
 	// rename in solvers
-	for (UIKRigSolver* Solver : IKRigAsset->Solvers)
+	for (UIKRigSolver* Solver : Asset->Solvers)
 	{
 		Solver->Modify();
 		Solver->RenameGoal(OldName, NewName);
@@ -671,6 +671,17 @@ FName UIKRigController::RenameGoal(const FName& OldName, const FName& PotentialN
 	return NewName;
 }
 
+bool UIKRigController::ModifyGoal(const FName& GoalName) const
+{
+	if (UIKRigEffectorGoal* Goal = GetGoal(GoalName))
+	{
+		Goal->Modify();
+		return true;
+	}
+	
+	return false;
+}
+
 bool UIKRigController::SetGoalBone(const FName& GoalName, const FName& NewBoneName) const
 {
 	const int32 GoalIndex = GetGoalIndex(GoalName);
@@ -691,13 +702,13 @@ bool UIKRigController::SetGoalBone(const FName& GoalName, const FName& NewBoneNa
 	}
 	
 	FScopedTransaction Transaction(LOCTEXT("RenameGoal_Label", "Rename Goal"));
-	IKRigAsset->Modify();
+	Asset->Modify();
 
 	// update goal
-	IKRigAsset->Goals[GoalIndex]->BoneName = NewBoneName;
+	Asset->Goals[GoalIndex]->BoneName = NewBoneName;
 	
 	// update in solvers
-	for (UIKRigSolver* Solver : IKRigAsset->Solvers)
+	for (UIKRigSolver* Solver : Asset->Solvers)
 	{
 		Solver->Modify();
 		Solver->SetGoalBone(GoalName, NewBoneName);
@@ -713,7 +724,7 @@ bool UIKRigController::SetGoalBone(const FName& GoalName, const FName& NewBoneNa
 
 FName UIKRigController::GetBoneForGoal(const FName& GoalName) const
 {
-	for (const UIKRigEffectorGoal* Goal : IKRigAsset->Goals)
+	for (const UIKRigEffectorGoal* Goal : Asset->Goals)
 	{
 		if (Goal->GoalName == GoalName)
 		{
@@ -729,10 +740,10 @@ bool UIKRigController::ConnectGoalToSolver(const UIKRigEffectorGoal& Goal, int32
 	// can't add goal that is not present in the core
 	check(GetGoalIndex(Goal.GoalName) != INDEX_NONE);
 	// can't add goal to a solver with an invalid index
-	check(IKRigAsset->Solvers.IsValidIndex(SolverIndex))
+	check(Asset->Solvers.IsValidIndex(SolverIndex))
 
 	FScopedTransaction Transaction(LOCTEXT("AddGoalToSolver_Label", "Add Goal"));
-	UIKRigSolver* Solver = IKRigAsset->Solvers[SolverIndex];
+	UIKRigSolver* Solver = Asset->Solvers[SolverIndex];
 	Solver->Modify();
 	
 	Solver->AddGoal(Goal);
@@ -746,10 +757,10 @@ bool UIKRigController::DisconnectGoalFromSolver(const FName& GoalToRemove, int32
 	// can't remove goal that is not present in the core
 	check(GetGoalIndex(GoalToRemove) != INDEX_NONE);
 	// can't remove goal from a solver with an invalid index
-	check(IKRigAsset->Solvers.IsValidIndex(SolverIndex))
+	check(Asset->Solvers.IsValidIndex(SolverIndex))
 
     FScopedTransaction Transaction(LOCTEXT("RemoveGoalFromSolver_Label", "Remove Goal"));
-	UIKRigSolver* Solver = IKRigAsset->Solvers[SolverIndex];
+	UIKRigSolver* Solver = Asset->Solvers[SolverIndex];
 	Solver->Modify();
 	
 	Solver->RemoveGoal(GoalToRemove);
@@ -760,30 +771,30 @@ bool UIKRigController::DisconnectGoalFromSolver(const FName& GoalToRemove, int32
 
 bool UIKRigController::IsGoalConnectedToSolver(const FName& GoalName, int32 SolverIndex) const
 {
-	if (!IKRigAsset->Solvers.IsValidIndex(SolverIndex))
+	if (!Asset->Solvers.IsValidIndex(SolverIndex))
 	{
 		return false;
 	}
 
-	return IKRigAsset->Solvers[SolverIndex]->IsGoalConnected(GoalName);
+	return Asset->Solvers[SolverIndex]->IsGoalConnected(GoalName);
 }
 
 const TArray<UIKRigEffectorGoal*>& UIKRigController::GetAllGoals() const
 {
-	return IKRigAsset->Goals;
+	return Asset->Goals;
 }
 
 const UIKRigEffectorGoal* UIKRigController::GetGoal(int32 GoalIndex) const
 {
-	if (!IKRigAsset->Goals.IsValidIndex(GoalIndex))
+	if (!Asset->Goals.IsValidIndex(GoalIndex))
 	{
 		return nullptr;
 	}
 
-	return IKRigAsset->Goals[GoalIndex];
+	return Asset->Goals[GoalIndex];
 }
 
-const UIKRigEffectorGoal* UIKRigController::GetGoal(const FName& GoalName) const
+UIKRigEffectorGoal* UIKRigController::GetGoal(const FName& GoalName) const
 {
 	const int32 GoalIndex = GetGoalIndex(GoalName);
 	if (GoalIndex == INDEX_NONE)
@@ -791,10 +802,10 @@ const UIKRigEffectorGoal* UIKRigController::GetGoal(const FName& GoalName) const
 		return nullptr;
 	}
 	
-	return IKRigAsset->Goals[GoalIndex];
+	return Asset->Goals[GoalIndex];
 }
 
-UObject* UIKRigController::GetEffectorForGoal(const FName& GoalName, int32 SolverIndex) const
+UObject* UIKRigController::GetGoalSettingsForSolver(const FName& GoalName, int32 SolverIndex) const
 {
 	const int32 GoalIndex = GetGoalIndex(GoalName);
 	if (GoalIndex == INDEX_NONE)
@@ -802,12 +813,12 @@ UObject* UIKRigController::GetEffectorForGoal(const FName& GoalName, int32 Solve
 		return nullptr; // no goal with that index
 	}
 
-	if (!IKRigAsset->Solvers.IsValidIndex(SolverIndex))
+	if (!Asset->Solvers.IsValidIndex(SolverIndex))
 	{
 		return nullptr; // no solver with that index
 	}
 
-	return IKRigAsset->Solvers[SolverIndex]->GetEffectorWithGoal(GoalName);
+	return Asset->Solvers[SolverIndex]->GetGoalSettings(GoalName);
 }
 
 FTransform UIKRigController::GetGoalCurrentTransform(const FName& GoalName) const
@@ -822,7 +833,7 @@ FTransform UIKRigController::GetGoalCurrentTransform(const FName& GoalName) cons
 
 void UIKRigController::SetGoalCurrentTransform(const FName& GoalName, const FTransform& Transform) const
 {
-	UIKRigEffectorGoal* Goal = const_cast<UIKRigEffectorGoal*>(GetGoal(GoalName));
+	UIKRigEffectorGoal* Goal = GetGoal(GoalName);
 	check(Goal);
 
 	Goal->CurrentTransform = Transform;
@@ -830,12 +841,15 @@ void UIKRigController::SetGoalCurrentTransform(const FName& GoalName, const FTra
 
 void UIKRigController::ResetGoalTransforms() const
 {
-	for (UIKRigEffectorGoal* Goal : IKRigAsset->Goals)
+	FScopedTransaction Transaction(LOCTEXT("ResetGoalTransforms", "Reset All Goal Transforms"));
+	
+	for (UIKRigEffectorGoal* Goal : Asset->Goals)
     {
-    	const FTransform InitialTransform = GetBoneRetargetPose(Goal->BoneName);
+		Goal->Modify();
+    	const FTransform InitialTransform = GetRefPoseTransformOfBone(Goal->BoneName);
     	Goal->InitialTransform = InitialTransform;
     	Goal->CurrentTransform = InitialTransform;
-    }	
+    }
 }
 
 void UIKRigController::SanitizeGoalName(FString& InOutName)
@@ -864,9 +878,9 @@ void UIKRigController::SanitizeGoalName(FString& InOutName)
 
 int32 UIKRigController::GetGoalIndex(const FName& GoalName) const
 {
-	for (int32 i=0; i<IKRigAsset->Goals.Num(); ++i)
+	for (int32 i=0; i<Asset->Goals.Num(); ++i)
 	{	
-		if (IKRigAsset->Goals[i]->GoalName == GoalName)
+		if (Asset->Goals[i]->GoalName == GoalName)
 		{
 			return i;
 		}
@@ -877,12 +891,12 @@ int32 UIKRigController::GetGoalIndex(const FName& GoalName) const
 
 FName UIKRigController::GetGoalName(const int32& GoalIndex) const
 {
-	if (!IKRigAsset->Goals.IsValidIndex(GoalIndex))
+	if (!Asset->Goals.IsValidIndex(GoalIndex))
 	{
 		return NAME_None;
 	}
 
-	return IKRigAsset->Goals[GoalIndex]->GoalName;
+	return Asset->Goals[GoalIndex]->GoalName;
 }
 
 #undef LOCTEXT_NAMESPACE
diff --git a/Engine/Plugins/Animation/IKRig/Source/IKRigEditor/Private/RigEditor/IKRigDefinitionFactory.cpp b/Engine/Plugins/Animation/IKRig/Source/IKRigEditor/Private/RigEditor/IKRigDefinitionFactory.cpp
--- a/Engine/Plugins/Animation/IKRig/Source/IKRigEditor/Private/RigEditor/IKRigDefinitionFactory.cpp
+++ b/Engine/Plugins/Animation/IKRig/Source/IKRigEditor/Private/RigEditor/IKRigDefinitionFactory.cpp
@@ -38,7 +38,7 @@ UObject* UIKRigDefinitionFactory::FactoryCreateNew(
 	
 	// imports the skeleton data into the IK Rig
 	UIKRigController* Controller = UIKRigController::GetIKRigController(IKRig);
-	Controller->SetSourceSkeletalMesh(SkeletalMesh, true);
+	Controller->SetSkeletalMesh(SkeletalMesh);
 	
 	return IKRig;
 }
diff --git a/Engine/Plugins/Animation/IKRig/Source/IKRigEditor/Private/RigEditor/IKRigEditMode.cpp b/Engine/Plugins/Animation/IKRig/Source/IKRigEditor/Private/RigEditor/IKRigEditMode.cpp
--- a/Engine/Plugins/Animation/IKRig/Source/IKRigEditor/Private/RigEditor/IKRigEditMode.cpp
+++ b/Engine/Plugins/Animation/IKRig/Source/IKRigEditor/Private/RigEditor/IKRigEditMode.cpp
@@ -10,6 +10,8 @@
 #include "RigEditor/IKRigHitProxies.h"
 #include "RigEditor/IKRigToolkit.h"
 
+#define LOCTEXT_NAMESPACE "IKRetargeterEditMode"
+
 FName FIKRigEditMode::ModeName("IKRigAssetEditMode");
 
 FIKRigEditMode::FIKRigEditMode()
@@ -29,10 +31,9 @@ bool FIKRigEditMode::GetCameraTarget(FSphere& OutTarget) const
 	if (!OutGoalNames.IsEmpty())
 	{
 		TArray<FVector> GoalPoints;
-		for (const auto& GoalName : OutGoalNames)
-		{
-			FTransform GoalTransform = Controller->AssetController->GetBoneRetargetPose(Controller->AssetController->GetGoal(GoalName)->BoneName);
-			GoalPoints.Add(GoalTransform.GetLocation());
+		for (const FName& GoalName : OutGoalNames)
+		{			
+			GoalPoints.Add(Controller->AssetController->GetGoal(GoalName)->CurrentTransform.GetLocation());
 		}
 
 		// create a sphere that contains all the goal points
@@ -194,7 +195,12 @@ bool FIKRigEditMode::StartTracking(FEditorViewportClient* InViewportClient, FVie
 	{
 		return false; // not manipulating a required axis
 	}
-	
+
+	GEditor->BeginTransaction(LOCTEXT("ManipulateGoal", "Manipulate IK Rig Goal"));
+	for (const FName& SelectedGoal : SelectedGoalNames)
+	{
+		Controller->AssetController->ModifyGoal(SelectedGoal);
+	}
 	Controller->bManipulatingGoals = true;
 	return true;
 }
@@ -212,6 +218,7 @@ bool FIKRigEditMode::EndTracking(FEditorViewportClient* InViewportClient, FViewp
 		return false; // not handled
 	}
 
+	GEditor->EndTransaction();
 	Controller->bManipulatingGoals = false;
 	return true;
 }
@@ -330,3 +337,4 @@ void FIKRigEditMode::DrawHUD(FEditorViewportClient* ViewportClient, FViewport* V
 	FEdMode::DrawHUD(ViewportClient, Viewport, View, Canvas);
 }
 
+#undef LOCTEXT_NAMESPACE
diff --git a/Engine/Plugins/Animation/IKRig/Source/IKRigEditor/Private/RigEditor/IKRigEditorController.cpp b/Engine/Plugins/Animation/IKRig/Source/IKRigEditor/Private/RigEditor/IKRigEditorController.cpp
--- a/Engine/Plugins/Animation/IKRig/Source/IKRigEditor/Private/RigEditor/IKRigEditorController.cpp
+++ b/Engine/Plugins/Animation/IKRig/Source/IKRigEditor/Private/RigEditor/IKRigEditorController.cpp
@@ -4,6 +4,7 @@
 
 #include "IPersonaToolkit.h"
 #include "SKismetInspector.h"
+#include "Dialogs/CustomDialog.h"
 #include "Dialogs/Dialogs.h"
 
 #include "RigEditor/IKRigController.h"
@@ -55,7 +56,10 @@ void FIKRigEditorController::AddNewGoals(const TArray<FName>& GoalNames, const T
 	// add a default solver if there isn't one already
 	if (AssetController->GetNumSolvers() == 0)
 	{
-		PromptToAddSolver();	
+		if (!PromptToAddSolver())
+		{
+			return; // user cancelled
+		}
 	}
 
 	// get selected solvers
@@ -163,11 +167,11 @@ int32 FIKRigEditorController::GetSelectedSolverIndex()
 	return SelectedSolvers[0]->IndexInStack;
 }
 
-void FIKRigEditorController::PromptToAddSolver() const
+bool FIKRigEditorController::PromptToAddSolver() const
 {
 	if (AssetController->GetNumSolvers() > 0)
 	{
-		return;
+		return true;
 	}
 
 	// prompt user to add a default solver
@@ -175,18 +179,27 @@ void FIKRigEditorController::PromptToAddSolver() const
 	const TSharedPtr<FStructOnScope> StructToDisplay = MakeShareable(new FStructOnScope(FIKRigAddFirstSolverSettings::StaticStruct(), (uint8*)&Settings));
 	TSharedRef<SKismetInspector> KismetInspector = SNew(SKismetInspector);
 	KismetInspector->ShowSingleStruct(StructToDisplay);
+	
+	TSharedRef<SCustomDialog> AddSolverDialog = SNew(SCustomDialog)
+		.Title(FText(LOCTEXT("EditorController_IKRigFirstSolver", "Add Default Solver")))
+		.DialogContent(KismetInspector)
+		.Buttons({
+			SCustomDialog::FButton(LOCTEXT("OK", "OK")),
+			SCustomDialog::FButton(LOCTEXT("Cancel", "Cancel"))
+	});
 
-	SGenericDialogWidget::OpenDialog(
-        LOCTEXT("EditorController_IKRigFirstSolver", "Add Default Solver"), 
-        KismetInspector, 
-        SGenericDialogWidget::FArguments(), 
-        true);
+	if (AddSolverDialog->ShowModal() == 1)
+	{
+		return false; // cancelled
+	}
 
 	if (Settings.SolverType != nullptr)
 	{
 		AssetController->AddSolver(Settings.SolverType);
 		SolverStackView->RefreshStackView();
 	}
+
+	return true;
 }
 
 bool FIKRigEditorController::IsElementConnectedToSolver(TSharedRef<FIKRigTreeElement> TreeElement, int32 SolverIndex)
@@ -200,7 +213,7 @@ bool FIKRigEditorController::IsElementConnectedToSolver(TSharedRef<FIKRigTreeEle
 	if (TreeElement->ElementType == IKRigTreeElementType::BONE)
 	{
 		// is this bone affected by this solver?
-		return Solver->IsBoneAffectedBySolver(TreeElement->Key, AssetController->GetIKRigSkeleton());
+		return Solver->IsBoneAffectedBySolver(TreeElement->BoneName, AssetController->GetIKRigSkeleton());
 	}
 
 	if (TreeElement->ElementType == IKRigTreeElementType::BONE_SETTINGS)
@@ -212,13 +225,13 @@ bool FIKRigEditorController::IsElementConnectedToSolver(TSharedRef<FIKRigTreeEle
 	if (TreeElement->ElementType == IKRigTreeElementType::GOAL)
 	{
 		// is goal connected to the solver?
-		return AssetController->IsGoalConnectedToSolver(TreeElement->Key, SolverIndex);
+		return AssetController->IsGoalConnectedToSolver(TreeElement->GoalName, SolverIndex);
 	}
 
-	if (TreeElement->ElementType == IKRigTreeElementType::EFFECTOR)
+	if (TreeElement->ElementType == IKRigTreeElementType::SOLVERGOAL)
 	{
 		// is this an effector for this solver?
-		return TreeElement->EffectorSolverIndex == SolverIndex;
+		return TreeElement->SolverGoalIndex == SolverIndex;
 	}
 
 	checkNoEntry();
@@ -247,7 +260,7 @@ bool FIKRigEditorController::IsElementExcludedBone(TSharedRef<FIKRigTreeElement>
 	}
 	
 	// is this bone excluded?
-	return AssetController->GetBoneExcluded(TreeElement->Key);
+	return AssetController->GetBoneExcluded(TreeElement->BoneName);
 }
 
 void FIKRigEditorController::ShowDetailsForBone(const FName BoneName)
@@ -265,16 +278,15 @@ void FIKRigEditorController::ShowDetailsForBoneSettings(const FName BoneName, in
 
 void FIKRigEditorController::ShowDetailsForGoal(const FName GoalName)
 {
-	const UIKRigEffectorGoal* Goal = AssetController->GetGoal(GoalName);
-	DetailsView->SetObject(const_cast<UIKRigEffectorGoal*>(Goal));
+	DetailsView->SetObject(AssetController->GetGoal(GoalName));
 }
 
-void FIKRigEditorController::ShowDetailsForEffector(const FName GoalName, const int32 SolverIndex)
+void FIKRigEditorController::ShowDetailsForGoalSettings(const FName GoalName, const int32 SolverIndex)
 {
 	// get solver that owns this effector
 	if (const UIKRigSolver* SolverWithEffector = AssetController->GetSolver(SolverIndex))
 	{
-		if (UObject* EffectorSettings = SolverWithEffector->GetEffectorWithGoal(GoalName))
+		if (UObject* EffectorSettings = SolverWithEffector->GetGoalSettings(GoalName))
 		{
 			DetailsView->SetObject(EffectorSettings);
 		}
@@ -283,12 +295,12 @@ void FIKRigEditorController::ShowDetailsForEffector(const FName GoalName, const
 
 void FIKRigEditorController::ShowDetailsForSolver(const int32 SolverIndex)
 {
-	DetailsView->SetObject(const_cast<UIKRigSolver*>(AssetController->GetSolver(SolverIndex)));
+	DetailsView->SetObject(AssetController->GetSolver(SolverIndex));
 }
 
 void FIKRigEditorController::ShowEmptyDetails()
 {
-	DetailsView->SetObject(const_cast<UIKRigDefinition*>(AssetController->GetAsset()));
+	DetailsView->SetObject(AssetController->GetAsset());
 }
 
 void FIKRigEditorController::AddNewRetargetChain(const FName ChainName, const FName StartBone, const FName EndBone)
@@ -313,4 +325,12 @@ void FIKRigEditorController::AddNewRetargetChain(const FName ChainName, const FN
 		true);
 }
 
+void FIKRigEditorController::PlayAnimationAsset(UAnimationAsset* AssetToPlay)
+{
+	if (AssetToPlay && AnimInstance)
+	{
+		AnimInstance->SetAnimationAsset(AssetToPlay);
+	}
+}
+
 #undef LOCTEXT_NAMESPACE
diff --git a/Engine/Plugins/Animation/IKRig/Source/IKRigEditor/Private/RigEditor/IKRigMode.cpp b/Engine/Plugins/Animation/IKRig/Source/IKRigEditor/Private/RigEditor/IKRigMode.cpp
--- a/Engine/Plugins/Animation/IKRig/Source/IKRigEditor/Private/RigEditor/IKRigMode.cpp
+++ b/Engine/Plugins/Animation/IKRig/Source/IKRigEditor/Private/RigEditor/IKRigMode.cpp
@@ -7,6 +7,7 @@
 #include "ISkeletonEditorModule.h"
 #include "Modules/ModuleManager.h"
 #include "PersonaTabs.h"
+#include "RigEditor/IKRigAssetBrowserTabSummoner.h"
 #include "RigEditor/IKRigSkeletonTabSummoner.h"
 #include "RigEditor/IKRigSolverStackTabSummoner.h"
 #include "RigEditor/IKRigRetargetChainTabSummoner.h"
@@ -32,15 +33,15 @@ FIKRigMode::FIKRigMode(
 	TabFactories.RegisterFactory(PersonaModule.CreatePersonaViewportTabFactory(InHostingApp, ViewportArgs));
 	TabFactories.RegisterFactory(PersonaModule.CreateAdvancedPreviewSceneTabFactory(InHostingApp, InPreviewScene));
 	TabFactories.RegisterFactory(PersonaModule.CreateDetailsTabFactory(InHostingApp, FOnDetailsCreated::CreateSP(&IKRigEditor.Get(), &FIKRigEditorToolkit::HandleDetailsCreated)));
-	TabFactories.RegisterFactory(PersonaModule.CreateAnimationAssetBrowserTabFactory(InHostingApp, IKRigEditor->GetPersonaToolkit(), FOnOpenNewAsset::CreateSP(&IKRigEditor.Get(), &FIKRigEditorToolkit::HandleOpenNewAsset), FOnAnimationSequenceBrowserCreated::CreateSP(&IKRigEditor.Get(), &FIKRigEditorToolkit::HandleAnimationSequenceBrowserCreated), true));
 
 	// register custom tabs
+	TabFactories.RegisterFactory(MakeShared<FIKRigAssetBrowserTabSummoner>(IKRigEditor));
 	TabFactories.RegisterFactory(MakeShared<FIKRigSkeletonTabSummoner>(IKRigEditor));
 	TabFactories.RegisterFactory(MakeShared<FIKRigSolverStackTabSummoner>(IKRigEditor));
 	TabFactories.RegisterFactory(MakeShared<FIKRigRetargetChainTabSummoner>(IKRigEditor));
 
 	// create tab layout
-	TabLayout = FTabManager::NewLayout("Standalone_IKRigEditor_Layout_v1.122")
+	TabLayout = FTabManager::NewLayout("Standalone_IKRigEditor_Layout_v1.123")
 		->AddArea
 		(
 			FTabManager::NewPrimaryArea()
@@ -92,9 +93,9 @@ FIKRigMode::FIKRigMode(
                     (
 						FTabManager::NewStack()
 						->SetSizeCoefficient(0.6f)
-						->AddTab(FPersonaTabs::AssetBrowserID, ETabState::OpenedTab)
+						->AddTab(FIKRigAssetBrowserTabSummoner::TabID, ETabState::OpenedTab)
 						->AddTab(FIKRigRetargetChainTabSummoner::TabID, ETabState::OpenedTab)
-						->SetForegroundTab(FIKRigRetargetChainTabSummoner::TabID)
+						->SetForegroundTab(FIKRigAssetBrowserTabSummoner::TabID)
                     )
 				)
 			)
diff --git a/Engine/Plugins/Animation/IKRig/Source/IKRigEditor/Private/RigEditor/IKRigToolkit.cpp b/Engine/Plugins/Animation/IKRig/Source/IKRigEditor/Private/RigEditor/IKRigToolkit.cpp
--- a/Engine/Plugins/Animation/IKRig/Source/IKRigEditor/Private/RigEditor/IKRigToolkit.cpp
+++ b/Engine/Plugins/Animation/IKRig/Source/IKRigEditor/Private/RigEditor/IKRigToolkit.cpp
@@ -32,6 +32,11 @@ FIKRigEditorToolkit::FIKRigEditorToolkit()
 {
 }
 
+FIKRigEditorToolkit::~FIKRigEditorToolkit()
+{
+	EditorController->SkelMeshComponent->SetSkeletalMesh(nullptr);
+}
+
 void FIKRigEditorToolkit::InitAssetEditor(
 	const EToolkitMode::Type Mode,
     const TSharedPtr<IToolkitHost>& InitToolkitHost, 
@@ -45,7 +50,7 @@ void FIKRigEditorToolkit::InitAssetEditor(
 	PersonaToolkitArgs.OnPreviewSceneCreated = FOnPreviewSceneCreated::FDelegate::CreateSP(this, &FIKRigEditorToolkit::HandlePreviewSceneCreated);
 	
 	FPersonaModule& PersonaModule = FModuleManager::LoadModuleChecked<FPersonaModule>("Persona");
-	PersonaToolkit = PersonaModule.CreatePersonaToolkit(IKRigAsset, PersonaToolkitArgs, EditorController->AssetController->GetSkeleton());
+	PersonaToolkit = PersonaModule.CreatePersonaToolkit(IKRigAsset, PersonaToolkitArgs);
 	
 	// when/if preview mesh is changed, we need to reinitialize the anim instance
     PersonaToolkit->GetPreviewScene()->RegisterOnPreviewMeshChanged(FOnPreviewMeshChanged::CreateSP(this, &FIKRigEditorToolkit::HandlePreviewMeshChanged));
@@ -176,20 +181,6 @@ void FIKRigEditorToolkit::PostRedo(bool bSuccess)
 	EditorController->RefreshAllViews();
 }
 
-void FIKRigEditorToolkit::HandleOpenNewAsset(UObject* InNewAsset)
-{
-	if (UAnimationAsset* NewAnimationAsset = Cast<UAnimationAsset>(InNewAsset))
-	{
-		EditorController->AnimInstance->SetAnimationAsset(NewAnimationAsset);
-	}
-}
-
-void FIKRigEditorToolkit::HandleAnimationSequenceBrowserCreated(
-	const TSharedRef<IAnimationSequenceBrowser>& InSequenceBrowser)
-{
-	EditorController->SequenceBrowser = InSequenceBrowser;
-}
-
 void FIKRigEditorToolkit::HandlePreviewSceneCreated(const TSharedRef<IPersonaPreviewScene>& InPersonaPreviewScene)
 {
 	AAnimationEditorPreviewActor* Actor = InPersonaPreviewScene->GetWorld()->SpawnActor<AAnimationEditorPreviewActor>(AAnimationEditorPreviewActor::StaticClass(), FTransform::Identity);
@@ -205,7 +196,7 @@ void FIKRigEditorToolkit::HandlePreviewSceneCreated(const TSharedRef<IPersonaPre
 
 	// set the skeletal mesh on the component
 	// NOTE: this must be done AFTER setting the AnimInstance so that the correct root anim node is loaded
-	USkeletalMesh* Mesh = EditorController->AssetController->GetSourceSkeletalMesh();
+	USkeletalMesh* Mesh = EditorController->AssetController->GetSkeletalMesh();
 	EditorController->SkelMeshComponent->SetSkeletalMesh(Mesh);
 
 	// apply mesh to the preview scene
@@ -219,18 +210,16 @@ void FIKRigEditorToolkit::HandlePreviewSceneCreated(const TSharedRef<IPersonaPre
 
 void FIKRigEditorToolkit::HandlePreviewMeshChanged(USkeletalMesh* InOldSkeletalMesh, USkeletalMesh* InNewSkeletalMesh)
 {
-	if (!InOldSkeletalMesh)
-	{
-		return; // first time setup
-	}
-
 	if (InOldSkeletalMesh == InNewSkeletalMesh)
 	{
 		return; // already set to this skeletal mesh
 	}
 	
 	// update asset with new skeletal mesh (will copy new skeleton data)
-	EditorController->AssetController->SetSourceSkeletalMesh(InNewSkeletalMesh, false);
+	if (!EditorController->AssetController->SetSkeletalMesh(InNewSkeletalMesh))
+	{
+		return; // mesh was not set (incompatible for some reason) todo, show reason in UI someplace
+	}
 
 	// update anim instance to use new skeletal mesh
 	// this is required so that the bone containers passed around during update/eval are correctly sized
@@ -243,6 +232,8 @@ void FIKRigEditorToolkit::HandlePreviewMeshChanged(USkeletalMesh* InOldSkeletalM
 	}
 
 	EditorController->SkelMeshComponent->SetSkeletalMesh(InNewSkeletalMesh);
+	
+	EditorController->RefreshAllViews();
 }
 
 void FIKRigEditorToolkit::HandleDetailsCreated(const TSharedRef<class IDetailsView>& InDetailsView)
@@ -255,7 +246,6 @@ void FIKRigEditorToolkit::HandleDetailsCreated(const TSharedRef<class IDetailsVi
 void FIKRigEditorToolkit::OnFinishedChangingDetails(
     const FPropertyChangedEvent& PropertyChangedEvent)
 {
-	UE_LOG(LogTemp, Log, TEXT("FIKRigEditorToolkit::OnFinishedChangingProperties"));
 }
 
 void FIKRigEditorToolkit::HandleReset()
diff --git a/Engine/Plugins/Animation/IKRig/Source/IKRigEditor/Private/RigEditor/SIKRigAssetBrowser.cpp b/Engine/Plugins/Animation/IKRig/Source/IKRigEditor/Private/RigEditor/SIKRigAssetBrowser.cpp
new file mode 100644
--- /dev/null
+++ b/Engine/Plugins/Animation/IKRig/Source/IKRigEditor/Private/RigEditor/SIKRigAssetBrowser.cpp
@@ -0,0 +1,103 @@
+﻿// Copyright Epic Games, Inc. All Rights Reserved.
+
+#include "RigEditor/SIKRigAssetBrowser.h"
+
+#include "AnimPreviewInstance.h"
+#include "ContentBrowserModule.h"
+#include "IContentBrowserSingleton.h"
+#include "Animation/AnimMontage.h"
+#include "Animation/AnimSequence.h"
+#include "Animation/PoseAsset.h"
+#include "RigEditor/IKRigController.h"
+#include "RigEditor/IKRigEditorController.h"
+
+#define LOCTEXT_NAMESPACE "IKRigAssetBrowser"
+
+void SIKRigAssetBrowser::Construct(
+	const FArguments& InArgs,
+	TSharedRef<FIKRigEditorController> InEditorController)
+{
+	EditorController = InEditorController;
+	
+	ChildSlot
+    [
+        SNew(SVerticalBox)
+		+SVerticalBox::Slot()
+		[
+			SAssignNew(AssetBrowserBox, SBox)
+		]
+    ];
+
+	AddAssetBrowser();
+}
+
+void SIKRigAssetBrowser::AddAssetBrowser()
+{
+	FAssetPickerConfig AssetPickerConfig;
+	AssetPickerConfig.Filter.ClassNames.Add(UAnimSequence::StaticClass()->GetFName());
+	AssetPickerConfig.Filter.ClassNames.Add(UAnimMontage::StaticClass()->GetFName());
+	AssetPickerConfig.Filter.ClassNames.Add(UPoseAsset::StaticClass()->GetFName());
+	AssetPickerConfig.OnAssetDoubleClicked = FOnAssetSelected::CreateSP(this, &SIKRigAssetBrowser::OnAssetDoubleClicked);
+	AssetPickerConfig.GetCurrentSelectionDelegates.Add(&GetCurrentSelectionDelegate);
+	AssetPickerConfig.bAllowNullSelection = false;
+	AssetPickerConfig.InitialAssetViewType = EAssetViewType::Column;
+	AssetPickerConfig.OnShouldFilterAsset = FOnShouldFilterAsset::CreateSP(this, &SIKRigAssetBrowser::OnShouldFilterAsset);
+	AssetPickerConfig.bShowPathInColumnView = true;
+	AssetPickerConfig.bShowTypeInColumnView = false;
+
+	// hide all asset registry columns by default (we only really want the name and path)
+	TArray<UObject::FAssetRegistryTag> AssetRegistryTags;
+	UAnimSequence::StaticClass()->GetDefaultObject()->GetAssetRegistryTags(AssetRegistryTags);
+	for(UObject::FAssetRegistryTag& AssetRegistryTag : AssetRegistryTags)
+	{
+		AssetPickerConfig.HiddenColumnNames.Add(AssetRegistryTag.Name.ToString());
+	}
+
+	// Also hide the type column by default (but allow users to enable it, so don't use bShowTypeInColumnView)
+	AssetPickerConfig.HiddenColumnNames.Add(TEXT("Class"));
+
+	FContentBrowserModule& ContentBrowserModule = FModuleManager::Get().LoadModuleChecked<FContentBrowserModule>(TEXT("ContentBrowser"));
+
+	AssetBrowserBox->SetContent(ContentBrowserModule.Get().CreateAssetPicker(AssetPickerConfig));
+}
+
+void SIKRigAssetBrowser::OnAssetDoubleClicked(const FAssetData& AssetData)
+{
+	if (!AssetData.GetAsset())
+	{
+		return;
+	}
+	
+	UAnimationAsset* NewAnimationAsset = Cast<UAnimationAsset>(AssetData.GetAsset());
+	if (NewAnimationAsset && EditorController.Pin().IsValid())
+	{
+		EditorController.Pin()->PlayAnimationAsset(NewAnimationAsset);
+	}
+}
+
+bool SIKRigAssetBrowser::OnShouldFilterAsset(const struct FAssetData& AssetData)
+{
+	// is this an animation asset?
+	if (!AssetData.GetClass()->IsChildOf(UAnimationAsset::StaticClass()))
+	{
+		return true;
+	}
+	
+	// controller setup
+	FIKRigEditorController* Controller = EditorController.Pin().Get();
+	if (!Controller)
+	{
+		return true;
+	}
+
+	// get skeleton
+	USkeleton* DesiredSkeleton = Controller->AssetController->GetSkeleton();;
+	if (!DesiredSkeleton)
+	{
+		return true;
+	}
+
+	return !DesiredSkeleton->IsCompatibleSkeletonByAssetData(AssetData);
+}
+
+#undef LOCTEXT_NAMESPACE
diff --git a/Engine/Plugins/Animation/IKRig/Source/IKRigEditor/Private/RigEditor/SIKRigSkeleton.cpp b/Engine/Plugins/Animation/IKRig/Source/IKRigEditor/Private/RigEditor/SIKRigSkeleton.cpp
--- a/Engine/Plugins/Animation/IKRig/Source/IKRigEditor/Private/RigEditor/SIKRigSkeleton.cpp
+++ b/Engine/Plugins/Animation/IKRig/Source/IKRigEditor/Private/RigEditor/SIKRigSkeleton.cpp
@@ -19,7 +19,7 @@
 
 #define LOCTEXT_NAMESPACE "SIKRigSkeleton"
 
-FIKRigTreeElement::FIKRigTreeElement(const FName& InKey, IKRigTreeElementType InType)
+FIKRigTreeElement::FIKRigTreeElement(const FText& InKey, IKRigTreeElementType InType)
 {
 	Key = InKey;
 	ElementType = InType;
@@ -100,7 +100,7 @@ void SIKRigSkeletonItem::Construct(
 		case IKRigTreeElementType::GOAL:
 			Brush = FIKRigEditorStyle::Get().GetBrush("IKRig.Tree.Goal");
 			break;
-		case IKRigTreeElementType::EFFECTOR:
+		case IKRigTreeElementType::SOLVERGOAL:
 			Brush = FIKRigEditorStyle::Get().GetBrush("IKRig.Tree.Effector");
 			break;
 		default:
@@ -140,7 +140,7 @@ void SIKRigSkeletonItem::Construct(
 			.ColorAndOpacity(TextColor)
         ];
 
-		if (InEditorController->AssetController->GetRetargetRoot() == InRigTreeElement->Key)
+		if (InEditorController->AssetController->GetRetargetRoot() == InRigTreeElement->BoneName)
 		{	
 			HorizontalBox->AddSlot()
 			.AutoWidth()
@@ -186,12 +186,13 @@ void SIKRigSkeletonItem::OnNameCommitted(const FText& InText, ETextCommit::Type
 		return; 
 	}
 	
-	const FName OldName = WeakRigTreeElement.Pin()->Key;
+	const FName OldName = WeakRigTreeElement.Pin()->GoalName;
 	const FName PotentialNewName = FName(InText.ToString());
 	const FName NewName = Controller->AssetController->RenameGoal(OldName, PotentialNewName);
 	if (NewName != NAME_None)
 	{
-		WeakRigTreeElement.Pin()->Key = NewName;
+		WeakRigTreeElement.Pin()->Key = FText::FromName(NewName);
+		WeakRigTreeElement.Pin()->GoalName = NewName;
 	}
 
 	Controller->SkeletonView->RefreshTreeView();
@@ -199,7 +200,7 @@ void SIKRigSkeletonItem::OnNameCommitted(const FText& InText, ETextCommit::Type
 
 FText SIKRigSkeletonItem::GetName() const
 {
-	return (FText::FromName(WeakRigTreeElement.Pin()->Key));
+	return WeakRigTreeElement.Pin()->Key;
 }
 
 TSharedRef<FIKRigSkeletonDragDropOp> FIKRigSkeletonDragDropOp::New(TWeakPtr<FIKRigTreeElement> InElement)
@@ -269,7 +270,7 @@ void SIKRigSkeleton::SetSelectedGoalsFromViewport(const TArray<FName>& GoalNames
 	
 	for (const TSharedPtr<FIKRigTreeElement>& Item : AllElements)
 	{
-		if (GoalNames.Contains(Item->Key))
+		if (GoalNames.Contains(Item->GoalName))
 		{
 			TreeView->SetSelection(Item, ESelectInfo::Direct);
 		}
@@ -294,7 +295,7 @@ void SIKRigSkeleton::GetSelectedBoneChains(TArray<FIKRigSkeletonChain>& OutChain
 	TArray<int32> SelectedBones;
 	for (const TSharedPtr<FIKRigTreeElement>& BoneItem : SelectedBoneItems)
 	{
-		const FName BoneName = BoneItem.Get()->Key;
+		const FName BoneName = BoneItem.Get()->BoneName;
 		const int32 BoneIndex = Skeleton.GetBoneIndexFromName(BoneName);
 		SelectedBones.Add(BoneIndex);
 	}
@@ -408,7 +409,7 @@ void SIKRigSkeleton::HandleNewGoal()
 		}
 
 		// build default name for the new goal
-		const FName BoneName = Item->Key;
+		const FName BoneName = Item->BoneName;
 		const FName NewGoalName = FName(BoneName.ToString() + "_Goal");
 		
 		GoalNames.Add(NewGoalName);
@@ -454,11 +455,11 @@ void SIKRigSkeleton::HandleDeleteGoal()
 	{
 		if (Item->ElementType == IKRigTreeElementType::GOAL)
 		{
-			Controller->DeleteGoal(Item->Key);
+			Controller->DeleteGoal(Item->GoalName);
 		}
-		else if (Item->ElementType == IKRigTreeElementType::EFFECTOR)
+		else if (Item->ElementType == IKRigTreeElementType::SOLVERGOAL)
 		{
-			Controller->AssetController->DisconnectGoalFromSolver(Item->EffectorGoalName, Item->EffectorSolverIndex);
+			Controller->AssetController->DisconnectGoalFromSolver(Item->SolverGoalName, Item->SolverGoalIndex);
 		}
 	}
 
@@ -479,7 +480,7 @@ bool SIKRigSkeleton::CanDeleteGoal() const
 	for (const TSharedPtr<FIKRigTreeElement>& Item : SelectedItems)
 	{
 		const bool bIsGoal = Item->ElementType == IKRigTreeElementType::GOAL;
-		const bool bIsEffector = Item->ElementType == IKRigTreeElementType::EFFECTOR;
+		const bool bIsEffector = Item->ElementType == IKRigTreeElementType::SOLVERGOAL;
 		if (!(bIsGoal || bIsEffector))
 		{
 			return false;
@@ -531,7 +532,7 @@ void SIKRigSkeleton::ConnectSelectedGoalsToSelectedSolvers(bool bConnect)
 	UIKRigController* AssetController = Controller->AssetController;
 	for (const TSharedPtr<FIKRigTreeElement>& GoalElement : SelectedGoals)
 	{
-		const FName GoalName = GoalElement->Key;
+		const FName GoalName = GoalElement->GoalName;
 		const int32 GoalIndex = AssetController->GetGoalIndex(GoalName);
 		check(GoalIndex != INDEX_NONE);
 		const UIKRigEffectorGoal& EffectorGoal = *AssetController->GetGoal(GoalIndex);
@@ -570,7 +571,7 @@ int32 SIKRigSkeleton::GetNumSelectedGoalToSolverConnections(bool bCountOnlyConne
 	{
 		for (const TSharedPtr<FSolverStackElement>& Solver : SelectedSolvers)
 		{
-			const bool bIsConnected = Controller->AssetController->IsGoalConnectedToSolver(Goal->Key, Solver->IndexInStack);
+			const bool bIsConnected = Controller->AssetController->IsGoalConnectedToSolver(Goal->GoalName, Solver->IndexInStack);
 			if (bIsConnected == bCountOnlyConnected)
 			{
 				++NumMatched;
@@ -592,7 +593,7 @@ void SIKRigSkeleton::HandleSetRootBoneOnSolvers()
     // get name of selected root bone
     TArray<TSharedPtr<FIKRigTreeElement>> SelectedBones;
 	GetSelectedBones(SelectedBones);
-	const FName RootBoneName = SelectedBones[0]->Key;
+	const FName RootBoneName = SelectedBones[0]->BoneName;
 
 	// apply to all selected solvers (ignored on solvers that don't accept a root bone)
 	UIKRigController* AssetController = Controller->AssetController;
@@ -634,7 +635,7 @@ bool SIKRigSkeleton::CanSetRootBoneOnSolvers()
 	Controller->GetSelectedSolvers(SelectedSolvers);
 	for (const TSharedPtr<FSolverStackElement>& Solver : SelectedSolvers)
 	{
-		if (AssetController->GetSolver(Solver->IndexInStack)->CanSetRootBone())
+		if (AssetController->GetSolver(Solver->IndexInStack)->RequiresRootBone())
 		{
 			return true;
 		}
@@ -665,8 +666,8 @@ void SIKRigSkeleton::HandleAddBoneSettings()
 	{
 		for (const TSharedPtr<FSolverStackElement>& Solver : SelectedSolvers)
 		{
-			AssetController->AddBoneSetting(BoneItem->Key, Solver->IndexInStack);
-			BoneNameForSettings = BoneItem->Key;
+			AssetController->AddBoneSetting(BoneItem->BoneName, Solver->IndexInStack);
+			BoneNameForSettings = BoneItem->BoneName;
 			SolverIndex = Solver->IndexInStack;
         }
 	}
@@ -701,7 +702,7 @@ bool SIKRigSkeleton::CanAddBoneSettings()
 	{
 		for (const TSharedPtr<FSolverStackElement>& Solver : SelectedSolvers)
 		{
-			if (AssetController->CanAddBoneSetting(BoneItem->Key, Solver->IndexInStack))
+			if (AssetController->CanAddBoneSetting(BoneItem->BoneName, Solver->IndexInStack))
 			{
 				return true;
 			}
@@ -732,8 +733,8 @@ void SIKRigSkeleton::HandleRemoveBoneSettings()
 	{
 		for (const TSharedPtr<FSolverStackElement>& Solver : SelectedSolvers)
 		{
-			AssetController->RemoveBoneSetting(BoneItem->Key, Solver->IndexInStack);
-			BoneToShowInDetailsView = BoneItem->Key;
+			AssetController->RemoveBoneSetting(BoneItem->BoneName, Solver->IndexInStack);
+			BoneToShowInDetailsView = BoneItem->BoneName;
 		}
 	}
 
@@ -767,7 +768,7 @@ bool SIKRigSkeleton::CanRemoveBoneSettings()
 	{
 		for (const TSharedPtr<FSolverStackElement>& Solver : SelectedSolvers)
 		{
-			if (AssetController->CanRemoveBoneSetting(BoneItem->Key, Solver->IndexInStack))
+			if (AssetController->CanRemoveBoneSetting(BoneItem->BoneName, Solver->IndexInStack))
 			{
 				return true;
 			}
@@ -790,7 +791,7 @@ void SIKRigSkeleton::HandleExcludeBone()
 	GetSelectedBones(SelectedBones);
 	for (const TSharedPtr<FIKRigTreeElement>& BoneItem : SelectedBones)
 	{
-		Controller->AssetController->SetBoneExcluded(BoneItem->Key, true);
+		Controller->AssetController->SetBoneExcluded(BoneItem->BoneName, true);
 	}
 
 	// show greyed out bone name after being excluded
@@ -810,7 +811,7 @@ bool SIKRigSkeleton::CanExcludeBone()
 	GetSelectedBones(SelectedBones);
 	for (const TSharedPtr<FIKRigTreeElement>& BoneItem : SelectedBones)
 	{
-		if (!Controller->AssetController->GetBoneExcluded(BoneItem->Key))
+		if (!Controller->AssetController->GetBoneExcluded(BoneItem->BoneName))
 		{
 			return true;
 		}
@@ -832,7 +833,7 @@ void SIKRigSkeleton::HandleIncludeBone()
 	GetSelectedBones(SelectedBones);
 	for (const TSharedPtr<FIKRigTreeElement>& BoneItem : SelectedBones)
 	{
-		Controller->AssetController->SetBoneExcluded(BoneItem->Key, false);
+		Controller->AssetController->SetBoneExcluded(BoneItem->BoneName, false);
 	}
 
 	// show normal bone name after being included
@@ -852,7 +853,7 @@ bool SIKRigSkeleton::CanIncludeBone()
 	GetSelectedBones(SelectedBones);
 	for (const TSharedPtr<FIKRigTreeElement>& BoneItem : SelectedBones)
 	{
-		if (Controller->AssetController->GetBoneExcluded(BoneItem->Key))
+		if (Controller->AssetController->GetBoneExcluded(BoneItem->BoneName))
 		{
 			return true;
 		}
@@ -905,7 +906,7 @@ void SIKRigSkeleton::HandleSetRetargetRoot()
 	}
 
 	// set the first selected bone as the retarget root
-	Controller->AssetController->SetRetargetRoot(SelectedBones[0]->Key);
+	Controller->AssetController->SetRetargetRoot(SelectedBones[0]->BoneName);
 
 	// show root bone after being set
 	Controller->RefreshAllViews();
@@ -990,7 +991,9 @@ void SIKRigSkeleton::RefreshTreeView(bool IsInitialSetup)
 	for (const FName BoneName : Skeleton.BoneNames)
 	{
 		// create "Bone" tree element for this bone
-		TSharedPtr<FIKRigTreeElement> BoneElement = MakeShared<FIKRigTreeElement>(BoneName, IKRigTreeElementType::BONE);
+		const FText BoneDisplayName = FText::FromName(BoneName);
+		TSharedPtr<FIKRigTreeElement> BoneElement = MakeShared<FIKRigTreeElement>(BoneDisplayName, IKRigTreeElementType::BONE);
+		BoneElement.Get()->BoneName = BoneName;
 		const int32 BoneElementIndex = AllElements.Add(BoneElement);
 		BoneTreeElementIndices.Add(BoneName, BoneElementIndex);
 
@@ -999,8 +1002,9 @@ void SIKRigSkeleton::RefreshTreeView(bool IsInitialSetup)
 		{
 			if (Solvers[SolverIndex]->GetBoneSetting(BoneName))
 			{
-				const FName DisplayName = FName("Bone Settings for: " + Solvers[SolverIndex]->GetName());
-				TSharedPtr<FIKRigTreeElement> SettingsItem = MakeShared<FIKRigTreeElement>(DisplayName, IKRigTreeElementType::BONE_SETTINGS);
+				const FText SolverDisplayName = FText::FromString(AssetController->GetSolverUniqueName(SolverIndex));
+				const FText BoneSettingDisplayName = FText::Format(LOCTEXT("BoneSettings", "{0} settings for {1}"), BoneDisplayName, SolverDisplayName);
+				TSharedPtr<FIKRigTreeElement> SettingsItem = MakeShared<FIKRigTreeElement>(BoneSettingDisplayName, IKRigTreeElementType::BONE_SETTINGS);
 				SettingsItem->BoneSettingBoneName = BoneName;
 				SettingsItem->BoneSettingsSolverIndex = SolverIndex;
 				AllElements.Add(SettingsItem);
@@ -1019,27 +1023,29 @@ void SIKRigSkeleton::RefreshTreeView(bool IsInitialSetup)
 			}
 			
 			// make new element for goal
-			TSharedPtr<FIKRigTreeElement> GoalItem = MakeShared<FIKRigTreeElement>(Goal->GoalName, IKRigTreeElementType::GOAL);
+			const FText GoalDisplayName = FText::FromName(Goal->GoalName);
+			TSharedPtr<FIKRigTreeElement> GoalItem = MakeShared<FIKRigTreeElement>(GoalDisplayName, IKRigTreeElementType::GOAL);
+			GoalItem->GoalName = Goal->GoalName;
 			AllElements.Add(GoalItem);
 
 			// store hierarchy pointers for goal
 			BoneElement->Children.Add(GoalItem);
 			GoalItem->Parent = BoneElement;
 
-			// add all effectors connected to this goal
+			// add all solver settings connected to this goal
 			for (int32 SolverIndex=0; SolverIndex<Solvers.Num(); ++SolverIndex)
 			{
-				if (UObject* Effector = AssetController->GetEffectorForGoal(Goal->GoalName, SolverIndex))
+				if (UObject* GoalSettings = AssetController->GetGoalSettingsForSolver(Goal->GoalName, SolverIndex))
 				{
-					// make new element for effector
-					const FText EffectorPrefix = LOCTEXT("EffectorPrefix", "Effector for");
-					const FName DisplayName = FName(EffectorPrefix.ToString() + ": " + Solvers[SolverIndex]->GetName());
-					TSharedPtr<FIKRigTreeElement> EffectorItem = MakeShared<FIKRigTreeElement>(DisplayName, IKRigTreeElementType::EFFECTOR);
-					EffectorItem->EffectorSolverIndex = SolverIndex;
-					EffectorItem->EffectorGoalName = Goal->GoalName;
-					AllElements.Add(EffectorItem);
-					EffectorItem->Parent = GoalItem;
-					GoalItem->Children.Add(EffectorItem);
+					// make new element for solver goal
+					const FText SolverDisplayName = FText::FromString(AssetController->GetSolverUniqueName(SolverIndex));
+					const FText SolverGoalDisplayName = FText::Format(LOCTEXT("GoalSettingsForSolver", "{0} settings for solver {1}"), FText::FromName(Goal->GoalName), SolverDisplayName);
+					TSharedPtr<FIKRigTreeElement> SolverGoalItem = MakeShared<FIKRigTreeElement>(SolverGoalDisplayName, IKRigTreeElementType::SOLVERGOAL);
+					SolverGoalItem->SolverGoalIndex = SolverIndex;
+					SolverGoalItem->SolverGoalName = Goal->GoalName;
+					AllElements.Add(SolverGoalItem);
+					SolverGoalItem->Parent = GoalItem;
+					GoalItem->Children.Add(SolverGoalItem);
 				}
 			}
 		}
@@ -1116,7 +1122,7 @@ void SIKRigSkeleton::OnSelectionChanged(TSharedPtr<FIKRigTreeElement> Selection,
 		TArray<FName> SelectedGoalNames;
 		for (const TSharedPtr<FIKRigTreeElement>& Goal : SelectedGoals)
 		{
-			SelectedGoalNames.Add(Goal->Key);
+			SelectedGoalNames.Add(Goal->GoalName);
 		}
 		Controller->HandleGoalsSelectedInTreeView(SelectedGoalNames);
 	}
@@ -1141,15 +1147,15 @@ void SIKRigSkeleton::OnItemClicked(TSharedPtr<FIKRigTreeElement> InItem)
 	// update details view
 	if (InItem->ElementType == IKRigTreeElementType::BONE)
 	{
-		Controller->ShowDetailsForBone(InItem->Key);
+		Controller->ShowDetailsForBone(InItem->BoneName);
 	}
 	else if (InItem->ElementType == IKRigTreeElementType::GOAL)
 	{
-		Controller->ShowDetailsForGoal(InItem->Key);
+		Controller->ShowDetailsForGoal(InItem->GoalName);
 	}
-	else if (InItem->ElementType == IKRigTreeElementType::EFFECTOR)
+	else if (InItem->ElementType == IKRigTreeElementType::SOLVERGOAL)
 	{
-		Controller->ShowDetailsForEffector(InItem->EffectorGoalName, InItem->EffectorSolverIndex);
+		Controller->ShowDetailsForGoalSettings(InItem->SolverGoalName, InItem->SolverGoalIndex);
 	}
 	else if (InItem->ElementType == IKRigTreeElementType::BONE_SETTINGS)
 	{
@@ -1273,8 +1279,12 @@ FReply SIKRigSkeleton::OnAcceptDrop(
 	}
 
 	const FIKRigTreeElement& DraggedElement = *DragDropOp.Get()->Element.Pin().Get();
+	ensure(DraggedElement.ElementType == IKRigTreeElementType::GOAL);		// drag only supported for goals
+	ensure(TargetItem.Get()->ElementType == IKRigTreeElementType::BONE);	// drop only supported for bones
+
+	// re-parent the goal to a different bone
 	UIKRigController* AssetController = Controller->AssetController;
-	const bool bWasReparented = AssetController->SetGoalBone(DraggedElement.Key, TargetItem.Get()->Key);
+	const bool bWasReparented = AssetController->SetGoalBone(DraggedElement.GoalName, TargetItem.Get()->BoneName);
 	if (bWasReparented)
 	{
 		RefreshTreeView();
@@ -1294,7 +1304,7 @@ FReply SIKRigSkeleton::OnKeyDown(const FGeometry& MyGeometry, const FKeyEvent& I
 	}
 
 	// handle deleting selected items
-	if (Key == EKeys::Delete)
+	if (Key == EKeys::Delete || Key == EKeys::BackSpace)
 	{
 		TArray<TSharedPtr<FIKRigTreeElement>> SelectedItems = TreeView->GetSelectedItems();
 		for (const TSharedPtr<FIKRigTreeElement>& SelectedItem : SelectedItems)
@@ -1302,16 +1312,16 @@ FReply SIKRigSkeleton::OnKeyDown(const FGeometry& MyGeometry, const FKeyEvent& I
 			switch(SelectedItem->ElementType)
 			{
 				case IKRigTreeElementType::GOAL:
-					Controller->DeleteGoal(SelectedItem->Key);
+					Controller->DeleteGoal(SelectedItem->GoalName);
 					break;
-				case IKRigTreeElementType::EFFECTOR:
-					Controller->AssetController->DisconnectGoalFromSolver(SelectedItem->EffectorGoalName, SelectedItem->EffectorSolverIndex);
+				case IKRigTreeElementType::SOLVERGOAL:
+					Controller->AssetController->DisconnectGoalFromSolver(SelectedItem->SolverGoalName, SelectedItem->SolverGoalIndex);
 					break;
 				case IKRigTreeElementType::BONE_SETTINGS:
 					Controller->AssetController->RemoveBoneSetting(SelectedItem->BoneSettingBoneName, SelectedItem->BoneSettingsSolverIndex);
 					break;
 				default:
-					checkNoEntry()
+					continue; // can't delete anything else
 			}
 		}
 
diff --git a/Engine/Plugins/Animation/IKRig/Source/IKRigEditor/Private/RigEditor/SIKRigSolverStack.cpp b/Engine/Plugins/Animation/IKRig/Source/IKRigEditor/Private/RigEditor/SIKRigSolverStack.cpp
--- a/Engine/Plugins/Animation/IKRig/Source/IKRigEditor/Private/RigEditor/SIKRigSolverStack.cpp
+++ b/Engine/Plugins/Animation/IKRig/Source/IKRigEditor/Private/RigEditor/SIKRigSolverStack.cpp
@@ -373,7 +373,7 @@ void SIKRigSolverStack::RefreshStackView()
 	for (int32 i=0; i<NumSolvers; ++i)
 	{
 		const UIKRigSolver* Solver = AssetController->GetSolver(i);
-		const FText DisplayName = Solver ? FText::FromString(Solver->GetNiceName().ToString()) : FText::FromString("Unknown Solver");
+		const FText DisplayName = Solver ? FText::FromString(AssetController->GetSolverUniqueName(i)) : FText::FromString("Unknown Solver");
 		TSharedPtr<FSolverStackElement> SolverItem = FSolverStackElement::Make(DisplayName, i);
 		ListViewItems.Add(SolverItem);
 	}
diff --git a/Engine/Plugins/Animation/IKRig/Source/IKRigEditor/Public/RetargetEditor/IKRetargetEditor.h b/Engine/Plugins/Animation/IKRig/Source/IKRigEditor/Public/RetargetEditor/IKRetargetEditor.h
--- a/Engine/Plugins/Animation/IKRig/Source/IKRigEditor/Public/RetargetEditor/IKRetargetEditor.h
+++ b/Engine/Plugins/Animation/IKRig/Source/IKRigEditor/Public/RetargetEditor/IKRetargetEditor.h
@@ -84,36 +84,14 @@ private:
 	
 	/** preview scene setup */
 	void HandlePreviewSceneCreated(const TSharedRef<IPersonaPreviewScene>& InPersonaPreviewScene);
-	void HandlePreviewMeshChanged(USkeletalMesh* InOldSkeletalMesh, USkeletalMesh* InNewSkeletalMesh);
 	void HandleDetailsCreated(const TSharedRef<class IDetailsView>& InDetailsView);
 	void OnFinishedChangingDetails(const FPropertyChangedEvent& PropertyChangedEvent);
 	void SetupAnimInstance();
 	/** END preview scene setup */
 
-	/** edit retarget pose */
-	void HandleEditPose() const;
-	bool CanEditPose() const;
-	bool IsEditingPose() const;
-	/** END edit retarget pose*/
-
-	/** new/delete retarget pose */
-	void HandleNewPose();
-	FReply CreateNewPose();
-	void HandleDeletePose();
-	bool CanDeletePose() const;
-	void HandleResetPose();
-	TSharedPtr<SWindow> NewPoseWindow;
-	TSharedPtr<SEditableTextBox> NewPoseEditableText;
-	/** END new/delete retarget pose */
-
-	TArray<TSharedPtr<FName>> PoseNames;
-	FText GetCurrentPoseName() const;
-	void OnPoseSelected(TSharedPtr<FName> InPoseName, ESelectInfo::Type SelectInfo);
-	/* END edit reference pose */
-
 	void HandleSourceOrTargetIKRigAssetChanged();
 	
-	/** centralized management of across all views */
+	/** centralized management across all views */
 	TSharedRef<FIKRetargetEditorController> EditorController;
 	
 	friend FIKRetargetMode;
diff --git a/Engine/Plugins/Animation/IKRig/Source/IKRigEditor/Public/RetargetEditor/IKRetargetEditorController.h b/Engine/Plugins/Animation/IKRig/Source/IKRigEditor/Public/RetargetEditor/IKRetargetEditorController.h
--- a/Engine/Plugins/Animation/IKRig/Source/IKRigEditor/Public/RetargetEditor/IKRetargetEditorController.h
+++ b/Engine/Plugins/Animation/IKRig/Source/IKRigEditor/Public/RetargetEditor/IKRetargetEditorController.h
@@ -20,17 +20,23 @@ public:
 	/** Initialize the editor */
 	void Initialize(TSharedPtr<FIKRetargetEditor> InEditor, UIKRetargeter* InAsset);
 	/** Bind callbacks to this IK Rig */
-	void BindToIKRigAsset(UIKRigDefinition* InIKRig);
+	void BindToIKRigAsset(UIKRigDefinition* InIKRig) const;
 	/** callback when IK Rig asset requires reinitialization */
-	void OnIKRigNeedsInitialized(UIKRigDefinition* ModifiedIKRig);
+	void OnIKRigNeedsInitialized(UIKRigDefinition* ModifiedIKRig) const;
 	/** callback when IK Rig asset's retarget chain has been renamed */
-	void OnRetargetChainRenamed(UIKRigDefinition* ModifiedIKRig, FName OldName, FName NewName);
+	void OnRetargetChainRenamed(UIKRigDefinition* ModifiedIKRig, FName OldName, FName NewName) const;
 	/** callback when IK Retargeter asset requires reinitialization */
-	void OnRetargeterNeedsInitialized(const UIKRetargeter* Retargeter);
+	void OnRetargeterNeedsInitialized(const UIKRetargeter* Retargeter) const;
 	
 	/** all modifications to the data model should go through this controller */
 	UIKRetargeterController* AssetController;
 
+	/** Preview scene to be supplied by IHasPersonaToolkit::GetPersonaToolkit */
+	TSharedPtr<IPersonaToolkit> PersonaToolkit;
+
+	/** the persona toolkit */
+	TWeakPtr<FIKRetargetEditor> Editor;
+
 	/** viewport skeletal mesh */
 	UDebugSkelMeshComponent* SourceSkelMeshComponent;
 	UDebugSkelMeshComponent* TargetSkelMeshComponent;
@@ -47,30 +53,40 @@ public:
 	/** chain list view */
 	TSharedPtr<SIKRetargetChainMapList> ChainsView;
 
-	/** Preview scene to be supplied by IHasPersonaToolkit::GetPersonaToolkit */
-	TSharedPtr<IPersonaToolkit> PersonaToolkit;
-
-	/** the persona toolkit */
-	TWeakPtr<FIKRetargetEditor> Editor;
-
-	/** get the skeletal meshes we are copying animation between */
+	/** get the source skeletal mesh we are copying FROM */
 	USkeletalMesh* GetSourceSkeletalMesh() const;
+	/** get the target skeletal mesh we are copying TO */
 	USkeletalMesh* GetTargetSkeletalMesh() const;
 
 	/** get current chain pose */
 	FTransform GetTargetBoneTransform(const int32& TargetBoneIndex) const;
 	/** get the line segments to draw from this bone to each child */
 	bool GetTargetBoneLineSegments(const int32& TargetBoneIndex, FVector& OutStart, TArray<FVector>& OutChildren) const;
-	/** get if the target bone is being retargeted or not */
-	bool IsTargetBoneRetargeted(const int32& TargetBoneIndex);
 
 	/** get the retargeter that is running in the viewport (which is a duplicate of the source asset) */
 	const UIKRetargetProcessor* GetRetargetProcessor() const;
 
-	/** Sequence Browser and Edit Pose mode **/
+	/** Sequence Browser**/
 	void PlayAnimationAsset(UAnimationAsset* AssetToPlay);
-	void PlayPreviousAnimationAsset();
+	void PlayPreviousAnimationAsset() const;
 	UAnimationAsset* PreviousAsset = nullptr;
+	/** END Sequence Browser */
+
+	/** edit retarget poses */
+	void HandleEditPose() const;
+	bool CanEditPose() const;
+	bool IsEditingPose() const;
+	void HandleNewPose();
+	FReply CreateNewPose() const;
+	void HandleDeletePose() const;
+	bool CanDeletePose() const;
+	void HandleResetPose() const;
+	TSharedPtr<SWindow> NewPoseWindow;
+	TSharedPtr<SEditableTextBox> NewPoseEditableText;
+	TArray<TSharedPtr<FName>> PoseNames;
+	FText GetCurrentPoseName() const;
+	void OnPoseSelected(TSharedPtr<FName> InPoseName, ESelectInfo::Type SelectInfo) const;
+	/* END edit retarget poses */
 
-	void RefreshAllViews();
+	void RefreshAllViews() const;
 };
diff --git a/Engine/Plugins/Animation/IKRig/Source/IKRigEditor/Public/RetargetEditor/IKRetargeterController.h b/Engine/Plugins/Animation/IKRig/Source/IKRigEditor/Public/RetargetEditor/IKRetargeterController.h
--- a/Engine/Plugins/Animation/IKRig/Source/IKRigEditor/Public/RetargetEditor/IKRetargeterController.h
+++ b/Engine/Plugins/Animation/IKRig/Source/IKRigEditor/Public/RetargetEditor/IKRetargeterController.h
@@ -27,6 +27,17 @@ public:
 	/** Get access to the retargeter asset.
 	 *@warning Do not make modifications to the asset directly. Use this API instead. */
 	UIKRetargeter* GetAsset() const;
+
+	/** SOURCE / TARGET
+	* 
+	*/
+	/** Set the IK Rig to use as the source (to copy animation FROM) */
+	void SetSourceIKRig(UIKRigDefinition* SourceIKRig);
+	/** Set the IK Rig to use as the target (to copy animation TO) */
+	void SetTargetIKRig(UIKRigDefinition* TargetIKRig);
+	/** Get target skeletal mesh */
+	USkeletalMesh* GetTargetPreviewMesh();
+	
 	/** Get the USkeleton used on the Source asset */
 	USkeleton* GetSourceSkeletonAsset() const;
 
@@ -78,14 +89,16 @@ public:
 	/** Add a delta translation to the root bone (used in Edit Mode in the retarget editor) */
 	void AddTranslationOffsetToRetargetRootBone(FVector TranslationOffset) const;
 	/** Set whether to output retarget pose. Will output current retarget pose if true, or run retarget otherwise. */
-	void SetEditRetargetPoseMode(bool bOutputRetargetPose) const;
+	void SetEditRetargetPoseMode(bool bOutputRetargetPose, bool bReinitializeAfter=true) const;
 	/** Get whether in mode to output retarget pose (true) or run retarget (false). */
 	bool GetEditRetargetPoseMode() const;
+	/** Add a numbered suffix to the given pose name to make it unique. */
+	FName MakePoseNameUnique(FName PoseName) const;
 	/** END RETARGET POSE EDITING */
 
 private:
-
-	/** Called whenever the rig is modified in such a way that would require re-initialization by dependent systems.*/
+	
+	/** Called whenever the retargeter is modified in such a way that would require re-initialization by dependent systems.*/
 	DECLARE_MULTICAST_DELEGATE_OneParam(FOnRetargeterNeedsInitialized, const UIKRetargeter*);
 	FOnRetargeterNeedsInitialized RetargeterNeedsInitialized;
 
diff --git a/Engine/Plugins/Animation/IKRig/Source/IKRigEditor/Public/RigEditor/IKRigAssetBrowserTabSummoner.h b/Engine/Plugins/Animation/IKRig/Source/IKRigEditor/Public/RigEditor/IKRigAssetBrowserTabSummoner.h
new file mode 100644
--- /dev/null
+++ b/Engine/Plugins/Animation/IKRig/Source/IKRigEditor/Public/RigEditor/IKRigAssetBrowserTabSummoner.h
@@ -0,0 +1,21 @@
+﻿// Copyright Epic Games, Inc. All Rights Reserved.
+
+#pragma once
+
+#include "CoreMinimal.h"
+#include "WorkflowOrientedApp/WorkflowTabFactory.h"
+
+class FIKRigEditorToolkit;
+
+struct FIKRigAssetBrowserTabSummoner : public FWorkflowTabFactory
+{
+public:
+	static const FName TabID;
+	
+	FIKRigAssetBrowserTabSummoner(const TSharedRef<FIKRigEditorToolkit>& InIKRetargetEditor);
+	virtual TSharedRef<SWidget> CreateTabBody(const FWorkflowTabSpawnInfo& Info) const override;
+	virtual TSharedPtr<SToolTip> CreateTabToolTipWidget(const FWorkflowTabSpawnInfo& Info) const override;
+
+protected:
+	TWeakPtr<FIKRigEditorToolkit> IKRigEditor;
+};
diff --git a/Engine/Plugins/Animation/IKRig/Source/IKRigEditor/Public/RigEditor/IKRigController.h b/Engine/Plugins/Animation/IKRig/Source/IKRigEditor/Public/RigEditor/IKRigController.h
--- a/Engine/Plugins/Animation/IKRig/Source/IKRigEditor/Public/RigEditor/IKRigController.h
+++ b/Engine/Plugins/Animation/IKRig/Source/IKRigEditor/Public/RigEditor/IKRigController.h
@@ -7,6 +7,7 @@
 
 #include "IKRigController.generated.h"
 
+struct FIKRigInputSkeleton;
 struct FReferenceSkeleton;
 struct FIKRigSkeleton;
 class UIKRigSolver;
@@ -34,12 +35,11 @@ public:
 	/** SKELETON
 	 * 
 	 */
-	/** Sets the preview mesh to use, can optionally reinitialize the skeleton with bReImportBone=true. */
-	void SetSourceSkeletalMesh(USkeletalMesh* SkeletalMesh, bool bReImportBones) const;
+	/** Sets the preview mesh to use, can optionally reinitialize the skeleton with bReImportBone=true.
+	 * Returns true if the mesh was able to be set. False if it was incompatible for any reason. */
+	bool SetSkeletalMesh(USkeletalMesh* SkeletalMesh) const;
 	/** Get the skeletal mesh asset this IK Rig was initialized with */
-	USkeletalMesh* GetSourceSkeletalMesh() const;
-	/** Initializes the skeleton to use the supplied bones / hierarchy */
-	void SetSkeleton(const FReferenceSkeleton& InSkeleton) const;
+	USkeletalMesh* GetSkeletalMesh() const;
 	/** Get read-access to the IKRig skeleton representation */
 	const FIKRigSkeleton& GetIKRigSkeleton() const;
 	/** Get the USkeleton asset this rig was initialized with */
@@ -49,7 +49,7 @@ public:
 	/** Returns true if the given Bone is excluded, false otherwise. */
 	bool GetBoneExcluded(const FName& BoneName) const;
 	/** Get the global-space retarget pose transform of the given Bone. Can only be called AFTER skeleton is initialized. */
-	FTransform GetBoneRetargetPose(const FName& BoneName) const;
+	FTransform GetRefPoseTransformOfBone(const FName& BoneName) const;
 	/** END SKELETON */
 
 	/** SOLVERS
@@ -71,6 +71,8 @@ public:
 	void SetRootBone(const FName& RootBoneName, int32 SolverIndex) const;
 	/** Get read-only access to the array of solvers. */
 	const TArray<UIKRigSolver*>& GetSolverArray() const;
+	/** Get unique label for a given solver. Returns dash separated index and name like so, "0 - SolverName". */
+	FString GetSolverUniqueName(int32 SolverIndex);
 	/** END SOLVERS */
 
 	/** GOALS
@@ -82,6 +84,8 @@ public:
 	bool RemoveGoal(const FName& GoalName) const;
 	/** Rename a Goal. Returns new name, which may be different after being sanitized. Returns NAME_None if this fails.*/
 	FName RenameGoal(const FName& OldName, const FName& PotentialNewName) const;
+	/** Modify a Goal for a transaction. Returns true if Goal found.*/
+	bool ModifyGoal(const FName& GoalName) const;
 	/** The the Bone that the given Goal should be parented to / associated with. */
 	bool SetGoalBone(const FName& GoalName, const FName& NewBoneName) const;
 	/** The the Bone associated with the given Goal. */
@@ -100,11 +104,11 @@ public:
 	const TArray<UIKRigEffectorGoal*>& GetAllGoals() const;
 	/** Get read-only access to the Goal at the given index. */
 	const UIKRigEffectorGoal* GetGoal(int32 GoalIndex) const;
-	/** Get read-only access to the Goal with the given name. */
-	const UIKRigEffectorGoal* GetGoal(const FName& GoalName) const;
+	/** Get read-write access to the Goal with the given name. */
+	UIKRigEffectorGoal* GetGoal(const FName& GoalName) const;
 	/** Get the UObject for the settings associated with the given Goal in the given Solver.
 	 ** Solvers can define their own per-Goal settings depending on their needs. These are termed "Effectors". */
-	UObject* GetEffectorForGoal(const FName& GoalName, int32 SolverIndex) const;
+	UObject* GetGoalSettingsForSolver(const FName& GoalName, int32 SolverIndex) const;
 	/** Get the global-space transform of the given Goal. This may be set by the user in the editor, or at runtime. */
 	FTransform GetGoalCurrentTransform(const FName& GoalName) const;
 	/** Set the Goal to the given transform. */
@@ -162,10 +166,6 @@ public:
 	/** Sorts the Chains from Root to tip based on the Start Bone of each Chain. */
 	void SortRetargetChains() const;
 	/** END retarget chains */
-	
-	// BEGIN UObject
-	static void AddReferencedObjects(UObject* InThis, FReferenceCollector& Collector);
-	// END UObject
 
 private:
 
@@ -192,7 +192,7 @@ private:
 
 	/** The actual IKRigDefinition asset that this Controller modifies. */
 	UPROPERTY(transient)
-	TObjectPtr<UIKRigDefinition> IKRigAsset = nullptr;
+	TObjectPtr<UIKRigDefinition> Asset = nullptr;
 
 	/** Lazy-generated map of Controllers to IK Rig Assets. Avoids duplicate controllers. */
 	static TMap<UIKRigDefinition*, UIKRigController*> AssetToControllerMap;
diff --git a/Engine/Plugins/Animation/IKRig/Source/IKRigEditor/Public/RigEditor/IKRigEditorController.h b/Engine/Plugins/Animation/IKRig/Source/IKRigEditor/Public/RigEditor/IKRigEditorController.h
--- a/Engine/Plugins/Animation/IKRig/Source/IKRigEditor/Public/RigEditor/IKRigEditorController.h
+++ b/Engine/Plugins/Animation/IKRig/Source/IKRigEditor/Public/RigEditor/IKRigEditorController.h
@@ -52,7 +52,7 @@ public:
 	/** get index of the first selected solver, return INDEX_None if nothing selected */
 	int32 GetSelectedSolverIndex();
 	/** right after importing a skeleton, we ask user what solver they want to use */
-	void PromptToAddSolver() const;
+	bool PromptToAddSolver() const;
 	/** determine if the element is connected to the selected solver */
 	bool IsElementConnectedToSolver(TSharedRef<FIKRigTreeElement> TreeElement, int32 SolverIndex);
 	/** determine if the element is connected to ANY solver */
@@ -67,7 +67,7 @@ public:
 	/** show GOAL settings in details view */
 	void ShowDetailsForGoal(const FName GoalName);
 	/** show EFFECTOR settings in details view */
-	void ShowDetailsForEffector(const FName GoalName, const int32 SolverIndex);
+	void ShowDetailsForGoalSettings(const FName GoalName, const int32 SolverIndex);
 	/** show SOLVER settings in details view */
 	void ShowDetailsForSolver(const int32 SolverIndex);
 	/** show nothing in details view */
@@ -75,6 +75,9 @@ public:
 
 	/** create a new retarget chain */
 	void AddNewRetargetChain(const FName ChainName, const FName StartBone, const FName EndBone);
+
+	/** play preview animation on running anim instance in editor (before IK) */
+	void PlayAnimationAsset(UAnimationAsset* AssetToPlay);
 	
 	/** all modifications to the data model should go through this controller */
 	UIKRigController* AssetController;
@@ -86,9 +89,6 @@ public:
 	UPROPERTY(transient, NonTransactional)
 	TObjectPtr<class UIKRigAnimInstance> AnimInstance;
 
-	/** Sequence Browser **/
-	TWeakPtr<class IAnimationSequenceBrowser> SequenceBrowser;
-
 	/** asset properties tab */
 	TSharedPtr<class IDetailsView> DetailsView;
 
diff --git a/Engine/Plugins/Animation/IKRig/Source/IKRigEditor/Public/RigEditor/IKRigToolkit.h b/Engine/Plugins/Animation/IKRig/Source/IKRigEditor/Public/RigEditor/IKRigToolkit.h
--- a/Engine/Plugins/Animation/IKRig/Source/IKRigEditor/Public/RigEditor/IKRigToolkit.h
+++ b/Engine/Plugins/Animation/IKRig/Source/IKRigEditor/Public/RigEditor/IKRigToolkit.h
@@ -42,7 +42,7 @@ class FIKRigEditorToolkit :
 public:
 
 	FIKRigEditorToolkit();
-	virtual ~FIKRigEditorToolkit() override {};
+	virtual ~FIKRigEditorToolkit() override;
 
 	void InitAssetEditor(
 		const EToolkitMode::Type Mode,
@@ -82,11 +82,6 @@ public:
 	virtual void PostRedo( bool bSuccess );
 	/** END FSelfRegisteringEditorUndoClient interface */
 
-	/** animation asset browser */
-	void HandleOpenNewAsset(UObject* InNewAsset);
-	void HandleAnimationSequenceBrowserCreated(const TSharedRef<IAnimationSequenceBrowser>& InSequenceBrowser);
-	/* END animation asset browser */
-
 	TSharedRef<FIKRigEditorController> GetController() const {return EditorController;};
 
 private:
diff --git a/Engine/Plugins/Animation/IKRig/Source/IKRigEditor/Public/RigEditor/SIKRigAssetBrowser.h b/Engine/Plugins/Animation/IKRig/Source/IKRigEditor/Public/RigEditor/SIKRigAssetBrowser.h
new file mode 100644
--- /dev/null
+++ b/Engine/Plugins/Animation/IKRig/Source/IKRigEditor/Public/RigEditor/SIKRigAssetBrowser.h
@@ -0,0 +1,40 @@
+﻿// Copyright Epic Games, Inc. All Rights Reserved.
+
+#pragma once
+
+#include "CoreMinimal.h"
+
+#include "ContentBrowserDelegates.h"
+#include "Widgets/Layout/SBox.h"
+
+
+class FIKRigEditorController;
+
+class SIKRigAssetBrowser : public SBox
+{
+public:
+	SLATE_BEGIN_ARGS(SIKRigAssetBrowser) {}
+	SLATE_END_ARGS()
+
+	void Construct(const FArguments& InArgs, TSharedRef<FIKRigEditorController> InEditorController);
+
+private:
+
+	void AddAssetBrowser();
+
+	void OnPathChange(const FString& NewPath);
+	
+	void OnAssetDoubleClicked(const FAssetData& AssetData);
+	bool OnShouldFilterAsset(const struct FAssetData& AssetData);
+
+	/** Used to get the currently selected assets */
+	FGetCurrentSelectionDelegate GetCurrentSelectionDelegate;
+		
+	/** editor controller */
+	TWeakPtr<FIKRigEditorController> EditorController;
+
+	/** the animation asset browser */
+	TSharedPtr<SBox> AssetBrowserBox;
+
+	friend FIKRigEditorController;
+};
diff --git a/Engine/Plugins/Animation/IKRig/Source/IKRigEditor/Public/RigEditor/SIKRigSkeleton.h b/Engine/Plugins/Animation/IKRig/Source/IKRigEditor/Public/RigEditor/SIKRigSkeleton.h
--- a/Engine/Plugins/Animation/IKRig/Source/IKRigEditor/Public/RigEditor/SIKRigSkeleton.h
+++ b/Engine/Plugins/Animation/IKRig/Source/IKRigEditor/Public/RigEditor/SIKRigSkeleton.h
@@ -11,23 +11,19 @@
 
 #include "SIKRigSolverStack.h"
 
-#include "SIKRigSkeleton.generated.h"
-
 struct FIKRigSkeletonChain;
 class FIKRigEditorController;
 class SIKRigSkeleton;
 class FIKRigEditorToolkit;
 class USkeletalMesh;
 
-enum class IKRigTreeElementType { BONE, GOAL, EFFECTOR, BONE_SETTINGS };
+enum class IKRigTreeElementType { BONE, GOAL, SOLVERGOAL, BONE_SETTINGS };
 
 class FIKRigTreeElement : public TSharedFromThis<FIKRigTreeElement>
 {
 public:
 	
-	FIKRigTreeElement(
-		const FName& InKey,
-		IKRigTreeElementType InType);
+	FIKRigTreeElement(const FText& InKey, IKRigTreeElementType InType);
 
 	TSharedRef<ITableRow> MakeTreeRowWidget(
 		TSharedRef<FIKRigEditorController> InEditorController,
@@ -36,19 +32,25 @@ public:
 		TSharedRef<FUICommandList> InCommandList,
 		TSharedPtr<SIKRigSkeleton> InHierarchy);
 
-	FName Key;
+	FText Key;
 	IKRigTreeElementType ElementType;
 	TSharedPtr<FIKRigTreeElement> Parent;
 	TArray<TSharedPtr<FIKRigTreeElement>> Children;
 
-	/** effector meta-data */
-	FName EffectorGoalName = NAME_None;
-	int32 EffectorSolverIndex = INDEX_NONE;
+	/** effector meta-data (if it is an effector) */
+	FName SolverGoalName = NAME_None;
+	int32 SolverGoalIndex = INDEX_NONE;
 
-	/** bone setting meta-data */
+	/** bone setting meta-data (if it is a bone setting) */
 	FName BoneSettingBoneName = NAME_None;
 	int32 BoneSettingsSolverIndex = INDEX_NONE;
 
+	/** name of bone if it is one */
+	FName BoneName = NAME_None;
+	
+	/** name of goal if it is one */
+	FName GoalName = NAME_None;
+
 	/** delegate for when the context menu requests a rename */
 	void RequestRename();
 	DECLARE_DELEGATE(FOnRenameRequested);
@@ -110,7 +112,7 @@ public:
 	{
 		for (const TTuple<TSharedPtr<FIKRigTreeElement>, FSparseItemInfo>& Pair : OldSparseItemInfos)
 		{
-			if (Pair.Key->Key == ItemPtr->Key)
+			if (Pair.Key->Key.EqualTo(ItemPtr->Key))
 			{
 				// the SparseItemInfos now reference the new element, but keep the same expansion state
 				SparseItemInfos.Add(ItemPtr, Pair.Value);
@@ -129,17 +131,6 @@ private:
 	TSparseItemMap OldSparseItemInfos;
 };
 
-USTRUCT()
-struct FIKRigSkeletonImportSettings
-{
-	GENERATED_BODY()
-
-	FIKRigSkeletonImportSettings() : Mesh(nullptr) {}
-
-	UPROPERTY(EditAnywhere, Category = "Skeleton Import")
-	USkeletalMesh* Mesh;
-};
-
 class SIKRigSkeleton : public SCompoundWidget, public FEditorUndoClient
 {
 public:
diff --git a/Engine/Plugins/Experimental/FullBodyIK/Source/PBIK/Public/Core/PBIKSolver.h b/Engine/Plugins/Experimental/FullBodyIK/Source/PBIK/Public/Core/PBIKSolver.h
--- a/Engine/Plugins/Experimental/FullBodyIK/Source/PBIK/Public/Core/PBIKSolver.h
+++ b/Engine/Plugins/Experimental/FullBodyIK/Source/PBIK/Public/Core/PBIKSolver.h
@@ -72,7 +72,7 @@ struct PBIK_API FPBIKSolverSettings
 	GENERATED_BODY()
 
 	/** High iteration counts can help solve complex joint configurations with competing constraints, but will increase runtime cost. Default is 20. */
-	UPROPERTY(EditAnywhere, Category = SolverSettings, meta = (ClampMin = "0", UIMin = "0.0", UIMax = "200.0"))
+	UPROPERTY(EditAnywhere, Category = SolverSettings, meta = (ClampMin = "0", ClampMax = "1000", UIMin = "0.0", UIMax = "200.0"))
 	int32 Iterations = 20;
 
 	/** A global mass multiplier; higher values will make the joints more stiff, but require more iterations. Typical range is 0.0 to 10.0. */
diff --git a/Engine/Plugins/Experimental/FullBodyIK/Source/PBIK/Public/PBIK_Shared.h b/Engine/Plugins/Experimental/FullBodyIK/Source/PBIK/Public/PBIK_Shared.h
--- a/Engine/Plugins/Experimental/FullBodyIK/Source/PBIK/Public/PBIK_Shared.h
+++ b/Engine/Plugins/Experimental/FullBodyIK/Source/PBIK/Public/PBIK_Shared.h
@@ -46,10 +46,10 @@ struct PBIK_API FPBIKBoneSetting
 	 *Locked: no rotation is allowed in this axis (will remain at reference pose angle). */
 	UPROPERTY(EditAnywhere, Category = Limits)
 	EPBIKLimitType X;
-	/**Range is -180 to 180 (Default is 0). Degrees of rotation in the negative X direction to allow when joint is in "Limited" mode. */
+	/**Range is -180 to 0 (Default is 0). Degrees of rotation in the negative X direction to allow when joint is in "Limited" mode. */
 	UPROPERTY(EditAnywhere, Category = Limits, meta = (ClampMin = "-180", ClampMax = "0", UIMin = "-180.0", UIMax = "0.0"))
 	float MinX = 0.0f;
-	/**Range is -180 to 180 (Default is 0). Degrees of rotation in the positive X direction to allow when joint is in "Limited" mode. */
+	/**Range is 0 to 180 (Default is 0). Degrees of rotation in the positive X direction to allow when joint is in "Limited" mode. */
 	UPROPERTY(EditAnywhere, Category = Limits, meta = (ClampMin = "0", ClampMax = "180", UIMin = "0.0", UIMax = "180.0"))
 	float MaxX = 0.0f;
 
@@ -59,10 +59,10 @@ struct PBIK_API FPBIKBoneSetting
 	*Locked: no rotation is allowed in this axis (will remain at input pose angle). */
 	UPROPERTY(EditAnywhere, Category = Limits)
 	EPBIKLimitType Y;
-	/**Range is -180 to 180 (Default is 0). Degrees of rotation in the negative Y direction to allow when joint is in "Limited" mode. */
+	/**Range is -180 to 0 (Default is 0). Degrees of rotation in the negative Y direction to allow when joint is in "Limited" mode. */
 	UPROPERTY(EditAnywhere, Category = Limits, meta = (ClampMin = "-180", ClampMax = "0", UIMin = "-180.0", UIMax = "0.0"))
 	float MinY = 0.0f;
-	/**Range is -180 to 180 (Default is 0). Degrees of rotation in the positive Y direction to allow when joint is in "Limited" mode. */
+	/**Range is 0 to 180 (Default is 0). Degrees of rotation in the positive Y direction to allow when joint is in "Limited" mode. */
 	UPROPERTY(EditAnywhere, Category = Limits, meta = (ClampMin = "0", ClampMax = "180", UIMin = "0.0", UIMax = "180.0"))
 	float MaxY = 0.0f;
 
@@ -72,10 +72,10 @@ struct PBIK_API FPBIKBoneSetting
 	*Locked: no rotation is allowed in this axis (will remain at input pose angle). */
 	UPROPERTY(EditAnywhere, Category = Limits)
 	EPBIKLimitType Z;
-	/**Range is -180 to 180 (Default is 0). Degrees of rotation in the negative Z direction to allow when joint is in "Limited" mode. */
+	/**Range is -180 to 0 (Default is 0). Degrees of rotation in the negative Z direction to allow when joint is in "Limited" mode. */
 	UPROPERTY(EditAnywhere, Category = Limits, meta = (ClampMin = "-180", ClampMax = "0", UIMin = "-180.0", UIMax = "0.0"))
 	float MinZ = 0.0f;
-	/**Range is -180 to 180 (Default is 0). Degrees of rotation in the positive Z direction to allow when joint is in "Limited" mode. */
+	/**Range is 0 to 180 (Default is 0). Degrees of rotation in the positive Z direction to allow when joint is in "Limited" mode. */
 	UPROPERTY(EditAnywhere, Category = Limits, meta = (ClampMin = "0", ClampMax = "180", UIMin = "0.0", UIMax = "180.0"))
 	float MaxZ = 0.0f;
 
diff --git a/Engine/Plugins/Experimental/FullBodyIK/Source/PBIK/Public/RigUnit_PBIK.h b/Engine/Plugins/Experimental/FullBodyIK/Source/PBIK/Public/RigUnit_PBIK.h
--- a/Engine/Plugins/Experimental/FullBodyIK/Source/PBIK/Public/RigUnit_PBIK.h
+++ b/Engine/Plugins/Experimental/FullBodyIK/Source/PBIK/Public/RigUnit_PBIK.h
@@ -20,11 +20,11 @@ struct FPBIKDebug
 	GENERATED_BODY()
 
 	/** The scale of the debug drawing. */
-	UPROPERTY()
+	UPROPERTY(EditAnywhere, Category="Debug")
 	float DrawScale = 1.0f;
 
 	/** If true, turns on debug drawing for the node. */
-	UPROPERTY()
+	UPROPERTY(EditAnywhere, Category="Debug")
 	bool bDrawDebug = false;
 
 	void Draw(FControlRigDrawInterface* DrawInterface, FPBIKSolver* Solver) const
@@ -55,24 +55,24 @@ struct FPBIKEffector
 	FPBIKEffector()	: Bone(NAME_None) {}
 
 	/** The bone that this effector will pull on. */
-	UPROPERTY(meta = (Constant, CustomWidget = "BoneName"))
+	UPROPERTY(EditAnywhere, Category="Effector", meta = (Constant, CustomWidget = "BoneName"))
 	FName Bone;
 
 	/** The target location and rotation for this effector. The solver will try to get the specified bone to reach this location.*/
-	UPROPERTY()
+	UPROPERTY(EditAnywhere, Category="Effector")
 	FTransform Transform;
 
 	/** Range 0-1 (default is 1.0). This blends the effector translation and rotation between its transform and the transform of the bone from the input pose.
 	 * At 0.0, the effector will be placed at the location of the bone in the input pose (effectively pinning the skeleton to the input pose).
 	 * At 1.0, the effector will be placed at the supplied location.*/
-	UPROPERTY()
+	UPROPERTY(EditAnywhere, Category="Effector")
 	float OffsetAlpha = 1.0f;
 
 	/** Range 0-1 (default is 1.0). The strength of the effector when pulling the bone towards it's target location.
 	 * At 0.0, the effector does not pull at all, but the bones between the effector and the root will still slightly resist motion from other effectors.
 	 * This can thus act as a "stabilizer" of sorts for parts of the body that you do not want to behave in a pure FK fashion.
 	 */
-	UPROPERTY()
+	UPROPERTY(EditAnywhere, Category="Effector")
 	float StrengthAlpha = 1.0f;
 
 	/** Range 0-1 (default is 1.0). When enabled (greater than 0.0), the solver internally partitions the skeleton into 'chains' which extend from the effector to the nearest fork in the skeleton.
@@ -80,12 +80,12 @@ struct FPBIKEffector
 	 *This can improve the results for sparse bone chains, and significantly improve convergence on dense bone chains.
 	 *But it may cause undesirable results in highly constrained bone chains (like robot arms).
 	 */
-	UPROPERTY()
+	UPROPERTY(EditAnywhere, Category="Effector")
 	float PullChainAlpha = 1.0f;
 
 	/** Range 0-1 (default is 1.0).
 	 *Blends the effector bone rotation between the rotation of the effector transform (1.0) and the rotation of the input bone (0.0).*/
-	UPROPERTY()
+	UPROPERTY(EditAnywhere, Category="Effector")
 	float PinRotation = 1.0f;
 };
 
