Fixing Unity build in IK Rig plugin

diff --git a/Engine/Plugins/Animation/IKRig/Source/IKRigEditor/Private/RetargetEditor/IKRetargetEditorController.cpp b/Engine/Plugins/Animation/IKRig/Source/IKRigEditor/Private/RetargetEditor/IKRetargetEditorController.cpp
--- a/Engine/Plugins/Animation/IKRig/Source/IKRigEditor/Private/RetargetEditor/IKRetargetEditorController.cpp
+++ b/Engine/Plugins/Animation/IKRig/Source/IKRigEditor/Private/RetargetEditor/IKRetargetEditorController.cpp
@@ -3,6 +3,7 @@
 #include "RetargetEditor/IKRetargetEditorController.h"
 
 #include "EditorModeManager.h"
+#include "Widgets/Input/SButton.h"
 #include "Animation/DebugSkelMeshComponent.h"
 #include "RetargetEditor/IKRetargetAnimInstance.h"
 #include "RetargetEditor/IKRetargetEditMode.h"
diff --git a/Engine/Plugins/Animation/IKRig/Source/IKRigEditor/Private/RetargetEditor/IKRetargeterController.cpp b/Engine/Plugins/Animation/IKRig/Source/IKRigEditor/Private/RetargetEditor/IKRetargeterController.cpp
--- a/Engine/Plugins/Animation/IKRig/Source/IKRigEditor/Private/RetargetEditor/IKRetargeterController.cpp
+++ b/Engine/Plugins/Animation/IKRig/Source/IKRigEditor/Private/RetargetEditor/IKRetargeterController.cpp
@@ -2,6 +2,7 @@
 
 #include "RetargetEditor/IKRetargeterController.h"
 
+#include "ScopedTransaction.h"
 #include "Algo/LevenshteinDistance.h"
 #include "Engine/SkeletalMesh.h"
 #include "Retargeter/IKRetargeter.h"
diff --git a/Engine/Plugins/Animation/IKRig/Source/IKRigEditor/Public/RetargetEditor/IKRetargetEditorController.h b/Engine/Plugins/Animation/IKRig/Source/IKRigEditor/Public/RetargetEditor/IKRetargetEditorController.h
--- a/Engine/Plugins/Animation/IKRig/Source/IKRigEditor/Public/RetargetEditor/IKRetargetEditorController.h
+++ b/Engine/Plugins/Animation/IKRig/Source/IKRigEditor/Public/RetargetEditor/IKRetargetEditorController.h
@@ -5,6 +5,10 @@
 #include "IPersonaToolkit.h"
 #include "UObject/ObjectMacros.h"
 #include "UObject/ObjectPtr.h"
+#include "Input/Reply.h"
+#include "Widgets/Input/SEditableTextBox.h"
+#include "Widgets/SWindow.h"
+
 class UIKRetargetProcessor;
 class SIKRetargetChainMapList;
 class UIKRetargetAnimInstance;
