Add validity message when compiling UAnimGraphNode_IKRig

diff --git a/Engine/Plugins/Animation/IKRig/Source/IKRigDeveloper/Private/AnimGraphNode_IKRig.cpp b/Engine/Plugins/Animation/IKRig/Source/IKRigDeveloper/Private/AnimGraphNode_IKRig.cpp
--- a/Engine/Plugins/Animation/IKRig/Source/IKRigDeveloper/Private/AnimGraphNode_IKRig.cpp
+++ b/Engine/Plugins/Animation/IKRig/Source/IKRigDeveloper/Private/AnimGraphNode_IKRig.cpp
@@ -4,6 +4,7 @@
 #include "Animation/AnimInstance.h"
 #include "IKRigDefinition.h"
 #include "IKRigSolver.h"
+#include "Kismet2/CompilerResultsLog.h"
 
 /////////////////////////////////////////////////////
 // UAnimGraphNode_IKRig 
@@ -31,4 +32,14 @@ void UAnimGraphNode_IKRig::CopyNodeDataToPreviewNode(FAnimNode_Base* InPreviewNo
 	FAnimNode_IKRig* IKRigNode = static_cast<FAnimNode_IKRig*>(InPreviewNode);
 }
 
+void UAnimGraphNode_IKRig::ValidateAnimNodeDuringCompilation(USkeleton* ForSkeleton, FCompilerResultsLog& MessageLog)
+{
+	Super::ValidateAnimNodeDuringCompilation(ForSkeleton, MessageLog);
+
+	if (!IsValid(Node.RigDefinitionAsset))
+	{
+		MessageLog.Warning(*LOCTEXT("NoRigDefinitionAsset", "@@ - Please select a Rig Definition Asset.").ToString(), this);
+	}
+}
+
 #undef LOCTEXT_NAMESPACE
diff --git a/Engine/Plugins/Animation/IKRig/Source/IKRigDeveloper/Public/AnimGraphNode_IKRig.h b/Engine/Plugins/Animation/IKRig/Source/IKRigDeveloper/Public/AnimGraphNode_IKRig.h
--- a/Engine/Plugins/Animation/IKRig/Source/IKRigDeveloper/Public/AnimGraphNode_IKRig.h
+++ b/Engine/Plugins/Animation/IKRig/Source/IKRigDeveloper/Public/AnimGraphNode_IKRig.h
@@ -28,5 +28,6 @@ public:
 	// UAnimGraphNode_Base interface
 	virtual void CopyNodeDataToPreviewNode(FAnimNode_Base* AnimNode) override;
 	virtual void Draw(FPrimitiveDrawInterface* PDI, USkeletalMeshComponent* PreviewSkelMeshComp) const override;
+	virtual void ValidateAnimNodeDuringCompilation(USkeleton* ForSkeleton, FCompilerResultsLog& MessageLog) override;
 	// End of UAnimGraphNode_Base interface
 };
