IKRig Limb and Pole solvers

diff --git a/Engine/Plugins/Animation/IKRig/Source/IKRig/Private/IKRigSkeleton.cpp b/Engine/Plugins/Animation/IKRig/Source/IKRig/Private/IKRigSkeleton.cpp
--- a/Engine/Plugins/Animation/IKRig/Source/IKRig/Private/IKRigSkeleton.cpp
+++ b/Engine/Plugins/Animation/IKRig/Source/IKRig/Private/IKRigSkeleton.cpp
@@ -101,6 +101,22 @@ int32 FIKRigSkeleton::GetParentIndexThatIsNotExcluded(const int32 BoneIndex) con
 	return ParentIndex;
 }
 
+int32 FIKRigSkeleton::GetChildIndices(const int32 ParentBoneIndex, TArray<int32>& Children) const
+{
+	Children.Reset();
+
+	const int32 NumBones = BoneNames.Num();
+	for (int32 BoneIndex = ParentBoneIndex + 1; BoneIndex < NumBones; BoneIndex++)
+	{
+		if (ParentBoneIndex == GetParentIndex(BoneIndex))
+		{
+			Children.Add(BoneIndex);
+		}
+	}
+
+	return Children.Num();
+}
+
 void FIKRigSkeleton::ConvertLocalPoseToGlobal(
 	const TArray<int32>& InParentIndices,
 	const TArray<FTransform>& InLocalPose,
diff --git a/Engine/Plugins/Animation/IKRig/Source/IKRig/Private/Solvers/IKRig_LimbSolver.cpp b/Engine/Plugins/Animation/IKRig/Source/IKRig/Private/Solvers/IKRig_LimbSolver.cpp
new file mode 100644
--- /dev/null
+++ b/Engine/Plugins/Animation/IKRig/Source/IKRig/Private/Solvers/IKRig_LimbSolver.cpp
@@ -0,0 +1,240 @@
+﻿// Copyright Epic Games, Inc. All Rights Reserved.
+
+
+#include "Solvers/IKRig_LimbSolver.h"
+
+#define LOCTEXT_NAMESPACE "IKRig_LimbSolver"
+
+UIKRig_LimbSolver::UIKRig_LimbSolver()
+{
+	Effector = CreateDefaultSubobject<UIKRig_LimbEffector>(TEXT("Effector"));
+}
+
+void UIKRig_LimbSolver::Initialize(const FIKRigSkeleton& InSkeleton)
+{
+	Solver.Reset();
+	ChildrenToUpdate.Empty();
+	
+	if (Effector->GoalName == NAME_None || Effector->BoneName == NAME_None || RootName == NAME_None)
+	{
+		return;
+	}
+
+	int32 BoneIndex = InSkeleton.GetBoneIndexFromName(Effector->BoneName);
+	const int32 RootIndex = InSkeleton.GetBoneIndexFromName(RootName);
+	if (BoneIndex == INDEX_NONE || RootIndex == INDEX_NONE)
+	{
+		return;
+	}
+
+	// populate indices
+	TArray<int32> BoneIndices( {BoneIndex} );
+	BoneIndex = InSkeleton.GetParentIndex(BoneIndex);
+	while (BoneIndex != INDEX_NONE && BoneIndex >= RootIndex)
+	{
+		BoneIndices.Add(BoneIndex);
+		BoneIndex = InSkeleton.GetParentIndex(BoneIndex);
+	};
+
+	// if chain is not long enough
+	if (BoneIndices.Num() < 3)
+	{
+		return;
+	}
+
+	// sort the chain from root to end
+	Algo::Reverse(BoneIndices);
+
+	// initialize solver
+	for (int32 Index: BoneIndices)
+	{
+		const FVector Location = InSkeleton.CurrentPoseGlobal[Index].GetLocation();
+		Solver.AddLink(Location, Index);
+	}
+
+	const bool bInitialized = Solver.Initialize();
+	if (bInitialized)
+	{
+		// store children that needs propagation once solved
+		TArray<int32> Children;
+		for (int32 Index = 0; Index < BoneIndices.Num()-1; ++Index)
+		{
+			// store children if not already handled by the solver (part if the links)
+			InSkeleton.GetChildIndices(BoneIndices[Index], Children);
+			const int32 NextIndex = BoneIndices[Index+1];
+			for (const int32 ChildIndex: Children)
+			{
+				if (ChildIndex != NextIndex)
+				{
+					ChildrenToUpdate.Add(ChildIndex);
+					GatherChildren(ChildIndex, InSkeleton, ChildrenToUpdate);
+				}
+			}
+		}
+		// store end bone children
+		GatherChildren(BoneIndices.Last(), InSkeleton, ChildrenToUpdate);
+	}
+}
+
+void UIKRig_LimbSolver::Solve(FIKRigSkeleton& InOutRigSkeleton, const FIKRigGoalContainer& InGoals)
+{
+	if (Solver.NumLinks() < 3)
+	{
+		return;
+	}
+	
+	const FIKRigGoal* IKGoal = InGoals.FindGoalByName(Effector->GoalName);
+	if (!IKGoal)
+	{
+		return;
+	}
+
+	// update settings
+	FLimbSolverSettings Settings;
+	Settings.ReachPrecision = ReachPrecision;
+	Settings.MaxIterations = MaxIterations;
+	Settings.bEnableLimit = bEnableLimit;
+	Settings.MinRotationAngle = MinRotationAngle;
+	Settings.EndBoneForwardAxis = EndBoneForwardAxis; 
+	Settings.HingeRotationAxis = HingeRotationAxis;
+	Settings.bEnableTwistCorrection = bEnableTwistCorrection;
+	Settings.ReachStepAlpha = ReachStepAlpha;
+	Settings.PullDistribution = PullDistribution;
+	Settings.bAveragePull = bAveragePull;
+
+	// update settings
+	const FVector& GoalLocation = IKGoal->FinalBlendedPosition;
+	const FQuat& GoalRotation = IKGoal->FinalBlendedRotation;
+	const bool bModifiedLimb = Solver.Solve(
+		InOutRigSkeleton.CurrentPoseGlobal,
+		GoalLocation,
+		GoalRotation,
+		Settings);
+
+	// propagate if needed
+	if (bModifiedLimb)
+	{
+		// update chain bones local transform
+		for (int32 Index = 0; Index < Solver.NumLinks(); Index++)
+		{
+			InOutRigSkeleton.UpdateLocalTransformFromGlobal(Solver.GetBoneIndex(Index));
+		}
+
+		// propagate to children
+		for (const int32 ChildIndex: ChildrenToUpdate)
+		{
+			InOutRigSkeleton.UpdateGlobalTransformFromLocal(ChildIndex);
+		}
+	}
+}
+
+#if WITH_EDITOR
+
+void UIKRig_LimbSolver::UpdateSolverSettings(UIKRigSolver* InSettings)
+{
+	if (UIKRig_LimbSolver* Settings = Cast<UIKRig_LimbSolver>(InSettings))
+	{
+		ReachPrecision = Settings->ReachPrecision;
+		HingeRotationAxis = Settings->HingeRotationAxis;
+		MaxIterations = Settings->MaxIterations;
+		bEnableLimit = Settings->bEnableLimit;
+		MinRotationAngle = Settings->MinRotationAngle;
+		bAveragePull = Settings->bAveragePull;
+		PullDistribution = Settings->PullDistribution;
+		ReachStepAlpha = Settings->ReachStepAlpha;
+		bEnableTwistCorrection = Settings->bEnableTwistCorrection;
+		EndBoneForwardAxis = Settings->EndBoneForwardAxis; 
+	}
+}
+
+FText UIKRig_LimbSolver::GetNiceName() const
+{
+	return FText(LOCTEXT("SolverName", "Limb IK"));
+}
+
+bool UIKRig_LimbSolver::GetWarningMessage(FText& OutWarningMessage) const
+{
+	if (Effector->GoalName == NAME_None)
+	{
+		OutWarningMessage = LOCTEXT("MissingGoal", "Missing goal.");
+		return true;
+	}
+
+	if (RootName == NAME_None)
+	{
+		OutWarningMessage = LOCTEXT("MissingRoot", "Missing root.");
+		return true;
+	}
+	
+	return false;
+}
+
+bool UIKRig_LimbSolver::IsBoneAffectedBySolver(const FName& BoneName, const FIKRigSkeleton& IKRigSkeleton) const
+{
+	if (Effector->BoneName != NAME_None)
+	{
+		return IKRigSkeleton.IsBoneInDirectLineage(Effector->BoneName, RootName);
+	}
+	return false;
+}
+
+void UIKRig_LimbSolver::AddGoal(const UIKRigEffectorGoal& NewGoal)
+{
+	Effector->GoalName = NewGoal.GoalName;
+	Effector->BoneName = NewGoal.BoneName;
+}
+
+void UIKRig_LimbSolver::RemoveGoal(const FName& GoalName)
+{
+	if (Effector->GoalName == GoalName)
+	{
+		Effector->GoalName = NAME_None;
+		Effector->BoneName = NAME_None;
+	}
+}
+
+void UIKRig_LimbSolver::RenameGoal(const FName& OldName, const FName& NewName)
+{
+	if (Effector->GoalName == OldName)
+	{
+		Effector->GoalName = NewName;
+	}
+}
+
+void UIKRig_LimbSolver::SetGoalBone(const FName& GoalName, const FName& NewBoneName)
+{
+	if (Effector->GoalName == GoalName)
+	{
+		Effector->BoneName = NewBoneName;
+	}
+}
+
+bool UIKRig_LimbSolver::IsGoalConnected(const FName& GoalName) const
+{
+	return (Effector->GoalName == GoalName);
+}
+
+UObject* UIKRig_LimbSolver::GetGoalSettings(const FName& GoalName) const
+{
+	return (Effector->GoalName == GoalName) ? Effector : nullptr;
+}
+
+void UIKRig_LimbSolver::SetRootBone(const FName& RootBoneName)
+{
+	RootName = RootBoneName;
+}
+
+#endif
+
+void UIKRig_LimbSolver::GatherChildren(const int32 BoneIndex, const FIKRigSkeleton& InSkeleton, TArray<int32>& OutChildren)
+{
+	TArray<int32> Children;
+	InSkeleton.GetChildIndices(BoneIndex, Children);
+	for (int32 ChildIndex: Children)
+	{
+		OutChildren.Add(ChildIndex);
+		GatherChildren(ChildIndex, InSkeleton, OutChildren);
+	}
+}
+
+#undef LOCTEXT_NAMESPACE
diff --git a/Engine/Plugins/Animation/IKRig/Source/IKRig/Private/Solvers/IKRig_PoleSolver.cpp b/Engine/Plugins/Animation/IKRig/Source/IKRig/Private/Solvers/IKRig_PoleSolver.cpp
new file mode 100644
--- /dev/null
+++ b/Engine/Plugins/Animation/IKRig/Source/IKRig/Private/Solvers/IKRig_PoleSolver.cpp
@@ -0,0 +1,242 @@
+﻿// Copyright Epic Games, Inc. All Rights Reserved.
+
+#include "Solvers/IKRig_PoleSolver.h"
+#include "IKRigDataTypes.h"
+#include "IKRigSkeleton.h"
+
+#define LOCTEXT_NAMESPACE "UIKRig_PoleSolver"
+
+UIKRig_PoleSolver::UIKRig_PoleSolver()
+{
+	Effector = CreateDefaultSubobject<UIKRig_PoleSolverEffector>(TEXT("Effector"));
+}
+
+void UIKRig_PoleSolver::Initialize(const FIKRigSkeleton& IKRigSkeleton)
+{
+	Chain.Empty();
+	
+	int32 BoneIndex = IKRigSkeleton.GetBoneIndexFromName(EndName);
+	const int32 RootIndex = IKRigSkeleton.GetBoneIndexFromName(RootName);
+
+	if (BoneIndex == INDEX_NONE || RootIndex == INDEX_NONE)
+	{
+		return;
+	}
+
+	// populate chain
+	Chain.Add(BoneIndex);
+	BoneIndex = IKRigSkeleton.GetParentIndex(BoneIndex);
+	while (BoneIndex != INDEX_NONE && BoneIndex >= RootIndex)
+	{
+		Chain.Add(BoneIndex);
+		BoneIndex = IKRigSkeleton.GetParentIndex(BoneIndex);
+	};
+
+	// if chain is not long enough
+	if (Chain.Num() < 3)
+	{
+		Chain.Empty();
+		return;
+	}
+	
+	// sort the chain from root to end
+	Algo::Reverse(Chain);
+	
+	// store children that needs propagation once solved
+	TArray<int32> Children;
+	for (int32 Index = 0; Index < Chain.Num()-1; ++Index)
+	{
+		// store children if not already handled by the solver
+		IKRigSkeleton.GetChildIndices(Chain[Index], Children);
+		const int32 NextIndex = Chain[Index+1];
+		for (const int32 ChildIndex: Children)
+		{
+			if (ChildIndex != NextIndex)
+			{
+				ChildrenToUpdate.Add(ChildIndex);
+				GatherChildren(ChildIndex, IKRigSkeleton, ChildrenToUpdate);
+			}
+		}
+	}
+}
+
+void UIKRig_PoleSolver::Solve(FIKRigSkeleton& IKRigSkeleton, const FIKRigGoalContainer& Goals)
+{
+	if (Chain.Num() < 3)
+	{
+		return;
+	}
+	
+	const FIKRigGoal* IKGoal = Goals.FindGoalByName(Effector->GoalName);
+	if (!IKGoal)
+	{
+		return;
+	}
+
+	const int32 RootIndex = Chain[0];
+	const int32 KneeIndex = Chain[1];
+	const int32 EndIndex = Chain.Last();
+
+	TArray<FTransform>& InOutTransforms = IKRigSkeleton.CurrentPoseGlobal;
+
+	// initial configuration
+	const FVector RootLocation = InOutTransforms[RootIndex].GetLocation();
+	const FVector KneeLocation = InOutTransforms[KneeIndex].GetLocation();
+	const FVector EndLocation = InOutTransforms[EndIndex].GetLocation();
+	
+	const FVector RootToEnd = (EndLocation - RootLocation).GetSafeNormal();
+	const FVector RootToKnee = (KneeLocation - RootLocation).GetSafeNormal();
+	if (RootToEnd.IsZero() || RootToKnee.IsZero())
+	{
+		return;
+	}
+	const FVector InitPlane = FVector::CrossProduct(RootToEnd, RootToKnee).GetSafeNormal();
+
+	// target configuration
+	const FVector& GoalLocation = IKGoal->FinalBlendedPosition;
+	const FVector RootToPole = (GoalLocation - RootLocation).GetSafeNormal();
+	if (RootToPole.IsZero())
+	{
+		return;
+	}
+	const FVector TargetPlane = FVector::CrossProduct(RootToEnd, RootToPole).GetSafeNormal();
+
+	// compute delta rotation from InitialPlane to TargetPlane
+	if (InitPlane.IsZero() || InitPlane.Equals(TargetPlane))
+	{
+    	return;
+    }
+	
+	const FQuat DeltaRotation = FQuat::FindBetweenNormals(InitPlane, TargetPlane);
+	if (!DeltaRotation.IsIdentity())
+	{
+		// update transforms
+		for (int32 Index = 0; Index < Chain.Num()-1; Index++)
+		{
+			int32 BoneIndex = Chain[Index];
+			FTransform& BoneTransform = InOutTransforms[BoneIndex];
+			BoneTransform.SetRotation(DeltaRotation * BoneTransform.GetRotation());
+
+			const FVector BoneLocation = BoneTransform.GetLocation();
+			BoneTransform.SetLocation(RootLocation + DeltaRotation.RotateVector(BoneLocation - RootLocation));
+		}
+
+		// propagate to children
+		for (const int32 ChildIndex: ChildrenToUpdate)
+		{
+			IKRigSkeleton.UpdateGlobalTransformFromLocal(ChildIndex);
+		}
+	}
+}
+
+#if WITH_EDITOR
+
+void UIKRig_PoleSolver::UpdateSolverSettings(UIKRigSolver* InSettings)
+{
+	if(UIKRig_PoleSolver* Settings = Cast<UIKRig_PoleSolver>(InSettings))
+	{
+	}
+}
+
+FText UIKRig_PoleSolver::GetNiceName() const
+{
+	return FText(LOCTEXT("SolverName", "Pole Solver"));
+}
+
+bool UIKRig_PoleSolver::GetWarningMessage(FText& OutWarningMessage) const
+{
+	if (Effector->GoalName == NAME_None)
+	{
+		OutWarningMessage = LOCTEXT("MissingGoal", "Missing goal.");
+		return true;
+	}
+
+	if (EndName == NAME_None)
+	{
+		OutWarningMessage = LOCTEXT("MissingEnd", "Missing end bone.");
+		return true;
+	}
+	
+	if (RootName == NAME_None)
+	{
+		OutWarningMessage = LOCTEXT("MissingRoot", "Missing root bone.");
+		return true;
+	}
+
+	return false;
+}
+
+void UIKRig_PoleSolver::AddGoal(const UIKRigEffectorGoal& NewGoal)
+{
+	Effector->GoalName = NewGoal.GoalName;
+	Effector->BoneName = NewGoal.BoneName;
+}
+
+void UIKRig_PoleSolver::RemoveGoal(const FName& GoalName)
+{
+	if (Effector->GoalName == GoalName)
+	{
+		Effector->GoalName = NAME_None;
+		Effector->BoneName = NAME_None;
+	}
+}
+
+void UIKRig_PoleSolver::RenameGoal(const FName& OldName, const FName& NewName)
+{
+	if (Effector->GoalName == OldName)
+	{
+		Effector->GoalName = NewName;
+	}
+}
+
+void UIKRig_PoleSolver::SetGoalBone(const FName& GoalName, const FName& NewBoneName)
+{
+	if (Effector->GoalName == GoalName)
+	{
+		Effector->BoneName = NewBoneName;
+	}
+}
+
+bool UIKRig_PoleSolver::IsGoalConnected(const FName& GoalName) const
+{
+	return (Effector->GoalName == GoalName);
+}
+
+UObject* UIKRig_PoleSolver::GetGoalSettings(const FName& GoalName) const
+{
+	return (Effector->GoalName == GoalName) ? Effector : nullptr;
+}
+
+bool UIKRig_PoleSolver::IsBoneAffectedBySolver(const FName& BoneName, const FIKRigSkeleton& IKRigSkeleton) const
+{
+	if (EndName != NAME_None)
+	{
+		return IKRigSkeleton.IsBoneInDirectLineage(BoneName, RootName);
+	}
+	return false;
+}
+
+void UIKRig_PoleSolver::SetRootBone(const FName& RootBoneName)
+{
+	RootName = RootBoneName;
+}
+
+void UIKRig_PoleSolver::SetEndBone(const FName& EndBoneName)
+{
+	EndName = EndBoneName;
+}
+
+#endif
+
+void UIKRig_PoleSolver::GatherChildren(const int32 BoneIndex, const FIKRigSkeleton& InSkeleton, TArray<int32>& OutChildren)
+{
+	TArray<int32> Children;
+	InSkeleton.GetChildIndices(BoneIndex, Children);
+	for (int32 ChildIndex: Children)
+	{
+		OutChildren.Add(ChildIndex);
+		GatherChildren(ChildIndex, InSkeleton, OutChildren);
+	}
+}
+
+#undef LOCTEXT_NAMESPACE
diff --git a/Engine/Plugins/Animation/IKRig/Source/IKRig/Private/Solvers/LimbSolver.cpp b/Engine/Plugins/Animation/IKRig/Source/IKRig/Private/Solvers/LimbSolver.cpp
new file mode 100644
--- /dev/null
+++ b/Engine/Plugins/Animation/IKRig/Source/IKRig/Private/Solvers/LimbSolver.cpp
@@ -0,0 +1,627 @@
+﻿// Copyright Epic Games, Inc. All Rights Reserved.
+
+#include "Solvers/LimbSolver.h"
+
+void FLimbSolver::Reset()
+{
+	Length = 0.f;
+	Links.Empty();
+}
+
+bool FLimbSolver::Initialize()
+{
+	Length = 0.f;
+
+	// compute segments' length + total limb's length
+	for (int32 Index = 1; Index < NumLinks(); ++Index)
+	{
+		FLimbLink& Link = Links[Index];
+		Link.Length = FVector::Dist( Link.Location, Links[Index-1].Location);
+		Length += Link.Length;
+	}
+
+	return (NumLinks() > 2);
+}
+
+bool FLimbSolver::Solve(
+	TArray<FTransform>& InOutTransforms,
+	const FVector& InGoalLocation,
+	const FQuat& InGoalRotation,
+	const FLimbSolverSettings& InSettings)
+{
+	// rotate root so the end aligns with the goal
+	const bool bOriented = OrientTransformsTowardsGoal( InOutTransforms, InGoalLocation);
+
+	// update links locations
+	for (int32 Index = 0; Index < NumLinks(); ++Index)
+	{
+		const int32 BoneIndex = Links[Index].BoneIndex;
+		Links[Index].Location = InOutTransforms[BoneIndex].GetLocation();
+	}
+
+	// solve ik
+	const bool bReached = AnimationCore::SolveLimb(
+		Links,
+		Length,
+		InGoalLocation,
+		InSettings.ReachPrecision,
+		InSettings.MaxIterations,
+		InOutTransforms[Links[0].BoneIndex].GetUnitAxis(InSettings.HingeRotationAxis),
+		InSettings.bEnableLimit,
+		InSettings.MinRotationAngle,
+		InSettings.ReachStepAlpha,
+		InSettings.bAveragePull,
+		InSettings.PullDistribution
+		);
+
+	const bool bAdjustedKneeTwist = InSettings.bEnableTwistCorrection ? AdjustKneeTwist(
+															InOutTransforms,
+															FTransform( InGoalRotation, InGoalLocation),
+															InSettings.EndBoneForwardAxis ) : false;
+
+	bool bModifiedLimb = bOriented || bReached || bAdjustedKneeTwist;
+	
+	// adjust end bone orientation
+	const int32 EndIndex = Links.Last().BoneIndex;
+	FTransform& EndTransform = InOutTransforms[EndIndex];
+	if (bModifiedLimb || !EndTransform.GetRotation().Equals(InGoalRotation))
+	{
+		EndTransform.SetRotation(InGoalRotation);
+		bModifiedLimb = true;
+	}
+	
+	if (bModifiedLimb)
+	{
+		// update transforms
+		UpdateTransforms(InOutTransforms);
+	}
+
+	return bModifiedLimb;
+}
+
+int32 FLimbSolver::NumLinks() const
+{
+	return Links.Num();
+}
+
+int32 FLimbSolver::AddLink(const FVector& InLocation, int32 InBoneIndex)
+{
+	return Links.Emplace(InLocation, InBoneIndex);
+}
+
+int32 FLimbSolver::GetBoneIndex(int32 Index) const
+{
+	check(Index >= INDEX_NONE && Index < Links.Num());
+	return Links[Index].BoneIndex; 
+}
+
+bool FLimbSolver::OrientTransformsTowardsGoal(TArray<FTransform>& InOutTransforms, const FVector& InGoalLocation) const
+{
+	const int32 RootIndex = Links[0].BoneIndex;
+	const int32 EndIndex = Links.Last().BoneIndex; 
+	
+	const FVector RootLocation = InOutTransforms[RootIndex].GetLocation();
+	const FVector EndLocation = InOutTransforms[EndIndex].GetLocation();
+
+	const FVector InitialDir = (EndLocation - RootLocation).GetSafeNormal();
+	const FVector TargetDir = (InGoalLocation - RootLocation).GetSafeNormal();
+	
+	// find delta rotation take takes us from Old to New dir
+	if (!InitialDir.IsZero() && !InitialDir.Equals(TargetDir))
+	{
+		const FQuat DeltaRotation = FQuat::FindBetweenNormals(InitialDir, TargetDir);
+		if (!DeltaRotation.IsIdentity())
+		{
+			// rotate Leg so it is aligned with IK Target
+			for (const FLimbLink& Link: Links)
+			{
+				FTransform& BoneTransform = InOutTransforms[Link.BoneIndex]; 
+				
+				BoneTransform.SetRotation(DeltaRotation * BoneTransform.GetRotation());
+
+				const FVector BoneLocation = BoneTransform.GetLocation();
+				BoneTransform.SetLocation(RootLocation + DeltaRotation.RotateVector(BoneLocation - RootLocation));
+			}
+
+			return true;
+		}
+	}
+
+	return false;
+}
+
+void FLimbSolver::UpdateTransforms(TArray<FTransform>& InOutTransforms)
+{
+	// Rotations
+	for (int32 Index = 1; Index < NumLinks(); Index++)
+	{
+		const int32 BoneIndex = Links[Index].BoneIndex;
+		const FTransform& BoneTransform = InOutTransforms[BoneIndex];
+		const FVector& TargetBoneLocation = Links[Index].Location;
+		
+		const int32 ParentIndex = Links[Index-1].BoneIndex;
+		FTransform& ParentTransform = InOutTransforms[ParentIndex];
+		const FVector& TargetParentLocation = Links[Index-1].Location;
+
+		// calculate pre-translation vector between this bone and child
+		const FVector InitialDir = (BoneTransform.GetLocation() - ParentTransform.GetLocation()).GetSafeNormal();
+
+		// Get vector from the post-translation bone to it's child
+		const FVector TargetDir = (TargetBoneLocation - TargetParentLocation).GetSafeNormal();
+
+		const FQuat DeltaRotation = FQuat::FindBetweenNormals(InitialDir, TargetDir);
+		ParentTransform.SetRotation(DeltaRotation * ParentTransform.GetRotation());
+	}
+ 
+	// Translations (do it in one pass?)
+	for (int32 Index = 1; Index < NumLinks(); Index++)
+	{
+		const int32 BoneIndex = Links[Index].BoneIndex;
+		InOutTransforms[BoneIndex].SetTranslation(Links[Index].Location);
+	}
+}
+
+bool FLimbSolver::AdjustKneeTwist(TArray<FTransform>& InOutTransforms, const FTransform& InGoalTransform, const EAxis::Type FootBoneForwardAxis )
+{
+	const int32 EndIndex = Links.Last().BoneIndex;
+	const FVector FootFKLocation = InOutTransforms[EndIndex].GetLocation();
+	const FVector FootIKLocation = InGoalTransform.GetLocation();
+
+	const int32 RootIndex = Links[0].BoneIndex;
+	const FVector HipLocation = InOutTransforms[RootIndex].GetLocation();
+	const FVector FootAxisZ = (FootIKLocation - HipLocation).GetSafeNormal();
+
+	FVector FootFKAxisX = InOutTransforms[EndIndex].GetUnitAxis(FootBoneForwardAxis);
+	FVector FootIKAxisX = InGoalTransform.GetUnitAxis(FootBoneForwardAxis);
+
+	// Reorient X Axis to be perpendicular with FootAxisZ
+	const FVector TmpFKAxisY = FVector::CrossProduct(FootAxisZ, FootFKAxisX);
+	FootFKAxisX = FVector::CrossProduct(TmpFKAxisY, FootAxisZ);
+	
+	const FVector TmpIKAxisY = FVector::CrossProduct(FootAxisZ, FootIKAxisX);
+	FootIKAxisX = FVector::CrossProduct(TmpIKAxisY, FootAxisZ);
+
+	// Compare Axis X to see if we need a rotation to be performed
+	if (!FootFKAxisX.IsZero() && !FootFKAxisX.Equals(FootIKAxisX))
+	{
+		const FQuat DeltaRotation = FQuat::FindBetweenNormals(FootFKAxisX, FootIKAxisX);
+		if (!DeltaRotation.IsIdentity())
+		{
+			// rotate Leg so it is aligned with IK Target
+			for (const FLimbLink& Link: Links)
+			{
+				FTransform& BoneTransform = InOutTransforms[Link.BoneIndex]; 
+				
+				BoneTransform.SetRotation(DeltaRotation * BoneTransform.GetRotation());
+
+				const FVector BoneLocation = BoneTransform.GetLocation();
+				BoneTransform.SetLocation(HipLocation + DeltaRotation.RotateVector(BoneLocation - HipLocation));
+			}
+
+			return true;
+		}
+	}
+
+	return false;
+}
+
+namespace AnimationCore
+{
+	void SolveStraightLimb(TArray<FLimbLink>& InOutLinks, const FVector& GoalLocation)
+	{
+		const FVector& RootLocation = InOutLinks[0].Location;
+		const FVector Direction = (GoalLocation - RootLocation).GetSafeNormal();
+
+		for (int32 Index = 1; Index < InOutLinks.Num(); Index++)
+		{
+			InOutLinks[Index].Location = InOutLinks[Index-1].Location + Direction * InOutLinks[Index].Length;
+		}
+	}
+
+	void SolveTwoBoneIK(TArray<FLimbLink>& InOutLinks, const FVector& InTargetLocation, const FVector& HingeRotationAxis)
+	{
+		check(InOutLinks.Num() == 3);
+
+		FVector& pA = InOutLinks[2].Location; // Foot
+		FVector& pB = InOutLinks[1].Location; // Knee
+		FVector& pC = InOutLinks[0].Location; // Hip / Root
+		
+		// Move foot directly to target.
+		pA = InTargetLocation;
+
+		const FVector HipToFoot = pA - pC;
+
+		// Use Law of Cosines to work out solution.
+		// At this point we know the target location is reachable, and we are already aligned with that location.
+		// So the leg is in the right plane.
+		const float a = InOutLinks[1].Length;	// hip to knee
+		const float b = HipToFoot.Length();		// hip to foot
+		const float c = InOutLinks[2].Length;	// knee to foot
+		
+		const float Two_ab = 2.f * a * b;
+		const float CosC = !FMath::IsNearlyZero(Two_ab) ? (a * a + b * b - c * c) / Two_ab : 0.f;
+ 		const float C = FMath::Acos(CosC);
+		
+		// Project Knee onto Hip to Foot line.
+		const FVector HipToFootDir = !FMath::IsNearlyZero(b) ? HipToFoot / b : FVector::ZeroVector;
+		const FVector HipToKnee = pB - pC;
+		const FVector ProjKnee = pC + HipToKnee.ProjectOnToNormal(HipToFootDir);
+
+		const FVector ProjKneeToKnee = (pB - ProjKnee);
+		FVector BendDir = ProjKneeToKnee.GetSafeNormal(KINDA_SMALL_NUMBER);
+		
+		// If we have a HingeRotationAxis defined, we can cache 'BendDir'
+		// and use it when we can't determine it. (When limb is straight without a bend).
+		// We do this instead of using an explicit one, so we carry over the pole vector that animators use. 
+		// So they can animate it, and we try to extract it from the animation.
+		if ((HingeRotationAxis != FVector::ZeroVector) && (HipToFootDir != FVector::ZeroVector) && !FMath::IsNearlyZero(a))
+		{
+			const FVector HipToKneeDir = HipToKnee / a;
+			const float KneeBendDot = FVector::DotProduct(HipToKneeDir, HipToFootDir);
+
+			FVector& CachedRealBendDir = InOutLinks[1].RealBendDir;
+			FVector& CachedBaseBendDir = InOutLinks[1].BaseBendDir;
+
+			// Valid 'bend', cache 'BendDir'
+			if ((BendDir != FVector::ZeroVector) && (KneeBendDot < 0.99f))
+			{
+				CachedRealBendDir = BendDir;
+				CachedBaseBendDir = FVector::CrossProduct(HingeRotationAxis, HipToFootDir);
+			}
+			// Limb is too straight, can't determine BendDir accurately, so use cached value if possible.
+			else 
+			{
+				// If we have cached 'BendDir', then reorient it based on 'HingeRotationAxis'
+				if (CachedRealBendDir != FVector::ZeroVector)
+				{
+					const FVector CurrentBaseBendDir = FVector::CrossProduct(HingeRotationAxis, HipToFootDir);
+					const FQuat DeltaCachedToCurrBendDir = FQuat::FindBetweenNormals(CachedBaseBendDir, CurrentBaseBendDir);
+					BendDir = DeltaCachedToCurrBendDir.RotateVector(CachedRealBendDir);
+				}
+			}
+		}
+
+		// We just combine both lines into one to save a multiplication.
+		// const FVector NewProjectedKneeLoc = pC + HipToFootDir * a * CosC;
+		// const FVector NewKneeLoc = NewProjectedKneeLoc + Dir_LegLineToKnee * a * FMath::Sin(C);
+		const FVector NewKneeLoc = pC + a * (HipToFootDir * CosC + BendDir * FMath::Sin(C));
+		pB = NewKneeLoc;
+	}
+
+	FVector FindPlaneNormal(const TArray<FLimbLink>& InLinks, const FVector& RootLocation, const FVector& TargetLocation)
+	{
+		const FVector AxisX = (TargetLocation - RootLocation).GetSafeNormal();
+
+		for (int32 Index = 1; Index < InLinks.Num(); ++Index)
+		{
+			const FVector AxisY = (InLinks[Index].Location - RootLocation).GetSafeNormal();
+			const FVector PlaneNormal = FVector::CrossProduct(AxisX, AxisY);
+
+			// Make sure we have a valid normal (Axes were not coplanar).
+			if (PlaneNormal.SizeSquared() > SMALL_NUMBER)
+			{
+				return PlaneNormal.GetUnsafeNormal();
+			}
+		}
+
+		// All links are co-planar?
+		return FVector::UpVector;
+	}
+
+	void FABRIK_ApplyLinkConstraints_Forward(TArray<FLimbLink>& InOutLinks, int32 LinkIndex, float MinRotationAngleRadians)
+	{
+		if ((LinkIndex <= 0) || (LinkIndex >= InOutLinks.Num() - 1))
+		{
+			return;
+		}
+	
+		const FLimbLink& ChildLink = InOutLinks[LinkIndex + 1];
+		const FLimbLink& CurrentLink = InOutLinks[LinkIndex];
+		FLimbLink& ParentLink = InOutLinks[LinkIndex - 1];
+	
+		const FVector ChildAxisX = (ChildLink.Location - CurrentLink.Location).GetSafeNormal();
+		const FVector ChildAxisY = FVector::CrossProduct(CurrentLink.LinkAxisZ, ChildAxisX);
+		const FVector ParentAxisX = (ParentLink.Location - CurrentLink.Location).GetSafeNormal();
+	
+		const float ParentCos = FVector::DotProduct(ParentAxisX, ChildAxisX);
+		const float ParentSin = FVector::DotProduct(ParentAxisX, ChildAxisY);
+	
+		const bool bNeedsReorient = (ParentSin < 0.f) || (ParentCos > FMath::Cos(MinRotationAngleRadians));
+	
+		// Parent Link needs to be reoriented.
+		if (bNeedsReorient)
+		{
+			// folding over itself.
+			if (ParentCos > 0.f)
+			{
+				// Enforce minimum angle.
+				ParentLink.Location = CurrentLink.Location + CurrentLink.Length * (FMath::Cos(MinRotationAngleRadians) * ChildAxisX + FMath::Sin(MinRotationAngleRadians) * ChildAxisY);
+			}
+			else
+			{
+				// When opening up leg, allow it to extend in a full straight line.
+				ParentLink.Location = CurrentLink.Location - ChildAxisX * CurrentLink.Length;
+			}
+		}
+	}
+	
+	void FABRIK_ApplyLinkConstraints_Backward(TArray<FLimbLink>& InOutLinks, int32 LinkIndex, float MinRotationAngleRadians)
+	{
+		if ((LinkIndex <= 0) || (LinkIndex >= InOutLinks.Num() - 1))
+		{
+			return;
+		}
+	
+		FLimbLink& ChildLink = InOutLinks[LinkIndex + 1];
+		const FLimbLink& CurrentLink = InOutLinks[LinkIndex];
+		const FLimbLink& ParentLink = InOutLinks[LinkIndex - 1];
+	
+		const FVector ParentAxisX = (ParentLink.Location - CurrentLink.Location).GetSafeNormal();
+		const FVector ParentAxisY = FVector::CrossProduct(CurrentLink.LinkAxisZ, ParentAxisX);
+		const FVector ChildAxisX = (ChildLink.Location - CurrentLink.Location).GetSafeNormal();
+	
+		const float ChildCos = FVector::DotProduct(ChildAxisX, ParentAxisX);
+		const float ChildSin = FVector::DotProduct(ChildAxisX, ParentAxisY);
+
+		const float CosMinAngle = FMath::Cos(MinRotationAngleRadians);
+		
+		// Parent Link needs to be reoriented.
+		const bool bNeedsReorient = (ChildSin > 0.f) || (ChildCos > CosMinAngle);
+		if (bNeedsReorient)
+		{
+			// folding over itself.
+			if (ChildCos > 0.f)
+			{
+				// Enforce minimum angle.
+				const float SinMinAngle = FMath::Sin(MinRotationAngleRadians);
+				ChildLink.Location = CurrentLink.Location + ChildLink.Length * (CosMinAngle * ParentAxisX - SinMinAngle * ParentAxisY);
+			}
+			else
+			{
+				// When opening up leg, allow it to extend in a full straight line.
+				ChildLink.Location = CurrentLink.Location - ParentAxisX * ChildLink.Length;
+			}
+		}
+	}
+	
+	void FABRIK_ForwardReach(
+		TArray<FLimbLink>& InOutLinks,
+		const FVector& InTargetLocation,
+		float ReachStepAlpha,
+		bool bEnableRotationLimit,
+		float MinRotationAngleRadians)
+	{
+		const int32 NumLinks = InOutLinks.Num();
+		
+		// Move end effector towards target
+		// If we are compressing the chain, limit displacement.
+		// Due to how FABRIK works, if we push the target past the parent's joint, we flip the bone.
+		{
+			FVector EndEffectorToTarget = InTargetLocation - InOutLinks.Last().Location;
+
+			FVector EndEffectorToTargetDir;
+			float EndEffectToTargetSize;
+			EndEffectorToTarget.ToDirectionAndLength(EndEffectorToTargetDir, EndEffectToTargetSize);
+
+			float Displacement = EndEffectToTargetSize;
+			for (int32 Index = NumLinks-1; Index >= 1; Index--)
+			{
+				FVector EndEffectorToParent = InOutLinks[Index].Location - InOutLinks.Last().Location;
+				float ParentDisplacement = FVector::DotProduct(EndEffectorToParent, EndEffectorToTargetDir);
+
+				Displacement = (ParentDisplacement > 0.f) ? FMath::Min(Displacement, ParentDisplacement * ReachStepAlpha) : Displacement;
+			}
+
+			InOutLinks.Last().Location += EndEffectorToTargetDir * Displacement;
+		}
+
+		// "Forward Reaching" stage - adjust bones from end effector.
+		for (int32 Index = NumLinks-2; Index > 0; Index--)
+		{
+			const FLimbLink& ChildLink = InOutLinks[Index+1];
+			FLimbLink& CurrentLink = InOutLinks[Index];
+
+			CurrentLink.Location = ChildLink.Location + (CurrentLink.Location - ChildLink.Location).GetUnsafeNormal() * ChildLink.Length;
+
+			if (bEnableRotationLimit)
+			{
+				FABRIK_ApplyLinkConstraints_Forward(InOutLinks, Index, MinRotationAngleRadians);
+			}
+		}
+	}
+
+	void FABRIK_BackwardReach(
+		TArray<FLimbLink>& InOutLinks,
+		const FVector& InRootTargetLocation,
+		float ReachStepAlpha,
+		bool bEnableRotationLimit,
+		float MinRotationAngleRadians)
+	{
+		const int32 NumLinks = InOutLinks.Num();
+		
+		// Move Root back towards RootTarget
+		// If we are compressing the chain, limit displacement.
+		// Due to how FABRIK works, if we push the target past the parent's joint, we flip the bone.
+		{
+			FVector RootToRootTarget = InRootTargetLocation - InOutLinks[0].Location;
+			
+			FVector RootToRootTargetDir;
+			float RootToRootTargetSize;
+			RootToRootTarget.ToDirectionAndLength(RootToRootTargetDir, RootToRootTargetSize);
+
+			float Displacement = RootToRootTargetSize;
+			
+			for (int32 Index = 0; Index < NumLinks-1; Index++)
+			{
+				FVector RootToChild = InOutLinks[Index].Location - InOutLinks[0].Location;
+				float ChildDisplacement = FVector::DotProduct(RootToChild, RootToRootTargetDir);
+				
+				Displacement = (ChildDisplacement > 0.f) ? FMath::Min(Displacement, ChildDisplacement * ReachStepAlpha) : Displacement;				
+			}
+
+			InOutLinks[0].Location += RootToRootTargetDir * Displacement;
+		}
+
+		// "Backward Reaching" stage - adjust bones from root.
+		for (int32 Index = 1; Index < NumLinks-1; Index++)
+		{
+			const FLimbLink& ParentLink = InOutLinks[Index-1];
+			FLimbLink& CurrentLink = InOutLinks[Index];
+			CurrentLink.Location = ParentLink.Location + (CurrentLink.Location - ParentLink.Location).GetUnsafeNormal() * CurrentLink.Length;
+
+			if (bEnableRotationLimit)
+			{
+				FABRIK_ApplyLinkConstraints_Backward(InOutLinks, Index, MinRotationAngleRadians);
+			}
+		}
+	}
+	
+	void SolveFABRIK(
+		TArray<FLimbLink>& InOutLinks,
+		const FVector& InTargetLocation,
+		float InReachPrecision,
+		int32 InMaxIterations,
+		bool bUseAngleLimit,
+		float MinRotationAngleRadians,
+		float ReachStepAlpha,
+		float PullDistribution,
+		bool bAveragePull)
+	{
+		const int32 NumLinks = InOutLinks.Num();
+		
+		// Make sure precision is not too small.
+		const float ReachPrecision = FMath::Max(InReachPrecision, KINDA_SMALL_NUMBER);
+
+		const FVector RootTargetLocation = InOutLinks[0].Location;
+
+		// Check distance between foot and foot target location
+		float Slop = FVector::Dist(InOutLinks.Last().Location, InTargetLocation);
+		
+		if (Slop > ReachPrecision || bUseAngleLimit)
+		{
+			if (bUseAngleLimit)
+			{
+				// Since we've previously aligned the foot with the IK Target, we're solving IK in 2D space on a single plane.
+				// Find Plane Normal, to use in rotation constraints.
+				const FVector PlaneNormal = FindPlaneNormal(InOutLinks, RootTargetLocation, InTargetLocation);
+
+				for (int32 Index = NumLinks-2; Index > 0; Index--)
+				{
+					const FLimbLink& ChildLink = InOutLinks[Index + 1];
+					FLimbLink& CurrentLink = InOutLinks[Index];
+					const FLimbLink& ParentLink = InOutLinks[Index - 1];
+
+					const FVector ChildAxisX = (ChildLink.Location - CurrentLink.Location).GetSafeNormal();
+					const FVector ChildAxisY = FVector::CrossProduct(PlaneNormal, ChildAxisX);
+					const FVector ParentAxisX = (ParentLink.Location - CurrentLink.Location).GetSafeNormal();
+
+					// Orient Z, so that ChildAxisY points 'up' and produces positive Sin values.
+					CurrentLink.LinkAxisZ = FVector::DotProduct(ParentAxisX, ChildAxisY) > 0.f ? PlaneNormal : -PlaneNormal;
+				}
+			}
+
+			// Re-position limb to distribute pull
+			const FVector EndToTarget = InTargetLocation - InOutLinks.Last().Location;
+			const FVector RootToTarget = RootTargetLocation - InOutLinks[0].Location;
+			const FVector PullDistributionOffset = PullDistribution * (EndToTarget) + (1.f - PullDistribution) * (RootToTarget);
+			for (int32 Index = 0; Index < NumLinks; Index++)
+			{
+				InOutLinks[Index].Location += PullDistributionOffset;
+			}
+			
+			int32 IterationCount = 1;
+			const int32 MaxIterations = FMath::Max(InMaxIterations, 1);
+			do
+			{
+				const float PreviousSlop = Slop;
+
+				// Pull averaging only has a visual impact when we have more than 2 bones (3 links).
+				const bool bNeedPullAveraging =(NumLinks > 3) && bAveragePull && (Slop > 1.f); 
+				if (bNeedPullAveraging)
+				{
+					TArray<FLimbLink> ForwardPull = InOutLinks;
+					FABRIK_ForwardReach(ForwardPull, InTargetLocation, ReachStepAlpha, bUseAngleLimit, MinRotationAngleRadians);
+
+					TArray<FLimbLink> BackwardPull = InOutLinks;
+					FABRIK_BackwardReach(BackwardPull,RootTargetLocation, ReachStepAlpha, bUseAngleLimit, MinRotationAngleRadians);
+
+					// Average pulls
+					for (int32 LinkIndex = 0; LinkIndex < NumLinks; LinkIndex++)
+					{
+						InOutLinks[LinkIndex].Location = 0.5f * (ForwardPull[LinkIndex].Location + BackwardPull[LinkIndex].Location);
+					}
+				}
+				else
+				{
+					FABRIK_ForwardReach(InOutLinks, InTargetLocation, ReachStepAlpha, bUseAngleLimit, MinRotationAngleRadians);
+					FABRIK_BackwardReach(InOutLinks, RootTargetLocation, ReachStepAlpha, bUseAngleLimit, MinRotationAngleRadians);
+				}
+				
+				Slop = FVector::Dist( InOutLinks.Last().Location, InTargetLocation) + FVector::Dist(InOutLinks[0].Location, RootTargetLocation);
+
+				// Abort if we're not getting closer and enter a deadlock.
+				if (Slop > PreviousSlop)
+				{
+					break;
+				}
+
+			} while ((Slop > ReachPrecision) && (++IterationCount < MaxIterations));
+			
+			// Make sure our root is back at our root target.
+			if (!InOutLinks[0].Location.Equals(RootTargetLocation))
+			{
+				FABRIK_BackwardReach(InOutLinks, RootTargetLocation, ReachStepAlpha, bUseAngleLimit, MinRotationAngleRadians);
+			}
+
+			// If we reached, set target precisely
+			if (Slop <= ReachPrecision)
+			{
+				InOutLinks.Last().Location = InTargetLocation;
+			}
+		}
+	}
+	
+	bool SolveLimb(
+		TArray<FLimbLink>& InOutLinks,
+		float LimbLength,
+		const FVector& GoalLocation,
+		float Precision,
+		int32 InMaxIterations,
+		const FVector& HingeRotationAxis,
+		bool bUseAngleLimit,
+		float MinRotationAngle,
+		float ReachStepAlpha,
+		bool bAveragePull,
+		float PullDistribution
+		)
+	{
+		const FVector& EndLocation = InOutLinks.Last().Location;
+		const bool FootAtGoal = EndLocation.Equals(GoalLocation, Precision);
+		
+		if (FootAtGoal && !bUseAngleLimit)
+		{
+			return false;
+		}
+		
+		const bool bDirectSolve = ( InOutLinks.Num() == 3 );
+		const FVector& RootLocation = InOutLinks[0].Location;
+
+		const float MinRotationAngleRadians = FMath::DegreesToRadians(FMath::Clamp(MinRotationAngle, 0.f, 90.f));
+		
+		// If we can't reach, we just go in a straight line towards the target,
+		if (FVector::DistSquared(RootLocation, GoalLocation) >= FMath::Square(LimbLength))
+		{
+			SolveStraightLimb(InOutLinks, GoalLocation);
+		}
+		else if (bDirectSolve)
+		{
+			SolveTwoBoneIK(InOutLinks, GoalLocation, HingeRotationAxis);
+		}
+		else
+		{
+			SolveFABRIK(InOutLinks, GoalLocation, Precision, InMaxIterations, bUseAngleLimit, MinRotationAngleRadians, ReachStepAlpha, PullDistribution, bAveragePull);
+		}
+		
+		return true;
+	}
+}
\ No newline at end of file
diff --git a/Engine/Plugins/Animation/IKRig/Source/IKRig/Public/IKRigSkeleton.h b/Engine/Plugins/Animation/IKRig/Source/IKRig/Public/IKRigSkeleton.h
--- a/Engine/Plugins/Animation/IKRig/Source/IKRig/Public/IKRigSkeleton.h
+++ b/Engine/Plugins/Animation/IKRig/Source/IKRig/Public/IKRigSkeleton.h
@@ -114,6 +114,8 @@ struct IKRIG_API FIKRigSkeleton
 
 	int32 GetParentIndexThatIsNotExcluded(const int32 BoneIndex) const;
 
+	int32 GetChildIndices(const int32 ParentBoneIndex, TArray<int32>& Children) const;
+
 	static void ConvertLocalPoseToGlobal(
 		const TArray<int32>& InParentIndices,
 		const TArray<FTransform>& InLocalPose,
diff --git a/Engine/Plugins/Animation/IKRig/Source/IKRig/Public/IKRigSolver.h b/Engine/Plugins/Animation/IKRig/Source/IKRig/Public/IKRigSolver.h
--- a/Engine/Plugins/Animation/IKRig/Source/IKRig/Public/IKRigSolver.h
+++ b/Engine/Plugins/Animation/IKRig/Source/IKRig/Public/IKRigSolver.h
@@ -79,6 +79,12 @@ public:
 	virtual bool RequiresRootBone() const { return false; };
 	//** END ROOT BONE */
 
+	//** ROOT BONE (optional, implement if your solver requires a end bone) */
+	/** override to support SETTING END BONE for the solver */
+	virtual void SetEndBone(const FName& EndBoneName){};
+	virtual bool RequiresEndBone() const { return false; };
+	//** END ROOT BONE */
+
 	//** BONE SETTINGS (optional, implement if your solver supports per-bone settings) */
 	/** override to support ADDING PER-BONE settings for this solver */
 	virtual void AddBoneSetting(const FName& BoneName){};
diff --git a/Engine/Plugins/Animation/IKRig/Source/IKRig/Public/Solvers/IKRig_LimbSolver.h b/Engine/Plugins/Animation/IKRig/Source/IKRig/Public/Solvers/IKRig_LimbSolver.h
new file mode 100644
--- /dev/null
+++ b/Engine/Plugins/Animation/IKRig/Source/IKRig/Public/Solvers/IKRig_LimbSolver.h
@@ -0,0 +1,116 @@
+﻿// Copyright Epic Games, Inc. All Rights Reserved.
+
+#pragma once
+
+#include "IKRigDefinition.h"
+#include "IKRigSolver.h"
+
+#include "Solvers/LimbSolver.h"
+
+#include "IKRig_LimbSolver.generated.h"
+
+UCLASS()
+class IKRIG_API UIKRig_LimbEffector : public UObject
+{
+	GENERATED_BODY()
+
+public:
+
+	UPROPERTY(VisibleAnywhere, Category = "Limb IK Effector")
+	FName GoalName = NAME_None;
+
+	UPROPERTY(VisibleAnywhere, Category = "Limb IK Effector")
+	FName BoneName = NAME_None;
+};
+
+UCLASS(EditInlineNew)
+class IKRIG_API UIKRig_LimbSolver : public UIKRigSolver
+{
+	GENERATED_BODY()
+
+public:
+
+	UIKRig_LimbSolver();
+	
+	UPROPERTY(VisibleAnywhere, Category = "Limb Solver Settings")
+	FName RootName = NAME_None;
+
+	/** Precision (distance to the target) */
+	UPROPERTY(EditAnywhere, Category = "Limb Solver Settings", meta = (ClampMin = "0.0"))
+	float ReachPrecision = 0.01f;
+	
+	// TWO BONES SETTINGS
+	
+	/** Hinge Bones Rotation Axis. This is essentially the plane normal for (hip - knee - foot). */
+	UPROPERTY(EditAnywhere, Category = "Limb Solver Settings|Two Bones")
+	TEnumAsByte<EAxis::Type> HingeRotationAxis = EAxis::None;
+
+	// FABRIK SETTINGS
+
+	/** Number of Max Iterations to reach the target */
+	UPROPERTY(EditAnywhere, Category = "Limb Solver Settings|FABRIK", meta = (ClampMin = "0", UIMin = "0", UIMax = "100"))
+	int32 MaxIterations = 12;
+
+	/** Enable/Disable rotational limits */
+	UPROPERTY(EditAnywhere, Category = "Limb Solver Settings|FABRIK|Limits")
+	bool bEnableLimit = false;
+
+	/** Only used if bEnableRotationLimit is enabled. Prevents the leg from folding onto itself,
+	* and forces at least this angle between Parent and Child bone. */
+	UPROPERTY(EditAnywhere, Category = "Limb Solver Settings|FABRIK|Limits", meta = (ClampMin = "0.0", ClampMax = "90.0", UIMin = "0.0", UIMax = "90.0"))
+	float MinRotationAngle = 15.f;
+
+	/** Pull averaging only has a visual impact when we have more than 2 bones (3 links). */
+	UPROPERTY(EditAnywhere, Category = "Limb Solver Settings|FABRIK|Pull Averaging")
+	bool bAveragePull = true;
+
+	/** Re-position limb to distribute pull: 0 = foot, 0.5 = balanced, 1.f = hip */
+	UPROPERTY(EditAnywhere, Category = "Limb Solver Settings|FABRIK|Pull Averaging", meta = (ClampMin = "0.0", ClampMax = "1.0",UIMin = "0.0", UIMax = "1.0"))
+	float PullDistribution = 0.5f;
+	
+	/** Move end effector towards target. If we are compressing the chain, limit displacement.*/
+	UPROPERTY(EditAnywhere, Category = "Limb Solver Settings|FABRIK|Pull Averaging", meta = (ClampMin = "0.0", ClampMax = "1.0",UIMin = "0.0", UIMax = "1.0"))
+	float ReachStepAlpha = 0.7f;
+
+	// TWIST SETTINGS
+	
+	/** Enable Knee Twist correction, by comparing Foot FK with Foot IK orientation. */
+	UPROPERTY(EditAnywhere, Category = "Limb Solver Settings|Twist")
+	bool bEnableTwistCorrection = false;
+	
+	/** Forward Axis for Foot bone. */
+	UPROPERTY(EditAnywhere, Category = "Limb Solver Settings|Twist")
+	TEnumAsByte<EAxis::Type> EndBoneForwardAxis = EAxis::Y;
+	
+	/** UIKRigSolver interface */
+	virtual void Initialize(const FIKRigSkeleton& IKRigSkeleton) override;
+	virtual void Solve(FIKRigSkeleton& IKRigSkeleton, const FIKRigGoalContainer& Goals) override;
+	
+#if WITH_EDITOR
+	virtual void UpdateSolverSettings(UIKRigSolver* InSettings) override;
+	virtual FText GetNiceName() const override;
+	virtual bool GetWarningMessage(FText& OutWarningMessage) const override;
+	virtual bool IsBoneAffectedBySolver(const FName& BoneName, const FIKRigSkeleton& IKRigSkeleton) const override;
+	// goals
+	virtual void AddGoal(const UIKRigEffectorGoal& NewGoal) override;
+	virtual void RemoveGoal(const FName& GoalName) override;
+	virtual void RenameGoal(const FName& OldName, const FName& NewName) override;
+	virtual void SetGoalBone(const FName& GoalName, const FName& NewBoneName) override;
+	virtual bool IsGoalConnected(const FName& GoalName) const override;
+	virtual UObject* GetGoalSettings(const FName& GoalName) const override;
+	// root bone can be set on this solver
+	virtual bool RequiresRootBone() const override { return true; };
+	virtual void SetRootBone(const FName& RootBoneName) override;
+#endif
+	/** END UIKRigSolver interface */
+	
+private:
+
+	static void GatherChildren(const int32 BoneIndex, const FIKRigSkeleton& InSkeleton, TArray<int32>& OutChildren);
+	
+	UPROPERTY()
+	TObjectPtr<UIKRig_LimbEffector> Effector;
+	
+	FLimbSolver						Solver;
+	TArray<int32>					ChildrenToUpdate;
+};
diff --git a/Engine/Plugins/Animation/IKRig/Source/IKRig/Public/Solvers/IKRig_PoleSolver.h b/Engine/Plugins/Animation/IKRig/Source/IKRig/Public/Solvers/IKRig_PoleSolver.h
new file mode 100644
--- /dev/null
+++ b/Engine/Plugins/Animation/IKRig/Source/IKRig/Public/Solvers/IKRig_PoleSolver.h
@@ -0,0 +1,73 @@
+﻿// Copyright Epic Games, Inc. All Rights Reserved.
+
+#pragma once
+
+#include "IKRigDefinition.h"
+#include "IKRigSolver.h"
+
+#include "IKRig_PoleSolver.generated.h"
+
+UCLASS()
+class IKRIG_API UIKRig_PoleSolverEffector : public UObject
+{
+	GENERATED_BODY()
+
+public:
+
+	UPROPERTY(VisibleAnywhere, Category = "Pole Solver Effector")
+	FName GoalName = NAME_None;
+
+	UPROPERTY(VisibleAnywhere, Category = "Pole Solver Effector")
+	FName BoneName = NAME_None;
+};
+
+UCLASS(EditInlineNew)
+class IKRIG_API UIKRig_PoleSolver : public UIKRigSolver
+{
+	GENERATED_BODY()
+
+public:
+
+	UIKRig_PoleSolver();
+	
+	UPROPERTY(VisibleAnywhere, Category = "Pole Solver Settings")
+	FName RootName = NAME_None;
+
+	UPROPERTY(VisibleAnywhere, Category = "Pole Solver Settings")
+	FName EndName = NAME_None;
+	
+	/** UIKRigSolver interface */
+	virtual void Initialize(const FIKRigSkeleton& IKRigSkeleton) override;
+	virtual void Solve(FIKRigSkeleton& IKRigSkeleton, const FIKRigGoalContainer& Goals) override;
+#if WITH_EDITOR
+	virtual void UpdateSolverSettings(UIKRigSolver* InSettings) override;
+	virtual FText GetNiceName() const override;
+	virtual bool GetWarningMessage(FText& OutWarningMessage) const override;
+	virtual bool IsBoneAffectedBySolver(const FName& BoneName, const FIKRigSkeleton& IKRigSkeleton) const override;
+	// goals
+	virtual void AddGoal(const UIKRigEffectorGoal& NewGoal) override;
+	virtual void RemoveGoal(const FName& GoalName) override;
+	virtual void RenameGoal(const FName& OldName, const FName& NewName) override;
+	virtual void SetGoalBone(const FName& GoalName, const FName& NewBoneName) override;
+	virtual bool IsGoalConnected(const FName& GoalName) const override;
+	virtual UObject* GetGoalSettings(const FName& GoalName) const override;
+	// root bone can be set on this solver
+	virtual bool RequiresRootBone() const override { return true; };
+	virtual void SetRootBone(const FName& RootBoneName) override;
+	// end bone can be set on this solver
+	virtual void SetEndBone(const FName& EndBoneName) override;
+	virtual bool RequiresEndBone() const override { return true; };
+	/** END UIKRigSolver interface */
+#endif
+
+private:
+
+	static void GatherChildren(const int32 BoneIndex, const FIKRigSkeleton& InSkeleton, TArray<int32>& OutChildren);
+	
+	UPROPERTY()
+	TObjectPtr<UIKRig_PoleSolverEffector>	Effector;
+	
+	TArray<int32>							Chain;
+	TArray<int32>							ChildrenToUpdate;
+};
+
diff --git a/Engine/Plugins/Animation/IKRig/Source/IKRig/Public/Solvers/LimbSolver.h b/Engine/Plugins/Animation/IKRig/Source/IKRig/Public/Solvers/LimbSolver.h
new file mode 100644
--- /dev/null
+++ b/Engine/Plugins/Animation/IKRig/Source/IKRig/Public/Solvers/LimbSolver.h
@@ -0,0 +1,159 @@
+﻿// Copyright Epic Games, Inc. All Rights Reserved.
+
+#pragma once
+
+#include "CoreMinimal.h"
+#include "LimbSolver.generated.h"
+
+USTRUCT()
+struct IKRIG_API FLimbSolverSettings
+{
+	GENERATED_BODY()
+
+	/** Precision (distance to the target) */
+	UPROPERTY(EditAnywhere, Category = "Limb Solver Settings", meta = (ClampMin = "0.0"))
+	float ReachPrecision = 0.01f;
+	
+	// TWO BONES SETTINGS
+	
+	/** Hinge Bones Rotation Axis. This is essentially the plane normal for (hip - knee - foot). */
+	UPROPERTY(EditAnywhere, Category = "Limb Solver Settings|Two Bones")
+	TEnumAsByte<EAxis::Type> HingeRotationAxis = EAxis::None;
+
+	// FABRIK SETTINGS
+
+	/** Number of Max Iterations to reach the target */
+	UPROPERTY(EditAnywhere, Category = "Limb Solver Settings|FABRIK", meta = (ClampMin = "0", UIMin = "0", UIMax = "100"))
+	int32 MaxIterations = 12;
+
+	/** Enable/Disable rotational limits */
+	UPROPERTY(EditAnywhere, Category = "Limb Solver Settings|FABRIK|Limits")
+	bool bEnableLimit = false;
+
+	/** Only used if bEnableRotationLimit is enabled. Prevents the leg from folding onto itself,
+	* and forces at least this angle between Parent and Child bone. */
+	UPROPERTY(EditAnywhere, Category = "Limb Solver Settings|FABRIK|Limits", meta = (ClampMin = "0.0", ClampMax = "90.0", UIMin = "0.0", UIMax = "90.0"))
+	float MinRotationAngle = 15.f;
+
+	/** Pull averaging only has a visual impact when we have more than 2 bones (3 links). */
+	UPROPERTY(EditAnywhere, Category = "Limb Solver Settings|FABRIK|Pull Averaging")
+	bool bAveragePull = true;
+
+	/** Re-position limb to distribute pull: 0 = foot, 0.5 = balanced, 1.f = hip */
+	UPROPERTY(EditAnywhere, Category = "Limb Solver Settings|FABRIK|Pull Averaging", meta = (ClampMin = "0.0", ClampMax = "1.0",UIMin = "0.0", UIMax = "1.0"))
+	float PullDistribution = 0.5f;
+	
+	/** Move end effector towards target. If we are compressing the chain, limit displacement.*/
+	UPROPERTY(EditAnywhere, Category = "Limb Solver Settings|FABRIK|Pull Averaging", meta = (ClampMin = "0.0", ClampMax = "1.0",UIMin = "0.0", UIMax = "1.0"))
+	float ReachStepAlpha = 0.7f;
+
+	// TWIST SETTINGS
+	
+	/** Enable Knee Twist correction, by comparing Foot FK with Foot IK orientation. */
+	UPROPERTY(EditAnywhere, Category = "Limb Solver Settings|Twist")
+	bool bEnableTwistCorrection = false;
+	
+	/** Forward Axis for Foot bone. */
+	UPROPERTY(EditAnywhere, Category = "Limb Solver Settings|Twist")
+	TEnumAsByte<EAxis::Type> EndBoneForwardAxis = EAxis::Y;
+};
+
+USTRUCT()
+struct FLimbLink
+{
+	GENERATED_USTRUCT_BODY()
+
+	// Location of bone in component space.
+	FVector Location;
+
+	// Distance to its parent
+	float Length;
+
+	// Bone Index in SkeletalMesh
+	int32 BoneIndex;
+
+	// Axis utilites
+	FVector LinkAxisZ;
+	FVector RealBendDir;
+	FVector BaseBendDir;
+
+	FLimbLink()
+		: Location(FVector::ZeroVector)
+		, Length(0.f)
+		, BoneIndex( INDEX_NONE )
+		, LinkAxisZ(FVector::ZeroVector)
+		, RealBendDir(FVector::ZeroVector)
+		, BaseBendDir(FVector::ZeroVector)
+	{}
+
+	FLimbLink(const FVector& InLocation, int32 InBoneIndex)
+		: Location(InLocation)
+		, Length(0.f)
+		, BoneIndex(InBoneIndex)
+		, LinkAxisZ(FVector::ZeroVector)
+		, RealBendDir(FVector::ZeroVector)
+		, BaseBendDir(FVector::ZeroVector)
+	{}
+};
+
+USTRUCT()
+struct FLimbSolver
+{
+	GENERATED_USTRUCT_BODY()
+
+	void Reset();
+	bool Initialize();
+	bool Solve(
+		TArray<FTransform>& InOutTransforms,
+		const FVector& InGoalLocation,
+		const FQuat& InGoalRotation,
+		const FLimbSolverSettings& InSettings);
+
+	int32 AddLink(const FVector& InLocation, int32 InBoneIndex);
+	int32 GetBoneIndex(const int32 Index) const;
+	int32 NumLinks() const;
+	
+private:
+	
+	void UpdateTransforms(TArray<FTransform>& InOutTransforms);
+	bool OrientTransformsTowardsGoal(TArray<FTransform>& InOutTransforms, const FVector& InGoalLocation) const;
+	bool AdjustKneeTwist(TArray<FTransform>& InOutTransforms, const FTransform& InGoalTransform, const EAxis::Type FootBoneForwardAxis );
+
+	float Length = 0.f;
+	
+	TArray<FLimbLink> Links;
+};
+
+namespace AnimationCore
+{
+	/**
+	* Limb solver
+	*
+	* This solves a hybrid TwoBoneIK / FABRIK based on the number of links within the chain
+	*
+	* @param	InOutLinks			Array of chain links
+	* @param	LimbLength			The full length on the limb (sum of all links' length)
+	* @param	GoalLocation		The location of the IK target
+	* @param	Precision			Precision 
+	* @param	InMaxIterations		Number of Max Iteration (used for FABRIK only)
+	* @param	HingeRotationAxis	Hinge Bones Rotation Axis. This is essentially the plane normal for (hip - knee - foot). (used for TwoBones only)
+	* @param	bUseAngleLimit		Enable/Disable rotational limits (used for FABRIK only)
+	* @param	MinRotationAngle	Prevents the leg from folding onto itself, and forces at least this angle between Parent and Child bone (used for FABRIK only)
+	* @param	ReachStepAlpha		Move end effector towards target. If we are compressing the chain, limit displacement. (used for FABRIK only)
+	* * @param	bAveragePull		Pull averaging only has a visual impact when we have more than 2 bones (3 links). (used for FABRIK only)
+	* @param	PullDistribution	Re-position limb to distribute pull: 0 = foot, 0.5 = balanced, 1.f = hip. (used for FABRIK only)
+	* @return  true if modified. False if not.
+	*/
+	bool SolveLimb(
+		TArray<FLimbLink>& InOutLinks,
+		float LimbLength,
+		const FVector& GoalLocation,
+		float Precision,
+		int32 InMaxIterations,
+		const FVector& HingeRotationAxis,
+		bool bUseAngleLimit,
+		float MinRotationAngle,
+		float ReachStepAlpha,
+		bool bAveragePull,
+		float PullDistribution);
+}
\ No newline at end of file
diff --git a/Engine/Plugins/Animation/IKRig/Source/IKRigEditor/Private/RigEditor/IKRigController.cpp b/Engine/Plugins/Animation/IKRig/Source/IKRigEditor/Private/RigEditor/IKRigController.cpp
--- a/Engine/Plugins/Animation/IKRig/Source/IKRigEditor/Private/RigEditor/IKRigController.cpp
+++ b/Engine/Plugins/Animation/IKRig/Source/IKRigEditor/Private/RigEditor/IKRigController.cpp
@@ -524,6 +524,27 @@ void UIKRigController::SetRootBone(const FName& RootBoneName, int32 SolverIndex)
 	BroadcastNeedsReinitialized();
 }
 
+void UIKRigController::SetEndBone(const FName& EndBoneName, int32 SolverIndex) const
+{
+	if (!Asset->Solvers.IsValidIndex(SolverIndex))
+	{
+		return; // solver doesn't exist
+	}
+
+	if (Asset->Skeleton.GetBoneIndexFromName(EndBoneName) == INDEX_NONE)
+	{
+		return; // bone doesn't exist
+	}
+
+	FScopedTransaction Transaction(LOCTEXT("SetEndBone_Label", "Set End Bone"));
+	UIKRigSolver* Solver = Asset->Solvers[SolverIndex];
+	Solver->Modify();
+
+	Solver->SetEndBone(EndBoneName);
+
+	BroadcastNeedsReinitialized();
+}
+
 const TArray<UIKRigSolver*>& UIKRigController::GetSolverArray() const
 {
 	return Asset->Solvers;
diff --git a/Engine/Plugins/Animation/IKRig/Source/IKRigEditor/Private/RigEditor/IKRigSkeletonCommands.cpp b/Engine/Plugins/Animation/IKRig/Source/IKRigEditor/Private/RigEditor/IKRigSkeletonCommands.cpp
--- a/Engine/Plugins/Animation/IKRig/Source/IKRigEditor/Private/RigEditor/IKRigSkeletonCommands.cpp
+++ b/Engine/Plugins/Animation/IKRig/Source/IKRigEditor/Private/RigEditor/IKRigSkeletonCommands.cpp
@@ -11,6 +11,7 @@ void FIKRigSkeletonCommands::RegisterCommands()
 	UI_COMMAND(ConnectGoalToSolvers, "Connect Goal to Selected Solvers", "Make the selected goal an effector in the selected solvers.", EUserInterfaceActionType::Button, FInputChord());
 	UI_COMMAND(DisconnectGoalFromSolvers, "Disconnect Goal from Selected Solvers", "Remove effectors from the selected solvers that use this goal.", EUserInterfaceActionType::Button, FInputChord());
 	UI_COMMAND(SetRootBoneOnSolvers, "Set Root Bone on Selected Solvers", "Set the Root Bone setting on the selected solvers.", EUserInterfaceActionType::Button, FInputChord());
+	UI_COMMAND(SetEndBoneOnSolvers, "Set End Bone on Selected Solvers", "Set End Bone on Selected Solvers.", EUserInterfaceActionType::Button, FInputChord());
 	UI_COMMAND(AddBoneSettings, "Add Settings to Selected Bone", "Apply settings to the selected bone for all selected solvers (defined per-solver for limits, stiffness etc).", EUserInterfaceActionType::Button, FInputChord());
 	UI_COMMAND(RemoveBoneSettings, "Remove Settings on Selected Bone", "Remove all settings on the selected bone in all selected solvers.", EUserInterfaceActionType::Button, FInputChord());
 	UI_COMMAND(ExcludeBone, "Exclude Selected Bone From Solve", "Ignore bone in all solvers.", EUserInterfaceActionType::Button, FInputChord());
diff --git a/Engine/Plugins/Animation/IKRig/Source/IKRigEditor/Private/RigEditor/SIKRigSkeleton.cpp b/Engine/Plugins/Animation/IKRig/Source/IKRigEditor/Private/RigEditor/SIKRigSkeleton.cpp
--- a/Engine/Plugins/Animation/IKRig/Source/IKRigEditor/Private/RigEditor/SIKRigSkeleton.cpp
+++ b/Engine/Plugins/Animation/IKRig/Source/IKRigEditor/Private/RigEditor/SIKRigSkeleton.cpp
@@ -327,6 +327,12 @@ void SIKRigSkeleton::BindCommands()
         FExecuteAction::CreateSP(this, &SIKRigSkeleton::HandleSetRootBoneOnSolvers),
         FCanExecuteAction::CreateSP(this, &SIKRigSkeleton::CanSetRootBoneOnSolvers));
 
+	CommandList->MapAction(Commands.SetEndBoneOnSolvers,
+	FExecuteAction::CreateSP(this, &SIKRigSkeleton::HandleSetEndBoneOnSolvers),
+	FCanExecuteAction::CreateSP(this, &SIKRigSkeleton::CanSetEndBoneOnSolvers),
+	FIsActionChecked(),
+	FIsActionButtonVisible::CreateSP(this, &SIKRigSkeleton::HasEndBoneCompatibleSolverSelected));
+
 	CommandList->MapAction(Commands.AddBoneSettings,
         FExecuteAction::CreateSP(this, &SIKRigSkeleton::HandleAddBoneSettings),
         FCanExecuteAction::CreateSP(this, &SIKRigSkeleton::CanAddBoneSettings));
@@ -376,6 +382,7 @@ void SIKRigSkeleton::FillContextMenu(FMenuBuilder& MenuBuilder)
 	MenuBuilder.AddMenuEntry(Actions.AddBoneSettings);
 	MenuBuilder.AddMenuEntry(Actions.RemoveBoneSettings);
 	MenuBuilder.AddMenuEntry(Actions.SetRootBoneOnSolvers);
+	MenuBuilder.AddMenuEntry(Actions.SetEndBoneOnSolvers);
 	MenuBuilder.EndSection();
 
 	MenuBuilder.BeginSection("IncludeExclude", LOCTEXT("IncludeExcludeOperations", "Exclude Bones"));
@@ -644,6 +651,73 @@ bool SIKRigSkeleton::CanSetRootBoneOnSolvers()
 	return false;
 }
 
+void SIKRigSkeleton::HandleSetEndBoneOnSolvers()
+{
+	const TSharedPtr<FIKRigEditorController> Controller = EditorController.Pin();
+	if (!Controller.IsValid())
+	{
+		return;
+	}
+
+    // get name of selected root bone
+    TArray<TSharedPtr<FIKRigTreeElement>> SelectedBones;
+	GetSelectedBones(SelectedBones);
+	const FName RootBoneName = SelectedBones[0]->BoneName;
+
+	// apply to all selected solvers (ignored on solvers that don't accept a root bone)
+	UIKRigController* AssetController = Controller->AssetController;
+	TArray<TSharedPtr<FSolverStackElement>> SelectedSolvers;
+	Controller->GetSelectedSolvers(SelectedSolvers);
+	int32 SolverToShow = 0;
+	for (const TSharedPtr<FSolverStackElement>& Solver : SelectedSolvers)
+	{
+		AssetController->SetEndBone(RootBoneName, Solver->IndexInStack);
+		SolverToShow = Solver->IndexInStack;
+	}
+
+	// show solver that had it's root bone updated
+	Controller->ShowDetailsForSolver(SolverToShow);
+	
+	// show new icon when bone has settings applied
+	RefreshTreeView();
+}
+
+bool SIKRigSkeleton::CanSetEndBoneOnSolvers() const
+{
+	const TSharedPtr<FIKRigEditorController> Controller = EditorController.Pin();
+	if (!Controller.IsValid())
+	{
+		return false;
+	}
+	
+	// must have at least 1 bone selected
+	TArray<TSharedPtr<FIKRigTreeElement>> SelectedBones;
+	GetSelectedBones(SelectedBones);
+	if (SelectedBones.Num() != 1)
+	{
+		return false;
+	}
+
+	// must have at least 1 solver selected that accepts end bones
+	UIKRigController* AssetController = Controller->AssetController;
+	TArray<TSharedPtr<FSolverStackElement>> SelectedSolvers;
+	Controller->GetSelectedSolvers(SelectedSolvers);
+	for (const TSharedPtr<FSolverStackElement>& Solver : SelectedSolvers)
+	{
+		if (AssetController->GetSolver(Solver->IndexInStack)->RequiresEndBone())
+		{
+			return true;
+		}
+	}
+	
+	return false;
+}
+
+bool SIKRigSkeleton::HasEndBoneCompatibleSolverSelected() const
+{
+	return CanSetEndBoneOnSolvers();
+}
+
 void SIKRigSkeleton::HandleAddBoneSettings()
 {
 	const TSharedPtr<FIKRigEditorController> Controller = EditorController.Pin();
@@ -919,7 +993,7 @@ bool SIKRigSkeleton::CanSetRetargetRoot()
 	return !SelectedBones.IsEmpty();
 }
 
-void SIKRigSkeleton::GetSelectedBones(TArray<TSharedPtr<FIKRigTreeElement>>& OutBoneItems)
+void SIKRigSkeleton::GetSelectedBones(TArray<TSharedPtr<FIKRigTreeElement>>& OutBoneItems) const
 {
 	const TArray<TSharedPtr<FIKRigTreeElement>> SelectedItems = TreeView->GetSelectedItems();
 	for (const TSharedPtr<FIKRigTreeElement>& Item : SelectedItems)
diff --git a/Engine/Plugins/Animation/IKRig/Source/IKRigEditor/Public/RigEditor/IKRigController.h b/Engine/Plugins/Animation/IKRig/Source/IKRigEditor/Public/RigEditor/IKRigController.h
--- a/Engine/Plugins/Animation/IKRig/Source/IKRigEditor/Public/RigEditor/IKRigController.h
+++ b/Engine/Plugins/Animation/IKRig/Source/IKRigEditor/Public/RigEditor/IKRigController.h
@@ -69,6 +69,8 @@ public:
 	int32 GetNumSolvers() const;
 	/** Set the root bone on a given solver. (not all solvers support root bones, checks CanSetRootBone() first) */
 	void SetRootBone(const FName& RootBoneName, int32 SolverIndex) const;
+	/** Set the end bone on a given solver. (not all solvers require extra end bones, checks CanSetEndBone() first) */
+	void SetEndBone(const FName& EndBoneName, int32 SolverIndex) const;
 	/** Get read-only access to the array of solvers. */
 	const TArray<UIKRigSolver*>& GetSolverArray() const;
 	/** Get unique label for a given solver. Returns dash separated index and name like so, "0 - SolverName". */
diff --git a/Engine/Plugins/Animation/IKRig/Source/IKRigEditor/Public/RigEditor/IKRigSkeletonCommands.h b/Engine/Plugins/Animation/IKRig/Source/IKRigEditor/Public/RigEditor/IKRigSkeletonCommands.h
--- a/Engine/Plugins/Animation/IKRig/Source/IKRigEditor/Public/RigEditor/IKRigSkeletonCommands.h
+++ b/Engine/Plugins/Animation/IKRig/Source/IKRigEditor/Public/RigEditor/IKRigSkeletonCommands.h
@@ -29,6 +29,9 @@ public:
 
 	/** set root bone on selected solvers */
 	TSharedPtr< FUICommandInfo > SetRootBoneOnSolvers;
+
+	/** set end bone on selected solvers (Pole Solvers only)*/
+	TSharedPtr< FUICommandInfo > SetEndBoneOnSolvers;
 	
 	/** add bone setting to bone */
 	TSharedPtr< FUICommandInfo > AddBoneSettings;
diff --git a/Engine/Plugins/Animation/IKRig/Source/IKRigEditor/Public/RigEditor/SIKRigSkeleton.h b/Engine/Plugins/Animation/IKRig/Source/IKRigEditor/Public/RigEditor/SIKRigSkeleton.h
--- a/Engine/Plugins/Animation/IKRig/Source/IKRigEditor/Public/RigEditor/SIKRigSkeleton.h
+++ b/Engine/Plugins/Animation/IKRig/Source/IKRigEditor/Public/RigEditor/SIKRigSkeleton.h
@@ -171,6 +171,12 @@ private:
 	bool CanSetRootBoneOnSolvers();
 	/** END setting root bone */
 
+	/** setting end bone */
+	void HandleSetEndBoneOnSolvers();
+	bool CanSetEndBoneOnSolvers() const;
+	bool HasEndBoneCompatibleSolverSelected() const;
+	/** END setting end bone */
+
 	/** per-bone settings */
 	void HandleAddBoneSettings();
 	bool CanAddBoneSettings();
@@ -193,7 +199,7 @@ private:
 	/** END retarget chains */
 
 	/** selection state queries */
-	void GetSelectedBones(TArray<TSharedPtr<FIKRigTreeElement>>& OutBoneItems);
+	void GetSelectedBones(TArray<TSharedPtr<FIKRigTreeElement>>& OutBoneItems) const;
 	void GetSelectedGoals(TArray<TSharedPtr<FIKRigTreeElement>>& OutSelectedGoals) const;
 	void SetSelectedGoalsFromViewport(const TArray<FName>& GoalNames);
 	void GetSelectedBoneChains(TArray<FIKRigSkeletonChain>& OutChains);
