Fixed max clamping

diff --git a/Engine/Plugins/Animation/IKRig/Source/IKRig/Public/Solvers/IKRig_SetTransform.h b/Engine/Plugins/Animation/IKRig/Source/IKRig/Public/Solvers/IKRig_SetTransform.h
--- a/Engine/Plugins/Animation/IKRig/Source/IKRig/Public/Solvers/IKRig_SetTransform.h
+++ b/Engine/Plugins/Animation/IKRig/Source/IKRig/Public/Solvers/IKRig_SetTransform.h
@@ -23,7 +23,7 @@ public:
 	bool bEnableRotation = true;
 
 	/** Blend the effector on/off. Range is 0-1. Default is 1.0. */
-	UPROPERTY(EditAnywhere, Category = "Set Transform Effector", meta = (ClampMin = "0", ClampMa = "1", UIMin = "0.0", UIMax = "1.0"))
+	UPROPERTY(EditAnywhere, Category = "Set Transform Effector", meta = (ClampMin = "0.0", ClampMax = "1.0", UIMin = "0.0", UIMax = "1.0"))
 	float Alpha = 1.0f;
 };
 
