Added alpha blend parameter for pole effector

+ fixed IsBoneAffectedBySolver

diff --git a/Engine/Plugins/Animation/IKRig/Source/IKRig/Private/Solvers/IKRig_PoleSolver.cpp b/Engine/Plugins/Animation/IKRig/Source/IKRig/Private/Solvers/IKRig_PoleSolver.cpp
--- a/Engine/Plugins/Animation/IKRig/Source/IKRig/Private/Solvers/IKRig_PoleSolver.cpp
+++ b/Engine/Plugins/Animation/IKRig/Source/IKRig/Private/Solvers/IKRig_PoleSolver.cpp
@@ -73,6 +73,12 @@ void UIKRig_PoleSolver::Solve(FIKRigSkeleton& IKRigSkeleton, const FIKRigGoalCon
 		return;
 	}
 
+	const bool bHasAlpha = Effector->Alpha > KINDA_SMALL_NUMBER;
+	if (!bHasAlpha)
+	{
+		return;
+	}
+
 	const int32 RootIndex = Chain[0];
 	const int32 KneeIndex = Chain[1];
 	const int32 EndIndex = Chain.Last();
@@ -115,10 +121,16 @@ void UIKRig_PoleSolver::Solve(FIKRigSkeleton& IKRigSkeleton, const FIKRigGoalCon
 		{
 			int32 BoneIndex = Chain[Index];
 			FTransform& BoneTransform = InOutTransforms[BoneIndex];
-			BoneTransform.SetRotation(DeltaRotation * BoneTransform.GetRotation());
 
-			const FVector BoneLocation = BoneTransform.GetLocation();
-			BoneTransform.SetLocation(RootLocation + DeltaRotation.RotateVector(BoneLocation - RootLocation));
+			// rotation
+			const FQuat BoneRotation = BoneTransform.GetRotation();
+			const FQuat TargetRotation = FMath::Lerp(BoneRotation, DeltaRotation * BoneRotation, Effector->Alpha);
+			BoneTransform.SetRotation(TargetRotation);
+
+			// translation
+			const FVector BoneTranslation = BoneTransform.GetLocation();
+			const FVector TargetTranslation = FMath::Lerp( BoneTranslation, RootLocation + DeltaRotation.RotateVector(BoneTranslation - RootLocation), Effector->Alpha);
+			BoneTransform.SetTranslation(TargetTranslation);
 		}
 
 		// propagate to children
@@ -133,8 +145,9 @@ void UIKRig_PoleSolver::Solve(FIKRigSkeleton& IKRigSkeleton, const FIKRigGoalCon
 
 void UIKRig_PoleSolver::UpdateSolverSettings(UIKRigSolver* InSettings)
 {
-	if(UIKRig_PoleSolver* Settings = Cast<UIKRig_PoleSolver>(InSettings))
+	if (UIKRig_PoleSolver* Settings = Cast<UIKRig_PoleSolver>(InSettings))
 	{
+		Effector->Alpha = Settings->Effector->Alpha;
 	}
 }
 
@@ -209,11 +222,7 @@ UObject* UIKRig_PoleSolver::GetGoalSettings(const FName& GoalName) const
 
 bool UIKRig_PoleSolver::IsBoneAffectedBySolver(const FName& BoneName, const FIKRigSkeleton& IKRigSkeleton) const
 {
-	if (EndName != NAME_None)
-	{
-		return IKRigSkeleton.IsBoneInDirectLineage(BoneName, RootName);
-	}
-	return false;
+	return IKRigSkeleton.IsBoneInDirectLineage(BoneName, RootName);
 }
 
 void UIKRig_PoleSolver::SetRootBone(const FName& RootBoneName)
diff --git a/Engine/Plugins/Animation/IKRig/Source/IKRig/Public/Solvers/IKRig_PoleSolver.h b/Engine/Plugins/Animation/IKRig/Source/IKRig/Public/Solvers/IKRig_PoleSolver.h
--- a/Engine/Plugins/Animation/IKRig/Source/IKRig/Public/Solvers/IKRig_PoleSolver.h
+++ b/Engine/Plugins/Animation/IKRig/Source/IKRig/Public/Solvers/IKRig_PoleSolver.h
@@ -19,6 +19,10 @@ public:
 
 	UPROPERTY(VisibleAnywhere, Category = "Pole Solver Effector")
 	FName BoneName = NAME_None;
+
+	/** Blend the effector on/off. Range is 0-1. Default is 1.0. */
+	UPROPERTY(EditAnywhere, Category = "Pole Solver Effector", meta = (ClampMin = "0.0", ClampMax = "1.0", UIMin = "0.0", UIMax = "1.0"))
+	float Alpha = 1.0f;
 };
 
 UCLASS(EditInlineNew)
