Fixing unity build issue in IK Rig plugin

diff --git a/Engine/Plugins/Animation/IKRig/Source/IKRigEditor/Public/RigEditor/IKRigEditorController.h b/Engine/Plugins/Animation/IKRig/Source/IKRigEditor/Public/RigEditor/IKRigEditorController.h
--- a/Engine/Plugins/Animation/IKRig/Source/IKRigEditor/Public/RigEditor/IKRigEditorController.h
+++ b/Engine/Plugins/Animation/IKRig/Source/IKRigEditor/Public/RigEditor/IKRigEditorController.h
@@ -5,6 +5,7 @@
 #include "IKRigSolver.h"
 #include "SIKRigRetargetChainList.h"
 #include "Templates/SubclassOf.h"
+#include "Animation/AnimationAsset.h"
 
 #include "IKRigEditorController.generated.h"
 
