Add smooth transition for Limb IK

diff --git a/Engine/Plugins/Animation/IKRig/Source/IKRig/Private/Solvers/LimbSolver.cpp b/Engine/Plugins/Animation/IKRig/Source/IKRig/Private/Solvers/LimbSolver.cpp
--- a/Engine/Plugins/Animation/IKRig/Source/IKRig/Private/Solvers/LimbSolver.cpp
+++ b/Engine/Plugins/Animation/IKRig/Source/IKRig/Private/Solvers/LimbSolver.cpp
@@ -384,7 +384,7 @@ namespace AnimationCore
 		}
 	}
 	
-	void FABRIK_ForwardReach(
+	void  FABRIK_ForwardReach(
 		TArray<FLimbLink>& InOutLinks,
 		const FVector& InTargetLocation,
 		float ReachStepAlpha,
@@ -536,9 +536,9 @@ namespace AnimationCore
 				const float PreviousSlop = Slop;
 
 				// Pull averaging only has a visual impact when we have more than 2 bones (3 links).
-				const bool bNeedPullAveraging =(NumLinks > 3) && bAveragePull && (Slop > 1.f); 
+				const bool bNeedPullAveraging =(NumLinks > 3) && bAveragePull && (Slop > 1.f);
 				if (bNeedPullAveraging)
-				{
+				{	
 					TArray<FLimbLink> ForwardPull = InOutLinks;
 					FABRIK_ForwardReach(ForwardPull, InTargetLocation, ReachStepAlpha, bUseAngleLimit, MinRotationAngleRadians);
 
@@ -573,11 +573,10 @@ namespace AnimationCore
 				FABRIK_BackwardReach(InOutLinks, RootTargetLocation, ReachStepAlpha, bUseAngleLimit, MinRotationAngleRadians);
 			}
 
-			// If we reached, set target precisely
-			if (Slop <= ReachPrecision)
-			{
-				InOutLinks.Last().Location = InTargetLocation;
-			}
+			// Replace end bone based on how close we got to the end target.
+			FLimbLink& EndLink = InOutLinks[NumLinks-1];
+			const FLimbLink& ParentLink = InOutLinks[NumLinks-2];
+			EndLink.Location = ParentLink.Location + (EndLink.Location - ParentLink.Location).GetUnsafeNormal() * EndLink.Length;
 		}
 	}
 	
@@ -604,21 +603,23 @@ namespace AnimationCore
 		}
 		
 		const bool bDirectSolve = ( InOutLinks.Num() == 3 );
-		const FVector& RootLocation = InOutLinks[0].Location;
-
-		const float MinRotationAngleRadians = FMath::DegreesToRadians(FMath::Clamp(MinRotationAngle, 0.f, 90.f));
-		
-		// If we can't reach, we just go in a straight line towards the target,
-		if (FVector::DistSquared(RootLocation, GoalLocation) >= FMath::Square(LimbLength))
-		{
-			SolveStraightLimb(InOutLinks, GoalLocation);
-		}
-		else if (bDirectSolve)
+		if (bDirectSolve)
 		{
-			SolveTwoBoneIK(InOutLinks, GoalLocation, HingeRotationAxis);
+			const FVector& RootLocation = InOutLinks[0].Location;
+			
+			// If we can't reach, we just go in a straight line towards the target,
+			if (FVector::DistSquared(RootLocation, GoalLocation) >= FMath::Square(LimbLength))
+			{
+				SolveStraightLimb(InOutLinks, GoalLocation);
+			}
+			else
+			{
+				SolveTwoBoneIK(InOutLinks, GoalLocation, HingeRotationAxis);
+			}
 		}
 		else
 		{
+			const float MinRotationAngleRadians = FMath::DegreesToRadians(FMath::Clamp(MinRotationAngle, 0.f, 90.f));
 			SolveFABRIK(InOutLinks, GoalLocation, Precision, InMaxIterations, bUseAngleLimit, MinRotationAngleRadians, ReachStepAlpha, PullDistribution, bAveragePull);
 		}
 		
