Modify transform using references

diff --git a/Engine/Plugins/Animation/IKRig/Source/IKRig/Private/Solvers/IKRig_PBIKSolver.cpp b/Engine/Plugins/Animation/IKRig/Source/IKRig/Private/Solvers/IKRig_PBIKSolver.cpp
--- a/Engine/Plugins/Animation/IKRig/Source/IKRig/Private/Solvers/IKRig_PBIKSolver.cpp
+++ b/Engine/Plugins/Animation/IKRig/Source/IKRig/Private/Solvers/IKRig_PBIKSolver.cpp
@@ -66,11 +66,12 @@ void UIKRigPBIKSolver::Solve(FIKRigSkeleton& IKRigSkeleton, const FIKRigGoalCont
 		return;
 	}
 
+	TArray<FTransform>& InOutTransforms = IKRigSkeleton.CurrentPoseGlobal;
+	
 	// set bones to input pose
 	for(int32 BoneIndex = 0; BoneIndex < Solver.GetNumBones(); BoneIndex++)
 	{
-		const FTransform GlobalTransform = IKRigSkeleton.CurrentPoseGlobal[BoneIndex];
-		Solver.SetBoneTransform(BoneIndex, GlobalTransform);
+		Solver.SetBoneTransform(BoneIndex, InOutTransforms[BoneIndex]);
 	}
 
 	// update bone settings
@@ -124,9 +125,7 @@ void UIKRigPBIKSolver::Solve(FIKRigSkeleton& IKRigSkeleton, const FIKRigGoalCont
 	// copy transforms back
 	for(int32 BoneIndex = 0; BoneIndex < Solver.GetNumBones(); BoneIndex++)
 	{
-		FTransform OutTransform;
-		Solver.GetBoneGlobalTransform(BoneIndex, OutTransform);
-		IKRigSkeleton.CurrentPoseGlobal[BoneIndex] = OutTransform;
+		Solver.GetBoneGlobalTransform(BoneIndex, InOutTransforms[BoneIndex]);
 	}
 }
 
diff --git a/Engine/Plugins/Experimental/FullBodyIK/Source/PBIK/Private/Core/PBIKSolver.cpp b/Engine/Plugins/Experimental/FullBodyIK/Source/PBIK/Private/Core/PBIKSolver.cpp
--- a/Engine/Plugins/Experimental/FullBodyIK/Source/PBIK/Private/Core/PBIKSolver.cpp
+++ b/Engine/Plugins/Experimental/FullBodyIK/Source/PBIK/Private/Core/PBIKSolver.cpp
@@ -758,7 +758,7 @@ PBIK::FBoneSettings* FPBIKSolver::GetBoneSettings(const int32 Index)
 void FPBIKSolver::GetBoneGlobalTransform(const int32 Index, FTransform& OutTransform)
 {
 	check(Index >= 0 && Index < Bones.Num());
-	PBIK::FBone& Bone = Bones[Index];
+	const PBIK::FBone& Bone = Bones[Index];
 	OutTransform.SetLocation(Bone.Position);
 	OutTransform.SetRotation(Bone.Rotation);
 }
