Solver settings naming consistency

diff --git a/Engine/Plugins/Animation/IKRig/Source/IKRig/Public/Solvers/IKRig_LimbSolver.h b/Engine/Plugins/Animation/IKRig/Source/IKRig/Public/Solvers/IKRig_LimbSolver.h
--- a/Engine/Plugins/Animation/IKRig/Source/IKRig/Public/Solvers/IKRig_LimbSolver.h
+++ b/Engine/Plugins/Animation/IKRig/Source/IKRig/Public/Solvers/IKRig_LimbSolver.h
@@ -32,54 +32,54 @@ public:
 
 	UIKRig_LimbSolver();
 	
-	UPROPERTY(VisibleAnywhere, Category = "Limb Solver Settings")
+	UPROPERTY(VisibleAnywhere, Category = "Limb IK Settings")
 	FName RootName = NAME_None;
 
 	/** Precision (distance to the target) */
-	UPROPERTY(EditAnywhere, Category = "Limb Solver Settings", meta = (ClampMin = "0.0"))
+	UPROPERTY(EditAnywhere, Category = "Limb IK Settings", meta = (ClampMin = "0.0"))
 	float ReachPrecision = 0.01f;
 	
 	// TWO BONES SETTINGS
 	
 	/** Hinge Bones Rotation Axis. This is essentially the plane normal for (hip - knee - foot). */
-	UPROPERTY(EditAnywhere, Category = "Limb Solver Settings|Two Bones")
+	UPROPERTY(EditAnywhere, Category = "Limb IK Settings|Two Bones")
 	TEnumAsByte<EAxis::Type> HingeRotationAxis = EAxis::None;
 
 	// FABRIK SETTINGS
 
 	/** Number of Max Iterations to reach the target */
-	UPROPERTY(EditAnywhere, Category = "Limb Solver Settings|FABRIK", meta = (ClampMin = "0", UIMin = "0", UIMax = "100"))
+	UPROPERTY(EditAnywhere, Category = "Limb IK Settings|FABRIK", meta = (ClampMin = "0", UIMin = "0", UIMax = "100"))
 	int32 MaxIterations = 12;
 
 	/** Enable/Disable rotational limits */
-	UPROPERTY(EditAnywhere, Category = "Limb Solver Settings|FABRIK|Limits")
+	UPROPERTY(EditAnywhere, Category = "Limb IK Settings|FABRIK|Limits")
 	bool bEnableLimit = false;
 
 	/** Only used if bEnableRotationLimit is enabled. Prevents the leg from folding onto itself,
 	* and forces at least this angle between Parent and Child bone. */
-	UPROPERTY(EditAnywhere, Category = "Limb Solver Settings|FABRIK|Limits", meta = (ClampMin = "0.0", ClampMax = "90.0", UIMin = "0.0", UIMax = "90.0"))
+	UPROPERTY(EditAnywhere, Category = "Limb IK Settings|FABRIK|Limits", meta = (ClampMin = "0.0", ClampMax = "90.0", UIMin = "0.0", UIMax = "90.0"))
 	float MinRotationAngle = 15.f;
 
 	/** Pull averaging only has a visual impact when we have more than 2 bones (3 links). */
-	UPROPERTY(EditAnywhere, Category = "Limb Solver Settings|FABRIK|Pull Averaging")
+	UPROPERTY(EditAnywhere, Category = "Limb IK Settings|FABRIK|Pull Averaging")
 	bool bAveragePull = true;
 
 	/** Re-position limb to distribute pull: 0 = foot, 0.5 = balanced, 1.f = hip */
-	UPROPERTY(EditAnywhere, Category = "Limb Solver Settings|FABRIK|Pull Averaging", meta = (ClampMin = "0.0", ClampMax = "1.0",UIMin = "0.0", UIMax = "1.0"))
+	UPROPERTY(EditAnywhere, Category = "Limb IK Settings|FABRIK|Pull Averaging", meta = (ClampMin = "0.0", ClampMax = "1.0",UIMin = "0.0", UIMax = "1.0"))
 	float PullDistribution = 0.5f;
 	
 	/** Move end effector towards target. If we are compressing the chain, limit displacement.*/
-	UPROPERTY(EditAnywhere, Category = "Limb Solver Settings|FABRIK|Pull Averaging", meta = (ClampMin = "0.0", ClampMax = "1.0",UIMin = "0.0", UIMax = "1.0"))
+	UPROPERTY(EditAnywhere, Category = "Limb IK Settings|FABRIK|Pull Averaging", meta = (ClampMin = "0.0", ClampMax = "1.0",UIMin = "0.0", UIMax = "1.0"))
 	float ReachStepAlpha = 0.7f;
 
 	// TWIST SETTINGS
 	
 	/** Enable Knee Twist correction, by comparing Foot FK with Foot IK orientation. */
-	UPROPERTY(EditAnywhere, Category = "Limb Solver Settings|Twist")
+	UPROPERTY(EditAnywhere, Category = "Limb IK Settings|Twist")
 	bool bEnableTwistCorrection = false;
 	
 	/** Forward Axis for Foot bone. */
-	UPROPERTY(EditAnywhere, Category = "Limb Solver Settings|Twist")
+	UPROPERTY(EditAnywhere, Category = "Limb IK Settings|Twist")
 	TEnumAsByte<EAxis::Type> EndBoneForwardAxis = EAxis::Y;
 	
 	/** UIKRigSolver interface */
