IK Rig: Limb IK Missing Warning When Bone Chain < 3 Bones

diff --git a/Engine/Plugins/Animation/IKRig/Source/IKRig/Private/Solvers/IKRig_LimbSolver.cpp b/Engine/Plugins/Animation/IKRig/Source/IKRig/Private/Solvers/IKRig_LimbSolver.cpp
--- a/Engine/Plugins/Animation/IKRig/Source/IKRig/Private/Solvers/IKRig_LimbSolver.cpp
+++ b/Engine/Plugins/Animation/IKRig/Source/IKRig/Private/Solvers/IKRig_LimbSolver.cpp
@@ -165,17 +165,19 @@ bool UIKRig_LimbSolver::GetWarningMessage(FText& OutWarningMessage) const
 		OutWarningMessage = LOCTEXT("MissingRoot", "Missing root.");
 		return true;
 	}
+
+	if (Solver.NumLinks() < 3)
+	{
+		OutWarningMessage = LOCTEXT("Requires3BonesChain", "Requires at least 3 bones between root and goal.");
+		return true;
+	}
 	
 	return false;
 }
 
 bool UIKRig_LimbSolver::IsBoneAffectedBySolver(const FName& BoneName, const FIKRigSkeleton& IKRigSkeleton) const
 {
-	if (Effector->BoneName != NAME_None)
-	{
-		return IKRigSkeleton.IsBoneInDirectLineage(Effector->BoneName, RootName);
-	}
-	return false;
+	return IKRigSkeleton.IsBoneInDirectLineage(BoneName, RootName);
 }
 
 void UIKRig_LimbSolver::AddGoal(const UIKRigEffectorGoal& NewGoal)
diff --git a/Engine/Plugins/Animation/IKRig/Source/IKRig/Private/Solvers/IKRig_PoleSolver.cpp b/Engine/Plugins/Animation/IKRig/Source/IKRig/Private/Solvers/IKRig_PoleSolver.cpp
--- a/Engine/Plugins/Animation/IKRig/Source/IKRig/Private/Solvers/IKRig_PoleSolver.cpp
+++ b/Engine/Plugins/Animation/IKRig/Source/IKRig/Private/Solvers/IKRig_PoleSolver.cpp
@@ -175,6 +175,12 @@ bool UIKRig_PoleSolver::GetWarningMessage(FText& OutWarningMessage) const
 		OutWarningMessage = LOCTEXT("MissingRoot", "Missing root bone.");
 		return true;
 	}
+	
+	if (Chain.Num() < 3)
+	{
+		OutWarningMessage = LOCTEXT("Requires3BonesChain", "Requires at least 3 bones between root and end bones.");
+		return true;
+	}
 
 	return false;
 }
diff --git a/Engine/Plugins/Animation/IKRig/Source/IKRigEditor/Private/RigEditor/IKRigEditorController.cpp b/Engine/Plugins/Animation/IKRig/Source/IKRigEditor/Private/RigEditor/IKRigEditorController.cpp
--- a/Engine/Plugins/Animation/IKRig/Source/IKRigEditor/Private/RigEditor/IKRigEditorController.cpp
+++ b/Engine/Plugins/Animation/IKRig/Source/IKRigEditor/Private/RigEditor/IKRigEditorController.cpp
@@ -23,6 +23,9 @@ void FIKRigEditorController::Initialize(TSharedPtr<FIKRigEditorToolkit> Toolkit,
 	if (!AssetController->OnIKRigNeedsInitialized().IsBoundToObject(this))
 	{
 		AssetController->OnIKRigNeedsInitialized().AddSP(this, &FIKRigEditorController::OnIKRigNeedsInitialized);
+
+		// Initialize editor's instances at first initialization
+		InitializeSolvers();
 	}
 }
 
@@ -34,6 +37,9 @@ void FIKRigEditorController::OnIKRigNeedsInitialized(UIKRigDefinition* ModifiedI
 	}
 
 	AnimInstance->SetProcessorNeedsInitialized();
+
+	// Initialize editor's instances on request
+	InitializeSolvers();
 }
 
 void FIKRigEditorController::Reset() const
@@ -344,4 +350,17 @@ void FIKRigEditorController::PlayAnimationAsset(UAnimationAsset* AssetToPlay)
 	}
 }
 
+void FIKRigEditorController::InitializeSolvers() const
+{
+	if (AssetController)
+	{
+		const FIKRigSkeleton& IKRigSkeleton = AssetController->GetIKRigSkeleton();
+		const TArray<UIKRigSolver*>& Solvers = AssetController->GetSolverArray(); 
+		for (UIKRigSolver* Solver: Solvers)
+		{
+			Solver->Initialize(IKRigSkeleton);
+		}
+	}
+}
+
 #undef LOCTEXT_NAMESPACE
diff --git a/Engine/Plugins/Animation/IKRig/Source/IKRigEditor/Public/RigEditor/IKRigEditorController.h b/Engine/Plugins/Animation/IKRig/Source/IKRigEditor/Public/RigEditor/IKRigEditorController.h
--- a/Engine/Plugins/Animation/IKRig/Source/IKRigEditor/Public/RigEditor/IKRigEditorController.h
+++ b/Engine/Plugins/Animation/IKRig/Source/IKRigEditor/Public/RigEditor/IKRigEditorController.h
@@ -112,6 +112,9 @@ public:
 
 private:
 
+	/** Initializes editor's solvers instances */
+	void InitializeSolvers() const;
+	
 	/** UI and viewport selection state */
 	TArray<FName> SelectedGoals;	
 };
