Optimize transform propagation through FIKRigSkeleton

diff --git a/Engine/Plugins/Animation/IKRig/Source/IKRig/Private/IKRigSkeleton.cpp b/Engine/Plugins/Animation/IKRig/Source/IKRig/Private/IKRigSkeleton.cpp
--- a/Engine/Plugins/Animation/IKRig/Source/IKRig/Private/IKRigSkeleton.cpp
+++ b/Engine/Plugins/Animation/IKRig/Source/IKRig/Private/IKRigSkeleton.cpp
@@ -2,6 +2,9 @@
 
 #include "IKRigSkeleton.h"
 
+// This is the default end of branch index value, meaning we haven't cached it yet
+#define FIKRIGSKELETON_INVALID_EO_BRANCH_INDEX -2
+
 void FIKRigSkeleton::SetInputSkeleton(const FReferenceSkeleton& RefSkeleton, const TArray<FName>& InExcludedBones)
 {
 	const FIKRigInputSkeleton InputSkeleton = FIKRigInputSkeleton(RefSkeleton);
@@ -40,6 +43,9 @@ void FIKRigSkeleton::SetInputSkeleton(const FIKRigInputSkeleton& InputSkeleton,
 
 	// initialize CURRENT LOCAL pose from global pose
 	UpdateAllLocalTransformFromGlobal();
+
+	// initialize branch caching
+	CachedEndOfBranchIndices.Init(FIKRIGSKELETON_INVALID_EO_BRANCH_INDEX, ParentIndices.Num());
 }
 
 void FIKRigSkeleton::Reset()
@@ -50,6 +56,7 @@ void FIKRigSkeleton::Reset()
 	CurrentPoseGlobal.Reset();
 	CurrentPoseLocal.Reset();
 	RefPoseGlobal.Reset();
+	CachedEndOfBranchIndices.Reset();
 }
 
 int32 FIKRigSkeleton::GetBoneIndexFromName(const FName InName) const
@@ -65,14 +72,15 @@ int32 FIKRigSkeleton::GetBoneIndexFromName(const FName InName) const
 	return INDEX_NONE;
 }
 
-FName FIKRigSkeleton::GetBoneNameFromIndex(const int32 BoneIndex) const
+const FName& FIKRigSkeleton::GetBoneNameFromIndex(const int32 BoneIndex) const
 {
 	if (BoneNames.IsValidIndex(BoneIndex))
 	{
 		return BoneNames[BoneIndex];
 	}
 
-	return NAME_None;
+	static const FName InvalidName(NAME_None);  
+	return InvalidName;
 }
 
 int32 FIKRigSkeleton::GetParentIndex(const int32 BoneIndex) const
@@ -105,18 +113,64 @@ int32 FIKRigSkeleton::GetChildIndices(const int32 ParentBoneIndex, TArray<int32>
 {
 	Children.Reset();
 
-	const int32 NumBones = BoneNames.Num();
-	for (int32 BoneIndex = ParentBoneIndex + 1; BoneIndex < NumBones; BoneIndex++)
+	const int32 LastBranchIndex = GetCachedEndOfBranchIndex(ParentBoneIndex);
+	if (LastBranchIndex != INDEX_NONE)
 	{
-		if (ParentBoneIndex == GetParentIndex(BoneIndex))
+		for (int32 BoneIndex = ParentBoneIndex + 1; BoneIndex <= LastBranchIndex; BoneIndex++)
 		{
-			Children.Add(BoneIndex);
+			if (GetParentIndex(BoneIndex) == ParentBoneIndex)
+			{
+				Children.Add(BoneIndex);
+			}
 		}
 	}
 
 	return Children.Num();
 }
 
+int32 FIKRigSkeleton::GetCachedEndOfBranchIndex(const int32 InBoneIndex) const
+{
+	if (!CachedEndOfBranchIndices.IsValidIndex(InBoneIndex))
+	{
+		return INDEX_NONE;
+	}
+
+	// already cached
+	if (CachedEndOfBranchIndices[InBoneIndex] != FIKRIGSKELETON_INVALID_EO_BRANCH_INDEX)
+	{
+		return CachedEndOfBranchIndices[InBoneIndex];
+	}
+
+	const int32 NumBones = BoneNames.Num();
+	
+	// if we're asking for root's branch, get the last bone  
+	if (InBoneIndex == 0)
+	{
+		CachedEndOfBranchIndices[InBoneIndex] = NumBones-1;
+		return CachedEndOfBranchIndices[InBoneIndex];
+	}
+
+	CachedEndOfBranchIndices[InBoneIndex] = INDEX_NONE;
+	
+	// store ref parent
+	const int32 RefParentIndex = GetParentIndex(InBoneIndex);
+
+	int32 BoneIndex = InBoneIndex + 1;
+	int32 ParentIndex = GetParentIndex(BoneIndex);
+
+	// if next bellow bone's parent is less than or equal to RefParentIndex, we are leaving the branch
+	// so ne need to go further
+	while (ParentIndex > RefParentIndex && BoneIndex < NumBones)
+	{
+		CachedEndOfBranchIndices[InBoneIndex] = BoneIndex;
+				
+		BoneIndex++;
+		ParentIndex = GetParentIndex(BoneIndex);
+	}
+
+	return CachedEndOfBranchIndices[InBoneIndex];
+}
+
 void FIKRigSkeleton::ConvertLocalPoseToGlobal(
 	const TArray<int32>& InParentIndices,
 	const TArray<FTransform>& InLocalPose,
@@ -190,10 +244,14 @@ void FIKRigSkeleton::UpdateLocalTransformFromGlobal(const int32 BoneIndex)
 
 void FIKRigSkeleton::PropagateGlobalPoseBelowBone(const int32 StartBoneIndex)
 {
-	for (int32 BoneIndex=StartBoneIndex+1; BoneIndex<CurrentPoseGlobal.Num(); ++BoneIndex)
+	const int32 LastBranchIndex = GetCachedEndOfBranchIndex(StartBoneIndex);
+	if (LastBranchIndex != INDEX_NONE)
 	{
-		UpdateGlobalTransformFromLocal(BoneIndex);
-		UpdateLocalTransformFromGlobal(BoneIndex);
+		for (int32 BoneIndex=StartBoneIndex+1; BoneIndex<=LastBranchIndex; ++BoneIndex)
+		{
+			UpdateGlobalTransformFromLocal(BoneIndex);
+			UpdateLocalTransformFromGlobal(BoneIndex);
+		}
 	}
 }
 
diff --git a/Engine/Plugins/Animation/IKRig/Source/IKRig/Public/IKRigSkeleton.h b/Engine/Plugins/Animation/IKRig/Source/IKRig/Public/IKRigSkeleton.h
--- a/Engine/Plugins/Animation/IKRig/Source/IKRig/Public/IKRigSkeleton.h
+++ b/Engine/Plugins/Animation/IKRig/Source/IKRig/Public/IKRigSkeleton.h
@@ -108,13 +108,15 @@ struct IKRIG_API FIKRigSkeleton
 	
 	int32 GetBoneIndexFromName(const FName InName) const;
 
-	FName GetBoneNameFromIndex(const int32 BoneIndex) const;
+	const FName& GetBoneNameFromIndex(const int32 BoneIndex) const;
 	
 	int32 GetParentIndex(const int32 BoneIndex) const;
 
 	int32 GetParentIndexThatIsNotExcluded(const int32 BoneIndex) const;
 
 	int32 GetChildIndices(const int32 ParentBoneIndex, TArray<int32>& Children) const;
+	
+	int32 GetCachedEndOfBranchIndex(const int32 InBoneIndex) const;
 
 	static void ConvertLocalPoseToGlobal(
 		const TArray<int32>& InParentIndices,
@@ -138,4 +140,9 @@ struct IKRIG_API FIKRigSkeleton
 	static void NormalizeRotations(TArray<FTransform>& Transforms);
 
 	void GetChainsInList(const TArray<int32>& SelectedBones, TArray<FIKRigSkeletonChain>& OutChains) const;
+
+private:
+
+	/** The branch bellow a specific bone. It stores the last element of the branch instead of the indices */
+	mutable TArray<int32> CachedEndOfBranchIndices;
 };
