Fixed bone rendering in IK Rig.
Misc IK Rig fixes

diff --git a/Engine/Plugins/Animation/IKRig/Source/IKRig/Private/AnimNodes/AnimNode_RetargetPoseFromMesh.cpp b/Engine/Plugins/Animation/IKRig/Source/IKRig/Private/AnimNodes/AnimNode_RetargetPoseFromMesh.cpp
--- a/Engine/Plugins/Animation/IKRig/Source/IKRig/Private/AnimNodes/AnimNode_RetargetPoseFromMesh.cpp
+++ b/Engine/Plugins/Animation/IKRig/Source/IKRig/Private/AnimNodes/AnimNode_RetargetPoseFromMesh.cpp
@@ -41,6 +41,13 @@ void FAnimNode_RetargetPoseFromMesh::Evaluate_AnyThread(FPoseContext& Output)
 		return;
 	}
 
+	if (SourceMeshComponentSpaceBoneTransforms.IsEmpty())
+	{
+		// it's possible in editor to have anim instances initialized before PreUpdate() is called
+		// which results in trying to run the retargeter without an source pose to copy from
+		return;
+	}
+
 #if WITH_EDITOR
 	// live preview IK Rig solver settings in the retarget, editor only
 	// NOTE: this copies goal targets as well, but these are overwritten by IK chain goals
diff --git a/Engine/Plugins/Animation/IKRig/Source/IKRig/Private/IKRigDefinition.cpp b/Engine/Plugins/Animation/IKRig/Source/IKRig/Private/IKRigDefinition.cpp
--- a/Engine/Plugins/Animation/IKRig/Source/IKRig/Private/IKRigDefinition.cpp
+++ b/Engine/Plugins/Animation/IKRig/Source/IKRig/Private/IKRigDefinition.cpp
@@ -26,14 +26,6 @@ void UIKRigDefinition::Serialize(FArchive& Ar)
 	Ar.UsingCustomVersion(FIKRigObjectVersion::GUID);
 }
 
-#if WITH_EDITOR
-void UIKRigDefinition::PostEditUndo()
-{
-	Super::PostEditUndo();
-	IKRigEditUndo.Broadcast();
-}
-#endif
-
 const FBoneChain* UIKRigDefinition::GetRetargetChainByName(FName ChainName) const
 {
 	for (const FBoneChain& Chain : RetargetDefinition.BoneChains)
diff --git a/Engine/Plugins/Animation/IKRig/Source/IKRig/Private/Retargeter/IKRetargetProcessor.cpp b/Engine/Plugins/Animation/IKRig/Source/IKRig/Private/Retargeter/IKRetargetProcessor.cpp
--- a/Engine/Plugins/Animation/IKRig/Source/IKRig/Private/Retargeter/IKRetargetProcessor.cpp
+++ b/Engine/Plugins/Animation/IKRig/Source/IKRig/Private/Retargeter/IKRetargetProcessor.cpp
@@ -174,43 +174,47 @@ void FTargetSkeleton::GenerateRetargetPose(const FIKRetargetPose* InRetargetPose
 	// create a retarget pose by copying the ref pose
 	FRetargetSkeleton::GenerateRetargetPose();
 	
+	// no retarget pose specified (will use default pose from skeletal mesh with no offsets)
+	if (InRetargetPose==nullptr  || RetargetRootBone == NAME_None)
+	{
+		return;
+	}
+
 	// apply retarget pose offsets (retarget pose is stored as offset relative to reference pose)
-	if (InRetargetPose && RetargetRootBone != NAME_None)
+	const TArray<FTransform>& RefPoseLocal = SkeletalMesh->GetRefSkeleton().GetRefBonePose();
+	
+	// apply root translation offset
+	const int32 RootBoneIndex = FindBoneIndexByName(RetargetRootBone);
+	if (RootBoneIndex != INDEX_NONE)
+	{
+		FTransform& RootTransform = RetargetGlobalPose[RootBoneIndex];
+		RootTransform.AddToTranslation(InRetargetPose->RootTranslationOffset);
+		UpdateLocalTransformOfSingleBone(RootBoneIndex, RetargetLocalPose, RetargetGlobalPose);
+	}
+
+	// apply bone rotation offsets
+	for (const TTuple<FName, FQuat>& BoneRotationOffset : InRetargetPose->BoneRotationOffsets)
 	{
-		// apply root translation offset
-		const int32 RootBoneIndex = FindBoneIndexByName(RetargetRootBone);
-		if (RootBoneIndex != INDEX_NONE)
+		const int32 BoneIndex = FindBoneIndexByName(BoneRotationOffset.Key);
+		if (BoneIndex == INDEX_NONE)
 		{
-			FTransform& RootTransform = RetargetGlobalPose[RootBoneIndex];
-			RootTransform.AddToTranslation(InRetargetPose->RootTranslationOffset);
-			UpdateLocalTransformOfSingleBone(RootBoneIndex, RetargetLocalPose, RetargetGlobalPose);
-			UpdateGlobalTransformsBelowBone(RootBoneIndex, RetargetLocalPose, RetargetGlobalPose);
+			// this can happen if a retarget pose recorded a bone offset for a bone that is not present in the
+			// target skeleton; ie, the retarget pose was generated from a different Skeletal Mesh with extra bones
+			continue;
 		}
 
-		// apply bone rotation offsets
-		for (const TTuple<FName, FQuat>& BoneRotationOffset : InRetargetPose->BoneRotationOffsets)
+		if (!IsBoneInAnyTargetChain[BoneIndex] && BoneIndex!=RootBoneIndex)
 		{
-			const int32 BoneIndex = FindBoneIndexByName(BoneRotationOffset.Key);
-			if (BoneIndex == INDEX_NONE)
-			{
-				// this can happen if a retarget pose recorded a bone offset for a bone that is not present in the
-				// target skeleton; ie, the retarget pose was generated from a different Skeletal Mesh with extra bones
-				continue;
-			}
-
-			if (!IsBoneInAnyTargetChain[BoneIndex])
-			{
-				// this can happen if a retarget pose includes bone edits from a bone chain that was subsequently removed,
-				// and the asset has not run through the "CleanChainMapping" operation yet (happens on load)
-				continue;
-			}
-
-			FQuat BoneRotation = BoneRotationOffset.Value * RetargetGlobalPose[BoneIndex].GetRotation();
-			RetargetGlobalPose[BoneIndex].SetRotation(BoneRotation);
-			UpdateLocalTransformOfSingleBone(BoneIndex, RetargetLocalPose, RetargetGlobalPose);
-			UpdateGlobalTransformsBelowBone(BoneIndex, RetargetLocalPose, RetargetGlobalPose);
+			// this can happen if a retarget pose includes bone edits from a bone chain that was subsequently removed,
+			// and the asset has not run through the "CleanChainMapping" operation yet (happens on load)
+			continue;
 		}
+
+		const FQuat LocalBoneRotation = BoneRotationOffset.Value * RefPoseLocal[BoneIndex].GetRotation();
+		RetargetLocalPose[BoneIndex].SetRotation(LocalBoneRotation);
 	}
+
+	UpdateGlobalTransformsBelowBone(RootBoneIndex, RetargetLocalPose, RetargetGlobalPose);
 }
 
 void FTargetSkeleton::Reset()
@@ -1135,6 +1139,14 @@ FTransform UIKRetargetProcessor::GetTargetBoneRetargetPoseGlobalTransform(const
 	return TargetSkeleton.RetargetGlobalPose[TargetBoneIndex];
 }
 
+FTransform UIKRetargetProcessor::GetTargetBoneRetargetPoseLocalTransform(const int32& TargetBoneIndex) const
+{
+	check(TargetSkeleton.BoneNames.IsValidIndex(TargetBoneIndex))
+
+	// get the current retarget pose
+	return TargetSkeleton.RetargetLocalPose[TargetBoneIndex];
+}
+
 #if WITH_EDITOR
 void UIKRetargetProcessor::SetNeedsInitialized()
 {
diff --git a/Engine/Plugins/Animation/IKRig/Source/IKRig/Private/Retargeter/IKRetargeter.cpp b/Engine/Plugins/Animation/IKRig/Source/IKRig/Private/Retargeter/IKRetargeter.cpp
--- a/Engine/Plugins/Animation/IKRig/Source/IKRig/Private/Retargeter/IKRetargeter.cpp
+++ b/Engine/Plugins/Animation/IKRig/Source/IKRig/Private/Retargeter/IKRetargeter.cpp
@@ -10,29 +10,31 @@ const FName UIKRetargeter::GetTargetPreviewMeshPropertyName() { return GET_MEMBE
 #endif
 const FName UIKRetargeter::GetDefaultPoseName() { return DefaultPoseName; }
 
-#if WITH_EDITOR
-void UIKRetargeter::PostEditUndo()
-{
-	Super::PostEditUndo();
-	IKRigEditUndo.Broadcast();
-};
-#endif
-
-void FIKRetargetPose::AddRotationDeltaToBone(FName BoneName, FQuat RotationDelta)
+void FIKRetargetPose::SetBoneRotationOffset(FName BoneName, FQuat RotationDelta, const FIKRigSkeleton& Skeleton)
 {
 	FQuat* RotOffset = BoneRotationOffsets.Find(BoneName);
 	if (RotOffset == nullptr)
 	{
 		// first time this bone has been modified in this pose
 		BoneRotationOffsets.Emplace(BoneName, RotationDelta);
+		SortHierarchically(Skeleton);
 		return;
 	}
 
-	// accumulate delta rotation
-	*RotOffset = RotationDelta * (*RotOffset);
+	*RotOffset = RotationDelta;
 }
 
 void FIKRetargetPose::AddTranslationDeltaToRoot(FVector TranslateDelta)
 {
 	RootTranslationOffset += TranslateDelta;
 }
+
+void FIKRetargetPose::SortHierarchically(const FIKRigSkeleton& Skeleton)
+{
+	// sort offsets hierarchically so that they are applied in leaf to root order
+	// when generating the component space retarget pose in the processor
+	BoneRotationOffsets.KeySort([Skeleton](FName A, FName B)
+	{
+		return Skeleton.GetBoneIndexFromName(A) > Skeleton.GetBoneIndexFromName(B);
+	});
+}
diff --git a/Engine/Plugins/Animation/IKRig/Source/IKRig/Private/Solvers/IKRig_PBIKSolver.cpp b/Engine/Plugins/Animation/IKRig/Source/IKRig/Private/Solvers/IKRig_PBIKSolver.cpp
--- a/Engine/Plugins/Animation/IKRig/Source/IKRig/Private/Solvers/IKRig_PBIKSolver.cpp
+++ b/Engine/Plugins/Animation/IKRig/Source/IKRig/Private/Solvers/IKRig_PBIKSolver.cpp
@@ -85,7 +85,6 @@ void UIKRigPBIKSolver::Solve(FIKRigSkeleton& IKRigSkeleton, const FIKRigGoalCont
 	}
 
 	// update effectors
-	const float OffsetAlpha = 1.0f; // this is constant because IKRig manages offset alphas itself
 	for (const UIKRig_FBIKEffector* Effector : Effectors)
 	{
 		if (Effector->IndexInSolver < 0)
@@ -98,15 +97,19 @@ void UIKRigPBIKSolver::Solve(FIKRigSkeleton& IKRigSkeleton, const FIKRigGoalCont
 		{
 			return;
 		}
+
+		PBIK::FEffectorSettings Settings;
+		Settings.PositionAlpha = 1.0f; // this is constant because IKRig manages offset alphas itself
+		Settings.RotationAlpha = 1.0f; // this is constant because IKRig manages offset alphas itself
+		Settings.StrengthAlpha = Effector->StrengthAlpha;
+		Settings.PullChainAlpha = Effector->PullChainAlpha;
+		Settings.PinRotation = Effector->PinRotation;
 		
 		Solver.SetEffectorGoal(
 			Effector->IndexInSolver,
 			Goal->FinalBlendedPosition,
 			Goal->FinalBlendedRotation,
-			OffsetAlpha,
-			Effector->StrengthAlpha,
-			Effector->PullChainAlpha,
-			Effector->PinRotation);
+			Settings);
 	}
 
 	// update settings
@@ -325,6 +328,19 @@ void UIKRigPBIKSolver::DrawBoneSettings(
 
 bool UIKRigPBIKSolver::IsBoneAffectedBySolver(const FName& BoneName, const FIKRigSkeleton& IKRigSkeleton) const
 {
+	// nothing is affected by solver without a root bone assigned or at least 1 effector
+	if (RootBone == NAME_None || Effectors.IsEmpty())
+	{
+		return false;
+	}
+
+	// has to be BELOW root
+	if (!IKRigSkeleton.IsBoneInDirectLineage(BoneName, RootBone))
+	{
+		return false;
+	}
+
+	// has to be ABOVE an effector
 	for (UIKRig_FBIKEffector* Effector : Effectors)
 	{
 		if (IKRigSkeleton.IsBoneInDirectLineage(Effector->BoneName, BoneName))
diff --git a/Engine/Plugins/Animation/IKRig/Source/IKRig/Public/AnimNodes/AnimNode_IKRig.h b/Engine/Plugins/Animation/IKRig/Source/IKRig/Public/AnimNodes/AnimNode_IKRig.h
--- a/Engine/Plugins/Animation/IKRig/Source/IKRig/Public/AnimNodes/AnimNode_IKRig.h
+++ b/Engine/Plugins/Animation/IKRig/Source/IKRig/Public/AnimNodes/AnimNode_IKRig.h
@@ -85,4 +85,5 @@ private:
 	
 	friend class UAnimGraphNode_IKRig;
 	friend struct FIKRigAnimInstanceProxy;
+	friend class UIKRigAnimInstance;
 };
diff --git a/Engine/Plugins/Animation/IKRig/Source/IKRig/Public/IKRigDefinition.h b/Engine/Plugins/Animation/IKRig/Source/IKRig/Public/IKRigDefinition.h
--- a/Engine/Plugins/Animation/IKRig/Source/IKRig/Public/IKRigDefinition.h
+++ b/Engine/Plugins/Animation/IKRig/Source/IKRig/Public/IKRigDefinition.h
@@ -35,11 +35,11 @@ public:
 	UPROPERTY(VisibleAnywhere, Category = "Goal Settings")
 	FName BoneName;
 
-	/** Range 0-1, default is 1. Blend between the Goal position between the input bone pose (0.0) and the current goal transform (1.0).*/
+	/** Range 0-1, default is 1. Blend between the input bone position (0.0) and the current goal position (1.0).*/
 	UPROPERTY(EditAnywhere, Category = "Goal Settings", meta = (ClampMin = "0.0", ClampMax = "1.0", UIMin = "0.0", UIMax = "1.0"))
 	float PositionAlpha = 1.0f;
 
-	/** Range 0-1, default is 1. Blend between the Goal rotation between the input bone pose (0.0) and the current goal transform (1.0).*/
+	/** Range 0-1, default is 1. Blend between the input bone rotation (0.0) and the current goal rotation (1.0).*/
 	UPROPERTY(EditAnywhere, Category = "Goal Settings", meta = (ClampMin = "0.0", ClampMax = "1.0", UIMin = "0.0", UIMax = "1.0"))
 	float RotationAlpha = 1.0f;
 
@@ -133,13 +133,25 @@ class IKRIG_API UIKRigDefinition : public UObject, public IInterface_PreviewMesh
 public:
 
 #if WITH_EDITORONLY_DATA
+
+	/**Draw bones in the viewport.*/
+	UPROPERTY(EditAnywhere, Category = "Viewport Bone Settings")
+	bool DrawBones = true;
+	
+	/**The size of the Bones in the editor viewport.*/
+	UPROPERTY(EditAnywhere, Category = "Viewport Bone Settings", meta = (ClampMin = "0.01", UIMin = "0.1", UIMax = "100.0"))
+	float BoneSize = 2.0f;
+
+	/**Draw bones in the viewport.*/
+	UPROPERTY(EditAnywhere, Category = "Viewport Goal Settings")
+	bool DrawGoals = true;
 	
 	/**The size of the Goals in the editor viewport.*/
-	UPROPERTY(EditAnywhere, Category = "Goal Drawing", meta = (ClampMin = "0.01", UIMin = "0.1", UIMax = "100.0"))
+	UPROPERTY(EditAnywhere, Category = "Viewport Goal Settings", meta = (ClampMin = "0.01", UIMin = "0.1", UIMax = "100.0"))
 	float GoalSize = 7.0f;
 
 	/** The thickness of the Goals in the editor viewport.*/
-	UPROPERTY(EditAnywhere, Category = "Goal Drawing",  meta = (ClampMin = "0.01", UIMin = "0.0", UIMax = "10.0"))
+	UPROPERTY(EditAnywhere, Category = "Viewport Goal Settings",  meta = (ClampMin = "0.01", UIMin = "0.0", UIMax = "10.0"))
 	float GoalThickness = 0.7f;
 
 	/** The controller responsible for managing this asset's data (all editor mutation goes through this) */
@@ -150,14 +162,6 @@ public:
 
 	virtual void Serialize(FArchive& Ar) override;
 
-#if WITH_EDITOR
-	// This delegate is called when an edit operation is undone on the rig asset.
-	DECLARE_MULTICAST_DELEGATE(OnIKRigEditUndo);
-	OnIKRigEditUndo IKRigEditUndo;
-	
-	virtual void PostEditUndo() override;
-#endif
-
 	/** The skeletal mesh that was used as the source of the skeleton data. Also used for preview.
 	* NOTE: the IK rig may be played back on ANY skeleton that is compatible with it's hierarchy. */
 	UPROPERTY(AssetRegistrySearchable, VisibleAnywhere, Category = "Imported Skeleton")
diff --git a/Engine/Plugins/Animation/IKRig/Source/IKRig/Public/Retargeter/IKRetargetProcessor.h b/Engine/Plugins/Animation/IKRig/Source/IKRig/Public/Retargeter/IKRetargetProcessor.h
--- a/Engine/Plugins/Animation/IKRig/Source/IKRig/Public/Retargeter/IKRetargetProcessor.h
+++ b/Engine/Plugins/Animation/IKRig/Source/IKRig/Public/Retargeter/IKRetargetProcessor.h
@@ -62,7 +62,11 @@ struct IKRIG_API FRetargetSkeleton
 struct FTargetSkeleton : public FRetargetSkeleton
 {
 	TArray<FTransform> OutputGlobalPose;
+	// true for bones that are in a target chain that is ALSO mapped to a source chain
+	// ie, bones that are actually posed based on a mapped source chain
 	TArray<bool> IsBoneRetargeted;
+	// true for bones that are in a target chain and thus are posed by a retarget pose
+	// the application of a retarget pose to a bone happens regardless of whether they are mapped to a source chain or not
 	TArray<bool> IsBoneInAnyTargetChain;
 
 	void Initialize(
@@ -342,12 +346,20 @@ public:
 	*/
 	FTransform GetTargetBoneRetargetPoseGlobalTransform(const int32& TargetBoneIndex) const;
 
+	FTransform GetTargetBoneRetargetPoseLocalTransform(const int32& TargetBoneIndex) const;
+
 	/** Get read-only access to the target skeleton. */
 	const FTargetSkeleton& GetTargetSkeleton() const { return TargetSkeleton; };
 
+	/** Get index of the root bone of the target skeleton. */
+	const int32 GetTargetSkeletonRootBone() const { return RootRetargeter.Target.BoneIndex; };
+
 	/** Get whether this processor is ready to call RunRetargeter() and generate new poses. */
 	bool IsInitialized() const { return bIsInitialized; };
 
+	/** Get the currently running IK Rig processor for the target */
+	UIKRigProcessor* GetTargetIKRigProcessor() const { return IKRigProcessor; };
+
 #if WITH_EDITOR
 	/** Set that this processor needs to be reinitialized. */
 	void SetNeedsInitialized();
diff --git a/Engine/Plugins/Animation/IKRig/Source/IKRig/Public/Retargeter/IKRetargeter.h b/Engine/Plugins/Animation/IKRig/Source/IKRig/Public/Retargeter/IKRetargeter.h
--- a/Engine/Plugins/Animation/IKRig/Source/IKRig/Public/Retargeter/IKRetargeter.h
+++ b/Engine/Plugins/Animation/IKRig/Source/IKRig/Public/Retargeter/IKRetargeter.h
@@ -59,7 +59,9 @@ struct IKRIG_API FIKRetargetPose
 {
 	GENERATED_BODY()
 	
-	FIKRetargetPose(){}
+public:
+	
+	FIKRetargetPose() = default;
 	
 	UPROPERTY(EditAnywhere, Category = RetargetPose)
 	FVector RootTranslationOffset = FVector::ZeroVector;
@@ -67,9 +69,12 @@ struct IKRIG_API FIKRetargetPose
 	UPROPERTY(EditAnywhere, Category = RetargetPose)
 	TMap<FName, FQuat> BoneRotationOffsets;
 
-	void AddRotationDeltaToBone(FName BoneName, FQuat RotationDelta);
+	// returns true if this is a new offset (not previously recorded)
+	void SetBoneRotationOffset(FName BoneName, FQuat RotationOffset, const FIKRigSkeleton& Skeleton);
 
 	void AddTranslationDeltaToRoot(FVector TranslateDelta);
+
+	void SortHierarchically(const FIKRigSkeleton& Skeleton);
 };
 
 
@@ -107,12 +112,6 @@ public:
 
 #if WITH_EDITOR
 	bool IsInEditRetargetPoseMode() const { return bEditRetargetPoseMode; };
-
-	// This delegate is called when an edit operation is undone on the rig asset.
-	DECLARE_MULTICAST_DELEGATE(OnIKRigEditUndo);
-	OnIKRigEditUndo IKRigEditUndo;
-	
-	virtual void PostEditUndo() override;
 #endif
 
 private:
@@ -158,10 +157,6 @@ public:
 	UPROPERTY(EditAnywhere, Category = PoseEditSettings, meta = (ClampMin = "0.0", UIMin = "0.01", UIMax = "10.0"))
 	float BoneDrawSize = 8.0f;
 
-	/** The visual thickness of the bones in the viewport when editing the retarget pose.*/
-	UPROPERTY(EditAnywhere, Category = PoseEditSettings, meta = (ClampMin = "0.0", UIMin = "0.01", UIMax = "10.0"))
-	float BoneDrawThickness = 1.0f;
-
 private:
 	/** A special editor-only mode which forces the retargeter to output the current retarget reference pose,
 	* rather than actually running the retarget and outputting the retargeted pose. Used in Edit-Pose mode.*/
diff --git a/Engine/Plugins/Animation/IKRig/Source/IKRig/Public/Solvers/IKRig_PBIKSolver.h b/Engine/Plugins/Animation/IKRig/Source/IKRig/Public/Solvers/IKRig_PBIKSolver.h
--- a/Engine/Plugins/Animation/IKRig/Source/IKRig/Public/Solvers/IKRig_PBIKSolver.h
+++ b/Engine/Plugins/Animation/IKRig/Source/IKRig/Public/Solvers/IKRig_PBIKSolver.h
@@ -198,7 +198,9 @@ public:
 	UPROPERTY(EditAnywhere, Category = "Full Body IK Settings")
 	bool bAllowStretch = false;
 
-	/** If true, solver will pull entire skeleton towards effectors BEFORE running constraint iterations. Default is true. */
+	/** (Default is true) If true, solver runs a pre-pass (before the constraint solve) which will translate the whole body by the average motion of the stretched effectors.
+	*This can decrease the amount of iterations needed to achieve a converged pose when the effectors are pulling the body far.
+	*This option has no effect if PinRoot is set to True.*/
 	UPROPERTY(EditAnywhere, Category = "Full Body IK Settings")
 	bool bPrePullRoot = true;
 
diff --git a/Engine/Plugins/Animation/IKRig/Source/IKRigEditor/Private/RetargetEditor/IKRetargetEditMode.cpp b/Engine/Plugins/Animation/IKRig/Source/IKRigEditor/Private/RetargetEditor/IKRetargetEditMode.cpp
--- a/Engine/Plugins/Animation/IKRig/Source/IKRigEditor/Private/RetargetEditor/IKRetargetEditMode.cpp
+++ b/Engine/Plugins/Animation/IKRig/Source/IKRigEditor/Private/RetargetEditor/IKRetargetEditMode.cpp
@@ -5,12 +5,15 @@
 #include "EditorViewportClient.h"
 #include "AssetEditorModeManager.h"
 #include "EditorScriptingHelpers.h"
+#include "IKRigDebugRendering.h"
+#include "IKRigProcessor.h"
 #include "IPersonaPreviewScene.h"
 #include "Animation/DebugSkelMeshComponent.h"
 #include "Materials/MaterialInstanceDynamic.h"
 #include "RetargetEditor/IKRetargetHitProxies.h"
 #include "RetargetEditor/IKRetargetEditor.h"
 #include "Retargeter/IKRetargetProcessor.h"
+#include "RigEditor/IKRigEditMode.h"
 
 FName FIKRetargetEditMode::ModeName("IKRetargetAssetEditMode");
 
@@ -59,38 +62,81 @@ void FIKRetargetEditMode::Render(const FSceneView* View, FViewport* Viewport, FP
 	const UIKRetargeter* Retargeter = Controller->AssetController->GetAsset();
 	const UIKRetargetProcessor* RetargetProcessor = Controller->GetRetargetProcessor();
 	const FTargetSkeleton& TargetSkeleton = RetargetProcessor->GetTargetSkeleton();
+	const int32 RootBoneIndex = RetargetProcessor->GetTargetSkeletonRootBone();
 
+	TSet<int32> OutAffectedBones;
+	TSet<int32> OutSelectedBones;
+	UIKRigProcessor* IKRigProcessor = RetargetProcessor->GetTargetIKRigProcessor();
+	GetAffectedBones(Controller.Get(), IKRigProcessor, OutAffectedBones, OutSelectedBones);
+
+	const float MaxDrawRadius = Controller->TargetSkelMeshComponent->Bounds.SphereRadius * 0.01f;
 	for (int32 BoneIndex = 0; BoneIndex<TargetSkeleton.BoneNames.Num(); ++BoneIndex)
 	{		
-		// filter out bones that are not part of a retargeted chain
-		if (!TargetSkeleton.IsBoneInAnyTargetChain[BoneIndex])
+		// filter out bones that cannot be edited as part of the retarget pose
+		const bool bIsRootBone = BoneIndex==RootBoneIndex;
+		const bool bIsInTargetChain = TargetSkeleton.IsBoneInAnyTargetChain[BoneIndex];
+		if (!(bIsRootBone || bIsInTargetChain))
 		{
 			continue;
 		}
-		
-		const FName& BoneName = TargetSkeleton.BoneNames[BoneIndex];
 
-		// is this chain currently selected?
-		const bool bIsSelected = IsBoneSelected(BoneName);
+		// selected bones are drawn with different color, affected bones are highlighted
+		const bool bIsSelected = OutSelectedBones.Contains(BoneIndex);
+		const bool bIsAffected = OutAffectedBones.Contains(BoneIndex);
+		FLinearColor LineColor = bIsAffected ? IKRigDebugRendering::AFFECTED_BONE_COLOR : IKRigDebugRendering::DESELECTED_BONE_COLOR;
+		LineColor = bIsSelected ? IKRigDebugRendering::SELECTED_BONE_COLOR : LineColor;
 
+		const float BoneRadiusSetting = Retargeter->BoneDrawSize;
+		const float BoneRadius = FMath::Min(1.0f, MaxDrawRadius) * BoneRadiusSetting;
+		
 		// get the location of the bone on the currently initialized target skeletal mesh
-		const FTransform BoneTransform = Controller->GetTargetBoneTransform(BoneIndex);
-
+		// along with array of child positions
+		const FTransform BoneTransform = Controller->GetTargetBoneGlobalTransform(BoneIndex);
 		FVector Start;
-		TArray<FVector> End;
-		if (Controller->GetTargetBoneLineSegments(BoneIndex, Start, End))
+		TArray<FVector> ChildLocations;
+		Controller->GetTargetBoneLineSegments(BoneIndex, Start, ChildLocations);
+		
+		// draw the bone proxy
+		PDI->SetHitProxy(new HIKRetargetEditorBoneProxy(TargetSkeleton.BoneNames[BoneIndex]));
+		IKRigDebugRendering::DrawWireBone(
+			PDI,
+			BoneTransform,
+			ChildLocations,
+			LineColor,
+			SDPG_Foreground,
+			BoneRadius,
+			bIsSelected || bIsAffected);
+		PDI->SetHitProxy(NULL);
+	}
+}
+
+void FIKRetargetEditMode::GetAffectedBones(
+	FIKRetargetEditorController* Controller,
+	UIKRigProcessor* Processor,
+	TSet<int32>& OutAffectedBones,
+	TSet<int32>& OutSelectedBones) const
+{
+	if (!(Controller && Processor))
+	{
+		return;
+	}
+
+	const FIKRigSkeleton& Skeleton = Processor->GetSkeleton();
+
+	// record selected bone indices
+	for (const FName& SelectedBone : BoneEdit.SelectedBones)
+	{
+		OutSelectedBones.Add(Skeleton.GetBoneIndexFromName(SelectedBone));
+	}
+
+	// "affected bones" are the selected bones AND their children, recursively
+	for (int32 SelectedBone : OutSelectedBones)
+	{
+		const int32 EndOfBranch = Skeleton.GetCachedEndOfBranchIndex(SelectedBone);
+		OutAffectedBones.Add(SelectedBone);
+		for (int32 BoneIndex=SelectedBone; BoneIndex<=EndOfBranch; ++BoneIndex)
 		{
-			// draw the chain rotation gizmo
-			PDI->SetHitProxy(new HIKRetargetEditorBoneProxy(BoneName));
-			DrawBoneProxy(
-				PDI,
-				BoneTransform,
-				Start,
-				End,
-				Retargeter->BoneDrawSize,
-				Retargeter->BoneDrawThickness,
-				bIsSelected);
-			PDI->SetHitProxy(NULL);
+			OutAffectedBones.Add(BoneIndex);
 		}
 	}
 }
@@ -113,7 +159,7 @@ bool FIKRetargetEditMode::UsesTransformWidget() const
 bool FIKRetargetEditMode::UsesTransformWidget(UE::Widget::EWidgetMode CheckMode) const
 {
 	const bool bIsOnlyRootSelected = IsOnlyRootSelected();
-	const bool bIsAnyBoneSelected = !SelectedBones.IsEmpty();
+	const bool bIsAnyBoneSelected = !BoneEdit.SelectedBones.IsEmpty();
 
 	if (!bIsAnyBoneSelected)
 	{
@@ -135,7 +181,7 @@ bool FIKRetargetEditMode::UsesTransformWidget(UE::Widget::EWidgetMode CheckMode)
 
 FVector FIKRetargetEditMode::GetWidgetLocation() const
 {
-	if (SelectedBones.IsEmpty())
+	if (BoneEdit.SelectedBones.IsEmpty())
 	{
 		return FVector::ZeroVector; 
 	}
@@ -147,13 +193,13 @@ FVector FIKRetargetEditMode::GetWidgetLocation() const
 	}
 
 	const FTargetSkeleton& TargetSkeleton = Controller->GetRetargetProcessor()->GetTargetSkeleton();
-	const int32 BoneIndex = TargetSkeleton.FindBoneIndexByName(SelectedBones.Last());
+	const int32 BoneIndex = TargetSkeleton.FindBoneIndexByName(BoneEdit.SelectedBones.Last());
 	if (BoneIndex == INDEX_NONE)
 	{
 		return FVector::ZeroVector;
 	}
 	
-	return Controller->GetTargetBoneTransform(BoneIndex).GetTranslation();
+	return Controller->GetTargetBoneGlobalTransform(BoneIndex).GetTranslation();
 }
 
 bool FIKRetargetEditMode::HandleClick(FEditorViewportClient* InViewportClient, HHitProxy* HitProxy, const FViewportClick& Click)
@@ -182,12 +228,13 @@ bool FIKRetargetEditMode::HandleClick(FEditorViewportClient* InViewportClient, H
 
 bool FIKRetargetEditMode::StartTracking(FEditorViewportClient* InViewportClient, FViewport* InViewport)
 {
-	TrackingState = None;
-	
+	TrackingState = FIKRetargetTrackingState::None;
+
+	// not manipulating any widget axes, so stop tracking
 	const EAxisList::Type CurrentAxis = InViewportClient->GetCurrentWidgetAxis();
 	if (CurrentAxis == EAxisList::None)
 	{
-		return false; // not manipulating a required axis
+		return false; 
 	}
 
 	const TSharedPtr<FIKRetargetEditorController> Controller = EditorController.Pin();
@@ -199,7 +246,7 @@ bool FIKRetargetEditMode::StartTracking(FEditorViewportClient* InViewportClient,
 	// get state of viewport
 	const bool bTranslating = InViewportClient->GetWidgetMode() == UE::Widget::EWidgetMode::WM_Translate;
 	const bool bRotating = InViewportClient->GetWidgetMode() == UE::Widget::EWidgetMode::WM_Rotate;
-	const bool bAnyBoneSelected = !SelectedBones.IsEmpty();
+	const bool bAnyBoneSelected = !BoneEdit.SelectedBones.IsEmpty();
 	const bool bOnlyRootSelected = IsOnlyRootSelected();
 
 	// is any bone being rotated?
@@ -209,6 +256,7 @@ bool FIKRetargetEditMode::StartTracking(FEditorViewportClient* InViewportClient,
 		GEditor->BeginTransaction(LOCTEXT("RotateRetargetPoseBone", "Rotate Retarget Pose Bone"));
 		Controller->AssetController->GetAsset()->Modify();
 		TrackingState = FIKRetargetTrackingState::RotatingBone;
+		UpdateWidgetTransform();
 		return true;
 	}
 
@@ -219,6 +267,7 @@ bool FIKRetargetEditMode::StartTracking(FEditorViewportClient* InViewportClient,
 		GEditor->BeginTransaction(LOCTEXT("TranslateRetargetPoseBone", "Translate Retarget Pose Bone"));
 		Controller->AssetController->GetAsset()->Modify();
 		TrackingState = FIKRetargetTrackingState::TranslatingRoot;
+		UpdateWidgetTransform();
 		return true;
 	}
 
@@ -229,7 +278,15 @@ bool FIKRetargetEditMode::EndTracking(FEditorViewportClient* InViewportClient, F
 {	
 	if (TrackingState == FIKRetargetTrackingState::None)
 	{
-		return false; // not handled
+		const bool bIsRootSelected = IsRootSelected();
+		const bool bTranslating = InViewportClient->GetWidgetMode() == UE::Widget::EWidgetMode::WM_Translate;
+		// forcibly prevent translation of anything but the root
+		if (!bIsRootSelected && bTranslating)
+		{
+			InViewportClient->SetWidgetMode(UE::Widget::EWidgetMode::WM_Rotate);
+		}
+		
+		return true; // not handled
 	}
 
 	GEditor->EndTransaction();
@@ -237,7 +294,12 @@ bool FIKRetargetEditMode::EndTracking(FEditorViewportClient* InViewportClient, F
 	return true;
 }
 
-bool FIKRetargetEditMode::InputDelta(FEditorViewportClient* InViewportClient, FViewport* InViewport, FVector& InDrag, FRotator& InRot, FVector& InScale)
+bool FIKRetargetEditMode::InputDelta(
+	FEditorViewportClient* InViewportClient,
+	FViewport* InViewport,
+	FVector& InDrag,
+	FRotator& InRot,
+	FVector& InScale)
 {
 	if (TrackingState == FIKRetargetTrackingState::None)
 	{
@@ -258,10 +320,23 @@ bool FIKRetargetEditMode::InputDelta(FEditorViewportClient* InViewportClient, FV
 			return false;
 		}
 
+		// accumulate the rotation from the viewport gizmo (amount of rotation since tracking started)
+		BoneEdit.AccumulatedGlobalOffset = InRot.Quaternion().Inverse() * BoneEdit.AccumulatedGlobalOffset;
+
+		// convert to a delta rotation in the local space of the last selected bone
+		const FQuat BoneGlobalOrig = BoneEdit.GlobalTransform.GetRotation();
+		const FQuat BoneGlobalPlusOffset = BoneEdit.AccumulatedGlobalOffset * BoneGlobalOrig;
+		const FQuat ParentInv = BoneEdit.ParentGlobalTransform.GetRotation().Inverse();
+		const FQuat BoneLocal = ParentInv * BoneGlobalOrig;
+		const FQuat BoneLocalPlusOffset = ParentInv * BoneGlobalPlusOffset;
+		const FQuat BoneLocalOffset = BoneLocal * BoneLocalPlusOffset.Inverse();
+		
 		// apply rotation delta to all selected bones
-		for (const FName& BoneName : SelectedBones)
+		for (int32 SelectionIndex=0; SelectionIndex<BoneEdit.SelectedBones.Num(); ++SelectionIndex)
 		{
-			Controller->AssetController->AddRotationOffsetToRetargetPoseBone(BoneName, InRot.Quaternion());
+			const FQuat TotalOffsetLocal = BoneLocalOffset * BoneEdit.PrevLocalOffsets[SelectionIndex];
+			const FName& BoneName = BoneEdit.SelectedBones[SelectionIndex];
+			Controller->AssetController->SetRotationOffsetForRetargetPoseBone(BoneName, TotalOffsetLocal);
 		}
 
 		return true;
@@ -285,7 +360,7 @@ bool FIKRetargetEditMode::InputDelta(FEditorViewportClient* InViewportClient, FV
 
 bool FIKRetargetEditMode::GetCustomDrawingCoordinateSystem(FMatrix& InMatrix, void* InData)
 {
-	if (SelectedBones.IsEmpty())
+	if (BoneEdit.SelectedBones.IsEmpty())
 	{
 		return false; // nothing selected to manipulate
 	}
@@ -296,9 +371,12 @@ bool FIKRetargetEditMode::GetCustomDrawingCoordinateSystem(FMatrix& InMatrix, vo
 		return false; 
 	}
 
-	const int32 BoneIndex = Controller->AssetController->GetAsset()->GetTargetIKRig()->Skeleton.GetBoneIndexFromName(SelectedBones[0]);
-	const FTransform CurrentTransform = Controller->GetTargetBoneTransform(BoneIndex);
-	InMatrix = CurrentTransform.ToMatrixNoScale().RemoveTranslation();
+	if (TrackingState == FIKRetargetTrackingState::None)
+	{
+		UpdateWidgetTransform();
+	}
+
+	InMatrix = BoneEdit.GlobalTransform.ToMatrixNoScale().RemoveTranslation();
 	return true;
 }
 
@@ -307,9 +385,43 @@ bool FIKRetargetEditMode::GetCustomInputCoordinateSystem(FMatrix& InMatrix, void
 	return GetCustomDrawingCoordinateSystem(InMatrix, InData);
 }
 
-bool FIKRetargetEditMode::IsOnlyRootSelected() const
+void FIKRetargetEditMode::UpdateWidgetTransform()
 {
-	if (SelectedBones.Num() != 1)
+	const TSharedPtr<FIKRetargetEditorController> Controller = EditorController.Pin();
+	if (!Controller.IsValid() || BoneEdit.SelectedBones.IsEmpty())
+	{
+		BoneEdit.GlobalTransform = FTransform::Identity;
+		return;
+	}
+
+	UIKRetargeterController* AssetController = Controller->AssetController;
+	const FIKRigSkeleton& Skeleton = AssetController->GetAsset()->GetTargetIKRig()->Skeleton;
+
+	BoneEdit.Name = BoneEdit.SelectedBones.Last();
+	BoneEdit.Index = Skeleton.GetBoneIndexFromName(BoneEdit.Name);
+	BoneEdit.GlobalTransform = Controller->GetTargetBoneGlobalTransform(BoneEdit.Index);
+	BoneEdit.AccumulatedGlobalOffset = FQuat::Identity;
+
+	BoneEdit.PrevLocalOffsets.Reset();
+	for (int32 SelectionIndex=0; SelectionIndex<BoneEdit.SelectedBones.Num(); ++SelectionIndex)
+	{
+		FQuat PrevLocalOffset = AssetController->GetRotationOffsetForRetargetPoseBone(BoneEdit.SelectedBones[SelectionIndex]);
+		BoneEdit.PrevLocalOffsets.Add(PrevLocalOffset);
+	}
+	
+	int32 ParentIndex = Skeleton.GetParentIndex(BoneEdit.Index);
+	if (ParentIndex != INDEX_NONE)
+	{
+		BoneEdit.ParentGlobalTransform = Controller->GetTargetBoneGlobalTransform(ParentIndex);
+	}else
+	{
+		BoneEdit.ParentGlobalTransform = FTransform::Identity;
+	}
+}
+
+bool FIKRetargetEditMode::IsRootSelected() const
+{
+	if (BoneEdit.SelectedBones.IsEmpty())
 	{
 		return false;
 	}
@@ -320,30 +432,32 @@ bool FIKRetargetEditMode::IsOnlyRootSelected() const
 		return false; 
 	}
 
-	return Controller->AssetController->GetTargetRootBone() == SelectedBones[0];
+	const FName& RootName = Controller->AssetController->GetTargetRootBone();
+	for (const FName& SelectedBone : BoneEdit.SelectedBones)
+	{
+		if (RootName == SelectedBone)
+		{
+			return true;
+		}
+	}
+
+	return false;
 }
 
-void FIKRetargetEditMode::DrawBoneProxy(
-	FPrimitiveDrawInterface* PDI,
-	const FTransform& BoneTransform,
-	const FVector& Start,
-	const TArray<FVector>& ChildPoints,
-	float Size,
-	float Thickness,
-	bool bIsSelected) const
+bool FIKRetargetEditMode::IsOnlyRootSelected() const
 {
-	const FLinearColor Color = bIsSelected ? FLinearColor::Green : FLinearColor::Yellow;
-	Size = FMath::Clamp(Size, 0.1f, 1000.0f);
-	Thickness = FMath::Clamp(Thickness, 0.1f, 1000.0f);
+	if (BoneEdit.SelectedBones.Num() != 1)
+	{
+		return false;
+	}
 
-	// draw a line to each child
-	for (const FVector& ChildPosition : ChildPoints)
+	const TSharedPtr<FIKRetargetEditorController> Controller = EditorController.Pin();
+	if (!Controller.IsValid())
 	{
-		PDI->DrawLine(Start, ChildPosition, Color, SDPG_Foreground, Thickness * 1.0f);
+		return false; 
 	}
-	
-	// draw a coordinate frame at the bone
-	DrawCoordinateSystem(PDI, Start, BoneTransform.GetRotation().Rotator(), Size, SDPG_Foreground, Thickness);
+
+	return Controller->AssetController->GetTargetRootBone() == BoneEdit.SelectedBones[0];
 }
 
 void FIKRetargetEditMode::Tick(FEditorViewportClient* ViewportClient, float DeltaTime)
@@ -360,31 +474,31 @@ void FIKRetargetEditMode::DrawHUD(FEditorViewportClient* ViewportClient, FViewpo
 
 bool FIKRetargetEditMode::IsBoneSelected(const FName& BoneName) const
 {
-	return SelectedBones.Contains(BoneName);
+	return BoneEdit.SelectedBones.Contains(BoneName);
 }
 
 void FIKRetargetEditMode::HandleBoneSelectedInViewport(const FName& BoneName, bool bReplace)
 {
 	if (bReplace)
 	{
-		SelectedBones.Reset();
+		BoneEdit.SelectedBones.Reset();
 		if (BoneName != NAME_None)
 		{
-			SelectedBones.Add(BoneName);
+			BoneEdit.SelectedBones.Add(BoneName);
 		}
 	}
 	else
 	{
 		if (BoneName != NAME_None)
 		{
-			const bool bAlreadySelected = SelectedBones.Contains(BoneName);
+			const bool bAlreadySelected = BoneEdit.SelectedBones.Contains(BoneName);
 			if (bAlreadySelected)
 			{
-				SelectedBones.Remove(BoneName);	
+				BoneEdit.SelectedBones.Remove(BoneName);	
 			}
 			else
 			{
-				SelectedBones.Add(BoneName);
+				BoneEdit.SelectedBones.Add(BoneName);
 			}
 		}	
 	}
diff --git a/Engine/Plugins/Animation/IKRig/Source/IKRigEditor/Private/RetargetEditor/IKRetargetEditor.cpp b/Engine/Plugins/Animation/IKRig/Source/IKRigEditor/Private/RetargetEditor/IKRetargetEditor.cpp
--- a/Engine/Plugins/Animation/IKRig/Source/IKRigEditor/Private/RetargetEditor/IKRetargetEditor.cpp
+++ b/Engine/Plugins/Animation/IKRig/Source/IKRigEditor/Private/RetargetEditor/IKRetargetEditor.cpp
@@ -264,6 +264,18 @@ TStatId FIKRetargetEditor::GetStatId() const
 	RETURN_QUICK_DECLARE_CYCLE_STAT(FIKRetargetEditor, STATGROUP_Tickables);
 }
 
+void FIKRetargetEditor::PostUndo(bool bSuccess)
+{
+	EditorController->AssetController->BroadcastNeedsReinitialized();
+	EditorController->RefreshAllViews();
+}
+
+void FIKRetargetEditor::PostRedo(bool bSuccess)
+{
+	EditorController->AssetController->BroadcastNeedsReinitialized();
+	EditorController->RefreshAllViews();
+}
+
 void FIKRetargetEditor::HandlePreviewSceneCreated(const TSharedRef<IPersonaPreviewScene>& InPersonaPreviewScene)
 {
 	AAnimationEditorPreviewActor* Actor = InPersonaPreviewScene->GetWorld()->SpawnActor<AAnimationEditorPreviewActor>(AAnimationEditorPreviewActor::StaticClass(), FTransform::Identity);
diff --git a/Engine/Plugins/Animation/IKRig/Source/IKRigEditor/Private/RetargetEditor/IKRetargetEditorController.cpp b/Engine/Plugins/Animation/IKRig/Source/IKRigEditor/Private/RetargetEditor/IKRetargetEditorController.cpp
--- a/Engine/Plugins/Animation/IKRig/Source/IKRigEditor/Private/RetargetEditor/IKRetargetEditorController.cpp
+++ b/Engine/Plugins/Animation/IKRig/Source/IKRigEditor/Private/RetargetEditor/IKRetargetEditorController.cpp
@@ -76,8 +76,6 @@ void FIKRetargetEditorController::OnRetargeterNeedsInitialized(const UIKRetarget
 	TargetAnimInstance->SetProcessorNeedsInitialized();
 	// refresh all the UI views
 	RefreshAllViews();
-	// keep playing animation
-	PlayPreviousAnimationAsset();
 }
 
 USkeletalMesh* FIKRetargetEditorController::GetSourceSkeletalMesh() const
@@ -100,7 +98,7 @@ USkeletalMesh* FIKRetargetEditorController::GetTargetSkeletalMesh() const
 	return AssetController->GetTargetPreviewMesh();
 }
 
-FTransform FIKRetargetEditorController::GetTargetBoneTransform(const int32& TargetBoneIndex) const
+FTransform FIKRetargetEditorController::GetTargetBoneGlobalTransform(const int32& TargetBoneIndex) const
 {
 	UIKRetargetAnimInstance* AnimInstance = TargetAnimInstance.Get();
 	check(AnimInstance);
@@ -108,7 +106,7 @@ FTransform FIKRetargetEditorController::GetTargetBoneTransform(const int32& Targ
 	const UIKRetargetProcessor* RetargetProcessor = AnimInstance->GetRetargetProcessor();
 	check(RetargetProcessor)
 
-	// get transform of root of chain
+	// get transform of bone
 	FTransform BoneTransform = RetargetProcessor->GetTargetBoneRetargetPoseGlobalTransform(TargetBoneIndex);
 
 	// scale and offset
@@ -118,6 +116,17 @@ FTransform FIKRetargetEditorController::GetTargetBoneTransform(const int32& Targ
 	return BoneTransform;
 }
 
+FTransform FIKRetargetEditorController::GetTargetBoneLocalTransform(const int32& TargetBoneIndex) const
+{
+	UIKRetargetAnimInstance* AnimInstance = TargetAnimInstance.Get();
+	check(AnimInstance);
+
+	const UIKRetargetProcessor* RetargetProcessor = AnimInstance->GetRetargetProcessor();
+	check(RetargetProcessor)
+
+	return RetargetProcessor->GetTargetBoneRetargetPoseLocalTransform(TargetBoneIndex);
+}
+
 bool FIKRetargetEditorController::GetTargetBoneLineSegments(
 	const int32& TargetBoneIndex,
 	FVector& OutStart,
diff --git a/Engine/Plugins/Animation/IKRig/Source/IKRigEditor/Private/RetargetEditor/IKRetargeterController.cpp b/Engine/Plugins/Animation/IKRig/Source/IKRigEditor/Private/RetargetEditor/IKRetargeterController.cpp
--- a/Engine/Plugins/Animation/IKRig/Source/IKRigEditor/Private/RetargetEditor/IKRetargeterController.cpp
+++ b/Engine/Plugins/Animation/IKRig/Source/IKRigEditor/Private/RetargetEditor/IKRetargeterController.cpp
@@ -28,9 +28,6 @@ UIKRetargeterController* UIKRetargeterController::GetController(UIKRetargeter* I
 	// clean the asset before editing
 	Controller->CleanChainMapping();
 	Controller->CleanPoseList();
-
-	// bind callback to reinitialize processor on the asset undo
-	InRetargeterAsset->IKRigEditUndo.AddUObject(Controller, &UIKRetargeterController::BroadcastNeedsReinitialized);
 	
 	return Controller;
 }
@@ -214,6 +211,9 @@ void UIKRetargeterController::CleanPoseList()
 			{
 				Pose.Value.BoneRotationOffsets.Remove(BoneToRemove);
 			}
+
+			// sort the pose offset from leaf to root
+			Pose.Value.SortHierarchically(Asset->TargetIKRigAsset->Skeleton);
 		}
 	}
 
@@ -390,9 +390,21 @@ const TMap<FName, FIKRetargetPose>& UIKRetargeterController::GetRetargetPoses()
 	return GetAsset()->RetargetPoses;
 }
 
-void UIKRetargeterController::AddRotationOffsetToRetargetPoseBone(FName BoneName, FQuat RotationOffset) const
+void UIKRetargeterController::SetRotationOffsetForRetargetPoseBone(FName BoneName, FQuat RotationOffset) const
+{
+	const FIKRigSkeleton& Skeleton = Asset->GetTargetIKRig()->Skeleton;
+	Asset->RetargetPoses[Asset->CurrentRetargetPose].SetBoneRotationOffset(BoneName, RotationOffset, Skeleton);
+}
+
+FQuat UIKRetargeterController::GetRotationOffsetForRetargetPoseBone(FName BoneName) const
 {
-	Asset->RetargetPoses[Asset->CurrentRetargetPose].AddRotationDeltaToBone(BoneName, RotationOffset);
+	TMap<FName, FQuat>& BoneOffsets = Asset->RetargetPoses[Asset->CurrentRetargetPose].BoneRotationOffsets;
+	if (!BoneOffsets.Contains(BoneName))
+	{
+		return FQuat::Identity;
+	}
+	
+	return BoneOffsets[BoneName];
 }
 
 void UIKRetargeterController::AddTranslationOffsetToRetargetRootBone(FVector TranslationOffset) const
diff --git a/Engine/Plugins/Animation/IKRig/Source/IKRigEditor/Private/RigEditor/IKRigAnimInstance.cpp b/Engine/Plugins/Animation/IKRig/Source/IKRigEditor/Private/RigEditor/IKRigAnimInstance.cpp
--- a/Engine/Plugins/Animation/IKRig/Source/IKRigEditor/Private/RigEditor/IKRigAnimInstance.cpp
+++ b/Engine/Plugins/Animation/IKRig/Source/IKRigEditor/Private/RigEditor/IKRigAnimInstance.cpp
@@ -20,6 +20,11 @@ void UIKRigAnimInstance::SetProcessorNeedsInitialized()
 	IKRigNode.SetProcessorNeedsInitialized();
 }
 
+UIKRigProcessor* UIKRigAnimInstance::GetCurrentlyRunningProcessor()
+{
+	return IKRigNode.IKRigProcessor;
+}
+
 FAnimInstanceProxy* UIKRigAnimInstance::CreateAnimInstanceProxy()
 {
 	return new FIKRigAnimInstanceProxy(this, &IKRigNode);
diff --git a/Engine/Plugins/Animation/IKRig/Source/IKRigEditor/Private/RigEditor/IKRigController.cpp b/Engine/Plugins/Animation/IKRig/Source/IKRigEditor/Private/RigEditor/IKRigController.cpp
--- a/Engine/Plugins/Animation/IKRig/Source/IKRigEditor/Private/RigEditor/IKRigController.cpp
+++ b/Engine/Plugins/Animation/IKRig/Source/IKRigEditor/Private/RigEditor/IKRigController.cpp
@@ -24,7 +24,6 @@ UIKRigController* UIKRigController::GetIKRigController(UIKRigDefinition* InIKRig
 		UIKRigController* Controller = NewObject<UIKRigController>();
 		Controller->Asset = InIKRigDefinition;
 		InIKRigDefinition->Controller = Controller;
-		InIKRigDefinition->IKRigEditUndo.AddUObject(Controller, &UIKRigController::BroadcastNeedsReinitialized);
 	}
 
 	return Cast<UIKRigController>(InIKRigDefinition->Controller);
@@ -771,7 +770,7 @@ bool UIKRigController::ConnectGoalToSolver(const UIKRigEffectorGoal& Goal, int32
 	// can't add goal to a solver with an invalid index
 	check(Asset->Solvers.IsValidIndex(SolverIndex))
 
-	FScopedTransaction Transaction(LOCTEXT("AddGoalToSolver_Label", "Add Goal"));
+	FScopedTransaction Transaction(LOCTEXT("ConnectGoalSolver_Label", "Connect Goal to Solver"));
 	UIKRigSolver* Solver = Asset->Solvers[SolverIndex];
 	Solver->Modify();
 	
@@ -788,7 +787,7 @@ bool UIKRigController::DisconnectGoalFromSolver(const FName& GoalToRemove, int32
 	// can't remove goal from a solver with an invalid index
 	check(Asset->Solvers.IsValidIndex(SolverIndex))
 
-    FScopedTransaction Transaction(LOCTEXT("RemoveGoalFromSolver_Label", "Remove Goal"));
+    FScopedTransaction Transaction(LOCTEXT("DisconnectGoalSolver_Label", "Disconnect Goal from Solver"));
 	UIKRigSolver* Solver = Asset->Solvers[SolverIndex];
 	Solver->Modify();
 	
diff --git a/Engine/Plugins/Animation/IKRig/Source/IKRigEditor/Private/RigEditor/IKRigEditMode.cpp b/Engine/Plugins/Animation/IKRig/Source/IKRigEditor/Private/RigEditor/IKRigEditMode.cpp
--- a/Engine/Plugins/Animation/IKRig/Source/IKRigEditor/Private/RigEditor/IKRigEditMode.cpp
+++ b/Engine/Plugins/Animation/IKRig/Source/IKRigEditor/Private/RigEditor/IKRigEditMode.cpp
@@ -5,10 +5,13 @@
 #include "EditorViewportClient.h"
 #include "AssetEditorModeManager.h"
 #include "IPersonaPreviewScene.h"
+#include "SkeletalDebugRendering.h"
 #include "Animation/DebugSkelMeshComponent.h"
 #include "Materials/MaterialInstanceDynamic.h"
+#include "RigEditor/IKRigAnimInstance.h"
 #include "RigEditor/IKRigHitProxies.h"
 #include "RigEditor/IKRigToolkit.h"
+#include "IKRigDebugRendering.h"
 
 #define LOCTEXT_NAMESPACE "IKRetargeterEditMode"
 
@@ -27,7 +30,8 @@ bool FIKRigEditMode::GetCameraTarget(FSphere& OutTarget) const
 	}
 	
 	// target union of selected goals
-	const TArray<FName>& OutGoalNames = Controller->GetSelectedGoals();
+	TArray<FName> OutGoalNames;
+	Controller->SkeletonView->GetSelectedGoalNames(OutGoalNames);
 	if (!OutGoalNames.IsEmpty())
 	{
 		TArray<FVector> GoalPoints;
@@ -64,7 +68,12 @@ void FIKRigEditMode::GetOnScreenDebugInfo(TArray<FText>& OutDebugInfo) const
 void FIKRigEditMode::Render(const FSceneView* View, FViewport* Viewport, FPrimitiveDrawInterface* PDI)
 {
 	FEdMode::Render(View, Viewport, PDI);
+	RenderGoals(PDI);
+	RenderBones(PDI);
+}
 
+void FIKRigEditMode::RenderGoals(FPrimitiveDrawInterface* PDI)
+{
 	const TSharedPtr<FIKRigEditorController> Controller = EditorController.Pin();
 	if (!Controller.IsValid())
 	{
@@ -73,18 +82,193 @@ void FIKRigEditMode::Render(const FSceneView* View, FViewport* Viewport, FPrimit
 	
 	UIKRigController* AssetController = Controller->AssetController;
 	const UIKRigDefinition* IKRigAsset = AssetController->GetAsset();
+	if (!IKRigAsset->DrawGoals)
+	{
+		return;
+	}
+	
 	TArray<UIKRigEffectorGoal*> Goals = AssetController->GetAllGoals();
 	for (const UIKRigEffectorGoal* Goal : Goals)
 	{
-		const bool bIsSelected = Controller->IsGoalSelected(Goal->GoalName);
-		PDI->SetHitProxy(new HIKRigEditorGoalProxy(Goal->GoalName));
+		const bool bIsSelected = Controller->SkeletonView->IsGoalSelected(Goal->GoalName);
 		const float Size = IKRigAsset->GoalSize * Goal->SizeMultiplier;
 		const float Thickness = IKRigAsset->GoalThickness * Goal->ThicknessMultiplier;
-		GoalDrawer.DrawGoal(PDI, Goal, bIsSelected, Size, Thickness);
+		PDI->SetHitProxy(new HIKRigEditorGoalProxy(Goal->GoalName));
+		IKRigDebugRendering::DrawGoal(PDI, Goal, bIsSelected, Size, Thickness);
 		PDI->SetHitProxy(NULL);
 	}
 }
 
+void FIKRigEditMode::RenderBones(FPrimitiveDrawInterface* PDI)
+{
+	// editor configured and initialized?
+	const TSharedPtr<FIKRigEditorController> Controller = EditorController.Pin();
+	if (!Controller.IsValid())
+	{
+		return; 
+	}
+
+	UIKRigController* AssetController = Controller->AssetController;
+	const UIKRigDefinition* IKRigAsset = AssetController->GetAsset();
+	if (!IKRigAsset->DrawGoals)
+	{
+		return;
+	}
+
+	// anim instance initialized?
+	if (!Controller->AnimInstance.IsValid())
+	{
+		return;
+	}
+
+	// IKRig processor initialized and running?
+	UIKRigProcessor* CurrentProcessor = Controller->AnimInstance->GetCurrentlyRunningProcessor();
+	if (!IsValid(CurrentProcessor))
+	{
+		return;
+	}
+	if (!CurrentProcessor->IsInitialized())
+	{
+		return;
+	}
+
+	// get affected / selected bones
+	TSet<int32> OutAffectedBones;
+	TSet<int32> OutSelectedBones;
+	GetAffectedBones(Controller.Get(), CurrentProcessor,OutAffectedBones,OutSelectedBones);
+
+	// draw affected bones
+	TArray<int32> ChildrenIndices;
+	const FIKRigSkeleton& Skeleton = CurrentProcessor->GetSkeleton();
+	const TArray<FTransform>& BoneTransforms = Skeleton.CurrentPoseGlobal;
+	const float MaxDrawRadius = Controller->SkelMeshComponent->Bounds.SphereRadius * 0.01f;
+	for (int32 BoneIndex=0; BoneIndex<Skeleton.BoneNames.Num(); ++BoneIndex)
+	{
+		// selected bones are drawn with different color
+		const bool bIsSelected = OutSelectedBones.Contains(BoneIndex);
+		const bool bIsAffected = OutAffectedBones.Contains(BoneIndex);
+		FLinearColor LineColor = bIsAffected ? IKRigDebugRendering::AFFECTED_BONE_COLOR : IKRigDebugRendering::DESELECTED_BONE_COLOR;
+		LineColor = bIsSelected ? IKRigDebugRendering::SELECTED_BONE_COLOR : LineColor;
+
+		// only draw axes on affected/selected bones
+		const bool bDrawAxes = bIsSelected || bIsAffected;
+		
+		const float BoneRadiusSetting = IKRigAsset->BoneSize;
+		const float BoneRadius = FMath::Min(1.0f, MaxDrawRadius) * BoneRadiusSetting;
+		
+		// draw line from bone to each child
+		FTransform BoneTransform = BoneTransforms[BoneIndex];
+		TArray<FVector> ChildPoints;
+		Skeleton.GetChildIndices(BoneIndex,ChildrenIndices);
+		for (int32 ChildIndex : ChildrenIndices)
+		{
+			ChildPoints.Add(BoneTransforms[ChildIndex].GetLocation());
+		}
+
+		PDI->SetHitProxy(new HIKRigEditorBoneProxy(Skeleton.BoneNames[BoneIndex]));
+		IKRigDebugRendering::DrawWireBone(
+			PDI,
+			BoneTransform,
+			ChildPoints,
+			LineColor,
+			SDPG_Foreground,
+			BoneRadius,
+			bDrawAxes);
+		PDI->SetHitProxy(nullptr);
+	}
+}
+
+void FIKRigEditMode::GetAffectedBones(
+	FIKRigEditorController* Controller,
+	UIKRigProcessor* Processor,
+	TSet<int32>& OutAffectedBones,
+	TSet<int32>& OutSelectedBones) const
+{
+	OutAffectedBones.Reset();
+	OutSelectedBones.Reset();
+
+	const FIKRigSkeleton& Skeleton = Processor->GetSkeleton();
+	
+	switch (Controller->GetLastSelectedType())
+	{
+		case EIKRigSelectionType::Hierarchy:
+		{
+			// get selected bones
+			TArray<TSharedPtr<FIKRigTreeElement>> SelectedBoneItems;
+			SIKRigSkeleton* SkeletonView = Controller->SkeletonView.Get();
+			SkeletonView->GetSelectedBones(SelectedBoneItems);
+
+			// record indices of all selected bones
+			for (TSharedPtr<FIKRigTreeElement> SelectedBone: SelectedBoneItems)
+			{
+				int32 BoneIndex = Skeleton.GetBoneIndexFromName(SelectedBone->BoneName);
+				OutSelectedBones.Add(BoneIndex);
+			}
+
+			// "affected bones" are the selected bones AND their children, recursively
+			for (int32 SelectedBone : OutSelectedBones)
+			{
+				const int32 EndOfBranch = Skeleton.GetCachedEndOfBranchIndex(SelectedBone);
+				OutAffectedBones.Add(SelectedBone);
+				for (int32 BoneIndex=SelectedBone; BoneIndex<=EndOfBranch; ++BoneIndex)
+				{
+					OutAffectedBones.Add(BoneIndex);
+				}
+			}
+		}
+		break;
+		
+		case EIKRigSelectionType::SolverStack:
+		{
+			// get selected solver
+			TArray<TSharedPtr<FSolverStackElement>> SelectedSolvers;
+			Controller->GetSelectedSolvers(SelectedSolvers);
+			if (SelectedSolvers.IsEmpty())
+			{
+				return;
+			}
+
+			// record which bones in the skeleton are affected by this solver
+			UIKRigController* AssetController = Controller->AssetController;
+			UIKRigSolver* SelectedSolver = AssetController->GetSolver(SelectedSolvers[0].Get()->IndexInStack);
+			for (int32 BoneIndex=0; BoneIndex<Skeleton.BoneNames.Num(); ++BoneIndex)
+			{
+				const FName& BoneName = Skeleton.BoneNames[BoneIndex];
+				if (SelectedSolver->IsBoneAffectedBySolver(BoneName, Skeleton))
+				{
+					OutAffectedBones.Add(BoneIndex);
+				}
+			}
+		}
+		break;
+		
+		case EIKRigSelectionType::RetargetChains:
+		{
+			const FName SelectedChainName = Controller->RetargetingView->GetSelectedChain();
+			if (SelectedChainName == NAME_None)
+			{
+				return;
+			}
+			UIKRigController* AssetController = Controller->AssetController;
+			const FName StartBone = AssetController->GetRetargetChainStartBone(SelectedChainName);
+			const FName EndBone = AssetController->GetRetargetChainEndBone(SelectedChainName);
+			const int32 ChainStartIndex = Skeleton.GetBoneIndexFromName(StartBone);
+			const int32 ChainEndIndex = Skeleton.GetBoneIndexFromName(EndBone);
+			OutSelectedBones.Add(ChainStartIndex);
+			OutAffectedBones.Add(ChainStartIndex);
+			for (int32 BoneIndex=ChainStartIndex; BoneIndex<=ChainEndIndex; ++BoneIndex)
+			{
+				OutSelectedBones.Add(BoneIndex);
+				OutAffectedBones.Add(BoneIndex);
+			}
+		}
+		break;
+		
+		default:
+			checkNoEntry()
+	}
+}
+
 bool FIKRigEditMode::AllowWidgetMove()
 {
 	const TSharedPtr<FIKRigEditorController> Controller = EditorController.Pin();
@@ -93,7 +277,7 @@ bool FIKRigEditMode::AllowWidgetMove()
 		return false; 
 	}
 	
-	return Controller->GetNumSelectedGoals() > 0;
+	return Controller->SkeletonView->GetNumSelectedGoals() > 0;
 }
 
 bool FIKRigEditMode::ShouldDrawWidget() const
@@ -104,7 +288,7 @@ bool FIKRigEditMode::ShouldDrawWidget() const
 		return false; 
 	}
 	
-	return Controller->GetNumSelectedGoals() > 0;
+	return Controller->SkeletonView->GetNumSelectedGoals() > 0;
 }
 
 bool FIKRigEditMode::UsesTransformWidget() const
@@ -115,7 +299,7 @@ bool FIKRigEditMode::UsesTransformWidget() const
 		return false; 
 	}
 	
-	return Controller->GetNumSelectedGoals() > 0;
+	return Controller->SkeletonView->GetNumSelectedGoals() > 0;
 }
 
 bool FIKRigEditMode::UsesTransformWidget(UE::Widget::EWidgetMode CheckMode) const
@@ -126,7 +310,7 @@ bool FIKRigEditMode::UsesTransformWidget(UE::Widget::EWidgetMode CheckMode) cons
 		return false; 
 	}
 	
-	return Controller->GetNumSelectedGoals() > 0;
+	return Controller->SkeletonView->GetNumSelectedGoals() > 0;
 }
 
 FVector FIKRigEditMode::GetWidgetLocation() const
@@ -137,7 +321,8 @@ FVector FIKRigEditMode::GetWidgetLocation() const
 		return FVector::ZeroVector; 
 	}
 	
-	const TArray<FName>& OutGoalNames = Controller->GetSelectedGoals();
+	TArray<FName> OutGoalNames;
+	Controller->SkeletonView->GetSelectedGoalNames(OutGoalNames);
 	if (OutGoalNames.IsEmpty())
 	{
 		return FVector::ZeroVector; 
@@ -157,19 +342,31 @@ bool FIKRigEditMode::HandleClick(FEditorViewportClient* InViewportClient, HHitPr
 	// check for selections
 	if (Click.GetKey() == EKeys::LeftMouseButton)
 	{
+		// draw bones based on the hierarchy when clicking in viewport
+		EditorController.Pin()->SetLastSelectedType(EIKRigSelectionType::Hierarchy);
+		
+		// clicking in empty space clears selection and shows empty details
+		if (!HitProxy)
+		{
+			Controller->ClearSelection();
+			return false;
+		}
+		
 		// selected goal
-		if (HitProxy && HitProxy->IsA(HIKRigEditorGoalProxy::StaticGetType()))
+		if (HitProxy->IsA(HIKRigEditorGoalProxy::StaticGetType()))
 		{
 			HIKRigEditorGoalProxy* GoalProxy = static_cast<HIKRigEditorGoalProxy*>(HitProxy);
 			const bool bReplaceSelection = !(InViewportClient->IsCtrlPressed() || InViewportClient->IsShiftPressed());
 			Controller->HandleGoalSelectedInViewport(GoalProxy->GoalName, bReplaceSelection);
 			return true;
 		}
-		else
+		// selected bone
+		if (HitProxy->IsA(HIKRigEditorBoneProxy::StaticGetType()))
 		{
-			// clicking in empty space clears selected goals
-			const bool bReplaceSelection = true;
-			Controller->HandleGoalSelectedInViewport(NAME_None, bReplaceSelection);
+			HIKRigEditorBoneProxy* BoneProxy = static_cast<HIKRigEditorBoneProxy*>(HitProxy);
+			const bool bReplaceSelection = !(InViewportClient->IsCtrlPressed() || InViewportClient->IsShiftPressed());
+			Controller->HandleBoneSelectedInViewport(BoneProxy->BoneName, bReplaceSelection);
+			return true;
 		}
 	}
 	
@@ -184,7 +381,8 @@ bool FIKRigEditMode::StartTracking(FEditorViewportClient* InViewportClient, FVie
 		return false; 
 	}
 	
-	const TArray<FName>& SelectedGoalNames = Controller->GetSelectedGoals();
+	TArray<FName> SelectedGoalNames;
+	Controller->SkeletonView->GetSelectedGoalNames(SelectedGoalNames);
 	if (SelectedGoalNames.IsEmpty())
 	{
 		return false; // no goals selected to manipulate
@@ -236,7 +434,8 @@ bool FIKRigEditMode::InputDelta(FEditorViewportClient* InViewportClient, FViewpo
 		return false; // not handled
 	}
 
-	const TArray<FName>& SelectedGoalNames = Controller->GetSelectedGoals();
+	TArray<FName> SelectedGoalNames;
+	Controller->SkeletonView->GetSelectedGoalNames(SelectedGoalNames);
 	UIKRigController* AssetController = Controller->AssetController;
 
 	// translate goals
@@ -274,13 +473,14 @@ bool FIKRigEditMode::GetCustomDrawingCoordinateSystem(FMatrix& InMatrix, void* I
 		return false;
 	}
 	
-	const TArray<FName>& SelectedGoals = Controller->GetSelectedGoals();
-	if (SelectedGoals.IsEmpty())
+	TArray<FName> SelectedGoalNames;
+	Controller->SkeletonView->GetSelectedGoalNames(SelectedGoalNames);
+	if (SelectedGoalNames.IsEmpty())
 	{
 		return false; // nothing selected to manipulate
 	}
 
-	if (const UIKRigEffectorGoal* Goal = Controller->AssetController->GetGoal(SelectedGoals[0]))
+	if (const UIKRigEffectorGoal* Goal = Controller->AssetController->GetGoal(SelectedGoalNames[0]))
 	{
 		InMatrix = Goal->CurrentTransform.ToMatrixNoScale().RemoveTranslation();
 		return true;
@@ -309,15 +509,16 @@ bool FIKRigEditMode::InputKey(FEditorViewportClient* ViewportClient, FViewport*
 			return false;
 		}
 	
-		TArray<FName> SelectedGoals = Controller->GetSelectedGoals();
-		if (SelectedGoals.IsEmpty())
+		TArray<FName> SelectedGoalNames;
+		Controller->SkeletonView->GetSelectedGoalNames(SelectedGoalNames);
+		if (SelectedGoalNames.IsEmpty())
 		{
 			return false; // nothing selected to manipulate
 		}
 
-		for (const FName& GoalName : SelectedGoals)
+		for (const FName& GoalName : SelectedGoalNames)
 		{
-			Controller->DeleteGoal(GoalName);
+			Controller->AssetController->RemoveGoal(GoalName);
 		}
 
 		Controller->RefreshAllViews();
diff --git a/Engine/Plugins/Animation/IKRig/Source/IKRigEditor/Private/RigEditor/IKRigEditorController.cpp b/Engine/Plugins/Animation/IKRig/Source/IKRigEditor/Private/RigEditor/IKRigEditorController.cpp
--- a/Engine/Plugins/Animation/IKRig/Source/IKRigEditor/Private/RigEditor/IKRigEditorController.cpp
+++ b/Engine/Plugins/Animation/IKRig/Source/IKRigEditor/Private/RigEditor/IKRigEditorController.cpp
@@ -18,6 +18,7 @@ void FIKRigEditorController::Initialize(TSharedPtr<FIKRigEditorToolkit> Toolkit,
 {
 	EditorToolkit = Toolkit;
 	AssetController = UIKRigController::GetIKRigController(IKRigAsset);
+	BoneDetails = NewObject<UIKRigBoneDetails>();
 	
 	// register callback to be informed when rig asset is modified by editor
 	if (!AssetController->OnIKRigNeedsInitialized().IsBoundToObject(this))
@@ -105,66 +106,22 @@ void FIKRigEditorController::AddNewGoals(const TArray<FName>& GoalNames, const T
 	}
 }
 
-void FIKRigEditorController::DeleteGoal(const FName& GoalToDelete)
+void FIKRigEditorController::ClearSelection()
 {
-	AssetController->RemoveGoal(GoalToDelete);
-	SelectedGoals.Remove(GoalToDelete);
-}
-
-bool FIKRigEditorController::IsGoalSelected(const FName& GoalName) const
-{
-	return SelectedGoals.Contains(GoalName);
-}
-
-void FIKRigEditorController::ReplaceGoalInSelection(const FName& OldName, const FName& NewName)
-{
-	const int32 GoalIndex = SelectedGoals.IndexOfByKey(OldName);
-	if (GoalIndex == INDEX_NONE)
-	{
-		return;
-	}
-
-	SelectedGoals[GoalIndex] = NewName;
-}
-
-int32 FIKRigEditorController::GetNumSelectedGoals()
-{
-	return SelectedGoals.Num();
+	SkeletonView->TreeView->ClearSelection();
+	ShowEmptyDetails();
 }
 
 void FIKRigEditorController::HandleGoalSelectedInViewport(const FName& GoalName, bool bReplace)
-{
-	if (bReplace)
-	{
-		SelectedGoals.Reset();
-		if (GoalName != NAME_None)
-		{
-			SelectedGoals.Add(GoalName);
-		}
-	}
-	else
-	{
-		if (GoalName != NAME_None)
-		{
-			const bool bAlreadySelected = SelectedGoals.Contains(GoalName);
-			if (bAlreadySelected)
-			{
-				SelectedGoals.Remove(GoalName);	
-			}
-			else
-			{
-				SelectedGoals.Add(GoalName);
-			}
-		}	
-	}
-	
-	SkeletonView->SetSelectedGoalsFromViewport(SelectedGoals);
+{	
+	SkeletonView->AddSelectedItemFromViewport(GoalName, IKRigTreeElementType::GOAL, bReplace);
 	ShowDetailsForGoal(GoalName);
 }
 
-void FIKRigEditorController::HandleGoalsSelectedInTreeView(const TArray<FName>& GoalNames)
+void FIKRigEditorController::HandleBoneSelectedInViewport(const FName& BoneName, bool bReplace)
 {
-	SelectedGoals = GoalNames;
+	SkeletonView->AddSelectedItemFromViewport(BoneName, IKRigTreeElementType::BONE, bReplace);
+	ShowDetailsForBone(BoneName);
 }
 
 void FIKRigEditorController::GetSelectedSolvers(TArray<TSharedPtr<FSolverStackElement>>& OutSelectedSolvers)
@@ -205,15 +162,14 @@ bool FIKRigEditorController::PromptToAddSolver() const
 			SCustomDialog::FButton(LOCTEXT("Cancel", "Cancel"))
 	});
 
-	if (AddSolverDialog->ShowModal() == 1)
+	if (AddSolverDialog->ShowModal() != 0)
 	{
-		return false; // cancelled
+		return false; // cancel button pressed, or window closed
 	}
 
 	if (Settings.SolverType != nullptr)
 	{
-		AssetController->AddSolver(Settings.SolverType);
-		SolverStackView->RefreshStackView();
+		SolverStackView->AddNewSolver(Settings.SolverType);
 	}
 
 	return true;
@@ -282,7 +238,8 @@ bool FIKRigEditorController::IsElementExcludedBone(TSharedRef<FIKRigTreeElement>
 
 void FIKRigEditorController::ShowDetailsForBone(const FName BoneName)
 {
-	ShowEmptyDetails();
+	BoneDetails->SetBone(BoneName);
+	DetailsView->SetObject(BoneDetails);
 }
 
 void FIKRigEditorController::ShowDetailsForBoneSettings(const FName BoneName, int32 SolverIndex)
@@ -350,6 +307,16 @@ void FIKRigEditorController::PlayAnimationAsset(UAnimationAsset* AssetToPlay)
 	}
 }
 
+EIKRigSelectionType FIKRigEditorController::GetLastSelectedType() const
+{
+	return LastSelectedType;
+}
+
+void FIKRigEditorController::SetLastSelectedType(EIKRigSelectionType SelectionType)
+{
+	LastSelectedType = SelectionType;
+}
+
 void FIKRigEditorController::InitializeSolvers() const
 {
 	if (AssetController)
diff --git a/Engine/Plugins/Animation/IKRig/Source/IKRigEditor/Private/RigEditor/IKRigToolkit.cpp b/Engine/Plugins/Animation/IKRig/Source/IKRigEditor/Private/RigEditor/IKRigToolkit.cpp
--- a/Engine/Plugins/Animation/IKRig/Source/IKRigEditor/Private/RigEditor/IKRigToolkit.cpp
+++ b/Engine/Plugins/Animation/IKRig/Source/IKRigEditor/Private/RigEditor/IKRigToolkit.cpp
@@ -173,11 +173,13 @@ TStatId FIKRigEditorToolkit::GetStatId() const
 
 void FIKRigEditorToolkit::PostUndo(bool bSuccess)
 {
+	EditorController->AssetController->BroadcastNeedsReinitialized();
 	EditorController->RefreshAllViews();
 }
 
 void FIKRigEditorToolkit::PostRedo(bool bSuccess)
 {
+	EditorController->AssetController->BroadcastNeedsReinitialized();
 	EditorController->RefreshAllViews();
 }
 
diff --git a/Engine/Plugins/Animation/IKRig/Source/IKRigEditor/Private/RigEditor/SIKRigRetargetChainList.cpp b/Engine/Plugins/Animation/IKRig/Source/IKRigEditor/Private/RigEditor/SIKRigRetargetChainList.cpp
--- a/Engine/Plugins/Animation/IKRig/Source/IKRigEditor/Private/RigEditor/SIKRigRetargetChainList.cpp
+++ b/Engine/Plugins/Animation/IKRig/Source/IKRigEditor/Private/RigEditor/SIKRigRetargetChainList.cpp
@@ -305,6 +305,16 @@ void SIKRigRetargetChainList::Construct(const FArguments& InArgs, TSharedRef<FIK
 	RefreshView();
 }
 
+FName SIKRigRetargetChainList::GetSelectedChain()
+{
+	TArray<TSharedPtr<FRetargetChainElement>> SelectedItems = ListView->GetSelectedItems();
+	if (SelectedItems.IsEmpty())
+	{
+		return NAME_None;
+	}
+	return SelectedItems[0].Get()->ChainName;
+}
+
 bool SIKRigRetargetChainList::IsAddChainEnabled() const
 {
 	const TSharedPtr<FIKRigEditorController> Controller = EditorController.Pin();
@@ -365,7 +375,7 @@ TSharedRef<ITableRow> SIKRigRetargetChainList::MakeListRowWidget(
 
 void SIKRigRetargetChainList::OnItemClicked(TSharedPtr<FRetargetChainElement> InItem)
 {
-	// TODO highlight chain in skeleton view
+	EditorController.Pin()->SetLastSelectedType(EIKRigSelectionType::RetargetChains);
 }
 
 FReply SIKRigRetargetChainList::OnKeyDown(const FGeometry& MyGeometry, const FKeyEvent& InKeyEvent)
diff --git a/Engine/Plugins/Animation/IKRig/Source/IKRigEditor/Private/RigEditor/SIKRigSkeleton.cpp b/Engine/Plugins/Animation/IKRig/Source/IKRigEditor/Private/RigEditor/SIKRigSkeleton.cpp
--- a/Engine/Plugins/Animation/IKRig/Source/IKRigEditor/Private/RigEditor/SIKRigSkeleton.cpp
+++ b/Engine/Plugins/Animation/IKRig/Source/IKRigEditor/Private/RigEditor/SIKRigSkeleton.cpp
@@ -186,7 +186,8 @@ void SIKRigSkeletonItem::OnNameCommitted(const FText& InText, ETextCommit::Type
 	{
 		return; 
 	}
-	
+
+	const FText OldText = WeakRigTreeElement.Pin()->Key;
 	const FName OldName = WeakRigTreeElement.Pin()->GoalName;
 	const FName PotentialNewName = FName(InText.ToString());
 	const FName NewName = Controller->AssetController->RenameGoal(OldName, PotentialNewName);
@@ -195,8 +196,8 @@ void SIKRigSkeletonItem::OnNameCommitted(const FText& InText, ETextCommit::Type
 		WeakRigTreeElement.Pin()->Key = FText::FromName(NewName);
 		WeakRigTreeElement.Pin()->GoalName = NewName;
 	}
-
-	Controller->ReplaceGoalInSelection(OldName, NewName);
+	
+	Controller->SkeletonView->ReplaceItemInSelection(OldText, WeakRigTreeElement.Pin()->Key);
 	Controller->RefreshAllViews();
 }
 
@@ -262,31 +263,105 @@ void SIKRigSkeleton::Construct(const FArguments& InArgs, TSharedRef<FIKRigEditor
 	RefreshTreeView(IsInitialSetup);
 }
 
-void SIKRigSkeleton::SetSelectedGoalsFromViewport(const TArray<FName>& GoalNames)
-{
-	if (GoalNames.IsEmpty())
+void SIKRigSkeleton::AddSelectedItemFromViewport(
+	const FName& ItemName,
+	IKRigTreeElementType ItemType,
+	const bool bReplace)
+{	
+	// nothing to add
+	if (ItemName == NAME_None)
 	{
-		TreeView->ClearSelection();
 		return;
 	}
-	
+
+	// record what was already selected
+	TArray<TSharedPtr<FIKRigTreeElement>> PreviouslySelectedItems = TreeView->GetSelectedItems();
+	// add/remove items as needed
 	for (const TSharedPtr<FIKRigTreeElement>& Item : AllElements)
 	{
-		if (GoalNames.Contains(Item->GoalName))
+		bool bIsBeingAdded = false;
+		switch (ItemType)
 		{
-			TreeView->SetSelection(Item, ESelectInfo::Direct);
+		case IKRigTreeElementType::GOAL:
+			if (ItemName == Item->GoalName)
+			{
+				bIsBeingAdded = true;
+			}
+			break;
+		case IKRigTreeElementType::BONE:
+			if (ItemName == Item->BoneName)
+			{
+				bIsBeingAdded = true;
+			}
+			break;
+		default:
+			ensureMsgf(false, TEXT("IKRig cannot select anything but bones and goals in viewport."));
+			return;
+		}
 
-			if(GetDefault<UPersonaOptions>()->bExpandTreeOnSelection)
+		if (bReplace)
+		{
+			if (bIsBeingAdded)
 			{
-				TSharedPtr<FIKRigTreeElement> ItemToExpand = Item->Parent;
-				while(ItemToExpand.IsValid())
-				{
-					TreeView->SetItemExpansion(ItemToExpand, true);
-					ItemToExpand = ItemToExpand->Parent;
-				}
+				TreeView->ClearSelection();
+				AddItemToSelection(Item);
+				return;
 			}
+			
+			continue;
+		}
+
+		// remove if already selected (invert)
+		if (bIsBeingAdded && PreviouslySelectedItems.Contains(Item))
+		{
+			RemoveItemFromSelection(Item);
+			continue;
+		}
+
+		// add if being added
+		if (bIsBeingAdded)
+		{
+			AddItemToSelection(Item);
+			continue;
+		}
+	}
+}
 
-			TreeView->RequestScrollIntoView(Item);
+void SIKRigSkeleton::AddItemToSelection(const TSharedPtr<FIKRigTreeElement>& InItem)
+{
+	TreeView->SetItemSelection(InItem, true, ESelectInfo::Direct);
+    
+	if(GetDefault<UPersonaOptions>()->bExpandTreeOnSelection)
+	{
+		TSharedPtr<FIKRigTreeElement> ItemToExpand = InItem->Parent;
+		while(ItemToExpand.IsValid())
+		{
+			TreeView->SetItemExpansion(ItemToExpand, true);
+			ItemToExpand = ItemToExpand->Parent;
+		}
+	}
+    
+	TreeView->RequestScrollIntoView(InItem);
+}
+
+void SIKRigSkeleton::RemoveItemFromSelection(const TSharedPtr<FIKRigTreeElement>& InItem)
+{
+	TreeView->SetItemSelection(InItem, false, ESelectInfo::Direct);
+}
+
+void SIKRigSkeleton::ReplaceItemInSelection(const FText& OldName, const FText& NewName)
+{
+	for (const TSharedPtr<FIKRigTreeElement>& Item : AllElements)
+	{
+		// remove old selection
+		if (Item->Key.CompareTo(OldName))
+		{
+			TreeView->SetItemSelection(Item, false, ESelectInfo::Direct);
+		}
+		// add new selection
+		if (Item->Key.CompareTo(NewName))
+		{
+			TreeView->SetItemSelection(Item, true, ESelectInfo::Direct);
 		}
 	}
 }
@@ -475,7 +550,7 @@ void SIKRigSkeleton::HandleDeleteGoal()
 	{
 		if (Item->ElementType == IKRigTreeElementType::GOAL)
 		{
-			Controller->DeleteGoal(Item->GoalName);
+			Controller->AssetController->RemoveGoal(Item->GoalName);
 		}
 		else if (Item->ElementType == IKRigTreeElementType::SOLVERGOAL)
 		{
@@ -1006,6 +1081,18 @@ bool SIKRigSkeleton::CanSetRetargetRoot()
 	return !SelectedBones.IsEmpty();
 }
 
+bool SIKRigSkeleton::IsBoneInSelection(TArray<TSharedPtr<FIKRigTreeElement>>& SelectedBoneItems, const FName& BoneName)
+{
+	for (const TSharedPtr<FIKRigTreeElement>& Item : SelectedBoneItems)
+	{
+		if (Item->BoneName == BoneName)
+		{
+			return true;
+		}
+	}
+	return false;
+}
+
 void SIKRigSkeleton::GetSelectedBones(TArray<TSharedPtr<FIKRigTreeElement>>& OutBoneItems) const
 {
 	const TArray<TSharedPtr<FIKRigTreeElement>> SelectedItems = TreeView->GetSelectedItems();
@@ -1031,6 +1118,38 @@ void SIKRigSkeleton::GetSelectedGoals(TArray<TSharedPtr<FIKRigTreeElement>>& Out
 	}
 }
 
+int32 SIKRigSkeleton::GetNumSelectedGoals()
+{
+	TArray<TSharedPtr<FIKRigTreeElement>> OutSelectedGoals;
+	GetSelectedGoals(OutSelectedGoals);
+	return OutSelectedGoals.Num();
+}
+
+void SIKRigSkeleton::GetSelectedGoalNames(TArray<FName>& OutSelectedGoalNames) const
+{
+	TArray<TSharedPtr<FIKRigTreeElement>> OutSelectedGoals;
+	GetSelectedGoals(OutSelectedGoals);
+	OutSelectedGoalNames.Reset();
+	for (TSharedPtr<FIKRigTreeElement> SelectedGoalItem : OutSelectedGoals)
+	{
+		OutSelectedGoalNames.Add(SelectedGoalItem->GoalName);
+	}
+}
+
+bool SIKRigSkeleton::IsGoalSelected(const FName& GoalName)
+{
+	TArray<TSharedPtr<FIKRigTreeElement>> OutSelectedGoals;
+	GetSelectedGoals(OutSelectedGoals);
+	for (TSharedPtr<FIKRigTreeElement> SelectedGoalItem : OutSelectedGoals)
+	{
+		if (SelectedGoalItem->GoalName == GoalName)
+		{
+			return true;
+		}
+	}
+	return false;
+}
+
 void SIKRigSkeleton::HandleRenameElement() const
 {
 	TArray<TSharedPtr<FIKRigTreeElement>> SelectedGoals;
@@ -1195,24 +1314,6 @@ void SIKRigSkeleton::HandleGetChildrenForTree(
 
 void SIKRigSkeleton::OnSelectionChanged(TSharedPtr<FIKRigTreeElement> Selection, ESelectInfo::Type SelectInfo)
 {
-	const TSharedPtr<FIKRigEditorController> Controller = EditorController.Pin();
-	if (!Controller.IsValid())
-	{
-		return;
-	}
-	
-	// gate any selection changes NOT made by user clicking mouse
-	if (SelectInfo == ESelectInfo::OnMouseClick)
-	{
-		TArray<TSharedPtr<FIKRigTreeElement>> SelectedGoals;
-		GetSelectedGoals(SelectedGoals);
-		TArray<FName> SelectedGoalNames;
-		for (const TSharedPtr<FIKRigTreeElement>& Goal : SelectedGoals)
-		{
-			SelectedGoalNames.Add(Goal->GoalName);
-		}
-		Controller->HandleGoalsSelectedInTreeView(SelectedGoalNames);
-	}
 }
 
 TSharedPtr<SWidget> SIKRigSkeleton::CreateContextMenu()
@@ -1264,6 +1365,7 @@ void SIKRigSkeleton::OnItemClicked(TSharedPtr<FIKRigTreeElement> InItem)
 
 	TreeView->LastClickCycles = CurrentCycles;
 	TreeView->LastSelected = InItem;
+	Controller->SetLastSelectedType(EIKRigSelectionType::Hierarchy);
 }
 
 void SIKRigSkeleton::OnItemDoubleClicked(TSharedPtr<FIKRigTreeElement> InItem)
@@ -1399,7 +1501,7 @@ FReply SIKRigSkeleton::OnKeyDown(const FGeometry& MyGeometry, const FKeyEvent& I
 			switch(SelectedItem->ElementType)
 			{
 				case IKRigTreeElementType::GOAL:
-					Controller->DeleteGoal(SelectedItem->GoalName);
+					Controller->AssetController->RemoveGoal(SelectedItem->GoalName);
 					break;
 				case IKRigTreeElementType::SOLVERGOAL:
 					Controller->AssetController->DisconnectGoalFromSolver(SelectedItem->SolverGoalName, SelectedItem->SolverGoalIndex);
diff --git a/Engine/Plugins/Animation/IKRig/Source/IKRigEditor/Private/RigEditor/SIKRigSolverStack.cpp b/Engine/Plugins/Animation/IKRig/Source/IKRigEditor/Private/RigEditor/SIKRigSolverStack.cpp
--- a/Engine/Plugins/Animation/IKRig/Source/IKRigEditor/Private/RigEditor/SIKRigSolverStack.cpp
+++ b/Engine/Plugins/Animation/IKRig/Source/IKRigEditor/Private/RigEditor/SIKRigSolverStack.cpp
@@ -102,6 +102,10 @@ void SIKRigSolverStackItem::Construct(
 					.Text(InStackElement->DisplayName)
 					.IsEnabled_Lambda([this]()
 					{
+						if (!IsSolverEnabled())
+						{
+							return false;
+						}
 						FText Warning;
 						return !GetWarningMessage(Warning);
 					})
@@ -151,27 +155,36 @@ void SIKRigSolverStackItem::Construct(
 
 bool SIKRigSolverStackItem::GetWarningMessage(FText& Message)
 {
-	if (!(StackElement.IsValid() && SolverStack.IsValid()))
+	if (UIKRigSolver* Solver = GetSolver())
 	{
-		return false;
+		return Solver->GetWarningMessage(Message);
 	}
-	if (!SolverStack.Pin()->EditorController.IsValid())
+	
+	return false;
+}
+
+bool SIKRigSolverStackItem::IsSolverEnabled() const
+{
+	if (const UIKRigSolver* Solver = GetSolver())
 	{
-		return false;
+		return Solver->IsEnabled();
 	}
-	const int32 SolverIndex = StackElement.Pin()->IndexInStack;
-	const UIKRigSolver* Solver = SolverStack.Pin()->EditorController.Pin()->AssetController->GetSolver(SolverIndex);
-	if (!Solver)
+	
+	return false;
+}
+
+UIKRigSolver* SIKRigSolverStackItem::GetSolver() const
+{
+	if (!(StackElement.IsValid() && SolverStack.IsValid()))
 	{
-		return false;
+		return nullptr;
 	}
-
-	if (Solver->GetWarningMessage(Message))
+	if (!SolverStack.Pin()->EditorController.IsValid())
 	{
-		return true;
+		return nullptr;
 	}
-					
-	return false;
+	const int32 SolverIndex = StackElement.Pin()->IndexInStack;
+	return SolverStack.Pin()->EditorController.Pin()->AssetController->GetSolver(SolverIndex);
 }
 
 TSharedRef<FIKRigSolverStackDragDropOp> FIKRigSolverStackDragDropOp::New(TWeakPtr<FSolverStackElement> InElement)
@@ -250,6 +263,7 @@ void SIKRigSolverStack::Construct(const FArguments& InArgs, TSharedRef<FIKRigEdi
 			.ListItemsSource( &ListViewItems )
 			.OnGenerateRow( this, &SIKRigSolverStack::MakeListRowWidget )
 			.OnSelectionChanged(this, &SIKRigSolverStack::OnSelectionChanged)
+			.OnMouseButtonClick(this, &SIKRigSolverStack::OnItemClicked)
         ]
     ];
 
@@ -359,7 +373,7 @@ void SIKRigSolverStack::RefreshStackView()
 	}
 
 	// record/restore selection
-	int32 IndexToSelect = 0; // default to first item
+	int32 IndexToSelect = -1; // default to nothing selected
 	const TArray<TSharedPtr<FSolverStackElement>> SelectedItems = ListView.Get()->GetSelectedItems();
 	if (!SelectedItems.IsEmpty())
 	{
@@ -379,7 +393,7 @@ void SIKRigSolverStack::RefreshStackView()
 	}
 
 	// restore selection
-	if (IndexToSelect < ListViewItems.Num())
+	if (ListViewItems.IsValidIndex(IndexToSelect))
 	{
 		ListView->SetSelection(ListViewItems[IndexToSelect]);
 	}
@@ -418,6 +432,17 @@ FReply SIKRigSolverStack::OnDragDetected(
 }
 
 void SIKRigSolverStack::OnSelectionChanged(TSharedPtr<FSolverStackElement> InItem, ESelectInfo::Type SelectInfo)
+{
+	ShowDetailsForItem(InItem);
+}
+
+void SIKRigSolverStack::OnItemClicked(TSharedPtr<FSolverStackElement> InItem)
+{
+	ShowDetailsForItem(InItem);
+	EditorController.Pin()->SetLastSelectedType(EIKRigSelectionType::SolverStack);
+}
+
+void SIKRigSolverStack::ShowDetailsForItem(TSharedPtr<FSolverStackElement> InItem)
 {
 	const TSharedPtr<FIKRigEditorController> Controller = EditorController.Pin();
 	if (!Controller.IsValid())
diff --git a/Engine/Plugins/Animation/IKRig/Source/IKRigEditor/Public/IKRigDebugRendering.h b/Engine/Plugins/Animation/IKRig/Source/IKRigEditor/Public/IKRigDebugRendering.h
new file mode 100644
--- /dev/null
+++ b/Engine/Plugins/Animation/IKRig/Source/IKRigEditor/Public/IKRigDebugRendering.h
@@ -0,0 +1,109 @@
+﻿// Copyright Epic Games, Inc. All Rights Reserved.
+
+#pragma once
+
+#include "IKRigDefinition.h"
+#include "SkeletalDebugRendering.h"
+
+namespace IKRigDebugRendering
+{
+	static const FLinearColor DESELECTED_BONE_COLOR = FLinearColor(0.0f,0.0f,0.025f,1.0f);
+	static const FLinearColor SELECTED_BONE_COLOR = FLinearColor(0.2f,1.0f,0.2f,1.0f);
+	static const FLinearColor AFFECTED_BONE_COLOR = FLinearColor(1.0f,1.0f,1.0f,1.0f);
+	
+	static TArray BoxPoints =
+	{
+		FVector(0.5f, 0.5f, 0.5f),
+		FVector(0.5f, -0.5f, 0.5f),
+		FVector(0.5f, -0.5f, 0.5f),
+		FVector(-0.5f, -0.5f, 0.5f),
+		FVector(-0.5f, -0.5f, 0.5f),
+		FVector(-0.5f, 0.5f, 0.5f),
+		FVector(-0.5f, 0.5f, 0.5f),
+		FVector(0.5f, 0.5f, 0.5f),
+
+		FVector(0.5f, 0.5f, -0.5f),
+		FVector(0.5f, -0.5f, -0.5f),
+		FVector(0.5f, -0.5f, -0.5f),
+		FVector(-0.5f, -0.5f, -0.5f),
+		FVector(-0.5f, -0.5f, -0.5f),
+		FVector(-0.5f, 0.5f, -0.5f),
+		FVector(-0.5f, 0.5f, -0.5f),
+		FVector(0.5f, 0.5f, -0.5f),
+
+		FVector(0.5f, 0.5f, 0.5f),
+		FVector(0.5f, 0.5f, -0.5f),
+		FVector(0.5f, -0.5f, 0.5f),
+		FVector(0.5f, -0.5f, -0.5f),
+		FVector(-0.5f, -0.5f, 0.5f),
+		FVector(-0.5f, -0.5f, -0.5f),
+		FVector(-0.5f, 0.5f, 0.5f),
+		FVector(-0.5f, 0.5f, -0.5f),
+	};
+
+	static void DrawGoal(
+		FPrimitiveDrawInterface* PDI,
+		const UIKRigEffectorGoal* Goal,
+		bool bIsSelected,
+		float Size,
+		float Thickness)
+	{
+		const FLinearColor Color = bIsSelected ? FLinearColor::Green : FLinearColor::Yellow;
+		const float Scale = FMath::Clamp(Size, 0.1f, 1000.0f);
+		const FTransform Transform = Goal->CurrentTransform;
+		for (int32 PointIndex = 0; PointIndex < BoxPoints.Num() - 1; PointIndex += 2)
+		{
+			PDI->DrawLine(Transform.TransformPosition(BoxPoints[PointIndex] * Scale), Transform.TransformPosition(BoxPoints[PointIndex + 1] * Scale), Color, SDPG_Foreground, Thickness);
+		}
+	}
+
+	// todo - refactor this and fix SkeletalDebugRendering
+	// this is copied/adapted from SkeletalDebugRendering because annoyingly the bone rendering puts the joint sphere at the
+	// END of the bone. This is so that bones can be rendered individually (needing only start/end point pairs).
+	// Which is easier, but doesn't make sense. We need the bones to render from parent to each immediate child.
+	// This causes several issues:
+	// 1. Bones lack a sphere at the root.
+	// 2. Selected bones render to their parent, rather than to their children.
+	// Also, this function draws the cone BETWEEN the joint spheres, not through them which is a cleaner end result.
+	// 
+	static void DrawWireBone(
+		FPrimitiveDrawInterface* PDI,
+		const FTransform& InBoneTransform,
+		const TArray<FVector>& InChildLocations,
+		const FLinearColor& InColor,
+		ESceneDepthPriorityGroup InDepthPriority,
+		const float SphereRadius,
+		const bool bDrawAxes)
+	{
+	#if ENABLE_DRAW_DEBUG
+		static const int32 NumSphereSides = 10;
+		static const int32 NumConeSides = 4;
+
+		const FVector BoneLocation = InBoneTransform.GetLocation();
+
+		// render sphere at joint origin
+		DrawWireSphere(PDI, BoneLocation, InColor, SphereRadius, NumSphereSides, InDepthPriority, 0.0f, 1.0f);
+		// draw axes at joint location
+		if (bDrawAxes)
+		{
+			SkeletalDebugRendering::DrawAxes(PDI, InBoneTransform, SDPG_Foreground, 0.f , SphereRadius);
+		}
+
+		// draw wire cones to each child
+		for (const FVector& ChildPoint : InChildLocations)
+		{
+			// offset start/end based on bone radius
+			const FVector RadiusOffset = (ChildPoint - BoneLocation).GetSafeNormal() * SphereRadius;
+			const FVector Start = BoneLocation + RadiusOffset;
+			const FVector End = ChildPoint - RadiusOffset;
+			
+			// calc cone size
+			const FVector EndToStart = (Start - End);
+			const float ConeLength = EndToStart.Size();
+			const float Angle = FMath::RadiansToDegrees(FMath::Atan(SphereRadius / ConeLength));
+			TArray<FVector> Verts;
+			DrawWireCone(PDI, Verts, FRotationMatrix::MakeFromX(EndToStart) * FTranslationMatrix(End), ConeLength, Angle, NumConeSides, InColor, InDepthPriority, 0.0f, 1.0f);
+		}
+	#endif
+	}
+}
diff --git a/Engine/Plugins/Animation/IKRig/Source/IKRigEditor/Public/RetargetEditor/IKRetargetEditMode.h b/Engine/Plugins/Animation/IKRig/Source/IKRigEditor/Public/RetargetEditor/IKRetargetEditMode.h
--- a/Engine/Plugins/Animation/IKRig/Source/IKRigEditor/Public/RetargetEditor/IKRetargetEditMode.h
+++ b/Engine/Plugins/Animation/IKRig/Source/IKRigEditor/Public/RetargetEditor/IKRetargetEditMode.h
@@ -7,6 +7,7 @@
 #include "Retargeter/IKRetargeter.h"
 #include "IPersonaEditMode.h"
 
+class UIKRigProcessor;
 class FIKRetargetEditorController;
 class FIKRetargetEditor;
 class FIKRetargetPreviewScene;
@@ -18,6 +19,19 @@ enum FIKRetargetTrackingState : int8
 	TranslatingRoot,
 };
 
+struct BoneEdit
+{
+	FName Name;							// name of last selected bone
+	int32 Index;						// index of last selected bone
+	FTransform ParentGlobalTransform;	// global transform of parent of last selected bone
+	FTransform GlobalTransform;			// global transform of last selected bone
+	FTransform LocalTransform;			// local transform of last selected bone
+	FQuat AccumulatedGlobalOffset;		// the accumulated offset from rotation gizmo
+	
+	TArray<FQuat> PrevLocalOffsets;		// the prev stored local offsets of all selected bones
+	TArray<FName> SelectedBones;		// the currently selected bones in the viewport
+};
+
 class FIKRetargetEditMode : public IPersonaEditMode
 {
 public:
@@ -54,27 +68,25 @@ public:
 
 private:
 
+	void GetAffectedBones(
+		FIKRetargetEditorController* Controller,
+		UIKRigProcessor* Processor,
+		TSet<int32>& OutAffectedBones,
+		TSet<int32>& OutSelectedBones) const;
+
 	UE::Widget::EWidgetMode CurrentWidgetMode;
-	
+
+	bool IsRootSelected() const;
 	bool IsOnlyRootSelected() const;
-	
 	bool IsBoneSelected(const FName& BoneName) const;
-	
+
+	BoneEdit BoneEdit;
+	void UpdateWidgetTransform();
 	void HandleBoneSelectedInViewport(const FName& BoneName, bool bReplace);
 	
 	/** The hosting app */
 	TWeakPtr<FIKRetargetEditorController> EditorController;
 
-	void DrawBoneProxy(
-		FPrimitiveDrawInterface* PDI,
-		const FTransform& BoneTransform,
-		const FVector& Start,
-		const TArray<FVector>& ChildPoints,
-		float Size,
-		float Thickness,
-		bool bIsSelected) const;
-
 	/** viewport selection/editing state */
-	TArray<FName> SelectedBones;
 	FIKRetargetTrackingState TrackingState;
 };
diff --git a/Engine/Plugins/Animation/IKRig/Source/IKRigEditor/Public/RetargetEditor/IKRetargetEditor.h b/Engine/Plugins/Animation/IKRig/Source/IKRigEditor/Public/RetargetEditor/IKRetargetEditor.h
--- a/Engine/Plugins/Animation/IKRig/Source/IKRigEditor/Public/RetargetEditor/IKRetargetEditor.h
+++ b/Engine/Plugins/Animation/IKRig/Source/IKRigEditor/Public/RetargetEditor/IKRetargetEditor.h
@@ -31,7 +31,7 @@ class FIKRetargetEditor :
 	public FPersonaAssetEditorToolkit,
 	public IHasPersonaToolkit,
 	public FGCObject,
-	public FEditorUndoClient,
+	public FSelfRegisteringEditorUndoClient,
 	public FTickableEditorObject
 {
 public:
@@ -72,6 +72,11 @@ public:
 	virtual TSharedRef<IPersonaToolkit> GetPersonaToolkit() const override { return EditorController->PersonaToolkit.ToSharedRef(); }
 	/** END IHasPersonaToolkit interface */
 
+	/** FSelfRegisteringEditorUndoClient interface */
+	virtual void PostUndo( bool bSuccess );
+	virtual void PostRedo( bool bSuccess );
+	/** END FSelfRegisteringEditorUndoClient interface */
+
 	TSharedRef<FIKRetargetEditorController> GetController() const {return EditorController;};
 
 private:
diff --git a/Engine/Plugins/Animation/IKRig/Source/IKRigEditor/Public/RetargetEditor/IKRetargetEditorController.h b/Engine/Plugins/Animation/IKRig/Source/IKRigEditor/Public/RetargetEditor/IKRetargetEditorController.h
--- a/Engine/Plugins/Animation/IKRig/Source/IKRigEditor/Public/RetargetEditor/IKRetargetEditorController.h
+++ b/Engine/Plugins/Animation/IKRig/Source/IKRigEditor/Public/RetargetEditor/IKRetargetEditorController.h
@@ -63,7 +63,8 @@ public:
 	USkeletalMesh* GetTargetSkeletalMesh() const;
 
 	/** get current chain pose */
-	FTransform GetTargetBoneTransform(const int32& TargetBoneIndex) const;
+	FTransform GetTargetBoneGlobalTransform(const int32& TargetBoneIndex) const;
+	FTransform GetTargetBoneLocalTransform(const int32& TargetBoneIndex) const;
 	/** get the line segments to draw from this bone to each child */
 	bool GetTargetBoneLineSegments(const int32& TargetBoneIndex, FVector& OutStart, TArray<FVector>& OutChildren) const;
 
diff --git a/Engine/Plugins/Animation/IKRig/Source/IKRigEditor/Public/RetargetEditor/IKRetargeterController.h b/Engine/Plugins/Animation/IKRig/Source/IKRigEditor/Public/RetargetEditor/IKRetargeterController.h
--- a/Engine/Plugins/Animation/IKRig/Source/IKRigEditor/Public/RetargetEditor/IKRetargeterController.h
+++ b/Engine/Plugins/Animation/IKRig/Source/IKRigEditor/Public/RetargetEditor/IKRetargeterController.h
@@ -84,8 +84,10 @@ public:
 	void SetCurrentRetargetPose(FName CurrentPose) const;
 	/** Get read-only access to list of retarget poses */
 	const TMap<FName, FIKRetargetPose>& GetRetargetPoses();
-	/** Add a delta rotation to a given bone for the current retarget pose (used in Edit Mode in the retarget editor) */
-	void AddRotationOffsetToRetargetPoseBone(FName BoneName, FQuat RotationOffset) const;
+	/** Set a delta rotation for a given bone for the current retarget pose (used in Edit Mode in the retarget editor) */
+	void SetRotationOffsetForRetargetPoseBone(FName BoneName, FQuat RotationOffset) const;
+	/** Get a delta rotation for a given bone for the current retarget pose (used in Edit Mode in the retarget editor) */
+	FQuat GetRotationOffsetForRetargetPoseBone(FName BoneName) const;
 	/** Add a delta translation to the root bone (used in Edit Mode in the retarget editor) */
 	void AddTranslationOffsetToRetargetRootBone(FVector TranslationOffset) const;
 	/** Set whether to output retarget pose. Will output current retarget pose if true, or run retarget otherwise. */
@@ -101,16 +103,16 @@ private:
 	/** Called whenever the retargeter is modified in such a way that would require re-initialization by dependent systems.*/
 	DECLARE_MULTICAST_DELEGATE_OneParam(FOnRetargeterNeedsInitialized, const UIKRetargeter*);
 	FOnRetargeterNeedsInitialized RetargeterNeedsInitialized;
+	
+public:
+	
+	FOnRetargeterNeedsInitialized& OnRetargeterNeedsInitialized(){ return RetargeterNeedsInitialized; };
 
 	void BroadcastNeedsReinitialized() const
 	{
 		RetargeterNeedsInitialized.Broadcast(GetAsset());
 	}
 	
-public:
-	
-	FOnRetargeterNeedsInitialized& OnRetargeterNeedsInitialized(){ return RetargeterNeedsInitialized; };
-	
 private:
 
 	FRetargetChainMap* GetChainMap(const FName& TargetChainName) const;
diff --git a/Engine/Plugins/Animation/IKRig/Source/IKRigEditor/Public/RigEditor/IKRigAnimInstance.h b/Engine/Plugins/Animation/IKRig/Source/IKRigEditor/Public/RigEditor/IKRigAnimInstance.h
--- a/Engine/Plugins/Animation/IKRig/Source/IKRigEditor/Public/RigEditor/IKRigAnimInstance.h
+++ b/Engine/Plugins/Animation/IKRig/Source/IKRigEditor/Public/RigEditor/IKRigAnimInstance.h
@@ -21,6 +21,8 @@ public:
 
 	void SetProcessorNeedsInitialized();
 
+	UIKRigProcessor* GetCurrentlyRunningProcessor();
+
 protected:
 	virtual FAnimInstanceProxy* CreateAnimInstanceProxy() override;
 
diff --git a/Engine/Plugins/Animation/IKRig/Source/IKRigEditor/Public/RigEditor/IKRigController.h b/Engine/Plugins/Animation/IKRig/Source/IKRigEditor/Public/RigEditor/IKRigController.h
--- a/Engine/Plugins/Animation/IKRig/Source/IKRigEditor/Public/RigEditor/IKRigController.h
+++ b/Engine/Plugins/Animation/IKRig/Source/IKRigEditor/Public/RigEditor/IKRigController.h
@@ -169,6 +169,12 @@ public:
 	void SortRetargetChains() const;
 	/** END retarget chains */
 
+	// force all currently connected processors to reinitialize using latest asset state
+	void BroadcastNeedsReinitialized() const
+	{
+		IKRigNeedsInitialized.Broadcast(GetAsset());
+	}
+
 private:
 
 	/** Called whenever the rig is modified in such a way that would require re-initialization by dependent systems.*/
@@ -179,11 +185,6 @@ private:
 	DECLARE_MULTICAST_DELEGATE_ThreeParams(FOnRetargetChainRenamed, UIKRigDefinition*, FName /*old name*/, FName /*new name*/);
 	FOnRetargetChainRenamed RetargetChainRenamed;
 
-	void BroadcastNeedsReinitialized() const
-	{
-		IKRigNeedsInitialized.Broadcast(GetAsset());
-	}
-	
 public:
 	
 	FOnIKRigNeedsInitialized& OnIKRigNeedsInitialized(){ return IKRigNeedsInitialized; };
diff --git a/Engine/Plugins/Animation/IKRig/Source/IKRigEditor/Public/RigEditor/IKRigEditMode.h b/Engine/Plugins/Animation/IKRig/Source/IKRigEditor/Public/RigEditor/IKRigEditMode.h
--- a/Engine/Plugins/Animation/IKRig/Source/IKRigEditor/Public/RigEditor/IKRigEditMode.h
+++ b/Engine/Plugins/Animation/IKRig/Source/IKRigEditor/Public/RigEditor/IKRigEditMode.h
@@ -6,67 +6,19 @@
 
 #include "IKRigDefinition.h"
 #include "IPersonaEditMode.h"
+#include "SkeletalDebugRendering.h"
 
 class FIKRigEditorController;
 class UIKRigEffectorGoal;
 class FIKRigEditorToolkit;
 class FIKRigPreviewScene;
-
-struct GoalGizmo
-{
-	GoalGizmo()
-	{
-		// wireframe cube lines (point pairs)
-		BoxPoints.Add(FVector(0.5f, 0.5f, 0.5f));
-		BoxPoints.Add(FVector(0.5f, -0.5f, 0.5f));
-		BoxPoints.Add(FVector(0.5f, -0.5f, 0.5f));
-		BoxPoints.Add(FVector(-0.5f, -0.5f, 0.5f));
-		BoxPoints.Add(FVector(-0.5f, -0.5f, 0.5f));
-		BoxPoints.Add(FVector(-0.5f, 0.5f, 0.5f));
-		BoxPoints.Add(FVector(-0.5f, 0.5f, 0.5f));
-		BoxPoints.Add(FVector(0.5f, 0.5f, 0.5f));
-
-		BoxPoints.Add(FVector(0.5f, 0.5f, -0.5f));
-		BoxPoints.Add(FVector(0.5f, -0.5f, -0.5f));
-		BoxPoints.Add(FVector(0.5f, -0.5f, -0.5f));
-		BoxPoints.Add(FVector(-0.5f, -0.5f, -0.5f));
-		BoxPoints.Add(FVector(-0.5f, -0.5f, -0.5f));
-		BoxPoints.Add(FVector(-0.5f, 0.5f, -0.5f));
-		BoxPoints.Add(FVector(-0.5f, 0.5f, -0.5f));
-		BoxPoints.Add(FVector(0.5f, 0.5f, -0.5f));
-
-		BoxPoints.Add(FVector(0.5f, 0.5f, 0.5f));
-		BoxPoints.Add(FVector(0.5f, 0.5f, -0.5f));
-		BoxPoints.Add(FVector(0.5f, -0.5f, 0.5f));
-		BoxPoints.Add(FVector(0.5f, -0.5f, -0.5f));
-		BoxPoints.Add(FVector(-0.5f, -0.5f, 0.5f));
-		BoxPoints.Add(FVector(-0.5f, -0.5f, -0.5f));
-		BoxPoints.Add(FVector(-0.5f, 0.5f, 0.5f));
-		BoxPoints.Add(FVector(-0.5f, 0.5f, -0.5f));
-	}
-
-	void DrawGoal(
-		FPrimitiveDrawInterface* PDI,
-		const UIKRigEffectorGoal* Goal,
-		bool bIsSelected,
-		float Size,
-		float Thickness) const
-	{
-		const FLinearColor Color = bIsSelected ? FLinearColor::Green : FLinearColor::Yellow;
-		const float Scale = FMath::Clamp(Size, 0.1f, 1000.0f);
-		const FTransform Transform = Goal->CurrentTransform;
-		for (int32 PointIndex = 0; PointIndex < BoxPoints.Num() - 1; PointIndex += 2)
-		{
-			PDI->DrawLine(Transform.TransformPosition(BoxPoints[PointIndex] * Scale), Transform.TransformPosition(BoxPoints[PointIndex + 1] * Scale), Color, SDPG_Foreground, Thickness);
-		}
-	}
-
-	TArray<FVector> BoxPoints;
-};
+class UIKRigProcessor;
 
 class FIKRigEditMode : public IPersonaEditMode
 {
+	
 public:
+	
 	static FName ModeName;
 	
 	FIKRigEditMode();
@@ -100,9 +52,15 @@ public:
 	/** END FEdMode interface */
 
 private:
+	
+	void RenderGoals(FPrimitiveDrawInterface* PDI);
+	void RenderBones(FPrimitiveDrawInterface* PDI);
+	void GetAffectedBones(
+		FIKRigEditorController* Controller,
+		UIKRigProcessor* Processor,
+		TSet<int32>& OutAffectedBones,
+		TSet<int32>& OutSelectedBones) const;
+	
 	/** The hosting app */
 	TWeakPtr<FIKRigEditorController> EditorController;
-
-	/** draws goals in viewport */
-	GoalGizmo GoalDrawer;
 };
diff --git a/Engine/Plugins/Animation/IKRig/Source/IKRigEditor/Public/RigEditor/IKRigEditorController.h b/Engine/Plugins/Animation/IKRig/Source/IKRigEditor/Public/RigEditor/IKRigEditorController.h
--- a/Engine/Plugins/Animation/IKRig/Source/IKRigEditor/Public/RigEditor/IKRigEditorController.h
+++ b/Engine/Plugins/Animation/IKRig/Source/IKRigEditor/Public/RigEditor/IKRigEditorController.h
@@ -18,8 +18,38 @@ class FSolverStackElement;
 class FIKRigTreeElement;
 class UDebugSkelMeshComponent;
 
+enum EIKRigSelectionType : int8
+{
+	Hierarchy,
+	SolverStack,
+	RetargetChains,
+};
+
+UCLASS(config = Engine, hidecategories = UObject)
+class IKRIGEDITOR_API UIKRigBoneDetails : public UObject
+{
+	GENERATED_BODY()
+
+public:
+	
+	// todo update bone info automatically using something else
+	void SetBone(const FName& BoneName)
+	{
+		SelectedBone = BoneName;
+	};
+
+	UPROPERTY(VisibleAnywhere, Category = "Bone Transforms")
+	FName SelectedBone;
+	
+	UPROPERTY(VisibleAnywhere, Category = "Bone Transforms")
+	FTransform CurrentTransform;
+
+	UPROPERTY(VisibleAnywhere, Category = "Bone Transforms")
+	FTransform DefaultTransform;
+};
+
 /** a home for cross-widget communication to synchronize state across all tabs and viewport */
-class FIKRigEditorController : public TSharedFromThis<FIKRigEditorController>
+class FIKRigEditorController : public TSharedFromThis<FIKRigEditorController>, FGCObject
 {
 public:
 
@@ -31,20 +61,12 @@ public:
 
 	/** create goals */
 	void AddNewGoals(const TArray<FName>& GoalNames, const TArray<FName>& BoneNames);
-	/** delete goals */
-	void DeleteGoal(const FName& GoalToDelete);
-	/** get array of the names of currently selected goals */
-	const TArray<FName>& GetSelectedGoals() const {return SelectedGoals;};
-	/** return true if a goal with the given name is selected */
-	bool IsGoalSelected(const FName& GoalName) const;
-	/** when goal is renamed, update selection list */
-	void ReplaceGoalInSelection(const FName& OldName, const FName& NewName);
-	/** return number of goals that are currently selected */
-	int32 GetNumSelectedGoals();
+	/** clear all selected objects */
+	void ClearSelection();
 	/** callback when goal is selected in the viewport */
 	void HandleGoalSelectedInViewport(const FName& GoalName, bool bReplace);
-	/** callback when goal is selected in the skeleton view */
-	void HandleGoalsSelectedInTreeView(const TArray<FName>& GoalNames);
+	/** callback when bone is selected in the viewport */
+	void HandleBoneSelectedInViewport(const FName& BoneName, bool bReplace);
 	/** reset all goals to initial transforms */
 	void Reset() const;
 	/** refresh all views */
@@ -63,7 +85,7 @@ public:
 	/** determine if the element is an excluded bone*/
 	bool IsElementExcludedBone(TSharedRef<FIKRigTreeElement> TreeElement);
 	
-	/** todo show BONE transform in details view */
+	/** show transform of bone in details view */
 	void ShowDetailsForBone(const FName BoneName);
 	/** show BONE settings in details view */
 	void ShowDetailsForBoneSettings(const FName BoneName, int32 SolverIndex);
@@ -110,13 +132,31 @@ public:
 	/** UI and viewport selection state */
 	bool bManipulatingGoals = false;
 
+	/** record which part of the UI was last selected */
+	EIKRigSelectionType GetLastSelectedType() const;
+	void SetLastSelectedType(EIKRigSelectionType SelectionType);
+	/** END selection type */
+
+	/** FGCObject interface */
+	virtual void AddReferencedObjects(FReferenceCollector& Collector) override
+	{
+		Collector.AddReferencedObject(BoneDetails);
+	};
+	virtual FString GetReferencerName() const override
+	{
+		return "IKRigEditorController";
+	};
+	/** END FGCObject interface */
+
 private:
 
 	/** Initializes editor's solvers instances */
 	void InitializeSolvers() const;
-	
-	/** UI and viewport selection state */
-	TArray<FName> SelectedGoals;	
+
+	EIKRigSelectionType LastSelectedType;
+
+	UPROPERTY()
+	TObjectPtr<UIKRigBoneDetails> BoneDetails;
 };
 
 /** only used for pop-up window to selected a first solver to add*/
diff --git a/Engine/Plugins/Animation/IKRig/Source/IKRigEditor/Public/RigEditor/IKRigHitProxies.h b/Engine/Plugins/Animation/IKRig/Source/IKRigEditor/Public/RigEditor/IKRigHitProxies.h
--- a/Engine/Plugins/Animation/IKRig/Source/IKRigEditor/Public/RigEditor/IKRigHitProxies.h
+++ b/Engine/Plugins/Animation/IKRig/Source/IKRigEditor/Public/RigEditor/IKRigHitProxies.h
@@ -5,11 +5,16 @@
 #include "CoreMinimal.h"
 #include "HitProxies.h"
 
-// todo allow bone selection to preview transform values
 struct HIKRigEditorBoneProxy : public HHitProxy
 {
 	DECLARE_HIT_PROXY();
 
+	FName BoneName;
+
+	HIKRigEditorBoneProxy(const FName& InBoneName)
+		: HHitProxy(HPP_World)
+		, BoneName(InBoneName) {}
+
 	virtual EMouseCursor::Type GetMouseCursor()
 	{
 		return EMouseCursor::Crosshairs;
diff --git a/Engine/Plugins/Animation/IKRig/Source/IKRigEditor/Public/RigEditor/SIKRigRetargetChainList.h b/Engine/Plugins/Animation/IKRig/Source/IKRigEditor/Public/RigEditor/SIKRigRetargetChainList.h
--- a/Engine/Plugins/Animation/IKRig/Source/IKRigEditor/Public/RigEditor/SIKRigRetargetChainList.h
+++ b/Engine/Plugins/Animation/IKRig/Source/IKRigEditor/Public/RigEditor/SIKRigRetargetChainList.h
@@ -115,6 +115,8 @@ public:
 
 	void Construct(const FArguments& InArgs, TSharedRef<FIKRigEditorController> InEditorController);
 
+	FName GetSelectedChain();
+
 private:
 	
 	/** menu for adding new solver commands */
@@ -133,6 +135,7 @@ private:
 
 	/** callback when "Add New Chain" clicked */
 	bool IsAddChainEnabled() const;
+	
 	/** when a chain is clicked on in the stack view */
 	void OnItemClicked(TSharedPtr<FRetargetChainElement> InItem);
 
diff --git a/Engine/Plugins/Animation/IKRig/Source/IKRigEditor/Public/RigEditor/SIKRigSkeleton.h b/Engine/Plugins/Animation/IKRig/Source/IKRigEditor/Public/RigEditor/SIKRigSkeleton.h
--- a/Engine/Plugins/Animation/IKRig/Source/IKRigEditor/Public/RigEditor/SIKRigSkeleton.h
+++ b/Engine/Plugins/Animation/IKRig/Source/IKRigEditor/Public/RigEditor/SIKRigSkeleton.h
@@ -141,6 +141,20 @@ public:
 
 	void Construct(const FArguments& InArgs, TSharedRef<FIKRigEditorController> InEditorController);
 
+	/** selection state queries */
+	static bool IsBoneInSelection(TArray<TSharedPtr<FIKRigTreeElement>>& SelectedBoneItems, const FName& BoneName);
+	void GetSelectedBones(TArray<TSharedPtr<FIKRigTreeElement>>& OutBoneItems) const;
+	void GetSelectedGoals(TArray<TSharedPtr<FIKRigTreeElement>>& OutSelectedGoals) const;
+	int32 GetNumSelectedGoals();
+	void GetSelectedGoalNames(TArray<FName>& OutSelectedGoalNames) const;
+	bool IsGoalSelected(const FName& GoalName);
+	void AddSelectedItemFromViewport(
+		const FName& ItemName,
+		IKRigTreeElementType ItemType,
+		const bool bReplace);
+	void GetSelectedBoneChains(TArray<FIKRigSkeletonChain>& OutChains);
+	/** END selection state queries */
+
 private:
 	/** SWidget interface */
 	virtual FReply OnKeyDown(const FGeometry& MyGeometry, const FKeyEvent& InKeyEvent);
@@ -149,6 +163,13 @@ private:
 	/** Bind commands that this widget handles */
 	void BindCommands();
 
+	/** directly, programmatically add item to selection */
+	void AddItemToSelection(const TSharedPtr<FIKRigTreeElement>& InItem);
+	/** directly, programmatically add item to selection */
+	void RemoveItemFromSelection(const TSharedPtr<FIKRigTreeElement>& InItem);
+	/** directly, programmatically replace item in selection */
+	void ReplaceItemInSelection(const FText& OldName, const FText& NewName);
+
 	/** creating / renaming / deleting goals */
 	void HandleNewGoal();
 	bool CanAddNewGoal();
@@ -198,13 +219,6 @@ private:
 	bool CanSetRetargetRoot();
 	/** END retarget chains */
 
-	/** selection state queries */
-	void GetSelectedBones(TArray<TSharedPtr<FIKRigTreeElement>>& OutBoneItems) const;
-	void GetSelectedGoals(TArray<TSharedPtr<FIKRigTreeElement>>& OutSelectedGoals) const;
-	void SetSelectedGoalsFromViewport(const TArray<FName>& GoalNames);
-	void GetSelectedBoneChains(TArray<FIKRigSkeletonChain>& OutChains);
-	/** END selection state queries */
-
 	/** centralized editor controls (facilitate cross-communication between multiple UI elements)*/
 	TWeakPtr<FIKRigEditorController> EditorController;
 	
diff --git a/Engine/Plugins/Animation/IKRig/Source/IKRigEditor/Public/RigEditor/SIKRigSolverStack.h b/Engine/Plugins/Animation/IKRig/Source/IKRigEditor/Public/RigEditor/SIKRigSolverStack.h
--- a/Engine/Plugins/Animation/IKRig/Source/IKRigEditor/Public/RigEditor/SIKRigSolverStack.h
+++ b/Engine/Plugins/Animation/IKRig/Source/IKRigEditor/Public/RigEditor/SIKRigSolverStack.h
@@ -11,6 +11,7 @@
 #include "Widgets/Views/SListView.h"
 #include "Framework/Commands/UICommandList.h"
 
+class UIKRigSolver;
 class FIKRigEditorController;
 class SIKRigSolverStack;
 class FIKRigEditorToolkit;
@@ -56,7 +57,10 @@ public:
 
 	bool GetWarningMessage(FText& Message);
 
+	bool IsSolverEnabled() const;
+
 private:
+	UIKRigSolver* GetSolver() const;
 	TWeakPtr<FSolverStackElement> StackElement;
 	TWeakPtr<SIKRigSolverStack> SolverStack;
 };
@@ -106,6 +110,9 @@ private:
 	void DeleteSolver(TSharedPtr<FSolverStackElement> SolverToDelete);
 	/** when a solver is selected on in the stack view */
 	void OnSelectionChanged(TSharedPtr<FSolverStackElement> InItem, ESelectInfo::Type SelectInfo);
+	/** when a solver is clicked on in the stack view */
+	void OnItemClicked(TSharedPtr<FSolverStackElement> InItem);
+	void ShowDetailsForItem(TSharedPtr<FSolverStackElement> InItem);
 
 	/** list view generate row callback */
 	TSharedRef<ITableRow> MakeListRowWidget(TSharedPtr<FSolverStackElement> InElement, const TSharedRef<STableViewBase>& OwnerTable);
diff --git a/Engine/Plugins/Experimental/FullBodyIK/Source/PBIK/Private/Core/PBIKSolver.cpp b/Engine/Plugins/Experimental/FullBodyIK/Source/PBIK/Private/Core/PBIKSolver.cpp
--- a/Engine/Plugins/Experimental/FullBodyIK/Source/PBIK/Private/Core/PBIKSolver.cpp
+++ b/Engine/Plugins/Experimental/FullBodyIK/Source/PBIK/Private/Core/PBIKSolver.cpp
@@ -12,16 +12,13 @@ namespace PBIK
 	{
 		check(InBone);
 		Bone = InBone;
-		SetGoal(Bone->Position, Bone->Rotation, 1.0f, 1.0f, 1.0f, 0.0f);
+		SetGoal(Bone->Position, Bone->Rotation, Settings);
 	}
 
 	void FEffector::SetGoal(
 		const FVector& InPositionGoal,
 		const FQuat& InRotationGoal,
-		float InTransformAlpha,
-		float InStrengthAlpha,
-		float InPullChainAlpha,
-		float InPinRotation)
+		const FEffectorSettings& InSettings)
 	{
 		PositionOrig = Bone->Position;
 		RotationOrig = Bone->Rotation;
@@ -29,18 +26,15 @@ namespace PBIK
 		Position = PositionGoal = InPositionGoal;
 		Rotation = RotationGoal = InRotationGoal;
 
-		TransformAlpha = InTransformAlpha;
-		StrengthAlpha = InStrengthAlpha;
-		PullChainAlpha = InPullChainAlpha;
-		PinRotation = InPinRotation;
+		Settings = InSettings;
 	}
 
 	void FEffector::UpdateFromInputs(const FBone& SolverRoot)
 	{
 		// blend effector transform by alpha and set pin goal transform
-		Position = FMath::Lerp(PositionOrig, PositionGoal, TransformAlpha);
-		Rotation = FMath::Lerp(RotationOrig, RotationGoal, TransformAlpha);
-		Pin.Pin()->SetGoal(Position, Rotation, StrengthAlpha);
+		Position = FMath::Lerp(PositionOrig, PositionGoal, Settings.PositionAlpha);
+		Rotation = FMath::Lerp(RotationOrig, RotationGoal, Settings.RotationAlpha);
+		Pin.Pin()->SetGoal(Position, Rotation, Settings.StrengthAlpha);
 
 		// update length of chain this effector controls in the input pose
 		if (ParentSubRoot)
@@ -241,7 +235,7 @@ void FPBIKSolver::ApplyPullChainAlpha()
 			continue;
 		}
 
-		if (Effector.PullChainAlpha < SMALL_NUMBER)
+		if (Effector.Settings.PullChainAlpha < SMALL_NUMBER)
 		{
 			continue;
 		}
@@ -271,8 +265,9 @@ void FPBIKSolver::ApplyPullChainAlpha()
 		(ChainEndNew - ChainStartNew).ToDirectionAndLength(ChainVecNew, ChainLenNew);
 
 		FQuat ChainDeltaRotation = FQuat::FindBetweenNormals(ChainVecOrig, ChainVecNew);
-		ChainDeltaRotation = FQuat::FastLerp(FQuat::Identity, ChainDeltaRotation, (Effector.PullChainAlpha * Effector.StrengthAlpha));
-		FVector ChainDeltaPosition = ChainVecNew * (ChainLenNew - ChainLenOrig) * (Effector.PullChainAlpha * Effector.StrengthAlpha);
+		const float ChainDeltaAlpha = Effector.Settings.PullChainAlpha * Effector.Settings.StrengthAlpha;
+		ChainDeltaRotation = FQuat::FastLerp(FQuat::Identity, ChainDeltaRotation, ChainDeltaAlpha);
+		FVector ChainDeltaPosition = ChainVecNew * (ChainLenNew - ChainLenOrig) * ChainDeltaAlpha;
 
 		FBone* Bone = Effector.Bone->Parent;
 		int32 ChainIndex = 0;
@@ -391,10 +386,10 @@ void FPBIKSolver::UpdateBonesFromBodies()
 		
 		Bone->Position = Bone->Parent->Position + Bone->Parent->Rotation * Bone->LocalPositionOrig;
 
-		if (Effector.PinRotation > SMALL_NUMBER)
+		if (Effector.Settings.PinRotation > SMALL_NUMBER)
 		{
 			// optionally pin rotation to that of effector
-			const float RotAmount = FMath::Clamp(0.0f, 1.0f, Effector.PinRotation);
+			const float RotAmount = FMath::Clamp(0.0f, 1.0f, Effector.Settings.PinRotation);
 			Bone->Rotation = FQuat::FastLerp(Bone->Rotation, Effector.Rotation, RotAmount).GetNormalized();
 		}else
 		{
@@ -776,11 +771,8 @@ void FPBIKSolver::SetEffectorGoal(
 	const int32 Index,
 	const FVector& InPosition,
 	const FQuat& InRotation,
-	const float OffsetAlpha,
-	const float StrengthAlpha,
-	const float PullChainAlpha,
-	const float PinRotation)
+	const PBIK::FEffectorSettings& InSettings)
 {
 	check(Index >= 0 && Index < Effectors.Num());
-	Effectors[Index].SetGoal(InPosition, InRotation, OffsetAlpha, StrengthAlpha, PullChainAlpha, PinRotation);
+	Effectors[Index].SetGoal(InPosition, InRotation, InSettings);
 }
diff --git a/Engine/Plugins/Experimental/FullBodyIK/Source/PBIK/Private/RigUnit_PBIK.cpp b/Engine/Plugins/Experimental/FullBodyIK/Source/PBIK/Private/RigUnit_PBIK.cpp
--- a/Engine/Plugins/Experimental/FullBodyIK/Source/PBIK/Private/RigUnit_PBIK.cpp
+++ b/Engine/Plugins/Experimental/FullBodyIK/Source/PBIK/Private/RigUnit_PBIK.cpp
@@ -158,14 +158,15 @@ FRigUnit_PBIK_Execute()
 		const FPBIKEffector& Effector = Effectors[E];
 		FVector Position = Effector.Transform.GetLocation();
 		FQuat Rotation = Effector.Transform.GetRotation();
-		Solver.SetEffectorGoal(
-			EffectorSolverIndices[E],
-			Position,
-			Rotation,
-			Effector.OffsetAlpha,
-			Effector.StrengthAlpha,
-			Effector.PullChainAlpha,
-			Effector.PinRotation);
+
+		PBIK::FEffectorSettings EffectorSettings;
+		EffectorSettings.PositionAlpha = Effector.PositionAlpha;
+		EffectorSettings.RotationAlpha = Effector.RotationAlpha;
+		EffectorSettings.StrengthAlpha = Effector.StrengthAlpha;
+		EffectorSettings.PullChainAlpha = Effector.PullChainAlpha;
+		EffectorSettings.PinRotation = Effector.PinRotation;
+		
+		Solver.SetEffectorGoal(EffectorSolverIndices[E], Position, Rotation, EffectorSettings);
 	}
 
 	// solve
diff --git a/Engine/Plugins/Experimental/FullBodyIK/Source/PBIK/Public/Core/PBIKSolver.h b/Engine/Plugins/Experimental/FullBodyIK/Source/PBIK/Public/Core/PBIKSolver.h
--- a/Engine/Plugins/Experimental/FullBodyIK/Source/PBIK/Public/Core/PBIKSolver.h
+++ b/Engine/Plugins/Experimental/FullBodyIK/Source/PBIK/Public/Core/PBIKSolver.h
@@ -24,6 +24,37 @@ FORCEINLINE static float SquaredEaseOut(const float& Input){ return (FMath::Pow(
 
 struct FRigidBody;
 
+struct FEffectorSettings
+{
+	/** Range 0-1, default is 1. Blend between the input bone position (0.0) and the current effector position (1.0).*/
+	UPROPERTY(EditAnywhere, Category = "Goal Settings", meta = (ClampMin = "0.0", ClampMax = "1.0", UIMin = "0.0", UIMax = "1.0"))
+	float PositionAlpha = 1.0f;
+
+	/** Range 0-1, default is 1. Blend between the input bone rotation (0.0) and the current effector rotation (1.0).*/
+	UPROPERTY(EditAnywhere, Category = "Goal Settings", meta = (ClampMin = "0.0", ClampMax = "1.0", UIMin = "0.0", UIMax = "1.0"))
+	float RotationAlpha = 1.0f;
+
+	/** Range 0-1 (default is 1.0). The strength of the effector when pulling the bone towards it's target location.
+	* At 0.0, the effector does not pull at all, but the bones between the effector and the root will still slightly resist motion from other effectors.
+	* This can thus act as a "stabilizer" of sorts for parts of the body that you do not want to behave in a pure FK fashion.
+	*/
+	UPROPERTY(EditAnywhere, Category="Effector")
+	float StrengthAlpha = 1.0f;
+
+	/** Range 0-1 (default is 1.0). When enabled (greater than 0.0), the solver internally partitions the skeleton into 'chains' which extend from the effector to the nearest fork in the skeleton.
+	*These chains are pre-rotated and translated, as a whole, towards the effector targets.
+	*This can improve the results for sparse bone chains, and significantly improve convergence on dense bone chains.
+	*But it may cause undesirable results in highly constrained bone chains (like robot arms).
+	*/
+	UPROPERTY(EditAnywhere, Category="Effector")
+	float PullChainAlpha = 1.0f;
+
+	/** Range 0-1 (default is 1.0).
+	*Blends the effector bone rotation between the rotation of the effector transform (1.0) and the rotation of the input bone (0.0).*/
+	UPROPERTY(EditAnywhere, Category="Effector")
+	float PinRotation = 1.0f;
+};
+
 struct FEffector
 {
 	FVector Position;
@@ -35,18 +66,15 @@ struct FEffector
 	FVector PositionGoal;
 	FQuat RotationGoal;
 
+	FEffectorSettings Settings;
+
 	FBone* Bone;
 	TWeakPtr<FPinConstraint> Pin;
-	float PinRotation;
 	FRigidBody* ParentSubRoot = nullptr;
 	float DistanceToSubRootInInputPose;
 	float DistToRootAlongBones;
 	float DistToRootStraightLine;
-
-	float TransformAlpha;
-	float StrengthAlpha;
 	
-	float PullChainAlpha;
 	TArray<float> DistancesFromEffector;
 	float DistToSubRootAlongBones;
 
@@ -55,15 +83,12 @@ struct FEffector
 	void SetGoal(
 		const FVector& InPositionGoal,
 		const FQuat& InRotationGoal,
-		float InTransformAlpha,
-		float InStrengthAlpha,
-		float InPullChainAlpha,
-		float InPinRotation);
+		const FEffectorSettings& InSettings);
 
 	void UpdateFromInputs(const FBone& SolverRoot);
 	void ApplyPreferredAngles();
 };
-
+	
 } // namespace
 
 USTRUCT()
@@ -135,10 +160,7 @@ public:
 		const int32 Index, 
 		const FVector& InPosition, 
 		const FQuat& InRotation, 
-		const float OffsetAlpha, 
-		const float StrengthAlpha,
-		const float PullChainAlpha,
-		const float PinRotation);
+		const PBIK::FEffectorSettings& Settings);
 
 	void GetBoneGlobalTransform(const int32 Index, FTransform& OutTransform);
 
diff --git a/Engine/Plugins/Experimental/FullBodyIK/Source/PBIK/Public/RigUnit_PBIK.h b/Engine/Plugins/Experimental/FullBodyIK/Source/PBIK/Public/RigUnit_PBIK.h
--- a/Engine/Plugins/Experimental/FullBodyIK/Source/PBIK/Public/RigUnit_PBIK.h
+++ b/Engine/Plugins/Experimental/FullBodyIK/Source/PBIK/Public/RigUnit_PBIK.h
@@ -62,11 +62,13 @@ struct FPBIKEffector
 	UPROPERTY(EditAnywhere, Category="Effector")
 	FTransform Transform;
 
-	/** Range 0-1 (default is 1.0). This blends the effector translation and rotation between its transform and the transform of the bone from the input pose.
-	 * At 0.0, the effector will be placed at the location of the bone in the input pose (effectively pinning the skeleton to the input pose).
-	 * At 1.0, the effector will be placed at the supplied location.*/
-	UPROPERTY(EditAnywhere, Category="Effector")
-	float OffsetAlpha = 1.0f;
+	/** Range 0-1, default is 1. Blend between the input bone position (0.0) and the current effector position (1.0).*/
+	UPROPERTY(EditAnywhere, Category = "Goal Settings", meta = (ClampMin = "0.0", ClampMax = "1.0", UIMin = "0.0", UIMax = "1.0"))
+	float PositionAlpha = 1.0f;
+
+	/** Range 0-1, default is 1. Blend between the input bone rotation (0.0) and the current effector rotation (1.0).*/
+	UPROPERTY(EditAnywhere, Category = "Goal Settings", meta = (ClampMin = "0.0", ClampMax = "1.0", UIMin = "0.0", UIMax = "1.0"))
+	float RotationAlpha = 1.0f;
 
 	/** Range 0-1 (default is 1.0). The strength of the effector when pulling the bone towards it's target location.
 	 * At 0.0, the effector does not pull at all, but the bones between the effector and the root will still slightly resist motion from other effectors.
