Fixing regression from earlier check-in 18276842: serialization of IKRetargeter asset broken.

diff --git a/Engine/Plugins/Animation/IKRig/Source/IKRig/Public/Retargeter/IKRetargeter.h b/Engine/Plugins/Animation/IKRig/Source/IKRig/Public/Retargeter/IKRetargeter.h
--- a/Engine/Plugins/Animation/IKRig/Source/IKRig/Public/Retargeter/IKRetargeter.h
+++ b/Engine/Plugins/Animation/IKRig/Source/IKRig/Public/Retargeter/IKRetargeter.h
@@ -118,18 +118,18 @@ private:
 
 	/** The rig to copy animation FROM.*/
 	UPROPERTY(VisibleAnywhere, Category = Rigs)
-	TWeakObjectPtr<UIKRigDefinition> SourceIKRigAsset = nullptr;
+	TObjectPtr<UIKRigDefinition> SourceIKRigAsset = nullptr;
 	
 	/** The rig to copy animation TO.*/
 	UPROPERTY(EditAnywhere, Category = Rigs)
-	TWeakObjectPtr<UIKRigDefinition> TargetIKRigAsset = nullptr;
+	TObjectPtr<UIKRigDefinition> TargetIKRigAsset = nullptr;
 
 public:
 
 #if WITH_EDITORONLY_DATA
 	/** The Skeletal Mesh to preview the retarget on.*/
 	UPROPERTY(EditAnywhere, Category = Rigs)
-	TWeakObjectPtr<USkeletalMesh> TargetPreviewMesh = nullptr;
+	TObjectPtr<USkeletalMesh> TargetPreviewMesh = nullptr;
 #endif
 	
 	/** When false, translational motion of skeleton root is not copied. Useful for debugging.*/
@@ -167,7 +167,7 @@ private:
 	UPROPERTY(Transient)
 	TObjectPtr<UObject> Controller;
 #endif
-
+	
 private:
 	/** The set of retarget poses available as options for retargeting.*/
 	UPROPERTY()
diff --git a/Engine/Plugins/Animation/IKRig/Source/IKRigEditor/Private/RetargetEditor/IKRetargeterController.cpp b/Engine/Plugins/Animation/IKRig/Source/IKRigEditor/Private/RetargetEditor/IKRetargeterController.cpp
--- a/Engine/Plugins/Animation/IKRig/Source/IKRigEditor/Private/RetargetEditor/IKRetargeterController.cpp
+++ b/Engine/Plugins/Animation/IKRig/Source/IKRigEditor/Private/RetargetEditor/IKRetargeterController.cpp
@@ -57,7 +57,7 @@ USkeletalMesh* UIKRetargeterController::GetTargetPreviewMesh()
 	}
 
 	// optionally prefer override if one is provided
-	if (Asset->TargetPreviewMesh.IsValid())
+	if (IsValid(Asset->TargetPreviewMesh))
 	{
 		return Asset->TargetPreviewMesh.Get();
 	}
@@ -68,17 +68,17 @@ USkeletalMesh* UIKRetargeterController::GetTargetPreviewMesh()
 
 FName UIKRetargeterController::GetSourceRootBone() const
 {
-	return Asset->SourceIKRigAsset.IsValid() ? Asset->SourceIKRigAsset->GetRetargetRoot() : FName("None");
+	return IsValid(Asset->SourceIKRigAsset) ? Asset->SourceIKRigAsset->GetRetargetRoot() : FName("None");
 }
 
 FName UIKRetargeterController::GetTargetRootBone() const
 {
-	return Asset->TargetIKRigAsset.IsValid() ? Asset->TargetIKRigAsset->GetRetargetRoot() : FName("None");
+	return IsValid(Asset->TargetIKRigAsset) ? Asset->TargetIKRigAsset->GetRetargetRoot() : FName("None");
 }
 
 void UIKRetargeterController::GetTargetChainNames(TArray<FName>& OutNames) const
 {
-	if (Asset->TargetIKRigAsset.IsValid())
+	if (IsValid(Asset->TargetIKRigAsset))
 	{
 		const TArray<FBoneChain>& Chains = Asset->TargetIKRigAsset->GetRetargetChains();
 		for (const FBoneChain& Chain : Chains)
@@ -90,7 +90,7 @@ void UIKRetargeterController::GetTargetChainNames(TArray<FName>& OutNames) const
 
 void UIKRetargeterController::GetSourceChainNames(TArray<FName>& OutNames) const
 {
-	if (Asset->SourceIKRigAsset.IsValid())
+	if (IsValid(Asset->SourceIKRigAsset))
 	{
 		const TArray<FBoneChain>& Chains = Asset->SourceIKRigAsset->GetRetargetChains();
 		for (const FBoneChain& Chain : Chains)
@@ -102,7 +102,7 @@ void UIKRetargeterController::GetSourceChainNames(TArray<FName>& OutNames) const
 
 void UIKRetargeterController::CleanChainMapping()
 {
-	if (!Asset->TargetIKRigAsset.IsValid())
+	if (!IsValid(Asset->TargetIKRigAsset))
 	{
 		// don't clean chain mappings, in case user is replacing with IK Rig asset that has some valid mappings
 		return;
@@ -191,7 +191,7 @@ void UIKRetargeterController::CleanPoseList()
 	}
 
 	// remove all bone offsets that are no longer part of the target skeleton
-	if (Asset->TargetIKRigAsset.IsValid())
+	if (IsValid(Asset->TargetIKRigAsset))
 	{
 		const TArray<FName> AllowedBoneNames = Asset->TargetIKRigAsset->Skeleton.BoneNames;
 		for (TTuple<FName, FIKRetargetPose>& Pose : Asset->RetargetPoses)
@@ -302,7 +302,7 @@ const TArray<FRetargetChainMap>& UIKRetargeterController::GetChainMappings()
 
 USkeleton* UIKRetargeterController::GetSourceSkeletonAsset() const
 {
-	if (!Asset->SourceIKRigAsset.IsValid())
+	if (!IsValid(Asset->SourceIKRigAsset))
 	{
 		return nullptr;
 	}
