Fixed camera framing on selected bones

diff --git a/Engine/Plugins/Animation/IKRig/Source/IKRigEditor/Private/RigEditor/IKRigEditMode.cpp b/Engine/Plugins/Animation/IKRig/Source/IKRigEditor/Private/RigEditor/IKRigEditMode.cpp
--- a/Engine/Plugins/Animation/IKRig/Source/IKRigEditor/Private/RigEditor/IKRigEditMode.cpp
+++ b/Engine/Plugins/Animation/IKRig/Source/IKRigEditor/Private/RigEditor/IKRigEditMode.cpp
@@ -29,17 +29,38 @@ bool FIKRigEditMode::GetCameraTarget(FSphere& OutTarget) const
 		return false; 
 	}
 	
-	// target union of selected goals
-	TArray<FName> OutGoalNames;
+	// target union of selected goals and bones
+	TArray<FName> OutGoalNames, OutBoneNames;
 	Controller->SkeletonView->GetSelectedGoalNames(OutGoalNames);
-	if (!OutGoalNames.IsEmpty())
+	Controller->SkeletonView->GetSelectedBoneNames(OutBoneNames);
+	
+	if (!OutGoalNames.IsEmpty() || !OutBoneNames.IsEmpty())
 	{
 		TArray<FVector> GoalPoints;
+
+		// get goal locations
 		for (const FName& GoalName : OutGoalNames)
 		{			
 			GoalPoints.Add(Controller->AssetController->GetGoal(GoalName)->CurrentTransform.GetLocation());
 		}
 
+		// get bone locations
+		const FIKRigSkeleton& FIKRigSkeleton = Controller->AssetController->GetIKRigSkeleton();
+		for (const FName& BoneName : OutBoneNames)
+		{
+			const int32 BoneIndex = FIKRigSkeleton.GetBoneIndexFromName(BoneName);
+			if (BoneIndex != INDEX_NONE)
+			{
+				TArray<int32> Children;
+				FIKRigSkeleton.GetChildIndices(BoneIndex, Children);
+				for (int32 ChildIndex: Children)
+				{
+					GoalPoints.Add(FIKRigSkeleton.CurrentPoseGlobal[ChildIndex].GetLocation());
+				}
+				GoalPoints.Add(FIKRigSkeleton.CurrentPoseGlobal[BoneIndex].GetLocation());
+			}
+		}
+
 		// create a sphere that contains all the goal points
 		OutTarget = FSphere(&GoalPoints[0], GoalPoints.Num());
 		return true;
diff --git a/Engine/Plugins/Animation/IKRig/Source/IKRigEditor/Private/RigEditor/SIKRigSkeleton.cpp b/Engine/Plugins/Animation/IKRig/Source/IKRigEditor/Private/RigEditor/SIKRigSkeleton.cpp
--- a/Engine/Plugins/Animation/IKRig/Source/IKRigEditor/Private/RigEditor/SIKRigSkeleton.cpp
+++ b/Engine/Plugins/Animation/IKRig/Source/IKRigEditor/Private/RigEditor/SIKRigSkeleton.cpp
@@ -1105,6 +1105,17 @@ void SIKRigSkeleton::GetSelectedBones(TArray<TSharedPtr<FIKRigTreeElement>>& Out
 	}
 }
 
+void SIKRigSkeleton::GetSelectedBoneNames(TArray<FName>& OutSelectedBoneNames) const
+{
+	TArray<TSharedPtr<FIKRigTreeElement>> OutSelectedBones;
+	GetSelectedBones(OutSelectedBones);
+	OutSelectedBoneNames.Reset();
+	for (TSharedPtr<FIKRigTreeElement> SelectedBoneItem : OutSelectedBones)
+	{
+		OutSelectedBoneNames.Add(SelectedBoneItem->BoneName);
+	}
+}
+
 void SIKRigSkeleton::GetSelectedGoals(TArray<TSharedPtr<FIKRigTreeElement>>& OutSelectedGoals) const
 {
 	OutSelectedGoals.Reset();
diff --git a/Engine/Plugins/Animation/IKRig/Source/IKRigEditor/Public/RigEditor/SIKRigSkeleton.h b/Engine/Plugins/Animation/IKRig/Source/IKRigEditor/Public/RigEditor/SIKRigSkeleton.h
--- a/Engine/Plugins/Animation/IKRig/Source/IKRigEditor/Public/RigEditor/SIKRigSkeleton.h
+++ b/Engine/Plugins/Animation/IKRig/Source/IKRigEditor/Public/RigEditor/SIKRigSkeleton.h
@@ -144,6 +144,7 @@ public:
 	/** selection state queries */
 	static bool IsBoneInSelection(TArray<TSharedPtr<FIKRigTreeElement>>& SelectedBoneItems, const FName& BoneName);
 	void GetSelectedBones(TArray<TSharedPtr<FIKRigTreeElement>>& OutBoneItems) const;
+	void GetSelectedBoneNames(TArray<FName>& OutSelectedBoneNames) const;
 	void GetSelectedGoals(TArray<TSharedPtr<FIKRigTreeElement>>& OutSelectedGoals) const;
 	int32 GetNumSelectedGoals();
 	void GetSelectedGoalNames(TArray<FName>& OutSelectedGoalNames) const;
