Add shortcuts for creation, deletion and renaming

diff --git a/Engine/Plugins/Animation/IKRig/Source/IKRigEditor/Private/RigEditor/IKRigSkeletonCommands.cpp b/Engine/Plugins/Animation/IKRig/Source/IKRigEditor/Private/RigEditor/IKRigSkeletonCommands.cpp
--- a/Engine/Plugins/Animation/IKRig/Source/IKRigEditor/Private/RigEditor/IKRigSkeletonCommands.cpp
+++ b/Engine/Plugins/Animation/IKRig/Source/IKRigEditor/Private/RigEditor/IKRigSkeletonCommands.cpp
@@ -6,8 +6,8 @@
 
 void FIKRigSkeletonCommands::RegisterCommands()
 {
-	UI_COMMAND(NewGoal, "New IK Goal", "Add new IK goal at the selected bone.", EUserInterfaceActionType::Button, FInputChord());
-	UI_COMMAND(DeleteGoal, "Delete Goal or Effector", "Delete the selected IK Goal or Effector.", EUserInterfaceActionType::Button, FInputChord(EKeys::Delete));
+	UI_COMMAND(NewGoal, "New IK Goal", "Add new IK goal at the selected bone.", EUserInterfaceActionType::Button, FInputChord(EKeys::N, EModifierKey::Control));
+	UI_COMMAND(DeleteElement, "Delete Goal or Effector", "Delete the selected IK Goal, Effector or Bone Settings.", EUserInterfaceActionType::Button, FInputChord(EKeys::Delete), FInputChord(EKeys::BackSpace));
 	UI_COMMAND(ConnectGoalToSolvers, "Connect Goal to Selected Solvers", "Make the selected goal an effector in the selected solvers.", EUserInterfaceActionType::Button, FInputChord());
 	UI_COMMAND(DisconnectGoalFromSolvers, "Disconnect Goal from Selected Solvers", "Remove effectors from the selected solvers that use this goal.", EUserInterfaceActionType::Button, FInputChord());
 	UI_COMMAND(SetRootBoneOnSolvers, "Set Root Bone on Selected Solvers", "Set the Root Bone setting on the selected solvers.", EUserInterfaceActionType::Button, FInputChord());
@@ -18,6 +18,7 @@ void FIKRigSkeletonCommands::RegisterCommands()
 	UI_COMMAND(IncludeBone, "Include Selected Bone In Solve", "Include bone in all solvers.", EUserInterfaceActionType::Button, FInputChord());
 	UI_COMMAND(NewRetargetChain, "New Retarget Chain from Selected Bones", "Create a new retarget bone chain from the selected bones.", EUserInterfaceActionType::Button, FInputChord());
 	UI_COMMAND(SetRetargetRoot, "Set Retargeting Root", "Set the Root Bone used for retargeting. Usually 'Pelvis'.", EUserInterfaceActionType::Button, FInputChord());
+	UI_COMMAND(RenameGoal, "Rename Goal", "Rename the selected goal.", EUserInterfaceActionType::Button, FInputChord(EKeys::F2));
 }
 
 #undef LOCTEXT_NAMESPACE
diff --git a/Engine/Plugins/Animation/IKRig/Source/IKRigEditor/Private/RigEditor/SIKRigSkeleton.cpp b/Engine/Plugins/Animation/IKRig/Source/IKRigEditor/Private/RigEditor/SIKRigSkeleton.cpp
--- a/Engine/Plugins/Animation/IKRig/Source/IKRigEditor/Private/RigEditor/SIKRigSkeleton.cpp
+++ b/Engine/Plugins/Animation/IKRig/Source/IKRigEditor/Private/RigEditor/SIKRigSkeleton.cpp
@@ -400,9 +400,9 @@ void SIKRigSkeleton::BindCommands()
         FExecuteAction::CreateSP(this, &SIKRigSkeleton::HandleNewGoal),
         FCanExecuteAction::CreateSP(this, &SIKRigSkeleton::CanAddNewGoal));
 	
-	CommandList->MapAction(Commands.DeleteGoal,
-        FExecuteAction::CreateSP(this, &SIKRigSkeleton::HandleDeleteGoal),
-        FCanExecuteAction::CreateSP(this, &SIKRigSkeleton::CanDeleteGoal));
+	CommandList->MapAction(Commands.DeleteElement,
+        FExecuteAction::CreateSP(this, &SIKRigSkeleton::HandleDeleteElement),
+        FCanExecuteAction::CreateSP(this, &SIKRigSkeleton::CanDeleteElement));
 
 	CommandList->MapAction(Commands.ConnectGoalToSolvers,
         FExecuteAction::CreateSP(this, &SIKRigSkeleton::HandleConnectGoalToSolver),
@@ -445,6 +445,10 @@ void SIKRigSkeleton::BindCommands()
 	CommandList->MapAction(Commands.SetRetargetRoot,
 		FExecuteAction::CreateSP(this, &SIKRigSkeleton::HandleSetRetargetRoot),
 		FCanExecuteAction::CreateSP(this, &SIKRigSkeleton::CanSetRetargetRoot));
+
+	CommandList->MapAction(Commands.RenameGoal,
+	FExecuteAction::CreateSP(this, &SIKRigSkeleton::HandleRenameGoal),
+		FCanExecuteAction::CreateSP(this, &SIKRigSkeleton::CanRenameGoal));
 }
 
 void SIKRigSkeleton::FillContextMenu(FMenuBuilder& MenuBuilder)
@@ -484,7 +488,7 @@ void SIKRigSkeleton::FillContextMenu(FMenuBuilder& MenuBuilder)
 	MenuBuilder.EndSection();
 }
 
-void SIKRigSkeleton::HandleNewGoal()
+void SIKRigSkeleton::HandleNewGoal() const
 {
 	const TSharedPtr<FIKRigEditorController> Controller = EditorController.Pin();
 	if (!Controller.IsValid())
@@ -515,7 +519,7 @@ void SIKRigSkeleton::HandleNewGoal()
 	Controller->AddNewGoals(GoalNames, BoneNames);
 }
 
-bool SIKRigSkeleton::CanAddNewGoal()
+bool SIKRigSkeleton::CanAddNewGoal() const
 {
 	// is anything selected?
 	TArray<TSharedPtr<FIKRigTreeElement>> SelectedItems = TreeView->GetSelectedItems();
@@ -536,7 +540,7 @@ bool SIKRigSkeleton::CanAddNewGoal()
 	return true;
 }
 
-void SIKRigSkeleton::HandleDeleteGoal()
+void SIKRigSkeleton::HandleDeleteElement()
 {
 	const TSharedPtr<FIKRigEditorController> Controller = EditorController.Pin();
 	if (!Controller.IsValid())
@@ -545,24 +549,31 @@ void SIKRigSkeleton::HandleDeleteGoal()
 	}
 
 	TArray<TSharedPtr<FIKRigTreeElement>> SelectedItems = TreeView->GetSelectedItems();
-	TSharedPtr<FIKRigTreeElement> ParentOfDeletedGoal;
-	for (const TSharedPtr<FIKRigTreeElement>& Item : SelectedItems)
+	for (const TSharedPtr<FIKRigTreeElement>& SelectedItem : SelectedItems)
 	{
-		if (Item->ElementType == IKRigTreeElementType::GOAL)
-		{
-			Controller->AssetController->RemoveGoal(Item->GoalName);
-		}
-		else if (Item->ElementType == IKRigTreeElementType::SOLVERGOAL)
+		switch(SelectedItem->ElementType)
 		{
-			Controller->AssetController->DisconnectGoalFromSolver(Item->SolverGoalName, Item->SolverGoalIndex);
+			case IKRigTreeElementType::GOAL:
+				Controller->AssetController->RemoveGoal(SelectedItem->GoalName);
+				break;
+			case IKRigTreeElementType::SOLVERGOAL:
+				Controller->AssetController->DisconnectGoalFromSolver(SelectedItem->SolverGoalName, SelectedItem->SolverGoalIndex);
+				break;
+			case IKRigTreeElementType::BONE_SETTINGS:
+				Controller->AssetController->RemoveBoneSetting(SelectedItem->BoneSettingBoneName, SelectedItem->BoneSettingsSolverIndex);
+				break;
+			default:
+				break; // can't delete anything else
 		}
 	}
+	
+	RefreshTreeView();
 
 	Controller->ShowEmptyDetails();
 	// update all views
 	Controller->RefreshAllViews();
 }
-bool SIKRigSkeleton::CanDeleteGoal() const
+bool SIKRigSkeleton::CanDeleteElement() const
 {
 	// is anything selected?
 	TArray<TSharedPtr<FIKRigTreeElement>> SelectedItems = TreeView->GetSelectedItems();
@@ -571,12 +582,11 @@ bool SIKRigSkeleton::CanDeleteGoal() const
 		return false;
 	}
 
-	// are all selected items goals or effectors?
+	// are all selected items goals, effectors or bone settings?
 	for (const TSharedPtr<FIKRigTreeElement>& Item : SelectedItems)
 	{
-		const bool bIsGoal = Item->ElementType == IKRigTreeElementType::GOAL;
-		const bool bIsEffector = Item->ElementType == IKRigTreeElementType::SOLVERGOAL;
-		if (!(bIsGoal || bIsEffector))
+		const bool bIsBone = Item->ElementType == IKRigTreeElementType::BONE;
+		if (bIsBone)
 		{
 			return false;
 		}
@@ -1161,7 +1171,7 @@ bool SIKRigSkeleton::IsGoalSelected(const FName& GoalName)
 	return false;
 }
 
-void SIKRigSkeleton::HandleRenameElement() const
+void SIKRigSkeleton::HandleRenameGoal() const
 {
 	TArray<TSharedPtr<FIKRigTreeElement>> SelectedGoals;
 	GetSelectedGoals(SelectedGoals);
@@ -1173,6 +1183,13 @@ void SIKRigSkeleton::HandleRenameElement() const
 	SelectedGoals[0]->RequestRename();
 }
 
+bool SIKRigSkeleton::CanRenameGoal() const
+{
+	TArray<TSharedPtr<FIKRigTreeElement>> SelectedGoals;
+	GetSelectedGoals(SelectedGoals);
+	return SelectedGoals.Num() == 1;
+}
+
 void SIKRigSkeleton::RefreshTreeView(bool IsInitialSetup)
 {
 	const TSharedPtr<FIKRigEditorController> Controller = EditorController.Pin();
@@ -1369,7 +1386,7 @@ void SIKRigSkeleton::OnItemClicked(TSharedPtr<FIKRigTreeElement> InItem)
 	{
 		RegisterActiveTimer(0.f,
             FWidgetActiveTimerDelegate::CreateLambda([this](double, float) {
-            HandleRenameElement();
+            HandleRenameGoal();
             return EActiveTimerReturnType::Stop;
         }));
 	}
@@ -1503,30 +1520,8 @@ FReply SIKRigSkeleton::OnKeyDown(const FGeometry& MyGeometry, const FKeyEvent& I
 		return FReply::Handled();
 	}
 
-	// handle deleting selected items
-	if (Key == EKeys::Delete || Key == EKeys::BackSpace)
+	if (CommandList.IsValid() && CommandList->ProcessCommandBindings(InKeyEvent))
 	{
-		TArray<TSharedPtr<FIKRigTreeElement>> SelectedItems = TreeView->GetSelectedItems();
-		for (const TSharedPtr<FIKRigTreeElement>& SelectedItem : SelectedItems)
-		{
-			switch(SelectedItem->ElementType)
-			{
-				case IKRigTreeElementType::GOAL:
-					Controller->AssetController->RemoveGoal(SelectedItem->GoalName);
-					break;
-				case IKRigTreeElementType::SOLVERGOAL:
-					Controller->AssetController->DisconnectGoalFromSolver(SelectedItem->SolverGoalName, SelectedItem->SolverGoalIndex);
-					break;
-				case IKRigTreeElementType::BONE_SETTINGS:
-					Controller->AssetController->RemoveBoneSetting(SelectedItem->BoneSettingBoneName, SelectedItem->BoneSettingsSolverIndex);
-					break;
-				default:
-					continue; // can't delete anything else
-			}
-		}
-
-		RefreshTreeView();
-		
 		return FReply::Handled();
 	}
 	
diff --git a/Engine/Plugins/Animation/IKRig/Source/IKRigEditor/Public/RigEditor/IKRigSkeletonCommands.h b/Engine/Plugins/Animation/IKRig/Source/IKRigEditor/Public/RigEditor/IKRigSkeletonCommands.h
--- a/Engine/Plugins/Animation/IKRig/Source/IKRigEditor/Public/RigEditor/IKRigSkeletonCommands.h
+++ b/Engine/Plugins/Animation/IKRig/Source/IKRigEditor/Public/RigEditor/IKRigSkeletonCommands.h
@@ -20,8 +20,8 @@ public:
 	
 	/** create new goal on selected bone */
 	TSharedPtr< FUICommandInfo > NewGoal;
-	/** delete selected goal */
-	TSharedPtr< FUICommandInfo > DeleteGoal;
+	/** delete selected element */
+	TSharedPtr< FUICommandInfo > DeleteElement;
 	/** add goal to selected solvers */
 	TSharedPtr< FUICommandInfo > ConnectGoalToSolvers;
 	/** remove goal to selected solvers */
@@ -48,6 +48,9 @@ public:
 	/** set the root of the retargeting */
 	TSharedPtr< FUICommandInfo > SetRetargetRoot;
 
+	/** rename selected goal */
+	TSharedPtr< FUICommandInfo > RenameGoal;
+	
 	/** initialize commands */
 	virtual void RegisterCommands() override;
 };
diff --git a/Engine/Plugins/Animation/IKRig/Source/IKRigEditor/Public/RigEditor/SIKRigSkeleton.h b/Engine/Plugins/Animation/IKRig/Source/IKRigEditor/Public/RigEditor/SIKRigSkeleton.h
--- a/Engine/Plugins/Animation/IKRig/Source/IKRigEditor/Public/RigEditor/SIKRigSkeleton.h
+++ b/Engine/Plugins/Animation/IKRig/Source/IKRigEditor/Public/RigEditor/SIKRigSkeleton.h
@@ -172,11 +172,12 @@ private:
 	void ReplaceItemInSelection(const FText& OldName, const FText& NewName);
 
 	/** creating / renaming / deleting goals */
-	void HandleNewGoal();
-	bool CanAddNewGoal();
-	void HandleRenameElement() const;
-	void HandleDeleteGoal();
-	bool CanDeleteGoal() const;
+	void HandleNewGoal() const;
+	bool CanAddNewGoal() const;
+	void HandleRenameGoal() const;
+	bool CanRenameGoal() const;
+	void HandleDeleteElement();
+	bool CanDeleteElement() const;
 	/** END creating / renaming / deleting goals */
 
 	/** connecting/disconnecting goals to solvers */
