Update selection properly and check solver's

diff --git a/Engine/Plugins/Animation/IKRig/Source/IKRigEditor/Private/RigEditor/IKRigEditMode.cpp b/Engine/Plugins/Animation/IKRig/Source/IKRigEditor/Private/RigEditor/IKRigEditMode.cpp
--- a/Engine/Plugins/Animation/IKRig/Source/IKRigEditor/Private/RigEditor/IKRigEditMode.cpp
+++ b/Engine/Plugins/Animation/IKRig/Source/IKRigEditor/Private/RigEditor/IKRigEditMode.cpp
@@ -252,12 +252,15 @@ void FIKRigEditMode::GetAffectedBones(
 			// record which bones in the skeleton are affected by this solver
 			UIKRigController* AssetController = Controller->AssetController;
 			UIKRigSolver* SelectedSolver = AssetController->GetSolver(SelectedSolvers[0].Get()->IndexInStack);
-			for (int32 BoneIndex=0; BoneIndex<Skeleton.BoneNames.Num(); ++BoneIndex)
+			if(SelectedSolver)
 			{
-				const FName& BoneName = Skeleton.BoneNames[BoneIndex];
-				if (SelectedSolver->IsBoneAffectedBySolver(BoneName, Skeleton))
+				for (int32 BoneIndex=0; BoneIndex<Skeleton.BoneNames.Num(); ++BoneIndex)
 				{
-					OutAffectedBones.Add(BoneIndex);
+					const FName& BoneName = Skeleton.BoneNames[BoneIndex];
+					if (SelectedSolver->IsBoneAffectedBySolver(BoneName, Skeleton))
+					{
+						OutAffectedBones.Add(BoneIndex);
+					}
 				}
 			}
 		}
diff --git a/Engine/Plugins/Animation/IKRig/Source/IKRigEditor/Private/RigEditor/SIKRigSolverStack.cpp b/Engine/Plugins/Animation/IKRig/Source/IKRigEditor/Private/RigEditor/SIKRigSolverStack.cpp
--- a/Engine/Plugins/Animation/IKRig/Source/IKRigEditor/Private/RigEditor/SIKRigSolverStack.cpp
+++ b/Engine/Plugins/Animation/IKRig/Source/IKRigEditor/Private/RigEditor/SIKRigSolverStack.cpp
@@ -137,7 +137,7 @@ void SIKRigSolverStackItem::Construct(
 			.Padding(3)
 			[
 				SNew(SButton)
-				.ToolTipText(LOCTEXT("DeleteSolver", "Delect solver and remove from stack."))
+				.ToolTipText(LOCTEXT("DeleteSolver", "Delete solver and remove from stack."))
 				.OnClicked_Lambda([InSolverStack, InStackElement]() -> FReply
 				{
 					InSolverStack.Get()->DeleteSolver(InStackElement);
@@ -381,22 +381,29 @@ void SIKRigSolverStack::RefreshStackView()
 	}
 
 	// generate all list items
+	static const FText UnknownSolverTxt = FText::FromString("Unknown Solver");
+	
 	ListViewItems.Reset();
 	UIKRigController* AssetController = Controller->AssetController;
 	const int32 NumSolvers = AssetController->GetNumSolvers();
 	for (int32 i=0; i<NumSolvers; ++i)
 	{
 		const UIKRigSolver* Solver = AssetController->GetSolver(i);
-		const FText DisplayName = Solver ? FText::FromString(AssetController->GetSolverUniqueName(i)) : FText::FromString("Unknown Solver");
+		const FText DisplayName = Solver ? FText::FromString(AssetController->GetSolverUniqueName(i)) : UnknownSolverTxt;
 		TSharedPtr<FSolverStackElement> SolverItem = FSolverStackElement::Make(DisplayName, i);
 		ListViewItems.Add(SolverItem);
 	}
 
-	// restore selection
-	if (ListViewItems.IsValidIndex(IndexToSelect))
+	if (NumSolvers && ListViewItems.IsValidIndex(IndexToSelect))
 	{
+		// restore selection
 		ListView->SetSelection(ListViewItems[IndexToSelect]);
 	}
+	else
+	{
+		// clear selection otherwise
+		ListView->ClearSelection();
+	}
 
 	ListView->RequestListRefresh();
 }
