IKRig: Use new SAdvancedTransformInputBox

diff --git a/Engine/Plugins/Animation/IKRig/Source/IKRig/Private/IKRigDefinition.cpp b/Engine/Plugins/Animation/IKRig/Source/IKRig/Private/IKRigDefinition.cpp
--- a/Engine/Plugins/Animation/IKRig/Source/IKRig/Private/IKRigDefinition.cpp
+++ b/Engine/Plugins/Animation/IKRig/Source/IKRig/Private/IKRigDefinition.cpp
@@ -7,6 +7,124 @@
 
 #define LOCTEXT_NAMESPACE	"IKRigDefinition"
 
+#if WITH_EDITOR
+
+TOptional<FTransform::FReal> UIKRigEffectorGoal::GetNumericValue(
+	ESlateTransformComponent::Type Component,
+	ESlateRotationRepresentation::Type Representation,
+	ESlateTransformSubComponent::Type SubComponent,
+	EIKRigTransformType::Type TransformType) const
+{
+	switch(TransformType)
+	{
+		case EIKRigTransformType::Current:
+		{
+			return SAdvancedTransformInputBox<FTransform>::GetNumericValueFromTransform(
+				CurrentTransform,
+				Component,
+				Representation,
+				SubComponent);
+		}
+		case EIKRigTransformType::Reference:
+		{
+			return SAdvancedTransformInputBox<FTransform>::GetNumericValueFromTransform(
+				InitialTransform,
+				Component,
+				Representation,
+				SubComponent);
+		}
+		default:
+		{
+			break;
+		}
+	}
+	return TOptional<FTransform::FReal>();
+}
+
+void UIKRigEffectorGoal::OnNumericValueChanged(ESlateTransformComponent::Type Component,
+	ESlateRotationRepresentation::Type Representation, ESlateTransformSubComponent::Type SubComponent,
+	FTransform::FReal Value, ETextCommit::Type CommitType, EIKRigTransformType::Type TransformType)
+{
+	Modify();
+
+	switch(TransformType)
+	{
+		case EIKRigTransformType::Current:
+		{
+			SAdvancedTransformInputBox<FTransform>::ApplyNumericValueChange(
+				CurrentTransform,
+				Value,
+				Component,
+				Representation,
+				SubComponent);
+			break;
+		}
+		case EIKRigTransformType::Reference:
+		{
+			SAdvancedTransformInputBox<FTransform>::ApplyNumericValueChange(
+				InitialTransform,
+				Value,
+				Component,
+				Representation,
+				SubComponent);
+			break;
+		}
+		default:
+		{
+			break;
+		}
+	}
+}
+
+bool UIKRigEffectorGoal::TransformDiffersFromDefault(TSharedPtr<IPropertyHandle> PropertyHandle,
+	ESlateTransformComponent::Type Component) const
+{
+	if(PropertyHandle->GetProperty()->GetFName() == GET_MEMBER_NAME_CHECKED(UIKRigEffectorGoal, CurrentTransform))
+	{
+		switch(Component)
+		{
+			case ESlateTransformComponent::Location:
+			{
+				return !(CurrentTransform.GetLocation() - InitialTransform.GetLocation()).IsNearlyZero();
+			}
+			case ESlateTransformComponent::Rotation:
+			{
+				return !(CurrentTransform.Rotator() - InitialTransform.Rotator()).IsNearlyZero();
+			}
+			case ESlateTransformComponent::Scale:
+			{
+				return !(CurrentTransform.GetScale3D() - InitialTransform.GetScale3D()).IsNearlyZero();
+			}
+		}
+	}
+	return false;
+}
+
+void UIKRigEffectorGoal::ResetTransformToDefault(TSharedPtr<IPropertyHandle> PropertyHandle,
+	ESlateTransformComponent::Type Component)
+{
+	switch(Component)
+	{
+		case ESlateTransformComponent::Location:
+		{
+			CurrentTransform.SetLocation(InitialTransform.GetLocation());
+			break;
+		}
+		case ESlateTransformComponent::Rotation:
+		{
+			CurrentTransform.SetRotation(InitialTransform.GetRotation());
+			break;
+		}
+		case ESlateTransformComponent::Scale:
+		{
+			CurrentTransform.SetScale3D(InitialTransform.GetScale3D());
+			break;
+		}
+	}
+}
+
+#endif
+
 FBoneChain* FRetargetDefinition::GetEditableBoneChainByName(FName ChainName)
 {
 	for (FBoneChain& Chain : BoneChains)
diff --git a/Engine/Plugins/Animation/IKRig/Source/IKRig/Public/IKRigDefinition.h b/Engine/Plugins/Animation/IKRig/Source/IKRig/Public/IKRigDefinition.h
--- a/Engine/Plugins/Animation/IKRig/Source/IKRig/Public/IKRigDefinition.h
+++ b/Engine/Plugins/Animation/IKRig/Source/IKRig/Public/IKRigDefinition.h
@@ -6,6 +6,9 @@
 #include "UObject/Object.h"
 #include "Interfaces/Interface_PreviewMeshProvider.h"
 #include "IKRigSkeleton.h"
+#if WITH_EDITOR
+#include "SAdvancedTransformInputBox.h"
+#endif
 #include "IKRigDefinition.generated.h"
 
 class UIKRigSolver;
@@ -17,6 +20,15 @@ enum class EIKRigGoalPreviewMode : uint8
 	Additive		UMETA(DisplayName = "Additive"),
 	Absolute		UMETA(DisplayName = "Absolute"),
 };
+
+namespace EIKRigTransformType
+{
+	enum Type : int8
+	{
+		Current,
+		Reference,
+	};
+}
 #endif
 
 UCLASS()
@@ -76,6 +88,29 @@ public:
 		SetFlags(RF_Transactional);
 	}
 	
+#endif
+
+#if WITH_EDITOR
+
+	TOptional<FTransform::FReal> GetNumericValue(
+		ESlateTransformComponent::Type Component,
+		ESlateRotationRepresentation::Type Representation,
+		ESlateTransformSubComponent::Type SubComponent,
+		EIKRigTransformType::Type TransformType
+	) const;
+	
+	void OnNumericValueChanged(
+		ESlateTransformComponent::Type Component,
+		ESlateRotationRepresentation::Type Representation,
+		ESlateTransformSubComponent::Type SubComponent,
+		FTransform::FReal Value,
+		ETextCommit::Type CommitType,
+		EIKRigTransformType::Type TransformType
+	);
+
+	bool TransformDiffersFromDefault(TSharedPtr<IPropertyHandle> PropertyHandle, ESlateTransformComponent::Type Component) const;
+	void ResetTransformToDefault(TSharedPtr<IPropertyHandle> PropertyHandle, ESlateTransformComponent::Type Component);
+	
 #endif
 };
 
diff --git a/Engine/Plugins/Animation/IKRig/Source/IKRigEditor/IKRigEditor.Build.cs b/Engine/Plugins/Animation/IKRig/Source/IKRigEditor/IKRigEditor.Build.cs
--- a/Engine/Plugins/Animation/IKRig/Source/IKRigEditor/IKRigEditor.Build.cs
+++ b/Engine/Plugins/Animation/IKRig/Source/IKRigEditor/IKRigEditor.Build.cs
@@ -45,6 +45,8 @@ namespace UnrealBuildTool.Rules
 					"IKRig",
 					"IKRigDeveloper",
 					"ToolWidgets",
+					"AnimationCore",
+                    "AnimationWidgets",
 				}
             );
 
diff --git a/Engine/Plugins/Animation/IKRig/Source/IKRigEditor/Private/IKRigEditor.cpp b/Engine/Plugins/Animation/IKRig/Source/IKRigEditor/Private/IKRigEditor.cpp
--- a/Engine/Plugins/Animation/IKRig/Source/IKRigEditor/Private/IKRigEditor.cpp
+++ b/Engine/Plugins/Animation/IKRig/Source/IKRigEditor/Private/IKRigEditor.cpp
@@ -15,6 +15,8 @@
 #include "RigEditor/IKRigCommands.h"
 #include "RigEditor/IKRigEditMode.h"
 #include "RigEditor/IKRigSkeletonCommands.h"
+#include "RigEditor/IKRigDetailCustomizations.h"
+#include "RigEditor/IKRigEditorController.h"
 
 IMPLEMENT_MODULE(FIKRigEditor, IKRigEditor)
 
@@ -42,6 +44,13 @@ void FIKRigEditor::StartupModule()
 	// register custom editor modes
 	FEditorModeRegistry::Get().RegisterMode<FIKRigEditMode>(FIKRigEditMode::ModeName, LOCTEXT("IKRigEditMode", "IKRig"), FSlateIcon(), false);
 	FEditorModeRegistry::Get().RegisterMode<FIKRetargetEditMode>(FIKRetargetEditMode::ModeName, LOCTEXT("IKRetargetEditMode", "IKRetarget"), FSlateIcon(), false);
+
+	// register detail customizations
+	FPropertyEditorModule& PropertyEditorModule = FModuleManager::LoadModuleChecked<FPropertyEditorModule>("PropertyEditor");
+	ClassesToUnregisterOnShutdown.Add(UIKRigBoneDetails::StaticClass()->GetFName());
+	PropertyEditorModule.RegisterCustomClassLayout(ClassesToUnregisterOnShutdown.Last(), FOnGetDetailCustomizationInstance::CreateStatic(&FIKRigGenericDetailCustomization::MakeInstance));
+	ClassesToUnregisterOnShutdown.Add(UIKRigEffectorGoal::StaticClass()->GetFName());
+	PropertyEditorModule.RegisterCustomClassLayout(ClassesToUnregisterOnShutdown.Last(), FOnGetDetailCustomizationInstance::CreateStatic(&FIKRigGenericDetailCustomization::MakeInstance));
 }
 
 void FIKRigEditor::ShutdownModule()
@@ -53,9 +62,6 @@ void FIKRigEditor::ShutdownModule()
 	FEditorModeRegistry::Get().UnregisterMode(FIKRigEditMode::ModeName);
 	FEditorModeRegistry::Get().UnregisterMode(FIKRetargetEditMode::ModeName);
 
-	FPropertyEditorModule& PropertyModule = FModuleManager::GetModuleChecked<FPropertyEditorModule>("PropertyEditor");
-	PropertyModule.UnregisterCustomPropertyTypeLayout("IKRigEffector");
-
 	// unregister IKRigDefinition asset action
 	if (IKRigDefinitionAssetAction.IsValid())
 	{
diff --git a/Engine/Plugins/Animation/IKRig/Source/IKRigEditor/Private/RigEditor/IKRigDetailCustomizations.cpp b/Engine/Plugins/Animation/IKRig/Source/IKRigEditor/Private/RigEditor/IKRigDetailCustomizations.cpp
new file mode 100644
--- /dev/null
+++ b/Engine/Plugins/Animation/IKRig/Source/IKRigEditor/Private/RigEditor/IKRigDetailCustomizations.cpp
@@ -0,0 +1,378 @@
+// Copyright Epic Games, Inc. All Rights Reserved.
+
+#include "IKRigDetailCustomizations.h"
+
+#include "DetailLayoutBuilder.h"
+#include "DetailCategoryBuilder.h"
+#include "DetailWidgetRow.h"
+#include "IDetailGroup.h"
+#include "IDetailPropertyRow.h"
+#include "Widgets/SWidget.h"
+#include "Widgets/Input/SSegmentedControl.h"
+
+#define LOCTEXT_NAMESPACE "IKRigDetailCustomizations"
+
+namespace IKRigDetailCustomizationsConstants
+{
+	static const float ItemWidth = 125.0f;
+}
+
+void FIKRigGenericDetailCustomization::CustomizeDetails(IDetailLayoutBuilder& DetailBuilder)
+{
+	TArray<TWeakObjectPtr<UObject>> ObjectsBeingCustomized = DetailBuilder.GetSelectedObjects();
+
+	// make sure all types are of the same class
+	UClass* DetailsClass = nullptr;
+	for(const TWeakObjectPtr<UObject>& ObjectBeingCustomized : ObjectsBeingCustomized)
+	{
+		if(!ObjectBeingCustomized.IsValid())
+		{
+			continue;
+		}
+
+		if(DetailsClass)
+		{
+			if(ObjectBeingCustomized->GetClass() != DetailsClass)
+			{
+				// multiple different things - fall back to default
+				// details panel behavior
+				return;
+			}
+		}
+		else
+		{
+			DetailsClass = ObjectBeingCustomized->GetClass();
+		}
+	}
+
+	// assuming the classes are all the same 
+	for(const TWeakObjectPtr<UObject>& ObjectBeingCustomized : ObjectsBeingCustomized)
+	{
+		if(!ObjectBeingCustomized.IsValid())
+		{
+			continue;
+		}
+
+		if(ObjectBeingCustomized->IsA<UIKRigBoneDetails>())
+		{
+			CustomizeDetailsForClass<UIKRigBoneDetails>(DetailBuilder, ObjectsBeingCustomized);	
+		}
+		else if(ObjectBeingCustomized->IsA<UIKRigEffectorGoal>())
+		{
+			CustomizeDetailsForClass<UIKRigEffectorGoal>(DetailBuilder, ObjectsBeingCustomized);	
+		}
+
+		break;
+	}
+}
+
+template <>
+void FIKRigGenericDetailCustomization::CustomizeDetailsForClass<UIKRigBoneDetails>(
+	IDetailLayoutBuilder& DetailBuilder,
+	TArray<TWeakObjectPtr<UObject>> ObjectsBeingCustomized)
+{
+	UIKRigBoneDetails* BoneDetails = Cast<UIKRigBoneDetails>(ObjectsBeingCustomized[0].Get());
+	if(BoneDetails == nullptr)
+	{
+		return;
+		
+	}
+	const TArray<FText> ButtonLabels =
+	{
+		LOCTEXT("CurrentTransform", "Current"),
+		LOCTEXT("ReferenceTransform", "Reference")
+	};
+	
+	const TArray<FText> ButtonTooltips =
+	{
+		LOCTEXT("CurrentBoneTransformTooltip", "The current transform of the bone"),
+		LOCTEXT("ReferenceBoneTransformTooltip", "The reference transform of a bone")
+	};
+
+	static TSharedRef<TArray<EIKRigTransformType::Type>> VisibleTransforms =
+		MakeShareable(new TArray<EIKRigTransformType::Type>({EIKRigTransformType::Current}));
+
+	TArray<TSharedRef<IPropertyHandle>> Properties;
+	Properties.Add(DetailBuilder.GetProperty(GET_MEMBER_NAME_CHECKED(UIKRigBoneDetails, CurrentTransform)));
+	Properties.Add(DetailBuilder.GetProperty(GET_MEMBER_NAME_CHECKED(UIKRigBoneDetails, ReferenceTransform)));
+
+	for(TSharedRef<IPropertyHandle> Property : Properties)
+	{
+		DetailBuilder.HideProperty(Property);
+	}
+
+	TSharedPtr<SSegmentedControl<EIKRigTransformType::Type>> TransformChoiceWidget;
+
+	DetailBuilder.EditCategory(TEXT("Selection")).SetSortOrder(1);
+
+	IDetailCategoryBuilder& CategoryBuilder = DetailBuilder.EditCategory(TEXT("Transforms"));
+	CategoryBuilder.SetSortOrder(2);
+	CategoryBuilder.AddCustomRow(FText::FromString(TEXT("TransformType")))
+	.ValueContent()
+	.MinDesiredWidth(375.f)
+	.MaxDesiredWidth(375.f)
+	.HAlign(HAlign_Left)
+	[
+		SNew(SHorizontalBox)
+		+SHorizontalBox::Slot()
+		.AutoWidth()
+		.HAlign(HAlign_Left)
+		.VAlign(VAlign_Center)
+		[
+			SAssignNew(TransformChoiceWidget, SSegmentedControl<EIKRigTransformType::Type>)
+			.SupportsMultiSelection(true)
+			.Values_Lambda([]()
+			{
+				return VisibleTransforms.Get();
+			})
+			.OnValuesChanged_Lambda([](TArray<EIKRigTransformType::Type> Values)
+			{
+				VisibleTransforms.Get() = Values;
+			})
+			+ SSegmentedControl<EIKRigTransformType::Type>::Slot(EIKRigTransformType::Current)
+			.Text(ButtonLabels[0])
+			.ToolTip(ButtonTooltips[0])
+			+ SSegmentedControl<EIKRigTransformType::Type>::Slot(EIKRigTransformType::Reference)
+			.Text(ButtonLabels[1])
+			.ToolTip(ButtonTooltips[1])
+		]
+	];
+
+	SAdvancedTransformInputBox<FTransform>::FArguments TransformWidgetArgs = SAdvancedTransformInputBox<FTransform>::FArguments()
+	.IsEnabled(false)
+	.DisplayRelativeWorld(true)
+	.DisplayScaleLock(false)
+	.AllowEditRotationRepresentation(true)
+	.Font(IDetailLayoutBuilder::GetDetailFont())
+	.UseQuaternionForRotation(true);
+	
+	for(int32 PropertyIndex=0;PropertyIndex<Properties.Num();PropertyIndex++)
+	{
+		const EIKRigTransformType::Type TransformType = (EIKRigTransformType::Type)PropertyIndex; 
+
+		TransformWidgetArgs
+		.OnGetIsComponentRelative_UObject(BoneDetails, &UIKRigBoneDetails::IsComponentRelative, TransformType)
+		.OnIsComponentRelativeChanged_UObject(BoneDetails, &UIKRigBoneDetails::OnComponentRelativeChanged, TransformType)
+		.Transform_UObject(BoneDetails, &UIKRigBoneDetails::GetTransform, TransformType);
+		
+		IDetailGroup& Group = CategoryBuilder.AddGroup(*ButtonLabels[PropertyIndex].ToString(), ButtonLabels[PropertyIndex], false, true);
+		Group.HeaderRow()
+		.Visibility( TAttribute<EVisibility>::CreateLambda([TransformChoiceWidget, TransformType]() -> EVisibility
+		{
+			return TransformChoiceWidget->HasValue(TransformType) ? EVisibility::Visible : EVisibility::Collapsed;
+		}))
+		.NameContent()
+		[
+			SNew(STextBlock)
+			.Font(IDetailLayoutBuilder::GetDetailFont())
+			.Text(ButtonLabels[PropertyIndex])
+			.ToolTipText(ButtonTooltips[PropertyIndex])
+		];
+
+		Group.AddPropertyRow(Properties[PropertyIndex])
+		.CustomWidget()
+		.NameContent()
+		.HAlign(HAlign_Fill)
+		[
+			SAdvancedTransformInputBox<FTransform>::ConstructLabel(TransformWidgetArgs, ESlateTransformComponent::Location)
+		]
+		.ValueContent()
+		.MinDesiredWidth(IKRigDetailCustomizationsConstants::ItemWidth * 3.0f)
+		.MaxDesiredWidth(IKRigDetailCustomizationsConstants::ItemWidth * 3.0f)
+		[
+			SAdvancedTransformInputBox<FTransform>::ConstructWidget(TransformWidgetArgs, ESlateTransformComponent::Location)
+		];
+
+		Group.AddPropertyRow(Properties[PropertyIndex])
+		.CustomWidget()
+		.NameContent()
+		.HAlign(HAlign_Fill)
+		[
+			SAdvancedTransformInputBox<FTransform>::ConstructLabel(TransformWidgetArgs, ESlateTransformComponent::Rotation)
+		]
+		.ValueContent()
+		.MinDesiredWidth(IKRigDetailCustomizationsConstants::ItemWidth * 3.0f)
+		.MaxDesiredWidth(IKRigDetailCustomizationsConstants::ItemWidth * 3.0f)
+		[
+			SAdvancedTransformInputBox<FTransform>::ConstructWidget(TransformWidgetArgs, ESlateTransformComponent::Rotation)
+		];
+
+		Group.AddPropertyRow(Properties[PropertyIndex])
+		.CustomWidget()
+		.NameContent()
+		.HAlign(HAlign_Fill)
+		[
+			SAdvancedTransformInputBox<FTransform>::ConstructLabel(TransformWidgetArgs, ESlateTransformComponent::Scale)
+		]
+		.ValueContent()
+		.MinDesiredWidth(IKRigDetailCustomizationsConstants::ItemWidth * 3.0f)
+		.MaxDesiredWidth(IKRigDetailCustomizationsConstants::ItemWidth * 3.0f)
+		[
+			SAdvancedTransformInputBox<FTransform>::ConstructWidget(TransformWidgetArgs, ESlateTransformComponent::Scale)
+		];
+	}
+}
+
+template <>
+void FIKRigGenericDetailCustomization::CustomizeDetailsForClass<UIKRigEffectorGoal>(
+	IDetailLayoutBuilder& DetailBuilder,
+	TArray<TWeakObjectPtr<UObject>> ObjectsBeingCustomized)
+{
+	UIKRigEffectorGoal* EffectorGoal = Cast<UIKRigEffectorGoal>(ObjectsBeingCustomized[0].Get());
+	if(EffectorGoal == nullptr)
+	{
+		return;
+		
+	}
+	const TArray<FText> ButtonLabels =
+	{
+		LOCTEXT("CurrentTransform", "Current"),
+		LOCTEXT("ReferenceTransform", "Reference")
+	};
+	
+	const TArray<FText> ButtonTooltips =
+	{
+		LOCTEXT("CurrentBoneTransformTooltip", "The current transform of the bone"),
+		LOCTEXT("ReferenceBoneTransformTooltip", "The reference transform of a bone")
+	};
+
+	static TSharedRef<TArray<EIKRigTransformType::Type>> VisibleTransforms =
+		MakeShareable(new TArray<EIKRigTransformType::Type>({EIKRigTransformType::Current}));
+
+	TArray<TSharedRef<IPropertyHandle>> Properties;
+	Properties.Add(DetailBuilder.GetProperty(GET_MEMBER_NAME_CHECKED(UIKRigEffectorGoal, CurrentTransform)));
+	Properties.Add(DetailBuilder.GetProperty(GET_MEMBER_NAME_CHECKED(UIKRigEffectorGoal, InitialTransform)));
+
+	for(TSharedRef<IPropertyHandle> Property : Properties)
+	{
+		DetailBuilder.HideProperty(Property);
+	}
+
+	TSharedPtr<SSegmentedControl<EIKRigTransformType::Type>> TransformChoiceWidget;
+
+	DetailBuilder.EditCategory(TEXT("Goal Settings")).SetSortOrder(1);
+	DetailBuilder.EditCategory(TEXT("Viewport Goal Settings")).SetSortOrder(3);
+
+	IDetailCategoryBuilder& CategoryBuilder = DetailBuilder.EditCategory(TEXT("Transforms"));
+	CategoryBuilder.SetSortOrder(2);
+	CategoryBuilder.AddCustomRow(FText::FromString(TEXT("TransformType")))
+	.ValueContent()
+	.MinDesiredWidth(375.f)
+	.MaxDesiredWidth(375.f)
+	.HAlign(HAlign_Left)
+	[
+		SNew(SHorizontalBox)
+		+SHorizontalBox::Slot()
+		.AutoWidth()
+		.HAlign(HAlign_Left)
+		.VAlign(VAlign_Center)
+		[
+			SAssignNew(TransformChoiceWidget, SSegmentedControl<EIKRigTransformType::Type>)
+			.SupportsMultiSelection(true)
+			.Values_Lambda([]()
+			{
+				return VisibleTransforms.Get();
+			})
+			.OnValuesChanged_Lambda([](TArray<EIKRigTransformType::Type> Values)
+			{
+				VisibleTransforms.Get() = Values;
+			})
+			+ SSegmentedControl<EIKRigTransformType::Type>::Slot(EIKRigTransformType::Current)
+			.Text(ButtonLabels[0])
+			.ToolTip(ButtonTooltips[0])
+			+ SSegmentedControl<EIKRigTransformType::Type>::Slot(EIKRigTransformType::Reference)
+			.Text(ButtonLabels[1])
+			.ToolTip(ButtonTooltips[1])
+		]
+	];
+
+	SAdvancedTransformInputBox<FTransform>::FArguments TransformWidgetArgs = SAdvancedTransformInputBox<FTransform>::FArguments()
+	.IsEnabled(true)
+	.DisplayRelativeWorld(false)
+	.DisplayScaleLock(true)
+	.AllowEditRotationRepresentation(true)
+	.Font(IDetailLayoutBuilder::GetDetailFont())
+	.UseQuaternionForRotation(true);
+	
+	for(int32 PropertyIndex=0;PropertyIndex<Properties.Num();PropertyIndex++)
+	{
+		const EIKRigTransformType::Type TransformType = (EIKRigTransformType::Type)PropertyIndex;
+
+		if(PropertyIndex > 0)
+		{
+			TransformWidgetArgs
+			.IsEnabled(false)
+			.DisplayScaleLock(false);
+		}
+
+		TransformWidgetArgs
+		.OnGetNumericValue_UObject(EffectorGoal, &UIKRigEffectorGoal::GetNumericValue, TransformType);
+		
+		IDetailGroup& Group = CategoryBuilder.AddGroup(*ButtonLabels[PropertyIndex].ToString(), ButtonLabels[PropertyIndex], false, true);
+		Group.HeaderRow()
+		.Visibility( TAttribute<EVisibility>::CreateLambda([TransformChoiceWidget, TransformType]() -> EVisibility
+		{
+			return TransformChoiceWidget->HasValue(TransformType) ? EVisibility::Visible : EVisibility::Collapsed;
+		}))
+		.NameContent()
+		[
+			SNew(STextBlock)
+			.Font(IDetailLayoutBuilder::GetDetailFont())
+			.Text(ButtonLabels[PropertyIndex])
+			.ToolTipText(ButtonTooltips[PropertyIndex])
+		];
+
+		Group.AddPropertyRow(Properties[PropertyIndex])
+		.OverrideResetToDefault(FResetToDefaultOverride::Create(
+			FIsResetToDefaultVisible::CreateUObject(EffectorGoal, &UIKRigEffectorGoal::TransformDiffersFromDefault, ESlateTransformComponent::Location),
+			FResetToDefaultHandler::CreateUObject(EffectorGoal, &UIKRigEffectorGoal::ResetTransformToDefault, ESlateTransformComponent::Location)))
+		.CustomWidget()
+		.NameContent()
+		.HAlign(HAlign_Fill)
+		[
+			SAdvancedTransformInputBox<FTransform>::ConstructLabel(TransformWidgetArgs, ESlateTransformComponent::Location)
+		]
+		.ValueContent()
+		.MinDesiredWidth(IKRigDetailCustomizationsConstants::ItemWidth * 3.0f)
+		.MaxDesiredWidth(IKRigDetailCustomizationsConstants::ItemWidth * 3.0f)
+		[
+			SAdvancedTransformInputBox<FTransform>::ConstructWidget(TransformWidgetArgs, ESlateTransformComponent::Location)
+		];
+
+		Group.AddPropertyRow(Properties[PropertyIndex])
+		.OverrideResetToDefault(FResetToDefaultOverride::Create(
+			FIsResetToDefaultVisible::CreateUObject(EffectorGoal, &UIKRigEffectorGoal::TransformDiffersFromDefault, ESlateTransformComponent::Rotation),
+			FResetToDefaultHandler::CreateUObject(EffectorGoal, &UIKRigEffectorGoal::ResetTransformToDefault, ESlateTransformComponent::Rotation)))
+		.CustomWidget()
+		.NameContent()
+		.HAlign(HAlign_Fill)
+		[
+			SAdvancedTransformInputBox<FTransform>::ConstructLabel(TransformWidgetArgs, ESlateTransformComponent::Rotation)
+		]
+		.ValueContent()
+		.MinDesiredWidth(IKRigDetailCustomizationsConstants::ItemWidth * 3.0f)
+		.MaxDesiredWidth(IKRigDetailCustomizationsConstants::ItemWidth * 3.0f)
+		[
+			SAdvancedTransformInputBox<FTransform>::ConstructWidget(TransformWidgetArgs, ESlateTransformComponent::Rotation)
+		];
+
+		Group.AddPropertyRow(Properties[PropertyIndex])
+		.OverrideResetToDefault(FResetToDefaultOverride::Create(
+			FIsResetToDefaultVisible::CreateUObject(EffectorGoal, &UIKRigEffectorGoal::TransformDiffersFromDefault, ESlateTransformComponent::Scale),
+			FResetToDefaultHandler::CreateUObject(EffectorGoal, &UIKRigEffectorGoal::ResetTransformToDefault, ESlateTransformComponent::Scale)))
+		.CustomWidget()
+		.NameContent()
+		.HAlign(HAlign_Fill)
+		[
+			SAdvancedTransformInputBox<FTransform>::ConstructLabel(TransformWidgetArgs, ESlateTransformComponent::Scale)
+		]
+		.ValueContent()
+		.MinDesiredWidth(IKRigDetailCustomizationsConstants::ItemWidth * 3.0f)
+		.MaxDesiredWidth(IKRigDetailCustomizationsConstants::ItemWidth * 3.0f)
+		[
+			SAdvancedTransformInputBox<FTransform>::ConstructWidget(TransformWidgetArgs, ESlateTransformComponent::Scale)
+		];
+	}
+}
+
+#undef LOCTEXT_NAMESPACE
diff --git a/Engine/Plugins/Animation/IKRig/Source/IKRigEditor/Private/RigEditor/IKRigDetailCustomizations.h b/Engine/Plugins/Animation/IKRig/Source/IKRigEditor/Private/RigEditor/IKRigDetailCustomizations.h
new file mode 100644
--- /dev/null
+++ b/Engine/Plugins/Animation/IKRig/Source/IKRigEditor/Private/RigEditor/IKRigDetailCustomizations.h
@@ -0,0 +1,30 @@
+// Copyright Epic Games, Inc. All Rights Reserved.
+
+#pragma once
+
+#include "CoreMinimal.h"
+#include "IDetailCustomization.h"
+#include "RigEditor/IKRigEditorController.h"
+
+class FIKRigGenericDetailCustomization : public IDetailCustomization
+{
+public:
+
+	static TSharedRef<IDetailCustomization> MakeInstance()
+	{
+		return MakeShareable(new FIKRigGenericDetailCustomization);
+	}
+
+	/** IDetailCustomization interface */
+	virtual void CustomizeDetails(IDetailLayoutBuilder& DetailBuilder) override;
+
+	/** Per class customization */
+	template<typename ClassToCustomize>
+	void CustomizeDetailsForClass(IDetailLayoutBuilder& DetailBuilder, TArray<TWeakObjectPtr<UObject>> ObjectsBeingCustomized) {}
+
+	/** Template specializations */
+	template<>
+	void CustomizeDetailsForClass<UIKRigBoneDetails>(IDetailLayoutBuilder& DetailBuilder, TArray<TWeakObjectPtr<UObject>> ObjectsBeingCustomized);
+	template<>
+	void CustomizeDetailsForClass<UIKRigEffectorGoal>(IDetailLayoutBuilder& DetailBuilder, TArray<TWeakObjectPtr<UObject>> ObjectsBeingCustomized);
+};
\ No newline at end of file
diff --git a/Engine/Plugins/Animation/IKRig/Source/IKRigEditor/Private/RigEditor/IKRigEditorController.cpp b/Engine/Plugins/Animation/IKRig/Source/IKRigEditor/Private/RigEditor/IKRigEditorController.cpp
--- a/Engine/Plugins/Animation/IKRig/Source/IKRigEditor/Private/RigEditor/IKRigEditorController.cpp
+++ b/Engine/Plugins/Animation/IKRig/Source/IKRigEditor/Private/RigEditor/IKRigEditorController.cpp
@@ -14,6 +14,95 @@
 
 #define LOCTEXT_NAMESPACE "IKRigEditorController"
 
+bool UIKRigBoneDetails::CurrentTransformRelative[3] = {true, true, true};
+bool UIKRigBoneDetails::ReferenceTransformRelative[3] = {true, true, true};
+
+TOptional<FTransform> UIKRigBoneDetails::GetTransform(EIKRigTransformType::Type TransformType) const
+{
+	if(!AnimInstancePtr.IsValid() || !AssetPtr.IsValid())
+	{
+		return TOptional<FTransform>();
+	}
+	
+	FTransform LocalTransform = FTransform::Identity;
+	FTransform GlobalTransform = FTransform::Identity;
+	const bool* IsRelative = nullptr;
+	
+	const int32 BoneIndex = AssetPtr->Skeleton.GetBoneIndexFromName(SelectedBone);
+	if(BoneIndex == INDEX_NONE)
+	{
+		return TOptional<FTransform>();
+	}
+
+	switch(TransformType)
+	{
+		case EIKRigTransformType::Current:
+		{
+			IsRelative = CurrentTransformRelative;
+			GlobalTransform = AnimInstancePtr->GetSkelMeshComponent()->GetBoneTransform(BoneIndex);
+			LocalTransform = AnimInstancePtr->GetSkelMeshComponent()->GetBoneSpaceTransforms()[BoneIndex];
+			break;
+		}
+		case EIKRigTransformType::Reference:
+		{
+			IsRelative = ReferenceTransformRelative;
+			GlobalTransform = AssetPtr->Skeleton.RefPoseGlobal[BoneIndex];
+			LocalTransform = GlobalTransform;
+			const int32 ParentBoneIndex = AssetPtr->Skeleton.ParentIndices[BoneIndex];
+			if(ParentBoneIndex != INDEX_NONE)
+			{
+				const FTransform ParentTransform = AssetPtr->Skeleton.RefPoseGlobal[ParentBoneIndex];;
+				LocalTransform = GlobalTransform.GetRelativeTransform(ParentTransform);
+			}
+			break;
+		}
+	}
+
+	FTransform Transform = LocalTransform;
+	if(!IsRelative[0]) Transform.SetLocation(GlobalTransform.GetLocation());
+	if(!IsRelative[1]) Transform.SetRotation(GlobalTransform.GetRotation());
+	if(!IsRelative[2]) Transform.SetScale3D(GlobalTransform.GetScale3D());
+	return Transform;
+}
+
+bool UIKRigBoneDetails::IsComponentRelative(
+	ESlateTransformComponent::Type Component,
+	EIKRigTransformType::Type TransformType) const
+{
+	switch(TransformType)
+	{
+		case EIKRigTransformType::Current:
+		{
+			return CurrentTransformRelative[(int32)Component]; 
+		}
+		case EIKRigTransformType::Reference:
+		{
+			return ReferenceTransformRelative[(int32)Component]; 
+		}
+	}
+	return true;
+}
+
+void UIKRigBoneDetails::OnComponentRelativeChanged(
+	ESlateTransformComponent::Type Component,
+	bool bIsRelative,
+	EIKRigTransformType::Type TransformType)
+{
+	switch(TransformType)
+	{
+		case EIKRigTransformType::Current:
+		{
+			CurrentTransformRelative[(int32)Component] = bIsRelative;
+			break; 
+		}
+		case EIKRigTransformType::Reference:
+		{
+			ReferenceTransformRelative[(int32)Component] = bIsRelative;
+			break; 
+		}
+	}
+}
+
 void FIKRigEditorController::Initialize(TSharedPtr<FIKRigEditorToolkit> Toolkit, UIKRigDefinition* IKRigAsset)
 {
 	EditorToolkit = Toolkit;
@@ -41,6 +130,10 @@ void FIKRigEditorController::OnIKRigNeedsInitialized(UIKRigDefinition* ModifiedI
 
 	// Initialize editor's instances on request
 	InitializeSolvers();
+
+	// update the bone details so it can pull on the current data
+	BoneDetails->AnimInstancePtr = AnimInstance;
+	BoneDetails->AssetPtr = ModifiedIKRig;
 }
 
 void FIKRigEditorController::Reset() const
diff --git a/Engine/Plugins/Animation/IKRig/Source/IKRigEditor/Private/RigEditor/IKRigToolkit.cpp b/Engine/Plugins/Animation/IKRig/Source/IKRigEditor/Private/RigEditor/IKRigToolkit.cpp
--- a/Engine/Plugins/Animation/IKRig/Source/IKRigEditor/Private/RigEditor/IKRigToolkit.cpp
+++ b/Engine/Plugins/Animation/IKRig/Source/IKRigEditor/Private/RigEditor/IKRigToolkit.cpp
@@ -261,6 +261,7 @@ void FIKRigEditorToolkit::SetupAnimInstance()
 	InAnimInstance->SetIKRigAsset(EditorController->AssetController->GetAsset());
 	EditorController->SkelMeshComponent->PreviewInstance = InAnimInstance;
 	InAnimInstance->InitializeAnimation();
+	EditorController->OnIKRigNeedsInitialized(EditorController->AssetController->GetAsset());
 }
 
 
diff --git a/Engine/Plugins/Animation/IKRig/Source/IKRigEditor/Public/IKRigEditor.h b/Engine/Plugins/Animation/IKRig/Source/IKRigEditor/Public/IKRigEditor.h
--- a/Engine/Plugins/Animation/IKRig/Source/IKRigEditor/Public/IKRigEditor.h
+++ b/Engine/Plugins/Animation/IKRig/Source/IKRigEditor/Public/IKRigEditor.h
@@ -15,4 +15,6 @@ private:
 	TSharedPtr<class FAssetTypeActions_AnimationAssetRetarget> RetargetAnimationAssetAction;
 	TSharedPtr<class FAssetTypeActions_IKRigDefinition> IKRigDefinitionAssetAction;
 	TSharedPtr<class FAssetTypeActions_IKRetargeter> IKRetargeterAssetAction;
+
+	TArray<FName> ClassesToUnregisterOnShutdown;
 };
diff --git a/Engine/Plugins/Animation/IKRig/Source/IKRigEditor/Public/RigEditor/IKRigEditorController.h b/Engine/Plugins/Animation/IKRig/Source/IKRigEditor/Public/RigEditor/IKRigEditorController.h
--- a/Engine/Plugins/Animation/IKRig/Source/IKRigEditor/Public/RigEditor/IKRigEditorController.h
+++ b/Engine/Plugins/Animation/IKRig/Source/IKRigEditor/Public/RigEditor/IKRigEditorController.h
@@ -6,6 +6,7 @@
 #include "SIKRigRetargetChainList.h"
 #include "Templates/SubclassOf.h"
 #include "Animation/AnimationAsset.h"
+#include "SAdvancedTransformInputBox.h"
 
 #include "IKRigEditorController.generated.h"
 
@@ -31,21 +32,36 @@ class IKRIGEDITOR_API UIKRigBoneDetails : public UObject
 	GENERATED_BODY()
 
 public:
-	
+
 	// todo update bone info automatically using something else
 	void SetBone(const FName& BoneName)
 	{
 		SelectedBone = BoneName;
 	};
 
-	UPROPERTY(VisibleAnywhere, Category = "Bone Transforms")
+	UPROPERTY(VisibleAnywhere, Category = "Selection")
 	FName SelectedBone;
 	
 	UPROPERTY(VisibleAnywhere, Category = "Bone Transforms")
 	FTransform CurrentTransform;
 
 	UPROPERTY(VisibleAnywhere, Category = "Bone Transforms")
-	FTransform DefaultTransform;
+	FTransform ReferenceTransform;
+
+	UPROPERTY()
+	TWeakObjectPtr<UAnimInstance> AnimInstancePtr;
+
+	UPROPERTY()
+	TWeakObjectPtr<UIKRigDefinition> AssetPtr;
+
+	TOptional<FTransform> GetTransform(EIKRigTransformType::Type TransformType) const;
+	bool IsComponentRelative(ESlateTransformComponent::Type Component, EIKRigTransformType::Type TransformType) const;
+	void OnComponentRelativeChanged(ESlateTransformComponent::Type Component, bool bIsRelative, EIKRigTransformType::Type TransformType);
+
+private:
+	
+	static bool CurrentTransformRelative[3];
+	static bool ReferenceTransformRelative[3];
 };
 
 /** a home for cross-widget communication to synchronize state across all tabs and viewport */
