Fixing crash when changing target mesh in IK Retargeter.

diff --git a/Engine/Plugins/Animation/IKRig/Source/IKRig/Private/IKRigProcessor.cpp b/Engine/Plugins/Animation/IKRig/Source/IKRig/Private/IKRigProcessor.cpp
--- a/Engine/Plugins/Animation/IKRig/Source/IKRig/Private/IKRigProcessor.cpp
+++ b/Engine/Plugins/Animation/IKRig/Source/IKRig/Private/IKRigProcessor.cpp
@@ -265,6 +265,11 @@ void UIKRigProcessor::CopyAllInputsFromSourceAssetAtRuntime(const UIKRigDefiniti
 {
 	check(SourceAsset)
 	
+	if (!bInitialized)
+	{
+		return;
+	}
+	
 	// copy goal settings
 	const TArray<UIKRigEffectorGoal*>& AssetGoals =  SourceAsset->GetGoalArray();
 	for (const UIKRigEffectorGoal* AssetGoal : AssetGoals)
diff --git a/Engine/Plugins/Animation/IKRig/Source/IKRig/Private/Retargeter/IKRetargetProcessor.cpp b/Engine/Plugins/Animation/IKRig/Source/IKRig/Private/Retargeter/IKRetargetProcessor.cpp
--- a/Engine/Plugins/Animation/IKRig/Source/IKRig/Private/Retargeter/IKRetargetProcessor.cpp
+++ b/Engine/Plugins/Animation/IKRig/Source/IKRig/Private/Retargeter/IKRetargetProcessor.cpp
@@ -214,7 +214,7 @@ void FTargetSkeleton::GenerateRetargetPose(const FIKRetargetPose* InRetargetPose
 		RetargetLocalPose[BoneIndex].SetRotation(LocalBoneRotation);
 	}
 
-	UpdateGlobalTransformsBelowBone(RootBoneIndex, RetargetLocalPose, RetargetGlobalPose);
+	UpdateGlobalTransformsBelowBone(0, RetargetLocalPose, RetargetGlobalPose);
 }
 
 void FTargetSkeleton::Reset()
