Pole solver's affected chain consistency

diff --git a/Engine/Plugins/Animation/IKRig/Source/IKRig/Private/Solvers/IKRig_PoleSolver.cpp b/Engine/Plugins/Animation/IKRig/Source/IKRig/Private/Solvers/IKRig_PoleSolver.cpp
--- a/Engine/Plugins/Animation/IKRig/Source/IKRig/Private/Solvers/IKRig_PoleSolver.cpp
+++ b/Engine/Plugins/Animation/IKRig/Source/IKRig/Private/Solvers/IKRig_PoleSolver.cpp
@@ -228,7 +228,20 @@ UObject* UIKRig_PoleSolver::GetGoalSettings(const FName& GoalName) const
 
 bool UIKRig_PoleSolver::IsBoneAffectedBySolver(const FName& BoneName, const FIKRigSkeleton& IKRigSkeleton) const
 {
-	return IKRigSkeleton.IsBoneInDirectLineage(BoneName, RootName);
+	const bool bAffected = IKRigSkeleton.IsBoneInDirectLineage(BoneName, RootName);
+	if (!bAffected)
+	{
+		return false;
+	}
+	
+	const int32 EndIndex = Chain.IsEmpty() ? INDEX_NONE : Chain.Last();
+	if (EndIndex == INDEX_NONE)
+	{
+		return false;
+	}
+
+	const int32 ChildIndex = IKRigSkeleton.GetBoneIndexFromName(BoneName);
+	return ChildIndex <= EndIndex;
 }
 
 void UIKRig_PoleSolver::SetRootBone(const FName& RootBoneName)
