Fix IKRig crash when changing target mesh in retargeter to same skeleton, different number of joints.

diff --git a/Engine/Plugins/Animation/IKRig/Source/IKRig/Private/AnimNodes/AnimNode_RetargetPoseFromMesh.cpp b/Engine/Plugins/Animation/IKRig/Source/IKRig/Private/AnimNodes/AnimNode_RetargetPoseFromMesh.cpp
--- a/Engine/Plugins/Animation/IKRig/Source/IKRig/Private/AnimNodes/AnimNode_RetargetPoseFromMesh.cpp
+++ b/Engine/Plugins/Animation/IKRig/Source/IKRig/Private/AnimNodes/AnimNode_RetargetPoseFromMesh.cpp
@@ -34,17 +34,15 @@ void FAnimNode_RetargetPoseFromMesh::Evaluate_AnyThread(FPoseContext& Output)
 	{
 		return;
 	}
-	
-	if (!Processor->IsInitialized())
-	{
-		Output.ResetToRefPose();
-		return;
-	}
 
-	if (SourceMeshComponentSpaceBoneTransforms.IsEmpty())
+	const bool bIsInitialized = Processor->IsInitialized();
+	const bool bInitializedWithSameMesh = Processor->GetTargetSkeleton().SkeletalMesh == Output.AnimInstanceProxy->GetSkelMeshComponent()->SkeletalMesh;
+	// it's possible in editor to have anim instances initialized before PreUpdate() is called
+	// which results in trying to run the retargeter without an source pose to copy from
+	const bool bSourceMeshBonesCopied = !SourceMeshComponentSpaceBoneTransforms.IsEmpty();
+	if (!(bIsInitialized && bInitializedWithSameMesh && bSourceMeshBonesCopied))
 	{
-		// it's possible in editor to have anim instances initialized before PreUpdate() is called
-		// which results in trying to run the retargeter without an source pose to copy from
+		Output.ResetToRefPose();
 		return;
 	}
 
diff --git a/Engine/Plugins/Animation/IKRig/Source/IKRig/Private/IKRigProcessor.cpp b/Engine/Plugins/Animation/IKRig/Source/IKRig/Private/IKRigProcessor.cpp
--- a/Engine/Plugins/Animation/IKRig/Source/IKRig/Private/IKRigProcessor.cpp
+++ b/Engine/Plugins/Animation/IKRig/Source/IKRig/Private/IKRigProcessor.cpp
@@ -20,6 +20,9 @@ void UIKRigProcessor::Initialize(const UIKRigDefinition* InRigAsset, const FIKRi
 
 	// ok, lets try to initialize
 	bTriedToInitialize = true;
+
+	// copy skeleton data from the actual skeleton we want to run on
+	Skeleton.SetInputSkeleton(InputSkeleton, InRigAsset->Skeleton.ExcludedBones);
 	
 	if (InRigAsset->Skeleton.BoneNames.IsEmpty())
 	{
@@ -32,9 +35,6 @@ void UIKRigProcessor::Initialize(const UIKRigDefinition* InRigAsset, const FIKRi
 		UE_LOG(LogTemp, Error, TEXT("Trying to initialize IKRig with a Skeleton that is missing required bones. See output log. %s"), *InRigAsset->GetName());
 		return;
 	}
-
-	// copy skeleton data from the actual skeleton we want to run on
-	Skeleton.SetInputSkeleton(InputSkeleton, InRigAsset->Skeleton.ExcludedBones);
 	
 	// initialize goals based on source asset
 	GoalContainer.Empty();
@@ -297,7 +297,6 @@ const FIKRigGoalContainer& UIKRigProcessor::GetGoalContainer() const
 
 FIKRigSkeleton& UIKRigProcessor::GetSkeleton()
 {
-	check(bInitialized);
 	return Skeleton;
 }
 
diff --git a/Engine/Plugins/Animation/IKRig/Source/IKRig/Private/Solvers/IKRig_PBIKSolver.cpp b/Engine/Plugins/Animation/IKRig/Source/IKRig/Private/Solvers/IKRig_PBIKSolver.cpp
--- a/Engine/Plugins/Animation/IKRig/Source/IKRig/Private/Solvers/IKRig_PBIKSolver.cpp
+++ b/Engine/Plugins/Animation/IKRig/Source/IKRig/Private/Solvers/IKRig_PBIKSolver.cpp
@@ -162,7 +162,7 @@ void UIKRigPBIKSolver::UpdateSolverSettings(UIKRigSolver* InSettings)
 		{
 			for (UIKRig_FBIKEffector* Effector : Effectors)
 			{
-				if (Effector->GoalName == InEffector->GoalName)
+				if (Effector && Effector->GoalName == InEffector->GoalName)
 				{
 					Effector->CopySettings(InEffector);
 					break;
@@ -175,7 +175,7 @@ void UIKRigPBIKSolver::UpdateSolverSettings(UIKRigSolver* InSettings)
 		{
 			for (UIKRig_PBIKBoneSettings* BoneSetting : BoneSettings)
 			{
-				if (BoneSetting->Bone == InBoneSetting->Bone)
+				if (BoneSetting && BoneSetting->Bone == InBoneSetting->Bone)
 				{
 					BoneSetting->CopySettings(InBoneSetting);
 					break;	
