fixed goal customization tooltip

diff --git a/Engine/Plugins/Animation/IKRig/Source/IKRigEditor/Private/RigEditor/IKRigDetailCustomizations.cpp b/Engine/Plugins/Animation/IKRig/Source/IKRigEditor/Private/RigEditor/IKRigDetailCustomizations.cpp
--- a/Engine/Plugins/Animation/IKRig/Source/IKRigEditor/Private/RigEditor/IKRigDetailCustomizations.cpp
+++ b/Engine/Plugins/Animation/IKRig/Source/IKRigEditor/Private/RigEditor/IKRigDetailCustomizations.cpp
@@ -87,7 +87,7 @@ void FIKRigGenericDetailCustomization::CustomizeDetailsForClass<UIKRigBoneDetail
 	const TArray<FText> ButtonTooltips =
 	{
 		LOCTEXT("CurrentBoneTransformTooltip", "The current transform of the bone"),
-		LOCTEXT("ReferenceBoneTransformTooltip", "The reference transform of a bone")
+		LOCTEXT("ReferenceBoneTransformTooltip", "The reference transform of the bone")
 	};
 
 	static TSharedRef<TArray<EIKRigTransformType::Type>> VisibleTransforms =
@@ -233,8 +233,8 @@ void FIKRigGenericDetailCustomization::CustomizeDetailsForClass<UIKRigEffectorGo
 	
 	const TArray<FText> ButtonTooltips =
 	{
-		LOCTEXT("CurrentBoneTransformTooltip", "The current transform of the bone"),
-		LOCTEXT("ReferenceBoneTransformTooltip", "The reference transform of a bone")
+		LOCTEXT("CurrentGoalTransformTooltip", "The current transform of the goal"),
+		LOCTEXT("ReferenceGoalTransformTooltip", "The reference transform of the goal")
 	};
 
 	static TSharedRef<TArray<EIKRigTransformType::Type>> VisibleTransforms =
