Crash when setting IK Rig input to Live Link while animation is playing.

diff --git a/Engine/Plugins/Animation/IKRig/Source/IKRigEditor/Private/RigEditor/IKRigAnimInstanceProxy.cpp b/Engine/Plugins/Animation/IKRig/Source/IKRigEditor/Private/RigEditor/IKRigAnimInstanceProxy.cpp
--- a/Engine/Plugins/Animation/IKRig/Source/IKRigEditor/Private/RigEditor/IKRigAnimInstanceProxy.cpp
+++ b/Engine/Plugins/Animation/IKRig/Source/IKRigEditor/Private/RigEditor/IKRigAnimInstanceProxy.cpp
@@ -12,7 +12,7 @@ FIKRigAnimInstanceProxy::FIKRigAnimInstanceProxy(UAnimInstance* InAnimInstance,
 
 void FIKRigAnimInstanceProxy::Initialize(UAnimInstance* InAnimInstance)
 {
-	FAnimInstanceProxy::Initialize(InAnimInstance);
+	FAnimPreviewInstanceProxy::Initialize(InAnimInstance);
 	IKRigNode->Source.SetLinkNode(&SingleNode);
 
 	// force this instance of the IK Rig evaluation to copy setting from the source IK Rig asset
