Reinitialize retargeter anim instance correctly.

diff --git a/Engine/Plugins/Animation/IKRig/Source/IKRigEditor/Private/RetargetEditor/IKRetargetAnimInstanceProxy.cpp b/Engine/Plugins/Animation/IKRig/Source/IKRigEditor/Private/RetargetEditor/IKRetargetAnimInstanceProxy.cpp
--- a/Engine/Plugins/Animation/IKRig/Source/IKRigEditor/Private/RetargetEditor/IKRetargetAnimInstanceProxy.cpp
+++ b/Engine/Plugins/Animation/IKRig/Source/IKRigEditor/Private/RetargetEditor/IKRetargetAnimInstanceProxy.cpp
@@ -12,7 +12,7 @@ FIKRetargetAnimInstanceProxy::FIKRetargetAnimInstanceProxy(UAnimInstance* InAnim
 
 void FIKRetargetAnimInstanceProxy::Initialize(UAnimInstance* InAnimInstance)
 {
-	FAnimInstanceProxy::Initialize(InAnimInstance);
+	FAnimPreviewInstanceProxy::Initialize(InAnimInstance);
 }
 
 bool FIKRetargetAnimInstanceProxy::Evaluate(FPoseContext& Output)
