Fixes bone/goal drawing toggle mismatch.

diff --git a/Engine/Plugins/Animation/IKRig/Source/IKRigEditor/Private/RigEditor/IKRigEditMode.cpp b/Engine/Plugins/Animation/IKRig/Source/IKRigEditor/Private/RigEditor/IKRigEditMode.cpp
--- a/Engine/Plugins/Animation/IKRig/Source/IKRigEditor/Private/RigEditor/IKRigEditMode.cpp
+++ b/Engine/Plugins/Animation/IKRig/Source/IKRigEditor/Private/RigEditor/IKRigEditMode.cpp
@@ -131,7 +131,7 @@ void FIKRigEditMode::RenderBones(FPrimitiveDrawInterface* PDI)
 
 	UIKRigController* AssetController = Controller->AssetController;
 	const UIKRigDefinition* IKRigAsset = AssetController->GetAsset();
-	if (!IKRigAsset->DrawGoals)
+	if (!IKRigAsset->DrawBones)
 	{
 		return;
 	}
