Add ability to clear retarget root.

diff --git a/Engine/Plugins/Animation/IKRig/Source/IKRigEditor/Private/RigEditor/IKRigSkeletonCommands.cpp b/Engine/Plugins/Animation/IKRig/Source/IKRigEditor/Private/RigEditor/IKRigSkeletonCommands.cpp
--- a/Engine/Plugins/Animation/IKRig/Source/IKRigEditor/Private/RigEditor/IKRigSkeletonCommands.cpp
+++ b/Engine/Plugins/Animation/IKRig/Source/IKRigEditor/Private/RigEditor/IKRigSkeletonCommands.cpp
@@ -17,7 +17,8 @@ void FIKRigSkeletonCommands::RegisterCommands()
 	UI_COMMAND(ExcludeBone, "Exclude Selected Bone From Solve", "Ignore bone in all solvers.", EUserInterfaceActionType::Button, FInputChord());
 	UI_COMMAND(IncludeBone, "Include Selected Bone In Solve", "Include bone in all solvers.", EUserInterfaceActionType::Button, FInputChord());
 	UI_COMMAND(NewRetargetChain, "New Retarget Chain from Selected Bones", "Create a new retarget bone chain from the selected bones.", EUserInterfaceActionType::Button, FInputChord());
-	UI_COMMAND(SetRetargetRoot, "Set Retargeting Root", "Set the Root Bone used for retargeting. Usually 'Pelvis'.", EUserInterfaceActionType::Button, FInputChord());
+	UI_COMMAND(SetRetargetRoot, "Set Retarget Root", "Set the Root Bone used for retargeting. Usually 'Pelvis'.", EUserInterfaceActionType::Button, FInputChord());
+	UI_COMMAND(ClearRetargetRoot, "Clear Retarget Root", "Clear the Root Bone used for retargeting.", EUserInterfaceActionType::Button, FInputChord());
 	UI_COMMAND(RenameGoal, "Rename Goal", "Rename the selected goal.", EUserInterfaceActionType::Button, FInputChord(EKeys::F2));
 }
 
diff --git a/Engine/Plugins/Animation/IKRig/Source/IKRigEditor/Private/RigEditor/SIKRigSkeleton.cpp b/Engine/Plugins/Animation/IKRig/Source/IKRigEditor/Private/RigEditor/SIKRigSkeleton.cpp
--- a/Engine/Plugins/Animation/IKRig/Source/IKRigEditor/Private/RigEditor/SIKRigSkeleton.cpp
+++ b/Engine/Plugins/Animation/IKRig/Source/IKRigEditor/Private/RigEditor/SIKRigSkeleton.cpp
@@ -446,6 +446,10 @@ void SIKRigSkeleton::BindCommands()
 		FExecuteAction::CreateSP(this, &SIKRigSkeleton::HandleSetRetargetRoot),
 		FCanExecuteAction::CreateSP(this, &SIKRigSkeleton::CanSetRetargetRoot));
 
+	CommandList->MapAction(Commands.ClearRetargetRoot,
+		FExecuteAction::CreateSP(this, &SIKRigSkeleton::HandleClearRetargetRoot),
+		FCanExecuteAction::CreateSP(this, &SIKRigSkeleton::CanClearRetargetRoot));
+
 	CommandList->MapAction(Commands.RenameGoal,
 	FExecuteAction::CreateSP(this, &SIKRigSkeleton::HandleRenameGoal),
 		FCanExecuteAction::CreateSP(this, &SIKRigSkeleton::CanRenameGoal));
@@ -484,6 +488,7 @@ void SIKRigSkeleton::FillContextMenu(FMenuBuilder& MenuBuilder)
 
 	MenuBuilder.BeginSection("Retargeting", LOCTEXT("RetargetingOperations", "Retargeting"));
 	MenuBuilder.AddMenuEntry(Actions.SetRetargetRoot);
+	MenuBuilder.AddMenuEntry(Actions.ClearRetargetRoot);
 	MenuBuilder.AddMenuEntry(Actions.NewRetargetChain);
 	MenuBuilder.EndSection();
 }
@@ -1091,6 +1096,29 @@ bool SIKRigSkeleton::CanSetRetargetRoot()
 	return !SelectedBones.IsEmpty();
 }
 
+void SIKRigSkeleton::HandleClearRetargetRoot()
+{
+	const TSharedPtr<FIKRigEditorController> Controller = EditorController.Pin();
+	if (!Controller.IsValid())
+	{
+		return;
+	}
+	
+	Controller->AssetController->SetRetargetRoot(NAME_None);
+	Controller->RefreshAllViews();
+}
+
+bool SIKRigSkeleton::CanClearRetargetRoot()
+{
+	const TSharedPtr<FIKRigEditorController> Controller = EditorController.Pin();
+	if (!Controller.IsValid())
+	{
+		return false;
+	}
+
+	return Controller->AssetController->GetRetargetRoot() != NAME_None;
+}
+
 bool SIKRigSkeleton::IsBoneInSelection(TArray<TSharedPtr<FIKRigTreeElement>>& SelectedBoneItems, const FName& BoneName)
 {
 	for (const TSharedPtr<FIKRigTreeElement>& Item : SelectedBoneItems)
@@ -1504,7 +1532,7 @@ FReply SIKRigSkeleton::OnAcceptDrop(
 	const bool bWasReparented = AssetController->SetGoalBone(DraggedElement.GoalName, TargetItem.Get()->BoneName);
 	if (bWasReparented)
 	{
-		RefreshTreeView();
+		Controller->RefreshAllViews();
 	}
 	
 	return FReply::Handled();
diff --git a/Engine/Plugins/Animation/IKRig/Source/IKRigEditor/Public/RigEditor/IKRigSkeletonCommands.h b/Engine/Plugins/Animation/IKRig/Source/IKRigEditor/Public/RigEditor/IKRigSkeletonCommands.h
--- a/Engine/Plugins/Animation/IKRig/Source/IKRigEditor/Public/RigEditor/IKRigSkeletonCommands.h
+++ b/Engine/Plugins/Animation/IKRig/Source/IKRigEditor/Public/RigEditor/IKRigSkeletonCommands.h
@@ -45,8 +45,10 @@ public:
 
 	/** create retarget chain */
 	TSharedPtr< FUICommandInfo > NewRetargetChain;
-	/** set the root of the retargeting */
+	/** set the retarget root */
 	TSharedPtr< FUICommandInfo > SetRetargetRoot;
+	/** clear the retarget root */
+	TSharedPtr< FUICommandInfo > ClearRetargetRoot;
 
 	/** rename selected goal */
 	TSharedPtr< FUICommandInfo > RenameGoal;
diff --git a/Engine/Plugins/Animation/IKRig/Source/IKRigEditor/Public/RigEditor/SIKRigSkeleton.h b/Engine/Plugins/Animation/IKRig/Source/IKRigEditor/Public/RigEditor/SIKRigSkeleton.h
--- a/Engine/Plugins/Animation/IKRig/Source/IKRigEditor/Public/RigEditor/SIKRigSkeleton.h
+++ b/Engine/Plugins/Animation/IKRig/Source/IKRigEditor/Public/RigEditor/SIKRigSkeleton.h
@@ -219,6 +219,8 @@ private:
 	bool CanAddNewRetargetChain();
 	void HandleSetRetargetRoot();
 	bool CanSetRetargetRoot();
+	void HandleClearRetargetRoot();
+	bool CanClearRetargetRoot();
 	/** END retarget chains */
 
 	/** centralized editor controls (facilitate cross-communication between multiple UI elements)*/
