Fixing undo when draging goal into different bone in IK Rig.
Replace bPinRoot and bPrePullRoot with RootBehavior enum.

diff --git a/Engine/Plugins/Animation/IKRig/Source/IKRig/Private/Solvers/IKRig_BodyMover.cpp b/Engine/Plugins/Animation/IKRig/Source/IKRig/Private/Solvers/IKRig_BodyMover.cpp
--- a/Engine/Plugins/Animation/IKRig/Source/IKRig/Private/Solvers/IKRig_BodyMover.cpp
+++ b/Engine/Plugins/Animation/IKRig/Source/IKRig/Private/Solvers/IKRig_BodyMover.cpp
@@ -237,6 +237,7 @@ void UIKRig_BodyMover::SetGoalBone(const FName& GoalName, const FName& NewBoneNa
 	}
 
 	// rename
+	Effectors[GoalIndex]->Modify();
 	Effectors[GoalIndex]->BoneName = NewBoneName;
 }
 
diff --git a/Engine/Plugins/Animation/IKRig/Source/IKRig/Private/Solvers/IKRig_LimbSolver.cpp b/Engine/Plugins/Animation/IKRig/Source/IKRig/Private/Solvers/IKRig_LimbSolver.cpp
--- a/Engine/Plugins/Animation/IKRig/Source/IKRig/Private/Solvers/IKRig_LimbSolver.cpp
+++ b/Engine/Plugins/Animation/IKRig/Source/IKRig/Private/Solvers/IKRig_LimbSolver.cpp
@@ -199,6 +199,7 @@ void UIKRig_LimbSolver::RenameGoal(const FName& OldName, const FName& NewName)
 {
 	if (Effector->GoalName == OldName)
 	{
+		Effector->Modify();
 		Effector->GoalName = NewName;
 	}
 }
@@ -207,6 +208,7 @@ void UIKRig_LimbSolver::SetGoalBone(const FName& GoalName, const FName& NewBoneN
 {
 	if (Effector->GoalName == GoalName)
 	{
+		Effector->Modify();
 		Effector->BoneName = NewBoneName;
 	}
 }
diff --git a/Engine/Plugins/Animation/IKRig/Source/IKRig/Private/Solvers/IKRig_PBIKSolver.cpp b/Engine/Plugins/Animation/IKRig/Source/IKRig/Private/Solvers/IKRig_PBIKSolver.cpp
--- a/Engine/Plugins/Animation/IKRig/Source/IKRig/Private/Solvers/IKRig_PBIKSolver.cpp
+++ b/Engine/Plugins/Animation/IKRig/Source/IKRig/Private/Solvers/IKRig_PBIKSolver.cpp
@@ -118,9 +118,8 @@ void UIKRigPBIKSolver::Solve(FIKRigSkeleton& IKRigSkeleton, const FIKRigGoalCont
 	Settings.bAllowStretch = bAllowStretch;
 	Settings.MassMultiplier = MassMultiplier;
 	Settings.MinMassMultiplier = MinMassMultiplier;
-	Settings.bPinRoot = bPinRoot;
+	Settings.RootBehavior = RootBehavior;
 	Settings.bStartSolveFromInputPose = bStartSolveFromInputPose;
-	Settings.bPrePullRoot = bPrePullRoot;
 
 	// solve
 	Solver.Solve(Settings);
@@ -153,9 +152,8 @@ void UIKRigPBIKSolver::UpdateSolverSettings(UIKRigSolver* InSettings)
 		bAllowStretch = Settings->bAllowStretch;
 		MassMultiplier = Settings->MassMultiplier;
 		MinMassMultiplier = Settings->MinMassMultiplier;
-		bPinRoot = Settings->bPinRoot;
+		RootBehavior = Settings->RootBehavior;
 		bStartSolveFromInputPose = Settings->bStartSolveFromInputPose;
-		bPrePullRoot = Settings->bPrePullRoot;
 
 		// copy effector settings
 		for (const UIKRig_FBIKEffector* InEffector : Settings->Effectors)
@@ -238,6 +236,7 @@ void UIKRigPBIKSolver::RenameGoal(const FName& OldName, const FName& NewName)
 	}
 
 	// rename
+	Effectors[GoalIndex]->Modify();
 	Effectors[GoalIndex]->GoalName = NewName;
 }
 
@@ -251,6 +250,7 @@ void UIKRigPBIKSolver::SetGoalBone(const FName& GoalName, const FName& NewBoneNa
 	}
 
 	// rename
+	Effectors[GoalIndex]->Modify();
 	Effectors[GoalIndex]->BoneName = NewBoneName;
 }
 
diff --git a/Engine/Plugins/Animation/IKRig/Source/IKRig/Public/IKRigSolver.h b/Engine/Plugins/Animation/IKRig/Source/IKRig/Public/IKRigSolver.h
--- a/Engine/Plugins/Animation/IKRig/Source/IKRig/Public/IKRigSolver.h
+++ b/Engine/Plugins/Animation/IKRig/Source/IKRig/Public/IKRigSolver.h
@@ -19,6 +19,8 @@ class IKRIG_API UIKRigSolver : public UObject
 	
 public:
 
+	UIKRigSolver() { SetFlags(RF_Transactional); }
+
 	//** RUNTIME */
 	/** override to setup internal data based on ref pose */
 	virtual void Initialize(const FIKRigSkeleton& IKRigSkeleton) PURE_VIRTUAL("Init");
diff --git a/Engine/Plugins/Animation/IKRig/Source/IKRig/Public/Solvers/IKRig_BodyMover.h b/Engine/Plugins/Animation/IKRig/Source/IKRig/Public/Solvers/IKRig_BodyMover.h
--- a/Engine/Plugins/Animation/IKRig/Source/IKRig/Public/Solvers/IKRig_BodyMover.h
+++ b/Engine/Plugins/Animation/IKRig/Source/IKRig/Public/Solvers/IKRig_BodyMover.h
@@ -13,7 +13,8 @@ class IKRIG_API UIKRig_BodyMoverEffector : public UObject
 	GENERATED_BODY()
 
 public:
-
+	UIKRig_BodyMoverEffector() { SetFlags(RF_Transactional); }
+	
 	UPROPERTY(VisibleAnywhere, Category = "Body Mover Effector")
 	FName GoalName;
 
diff --git a/Engine/Plugins/Animation/IKRig/Source/IKRig/Public/Solvers/IKRig_LimbSolver.h b/Engine/Plugins/Animation/IKRig/Source/IKRig/Public/Solvers/IKRig_LimbSolver.h
--- a/Engine/Plugins/Animation/IKRig/Source/IKRig/Public/Solvers/IKRig_LimbSolver.h
+++ b/Engine/Plugins/Animation/IKRig/Source/IKRig/Public/Solvers/IKRig_LimbSolver.h
@@ -15,6 +15,7 @@ class IKRIG_API UIKRig_LimbEffector : public UObject
 	GENERATED_BODY()
 
 public:
+	UIKRig_LimbEffector() { SetFlags(RF_Transactional); }
 
 	UPROPERTY(VisibleAnywhere, Category = "Limb IK Effector")
 	FName GoalName = NAME_None;
diff --git a/Engine/Plugins/Animation/IKRig/Source/IKRig/Public/Solvers/IKRig_PBIKSolver.h b/Engine/Plugins/Animation/IKRig/Source/IKRig/Public/Solvers/IKRig_PBIKSolver.h
--- a/Engine/Plugins/Animation/IKRig/Source/IKRig/Public/Solvers/IKRig_PBIKSolver.h
+++ b/Engine/Plugins/Animation/IKRig/Source/IKRig/Public/Solvers/IKRig_PBIKSolver.h
@@ -15,6 +15,8 @@ class IKRIG_API UIKRig_FBIKEffector : public UObject
 	GENERATED_BODY()
 
 public:
+	UIKRig_FBIKEffector() { SetFlags(RF_Transactional); }
+	
 	/** The Goal that is driving this effector's transform. */
 	UPROPERTY(VisibleAnywhere, Category = "Full Body IK Effector")
 	FName GoalName;
@@ -198,16 +200,13 @@ public:
 	UPROPERTY(EditAnywhere, Category = "Full Body IK Settings")
 	bool bAllowStretch = false;
 
-	/** (Default is true) If true, solver runs a pre-pass (before the constraint solve) which will translate the whole body by the average motion of the stretched effectors.
-	*This can decrease the amount of iterations needed to achieve a converged pose when the effectors are pulling the body far.
-	*This option has no effect if PinRoot is set to True.*/
-	UPROPERTY(EditAnywhere, Category = "Full Body IK Settings", meta=(EditCondition = "!bPinRoot"))
-	bool bPrePullRoot = true;
-
-	/** Lock the position and rotation of the solver root bone in-place (at animated position). Useful for partial-body solves. Default is false. */
+	/** (Default is PrePull) Set the behavior for the translation of the root.
+	*Pre Pull: translates the whole body by the average motion of the stretched effectors to help achieve faster convergence when reaching far.
+	*Pin to Input: locks the translation and rotation of the root bone to the input pose. Overrides any bone settings applied to the root.
+	*Free: treats the root bone like any other and allows it to move freely or according to any bone settings applied to it. */
 	UPROPERTY(EditAnywhere, Category = "Full Body IK Settings")
-	bool bPinRoot = false;
-
+	EPBIKRootBehavior RootBehavior;
+	
 	/** When true, the solver is reset each tick to start from the current input pose. If false, incoming animated poses are ignored and the solver starts from the results of the previous solve. Default is true. */
 	UPROPERTY(EditAnywhere, Category = "Full Body IK Settings")
 	bool bStartSolveFromInputPose = true;
diff --git a/Engine/Plugins/Animation/IKRig/Source/IKRigEditor/Private/RigEditor/IKRigController.cpp b/Engine/Plugins/Animation/IKRig/Source/IKRigEditor/Private/RigEditor/IKRigController.cpp
--- a/Engine/Plugins/Animation/IKRig/Source/IKRigEditor/Private/RigEditor/IKRigController.cpp
+++ b/Engine/Plugins/Animation/IKRig/Source/IKRigEditor/Private/RigEditor/IKRigController.cpp
@@ -729,10 +729,10 @@ bool UIKRigController::SetGoalBone(const FName& GoalName, const FName& NewBoneNa
 		return false; // goal is already using this bone
 	}
 	
-	FScopedTransaction Transaction(LOCTEXT("RenameGoal_Label", "Rename Goal"));
-	Asset->Modify();
+	FScopedTransaction Transaction(LOCTEXT("SetGoalBone_Label", "Set Goal Bone"));
 
 	// update goal
+	Asset->Goals[GoalIndex]->Modify();
 	Asset->Goals[GoalIndex]->BoneName = NewBoneName;
 	
 	// update in solvers
diff --git a/Engine/Plugins/Experimental/FullBodyIK/Source/PBIK/Private/Core/PBIKSolver.cpp b/Engine/Plugins/Experimental/FullBodyIK/Source/PBIK/Private/Core/PBIKSolver.cpp
--- a/Engine/Plugins/Experimental/FullBodyIK/Source/PBIK/Private/Core/PBIKSolver.cpp
+++ b/Engine/Plugins/Experimental/FullBodyIK/Source/PBIK/Private/Core/PBIKSolver.cpp
@@ -169,7 +169,7 @@ void FPBIKSolver::Solve(const FPBIKSolverSettings& Settings)
 	// optionally pin root in-place (convenience, does not require an effector)
 	if (RootPin.IsValid())
 	{
-		RootPin.Pin()->bEnabled = Settings.bPinRoot;
+		RootPin.Pin()->bEnabled = Settings.RootBehavior == EPBIKRootBehavior::PinToInput;
 		// pin to animated input root pose
 		RootPin.Pin()->SetGoal(SolverRoot->Position, SolverRoot->Rotation, 1.0f);
 	}
@@ -309,7 +309,7 @@ void FPBIKSolver::ApplyPreferredAngles()
 void FPBIKSolver::UpdateBodies(const FPBIKSolverSettings& Settings)
 {
 	// this creates a better pose from which to start the constraint solving
-	if (Settings.bPrePullRoot && !Settings.bPinRoot)
+	if (Settings.RootBehavior == EPBIKRootBehavior::PrePull)
 	{
 		// pull ALL bodies towards the effectors
 		PullRootTowardsEffectors();
diff --git a/Engine/Plugins/Experimental/FullBodyIK/Source/PBIK/Public/Core/PBIKSolver.h b/Engine/Plugins/Experimental/FullBodyIK/Source/PBIK/Public/Core/PBIKSolver.h
--- a/Engine/Plugins/Experimental/FullBodyIK/Source/PBIK/Public/Core/PBIKSolver.h
+++ b/Engine/Plugins/Experimental/FullBodyIK/Source/PBIK/Public/Core/PBIKSolver.h
@@ -91,6 +91,14 @@ struct FEffector
 	
 } // namespace
 
+UENUM()
+enum class EPBIKRootBehavior : uint8
+{
+	PrePull,
+	PinToInput,
+	Free,
+};
+
 USTRUCT()
 struct PBIK_API FPBIKSolverSettings
 {
@@ -112,15 +120,12 @@ struct PBIK_API FPBIKSolverSettings
 	UPROPERTY(EditAnywhere, Category = SolverSettings)
 	bool bAllowStretch = false;
 
-	/** Lock the position and rotation of the solver root bone in-place (at animated position). Useful for partial-body solves. Default is false. */
+	/** (Default is PrePull) Set the behavior of the solver root.
+	*Pre Pull: translates the root (and all children) by the average motion of the stretched effectors to help achieve faster convergence when reaching far.
+	*Pin to Input: locks the translation and rotation of the root bone to the input pose. Overrides any bone settings applied to the root. Good for partial-body solves.
+	*Free: treats the root bone like any other and allows it to move freely or according to any bone settings applied to it. */
 	UPROPERTY(EditAnywhere, Category = SolverSettings)
-	bool bPinRoot = false;
-
-	/** (Default is true) If true, solver runs a pre-pass (before the constraint solve) which will translate the whole body by the average motion of the stretched effectors.
-	 *This can decrease the amount of iterations needed to achieve a converged pose when the effectors are pulling the body far.
-	 *This option has no effect if PinRoot is set to True.*/
-	UPROPERTY(EditAnywhere, Category = SolverSettings, meta=(EditCondition = "!bPinRoot"))
-	bool bPrePullRoot = true;
+	EPBIKRootBehavior RootBehavior;
 
 	/** When true, the solver is reset each tick to start from the current input pose. If false, incoming animated poses are ignored and the solver starts from the results of the previous solve. Default is true. */
 	UPROPERTY(EditAnywhere, Category = SolverSettings)
