Fixed missing PostEditChange to notify that the duplicated animation sequence has changed

Missing this caused both the content browser and the asset browser in the retargeter with an asset pointing to the old skeleton asset

diff --git a/Engine/Plugins/Animation/IKRig/Source/IKRigEditor/Private/RetargetEditor/IKRetargetBatchOperation.cpp b/Engine/Plugins/Animation/IKRig/Source/IKRigEditor/Private/RetargetEditor/IKRetargetBatchOperation.cpp
--- a/Engine/Plugins/Animation/IKRig/Source/IKRigEditor/Private/RetargetEditor/IKRetargetBatchOperation.cpp
+++ b/Engine/Plugins/Animation/IKRig/Source/IKRigEditor/Private/RetargetEditor/IKRetargetBatchOperation.cpp
@@ -161,6 +161,7 @@ void FIKRetargetBatchOperation::RetargetAssets(
 		// replace references to other animation
 		AssetToRetarget->ReplaceReferredAnimations(RemappedAnimAssets);
 		AssetToRetarget->SetSkeleton(NewSkeleton);
+		AssetToRetarget->PostEditChange();
 		AssetToRetarget->MarkPackageDirty();
 	}
 
