IKRig anim graph node revamp

the approach is based on custom pins and detail panel customization to present a simplified UI within the Blueprint editor
this is prone to further changes as described in UE-134416

diff --git a/Engine/Plugins/Animation/IKRig/Source/IKRig/Private/AnimNodes/AnimNode_IKRig.cpp b/Engine/Plugins/Animation/IKRig/Source/IKRig/Private/AnimNodes/AnimNode_IKRig.cpp
--- a/Engine/Plugins/Animation/IKRig/Source/IKRig/Private/AnimNodes/AnimNode_IKRig.cpp
+++ b/Engine/Plugins/Animation/IKRig/Source/IKRig/Private/AnimNodes/AnimNode_IKRig.cpp
@@ -7,7 +7,7 @@
 #include "ActorComponents/IKRigInterface.h"
 #include "Drawing/ControlRigDrawInterface.h"
 #include "Animation/AnimInstanceProxy.h"
-
+#include "Algo/ForEach.h"
 
 void FAnimNode_IKRig::Evaluate_AnyThread(FPoseContext& Output) 
 {
@@ -176,8 +176,10 @@ void FAnimNode_IKRig::PreUpdate(const UAnimInstance* InAnimInstance)
 	if (!IKRigProcessor->IsInitialized())
 	{
 		const FReferenceSkeleton& RefSkeleton = InAnimInstance->GetSkelMeshComponent()->SkeletalMesh->GetRefSkeleton();
-		IKRigProcessor->Initialize(RigDefinitionAsset, RefSkeleton);
+ 		IKRigProcessor->Initialize(RigDefinitionAsset, RefSkeleton);
 	}
+
+	PropagateInputProperties(InAnimInstance);
 	
 	// cache list of goal creator components on the actor
 	// TODO tried doing this in Initialize_AnyThread but it would miss some GoalCreator components
@@ -206,6 +208,13 @@ void FAnimNode_IKRig::PreUpdate(const UAnimInstance* InAnimInstance)
 	}
 }
 
+void FAnimNode_IKRig::OnInitializeAnimInstance(const FAnimInstanceProxy* InProxy, const UAnimInstance* InAnimInstance)
+{
+	Super::OnInitializeAnimInstance(InProxy, InAnimInstance);
+	
+	InitializeProperties(InAnimInstance, GetTargetClass());
+}
+
 void FAnimNode_IKRig::SetProcessorNeedsInitialized()
 {
 	if (IKRigProcessor)
@@ -259,6 +268,92 @@ void FAnimNode_IKRig::CacheBones_AnyThread(const FAnimationCacheBonesContext& Co
 	}
 }
 
+void FAnimNode_IKRig::InitializeProperties(const UObject* InSourceInstance, UClass* InTargetClass)
+{
+	static const FString AlphaPosPropStr = GET_MEMBER_NAME_CHECKED(FIKRigGoal, PositionAlpha).ToString();
+	static const FString AlphaRotPropStr = GET_MEMBER_NAME_CHECKED(FIKRigGoal, RotationAlpha).ToString();
+	static const FString PositionPropStr = GET_MEMBER_NAME_CHECKED(FIKRigGoal, Position).ToString();
+	static const FString RotationPropStr = GET_MEMBER_NAME_CHECKED(FIKRigGoal, Rotation).ToString();
+	
+	SourceProperties.Reset(SourcePropertyNames.Num());
+	DestProperties.Reset(SourcePropertyNames.Num());
+	UpdateFunctions.Reset(SourcePropertyNames.Num());
+
+	check(SourcePropertyNames.Num() == DestPropertyNames.Num());
+
+	UClass* SourceClass = InSourceInstance->GetClass();
+	for (int32 PropIndex = 0; PropIndex < SourcePropertyNames.Num(); ++PropIndex)
+	{
+		const FName& SourceName = SourcePropertyNames[PropIndex];
+		
+		FProperty* SourceProperty = FindFProperty<FProperty>(SourceClass, SourceName);
+		SourceProperties.Add(SourceProperty);
+		DestProperties.Add(nullptr);
+
+		if (SourceProperty)
+		{
+			// store update functions for later use in PropagateInputProperties to avoid looking for properties
+			// while evaluating
+			const FName& GoalPropertyName = DestPropertyNames[PropIndex];
+
+			// we use the goal name's hash value as described in UAnimGraphNode_IKRig::CreateCustomPinsFromValidAsset
+			const int32 GoalIndex = Goals.IndexOfByPredicate([&GoalPropertyName](const FIKRigGoal& InGoal)
+			{
+				return GetTypeHash(InGoal.Name) == GoalPropertyName.GetNumber();
+			});
+
+			if (Goals.IsValidIndex(GoalIndex))
+			{
+				FIKRigGoal& Goal = Goals[GoalIndex];
+
+				const FString GoalPropStr = GoalPropertyName.ToString();
+
+				if (GoalPropStr.StartsWith(AlphaPosPropStr))
+				{
+					UpdateFunctions.Add([&Goal, SourceProperty](const UObject* InSourceInstance)
+					{
+						const float* AlphaValue = SourceProperty->ContainerPtrToValuePtr<float>(InSourceInstance);
+						Goal.PositionAlpha = FMath::Clamp<float>(*AlphaValue, 0.f, 1.f);
+					});
+				}
+				else if (GoalPropStr.StartsWith(AlphaRotPropStr))
+				{
+					UpdateFunctions.Add([&Goal, SourceProperty](const UObject* InSourceInstance)
+					{
+						const float* AlphaValue = SourceProperty->ContainerPtrToValuePtr<float>(InSourceInstance);
+						Goal.RotationAlpha = FMath::Clamp<float>(*AlphaValue, 0.f, 1.f);
+					});
+				}
+				else if (GoalPropStr.StartsWith(PositionPropStr))
+				{
+					UpdateFunctions.Add([&Goal, SourceProperty](const UObject* InSourceInstance)
+					{
+						Goal.Position = *SourceProperty->ContainerPtrToValuePtr<FVector>(InSourceInstance);
+					});
+				}
+				else if (GoalPropStr.StartsWith(RotationPropStr))
+				{
+					UpdateFunctions.Add([&Goal, SourceProperty](const UObject* InSourceInstance)
+					{
+						Goal.Rotation = *SourceProperty->ContainerPtrToValuePtr<FRotator>(InSourceInstance);
+					});
+				}
+			}
+		}
+	}
+}
+
+void FAnimNode_IKRig::PropagateInputProperties(const UObject* InSourceInstance)
+{
+	if (InSourceInstance)
+	{
+		Algo::ForEach(UpdateFunctions, [InSourceInstance](const UpdateFunction& InFunc)
+		{
+		 	InFunc(InSourceInstance);
+		});
+	}
+}
+
 void FAnimNode_IKRig::ConditionalDebugDraw(
 	FPrimitiveDrawInterface* PDI,
 	USkeletalMeshComponent* PreviewSkelMeshComp) const
diff --git a/Engine/Plugins/Animation/IKRig/Source/IKRig/Private/IKRigProcessor.cpp b/Engine/Plugins/Animation/IKRig/Source/IKRig/Private/IKRigProcessor.cpp
--- a/Engine/Plugins/Animation/IKRig/Source/IKRig/Private/IKRigProcessor.cpp
+++ b/Engine/Plugins/Animation/IKRig/Source/IKRig/Private/IKRigProcessor.cpp
@@ -317,42 +317,64 @@ void UIKRigProcessor::ResolveFinalGoalTransforms(const FTransform& WorldToCompon
 		FVector ComponentSpaceGoalPosition = Goal.Position;
 		FQuat ComponentSpaceGoalRotation = Goal.Rotation.Quaternion();
 
-		// put goal POSITION in Component Space
-		switch (Goal.PositionSpace)
+		// FIXME find a way to cache SourceBoneIndex to avoid calling Skeleton.GetBoneIndexFromName here
+		// we may use FGoalBone::OptSourceIndex for this
+		int SourceBoneIndex = INDEX_NONE;
+		if (Goal.TransformSource == EIKRigGoalTransformSource::Bone && Goal.SourceBone.BoneName != NAME_None)
 		{
-		case EIKRigGoalSpace::Additive:
-			// add position offset to bone position
-			ComponentSpaceGoalPosition = Skeleton.CurrentPoseGlobal[GoalBone.BoneIndex].GetLocation() + Goal.Position;
-			break;
-		case EIKRigGoalSpace::Component:
-			// was already supplied in Component Space
-			break;
-		case EIKRigGoalSpace::World:
-			// convert from World Space to Component Space
-			ComponentSpaceGoalPosition = WorldToComponent.TransformPosition(Goal.Position);
-			break;
-		default:
-			checkNoEntry();
-			break;
+			SourceBoneIndex = Skeleton.GetBoneIndexFromName(Goal.SourceBone.BoneName);
+		}
+
+		if (SourceBoneIndex != INDEX_NONE)
+		{
+			ComponentSpaceGoalPosition = Skeleton.CurrentPoseGlobal[SourceBoneIndex].GetLocation();
+		}
+		else
+		{
+			// put goal POSITION in Component Space
+			switch (Goal.PositionSpace)
+			{
+			case EIKRigGoalSpace::Additive:
+				// add position offset to bone position
+				ComponentSpaceGoalPosition = Skeleton.CurrentPoseGlobal[GoalBone.BoneIndex].GetLocation() + Goal.Position;
+				break;
+			case EIKRigGoalSpace::Component:
+				// was already supplied in Component Space
+				break;
+			case EIKRigGoalSpace::World:
+				// convert from World Space to Component Space
+				ComponentSpaceGoalPosition = WorldToComponent.TransformPosition(Goal.Position);
+				break;
+			default:
+				checkNoEntry();
+				break;
+			}
 		}
 		
 		// put goal ROTATION in Component Space
-		switch (Goal.RotationSpace)
+		if (SourceBoneIndex != INDEX_NONE)
+		{
+			ComponentSpaceGoalRotation = Skeleton.CurrentPoseGlobal[SourceBoneIndex].GetRotation();
+		}
+		else
 		{
-		case EIKRigGoalSpace::Additive:
-			// add rotation offset to bone rotation
-			ComponentSpaceGoalRotation = Goal.Rotation.Quaternion() * Skeleton.CurrentPoseGlobal[GoalBone.BoneIndex].GetRotation();
-			break;
-		case EIKRigGoalSpace::Component:
-			// was already supplied in Component Space
-			break;
-		case EIKRigGoalSpace::World:
-			// convert from World Space to Component Space
-			ComponentSpaceGoalRotation = WorldToComponent.TransformRotation(Goal.Rotation.Quaternion());
-			break;
-		default:
-			checkNoEntry();
-			break;
+			switch (Goal.RotationSpace)
+			{
+			case EIKRigGoalSpace::Additive:
+				// add rotation offset to bone rotation
+				ComponentSpaceGoalRotation = Goal.Rotation.Quaternion() * Skeleton.CurrentPoseGlobal[GoalBone.BoneIndex].GetRotation();
+				break;
+			case EIKRigGoalSpace::Component:
+				// was already supplied in Component Space
+				break;
+			case EIKRigGoalSpace::World:
+				// convert from World Space to Component Space
+				ComponentSpaceGoalRotation = WorldToComponent.TransformRotation(Goal.Rotation.Quaternion());
+				break;
+			default:
+				checkNoEntry();
+				break;
+			}
 		}
 
 		// blend by alpha from the input pose, to the supplied goal transform
diff --git a/Engine/Plugins/Animation/IKRig/Source/IKRig/Public/AnimNodes/AnimNode_IKRig.h b/Engine/Plugins/Animation/IKRig/Source/IKRig/Public/AnimNodes/AnimNode_IKRig.h
--- a/Engine/Plugins/Animation/IKRig/Source/IKRig/Public/AnimNodes/AnimNode_IKRig.h
+++ b/Engine/Plugins/Animation/IKRig/Source/IKRig/Public/AnimNodes/AnimNode_IKRig.h
@@ -5,7 +5,7 @@
 #include "CoreMinimal.h"
 
 #include "IKRigDataTypes.h"
-#include "Animation/AnimNodeBase.h"
+#include "Animation/AnimNode_CustomProperty.h"
 #include "AnimNode_IKRig.generated.h"
 
 class UIKRigProcessor;
@@ -14,7 +14,7 @@ class UIKRigDefinition;
 
 	
 USTRUCT(BlueprintInternalUseOnly)
-struct IKRIG_API FAnimNode_IKRig : public FAnimNode_Base
+struct IKRIG_API FAnimNode_IKRig : public FAnimNode_CustomProperty
 {
 	GENERATED_BODY()
 
@@ -27,7 +27,7 @@ struct IKRIG_API FAnimNode_IKRig : public FAnimNode_Base
 	TObjectPtr<UIKRigDefinition> RigDefinitionAsset = nullptr;
 
 	/** The input goal transforms used by the IK Rig solvers.*/
-	UPROPERTY(EditAnywhere, BlueprintReadWrite, Category = Goal, meta = (PinShownByDefault))
+	UPROPERTY(EditAnywhere, BlueprintReadWrite, Category = Goal, meta=(NeverAsPin, EditCondition=false, EditConditionHides))
 	TArray<FIKRigGoal> Goals;
 
 	/** optionally ignore the input pose and start from the reference pose each solve */
@@ -61,6 +61,10 @@ private:
 	
 	TMap<FCompactPoseBoneIndex, int32, FDefaultSetAllocator, TCompactPoseBoneIndexMapKeyFuncs<int32>> CompactPoseToRigIndices;
 
+	/** Cached functions used to update goals using custom properties to avoid looking for then when evaluating */
+	using UpdateFunction = std::function<void(const UObject*)>;
+	TArray<UpdateFunction> UpdateFunctions;
+	
 public:
 
 	// FAnimNode_Base interface
@@ -71,10 +75,19 @@ public:
 	virtual void Update_AnyThread(const FAnimationUpdateContext& Context) override;
 	virtual bool HasPreUpdate() const override { return true; }
 	virtual void PreUpdate(const UAnimInstance* InAnimInstance) override;
+	virtual bool NeedsOnInitializeAnimInstance() const override { return true; }
+	virtual void OnInitializeAnimInstance(const FAnimInstanceProxy* InProxy, const UAnimInstance* InAnimInstance) override;
 	// End of FAnimNode_Base interface
 
 	/** force reinitialization */
 	void SetProcessorNeedsInitialized();
+
+protected:
+
+	// FAnimNode_CustomProperty interface
+	virtual UClass* GetTargetClass() const override {return nullptr;}
+	virtual void InitializeProperties(const UObject* InSourceInstance, UClass* InTargetClass) override;
+	virtual void PropagateInputProperties(const UObject* InSourceInstance) override;
 	
 private:
 	
diff --git a/Engine/Plugins/Animation/IKRig/Source/IKRig/Public/IKRigDataTypes.h b/Engine/Plugins/Animation/IKRig/Source/IKRig/Public/IKRigDataTypes.h
--- a/Engine/Plugins/Animation/IKRig/Source/IKRig/Public/IKRigDataTypes.h
+++ b/Engine/Plugins/Animation/IKRig/Source/IKRig/Public/IKRigDataTypes.h
@@ -44,9 +44,12 @@ struct IKRIG_API FIKRigGoal
 	/** When TransformSource is set to "Bone" mode, the Position and Rotation will be driven by this Bone's input transform.
 	 * When using a Bone as the transform source, the Position and Rotation Alpha values can still be set as desired.
 	 * But the PositionSpace and RotationSpace are no longer relevant and will not be used.*/
-	UPROPERTY(EditAnywhere, BlueprintReadWrite, Category = Goal, meta = (CustomWidget = "BoneName"))
-	FName OptionalSourceBone;
-
+	UPROPERTY(EditAnywhere, Category = Goal)
+	FBoneReference SourceBone;
+	
+	// UPROPERTY(EditAnywhere, BlueprintReadWrite, Category = Goal, meta = (CustomWidget = "BoneName"))
+	// FName OptionalSourceBone;
+	
 	/** Position of the IK goal in Component Space of target actor component. */
 	UPROPERTY(EditAnywhere, BlueprintReadWrite, Category = Position, meta=(EditCondition="Source == EIKRigGoalSource::Manual"))
 	FVector Position;
diff --git a/Engine/Plugins/Animation/IKRig/Source/IKRig/Public/IKRigDefinition.h b/Engine/Plugins/Animation/IKRig/Source/IKRig/Public/IKRigDefinition.h
--- a/Engine/Plugins/Animation/IKRig/Source/IKRig/Public/IKRigDefinition.h
+++ b/Engine/Plugins/Animation/IKRig/Source/IKRig/Public/IKRigDefinition.h
@@ -51,7 +51,7 @@ public:
 	/** Range 0-1, default is 1. Blend between the input bone position (0.0) and the current goal position (1.0).*/
 	UPROPERTY(EditAnywhere, Category = "Goal Settings", meta = (ClampMin = "0.0", ClampMax = "1.0", UIMin = "0.0", UIMax = "1.0"))
 	float PositionAlpha = 1.0f;
-
+	
 	/** Range 0-1, default is 1. Blend between the input bone rotation (0.0) and the current goal rotation (1.0).*/
 	UPROPERTY(EditAnywhere, Category = "Goal Settings", meta = (ClampMin = "0.0", ClampMax = "1.0", UIMin = "0.0", UIMax = "1.0"))
 	float RotationAlpha = 1.0f;
@@ -63,7 +63,7 @@ public:
 	/** The initial transform of this Goal, as defined by the initial transform of the Goal's bone in the retarget pose.*/
 	UPROPERTY(VisibleAnywhere, Category = "Goal Settings")
 	FTransform InitialTransform;
-	
+
 	bool operator==(const UIKRigEffectorGoal& Other) const { return GoalName == Other.GoalName; }
 
 #if WITH_EDITORONLY_DATA
@@ -83,6 +83,14 @@ public:
 	UPROPERTY(EditAnywhere, Category = "Viewport Goal Settings",  meta = (ClampMin = "0.0", ClampMax = "10.0", UIMin = "0.0", UIMax = "5.0"))
 	float ThicknessMultiplier = 0.7f;
 
+	/** Should position data be exposed in Blueprint */
+	UPROPERTY(EditAnywhere, Category = "Exposure")
+	bool bExposePosition = false;
+
+	/** Should rotation data be exposed in Blueprint */
+	UPROPERTY(EditAnywhere, Category = "Exposure")
+	bool bExposeRotation = false;
+	
 	virtual void PostLoad() override
 	{
 		Super::PostLoad();
@@ -193,7 +201,7 @@ public:
 	/** The controller responsible for managing this asset's data (all editor mutation goes through this) */
 	UPROPERTY(Transient)
 	TObjectPtr<UObject> Controller;
-	
+
 #endif
 
 	virtual void Serialize(FArchive& Ar) override;
@@ -226,7 +234,7 @@ public:
 	virtual void SetPreviewMesh(USkeletalMesh* PreviewMesh, bool bMarkAsDirty = true) override;
 	virtual USkeletalMesh* GetPreviewMesh() const override;
 	/** END IInterface_PreviewMeshProvider interface */
-
+	
 private:
 
 	/** goals, used as effectors by solvers that support them */
diff --git a/Engine/Plugins/Animation/IKRig/Source/IKRig/Public/IKRigProcessor.h b/Engine/Plugins/Animation/IKRig/Source/IKRig/Public/IKRigProcessor.h
--- a/Engine/Plugins/Animation/IKRig/Source/IKRig/Public/IKRigProcessor.h
+++ b/Engine/Plugins/Animation/IKRig/Source/IKRig/Public/IKRigProcessor.h
@@ -18,6 +18,7 @@ struct FGoalBone
 	
 	FName BoneName;
 	int32 BoneIndex;
+	int32 OptSourceIndex = INDEX_NONE;
 };
 
 UCLASS()
diff --git a/Engine/Plugins/Animation/IKRig/Source/IKRigDeveloper/IKRigDeveloper.Build.cs b/Engine/Plugins/Animation/IKRig/Source/IKRigDeveloper/IKRigDeveloper.Build.cs
--- a/Engine/Plugins/Animation/IKRig/Source/IKRigDeveloper/IKRigDeveloper.Build.cs
+++ b/Engine/Plugins/Animation/IKRig/Source/IKRigDeveloper/IKRigDeveloper.Build.cs
@@ -32,6 +32,10 @@ namespace UnrealBuildTool.Rules
                         "Kismet",
                         "AnimGraph",
                         "BlueprintGraph",
+                        "Slate",
+                        "SlateCore",
+                        "EditorStyle",
+                        "Persona"
                     }
                 );
             }
diff --git a/Engine/Plugins/Animation/IKRig/Source/IKRigDeveloper/Private/AnimGraphNode_IKRig.cpp b/Engine/Plugins/Animation/IKRig/Source/IKRigDeveloper/Private/AnimGraphNode_IKRig.cpp
--- a/Engine/Plugins/Animation/IKRig/Source/IKRigDeveloper/Private/AnimGraphNode_IKRig.cpp
+++ b/Engine/Plugins/Animation/IKRig/Source/IKRigDeveloper/Private/AnimGraphNode_IKRig.cpp
@@ -5,12 +5,301 @@
 #include "IKRigDefinition.h"
 #include "IKRigSolver.h"
 #include "Kismet2/CompilerResultsLog.h"
+#include "AnimationGraphSchema.h"
+
+#include "DetailLayoutBuilder.h"
+#include "DetailWidgetRow.h"
+#include "IDetailChildrenBuilder.h"
+
+#include "BoneSelectionWidget.h"
 
 /////////////////////////////////////////////////////
 // UAnimGraphNode_IKRig 
 
 #define LOCTEXT_NAMESPACE "AnimGraphNode_IKRig"
 
+/////////////////////////////////////////////////////
+// FIKRigGoalLayout
+
+TSharedRef<SWidget> FIKRigGoalLayout::CreateManualValueWidget() const
+{
+	const TSharedPtr<IPropertyHandle> TransformSourceHandle = GoalPropHandle->GetChildHandle(GET_MEMBER_NAME_CHECKED(FIKRigGoal, TransformSource));
+
+	return SNew(SHorizontalBox)
+
+	// transform source combo box
+	+ SHorizontalBox::Slot()
+	.FillWidth(1.0f)
+	.HAlign(HAlign_Left)
+	[
+		TransformSourceHandle->CreatePropertyValueWidget()
+	];
+}
+
+TSharedRef<SWidget> FIKRigGoalLayout::CreateBoneValueWidget() const
+{
+	const TSharedPtr<IPropertyHandle> TransformSourceHandle = GoalPropHandle->GetChildHandle(GET_MEMBER_NAME_CHECKED(FIKRigGoal, TransformSource));
+	TSharedPtr<IPropertyHandle> SourceBoneHandle = GoalPropHandle->GetChildHandle(GET_MEMBER_NAME_CHECKED(FIKRigGoal, SourceBone));
+	
+	return SNew(SHorizontalBox)
+
+	// transform source combo box
+	+ SHorizontalBox::Slot()
+	.FillWidth(1.0f)
+	.HAlign(HAlign_Left)
+	[
+		TransformSourceHandle->CreatePropertyValueWidget()
+	]
+
+	// bone selector
+	+ SHorizontalBox::Slot()
+	.AutoWidth()
+	.HAlign(HAlign_Left)
+	.Padding(3.f, 0.f)
+	[
+		SNew(SBoneSelectionWidget)
+		.OnBoneSelectionChanged(this, &FIKRigGoalLayout::OnBoneSelectionChanged)
+		.OnGetSelectedBone(this, &FIKRigGoalLayout::GetSelectedBone)
+		.OnGetReferenceSkeleton(this, &FIKRigGoalLayout::GetReferenceSkeleton)
+	];
+}
+
+TSharedRef<SWidget> FIKRigGoalLayout::CreateValueWidget() const
+{
+	const EIKRigGoalTransformSource TransformSource = GetTransformSource();
+	if (TransformSource == EIKRigGoalTransformSource::Manual)
+	{
+		return CreateManualValueWidget();
+	}
+
+	if (TransformSource == EIKRigGoalTransformSource::Bone)
+	{
+		return CreateBoneValueWidget();
+	}
+
+	return SNullWidget::NullWidget;
+}
+
+void FIKRigGoalLayout::GenerateHeaderRowContent(FDetailWidgetRow& InOutGoalRow)
+{
+	InOutGoalRow
+	.NameContent()
+	[
+		SNew(STextBlock)
+		.Text(FText::FromName(GetName()))
+		.Font(IDetailLayoutBuilder::GetDetailFont())
+	]
+	.ValueContent()
+	[
+		CreateValueWidget()
+	];
+}
+
+void FIKRigGoalLayout::GenerateChildContent(IDetailChildrenBuilder& InOutChildrenBuilder)
+{
+	const EIKRigGoalTransformSource TransformSource = GetTransformSource();
+	if (TransformSource == EIKRigGoalTransformSource::Manual)
+	{
+		if (bExposePosition)
+		{
+			const TSharedPtr<IPropertyHandle> PosSpaceHandle = GoalPropHandle->GetChildHandle(GET_MEMBER_NAME_CHECKED(FIKRigGoal, PositionSpace));
+			InOutChildrenBuilder.AddProperty(PosSpaceHandle.ToSharedRef());
+		}
+
+		if (bExposeRotation)
+		{
+			const TSharedPtr<IPropertyHandle> RotSpaceHandle = GoalPropHandle->GetChildHandle(GET_MEMBER_NAME_CHECKED(FIKRigGoal, RotationSpace));
+			InOutChildrenBuilder.AddProperty(RotSpaceHandle.ToSharedRef());
+		}
+	}
+}
+
+FName FIKRigGoalLayout::GetGoalName(TSharedPtr<IPropertyHandle> InGoalHandle)
+{
+	if (!InGoalHandle.IsValid() || !InGoalHandle->IsValidHandle())
+	{
+		return  NAME_None;
+	}
+	
+	const TSharedPtr<IPropertyHandle> NameHandle = InGoalHandle->GetChildHandle(GET_MEMBER_NAME_CHECKED(FIKRigGoal, Name));
+	if (!NameHandle.IsValid() || !NameHandle->IsValidHandle())
+	{
+		return  NAME_None;
+	}
+
+	FName GoalName;
+	NameHandle->GetValue(GoalName);
+	return GoalName;
+}
+
+FName FIKRigGoalLayout::GetName() const
+{
+	return GetGoalName(GoalPropHandle);
+}
+
+EIKRigGoalTransformSource FIKRigGoalLayout::GetTransformSource() const
+{
+	if (!GoalPropHandle->IsValidHandle())
+	{
+		return EIKRigGoalTransformSource::Manual;
+	}
+	
+	const TSharedPtr<IPropertyHandle> TransformSourceHandle = GoalPropHandle->GetChildHandle(GET_MEMBER_NAME_CHECKED(FIKRigGoal, TransformSource));
+	if (!TransformSourceHandle->IsValidHandle())
+	{
+		return EIKRigGoalTransformSource::Manual;
+	}
+	
+	uint8 Source;
+	TransformSourceHandle->GetValue(Source);
+	return static_cast<EIKRigGoalTransformSource>(Source);
+}
+
+TSharedPtr<IPropertyHandle> FIKRigGoalLayout::GetBoneNameHandle() const
+{
+	if (!GoalPropHandle->IsValidHandle())
+	{
+		return nullptr;
+	}
+	
+	const TSharedPtr<IPropertyHandle> SourceBoneHandle = GoalPropHandle->GetChildHandle(GET_MEMBER_NAME_CHECKED(FIKRigGoal, SourceBone));
+	if (!SourceBoneHandle->IsValidHandle())
+	{
+		return nullptr;
+	}
+
+	return SourceBoneHandle->GetChildHandle(GET_MEMBER_NAME_CHECKED(FBoneReference, BoneName));
+}
+
+void FIKRigGoalLayout::OnBoneSelectionChanged(FName Name) const
+{
+	const TSharedPtr<IPropertyHandle> BoneNameProperty = GetBoneNameHandle();
+	if (BoneNameProperty.IsValid() && BoneNameProperty->IsValidHandle())
+	{
+		BoneNameProperty->SetValue(Name);
+	}
+}
+
+FName FIKRigGoalLayout::GetSelectedBone(bool& bMultipleValues) const
+{
+	const TSharedPtr<IPropertyHandle> BoneNameProperty = GetBoneNameHandle();
+	if (!BoneNameProperty.IsValid() || !BoneNameProperty->IsValidHandle())
+	{
+		return NAME_None;
+	}
+	
+	FString OutName;
+	const FPropertyAccess::Result Result = BoneNameProperty->GetValueAsFormattedString(OutName);
+	bMultipleValues = (Result == FPropertyAccess::MultipleValues);
+
+	return FName(*OutName);
+}
+
+const struct FReferenceSkeleton& FIKRigGoalLayout::GetReferenceSkeleton() const
+{
+	static const FReferenceSkeleton DummySkeleton;
+	
+	if (!GoalPropHandle->IsValidHandle())
+	{
+		return DummySkeleton;
+	}
+	
+	const TSharedPtr<IPropertyHandle> SourceBoneHandle = GoalPropHandle->GetChildHandle(GET_MEMBER_NAME_CHECKED(FIKRigGoal, SourceBone));
+	if (!SourceBoneHandle->IsValidHandle())
+	{
+		return DummySkeleton;
+	}
+	
+	TArray<UObject*> Objects;
+	SourceBoneHandle->GetOuterObjects(Objects);
+
+	USkeleton* TargetSkeleton = nullptr;
+
+	auto FindSkeletonForObject = [&TargetSkeleton](UObject* InObject)
+	{
+		for( ; InObject; InObject = InObject->GetOuter())
+		{
+			if (UAnimGraphNode_Base* AnimGraphNode = Cast<UAnimGraphNode_Base>(InObject))
+			{
+				TargetSkeleton = AnimGraphNode->GetAnimBlueprint()->TargetSkeleton;
+				break;
+			}
+		}
+
+		return TargetSkeleton != nullptr;
+	};
+	
+	for (UObject* Object : Objects)
+	{
+		if(FindSkeletonForObject(Object))
+		{
+			break;
+		}
+	}
+	
+	return TargetSkeleton ? TargetSkeleton->GetReferenceSkeleton() : DummySkeleton;
+}
+
+/////////////////////////////////////////////////////
+// FIKRigGoalArrayLayout
+
+void FIKRigGoalArrayLayout::GenerateChildContent(IDetailChildrenBuilder& ChildrenBuilder)
+{
+	if (NodePropHandle.IsValid())
+	{
+		const UObject* Object = nullptr;
+		const TSharedPtr<IPropertyHandle> RigDefAssetHandle = NodePropHandle->GetChildHandle(GET_MEMBER_NAME_CHECKED(FAnimNode_IKRig, RigDefinitionAsset));
+		if (RigDefAssetHandle->GetValue(Object) == FPropertyAccess::Fail || Object == nullptr)
+		{
+			return;
+		}
+
+		const UIKRigDefinition* IKRigDefinition = CastChecked<const UIKRigDefinition>(Object);
+		if (!IKRigDefinition)
+		{
+			return;
+		}
+
+		const TSharedPtr<IPropertyHandle> GoalsHandle = NodePropHandle->GetChildHandle(GET_MEMBER_NAME_CHECKED(FAnimNode_IKRig, Goals));
+		const TArray<UIKRigEffectorGoal*>& AssetGoals = IKRigDefinition->GetGoalArray();
+
+		// add customization for each goal
+		uint32 NumGoals = 0;
+		GoalsHandle->GetNumChildren(NumGoals);
+		for (uint32 Index = 0; Index < NumGoals; Index++)
+		{
+			TSharedPtr<IPropertyHandle> GoalHandle = GoalsHandle->GetChildHandle(Index);
+			if (GoalHandle.IsValid())
+			{
+				const int32 AssetGoalIndex = AssetGoals.IndexOfByPredicate([&GoalHandle](const UIKRigEffectorGoal* InAssetGoal)
+				{
+					return FIKRigGoalLayout::GetGoalName(GoalHandle) == InAssetGoal->GoalName;
+				});
+
+				if (AssetGoalIndex != INDEX_NONE)
+				{
+					const UIKRigEffectorGoal* AssetGoal = AssetGoals[AssetGoalIndex];
+					if (AssetGoal->bExposePosition || AssetGoal->bExposeRotation)
+					{
+						TSharedRef<FIKRigGoalLayout> ControlRigArgumentLayout = MakeShareable(
+						   new FIKRigGoalLayout(GoalHandle, AssetGoal->bExposePosition, AssetGoal->bExposeRotation));
+						ChildrenBuilder.AddCustomBuilder(ControlRigArgumentLayout);
+					}
+				}
+			}
+		}
+	}
+}
+
+UAnimGraphNode_IKRig::~UAnimGraphNode_IKRig()
+{
+	if (OnAssetPropertyChangedHandle.IsValid())
+	{
+		FCoreUObjectDelegates::OnObjectPropertyChanged.Remove(OnAssetPropertyChangedHandle);
+		OnAssetPropertyChangedHandle.Reset();
+	}
+}
+
 void UAnimGraphNode_IKRig::Draw(FPrimitiveDrawInterface* PDI, USkeletalMeshComponent* PreviewSkelMeshComp) const
 {
 	if(PreviewSkelMeshComp)
@@ -42,4 +331,394 @@ void UAnimGraphNode_IKRig::ValidateAnimNodeDuringCompilation(USkeleton* ForSkele
 	}
 }
 
+void UAnimGraphNode_IKRig::PostEditChangeProperty(struct FPropertyChangedEvent& PropertyChangedEvent)
+{	
+	const FName PropertyName = PropertyChangedEvent.GetPropertyName();
+	
+	if (PropertyName == GET_MEMBER_NAME_STRING_CHECKED(FAnimNode_IKRig, RigDefinitionAsset))
+	{
+		Node.Goals.Empty();
+		if (IsValid(Node.RigDefinitionAsset))
+		{
+			// create new goals based on the rig definition
+			const TArray<UIKRigEffectorGoal*>& AssetGoals = Node.RigDefinitionAsset->GetGoalArray();
+			for (const UIKRigEffectorGoal* AssetGoal: AssetGoals)
+			{
+				if (AssetGoal->bExposePosition || AssetGoal->bExposeRotation)
+				{
+					Node.Goals.Emplace(AssetGoal->GoalName);
+				}
+			}
+
+			BindPropertyChanges();
+		}
+		ReconstructNode();
+		return;
+	}
+	
+	if (PropertyName == GET_MEMBER_NAME_STRING_CHECKED(FAnimNode_IKRig, Goals))
+	{
+		ReconstructNode();
+		return;
+	}
+	
+	if (PropertyName == GET_MEMBER_NAME_STRING_CHECKED(FIKRigGoal, TransformSource))
+	{
+		ReconstructNode();
+		return;
+	}
+	
+	Super::PostEditChangeProperty(PropertyChangedEvent);
+}
+
+FName GetGoalSubPropertyPinPrettyName(const FName& InGoalName, const FName& InPropertyName)
+{
+	return *FString::Printf(TEXT("%s_%s"), *InGoalName.ToString(), *InPropertyName.ToString());
+}
+
+void UAnimGraphNode_IKRig::CreateCustomPins(TArray<UEdGraphPin*>* InOldPins)
+{
+	if (!IsValid(Node.RigDefinitionAsset))
+	{
+		return;
+	}
+
+	// the asset is not completely loaded so we'll use the old pins to sustain the current set of custom pins
+	if (Node.RigDefinitionAsset->HasAllFlags(RF_NeedPostLoad)) 
+	{
+		CreateCustomPinsFromUnloadedAsset(InOldPins);
+		return;
+	}
+
+	// generate pins based on the current asset
+	CreateCustomPinsFromValidAsset();
+}
+
+void UAnimGraphNode_IKRig::SetPinDefaultValue(UEdGraphPin* InPin, const FName& InPropertyName)
+{
+	static FIKRigGoal DefaultGoal;
+	static const TSharedPtr<FStructOnScope> StructOnScope =
+							MakeShareable(new FStructOnScope(FIKRigGoal::StaticStruct(), (uint8*)&DefaultGoal));
+	const UAnimationGraphSchema* AnimGraphDefaultSchema = GetDefault<UAnimationGraphSchema>();
+	
+	// default value
+	AnimGraphDefaultSchema->SetPinAutogeneratedDefaultValueBasedOnType(InPin);
+	if (const FProperty* Property = FIKRigGoal::StaticStruct()->FindPropertyByName(InPropertyName))
+	{
+		const uint8* Memory = Property->ContainerPtrToValuePtr<uint8>(StructOnScope->GetStructMemory());
+		FString DefaultValue;
+		Property->ExportTextItem(DefaultValue, Memory, nullptr, nullptr, PPF_None);
+		if (!DefaultValue.IsEmpty())
+		{
+			AnimGraphDefaultSchema->TrySetDefaultValue(*InPin, DefaultValue);
+		}
+	}
+}
+
+void UAnimGraphNode_IKRig::CreateCustomPinsFromUnloadedAsset(TArray<UEdGraphPin*>* InOldPins)
+{
+	// recreate pin from old pin
+	auto RecreateGoalPin = [this](const UEdGraphPin* InOldPin)
+	{
+		// pin's name is based on the property name within the FIKRigGoal structure + the index within the Goals array
+		const FName PropertyName = InOldPin->GetFName();
+		
+		UEdGraphPin* NewPin = CreatePin(EEdGraphPinDirection::EGPD_Input, InOldPin->PinType, PropertyName);
+
+		// pin's pretty name is the "GoalName_InPropertyName"
+		NewPin->PinFriendlyName = InOldPin->PinFriendlyName;
+
+		// default value
+		SetPinDefaultValue(NewPin, PropertyName);
+	};
+
+	// ensure that this is a goal related pin
+	auto bNeedsCreation = [&](const UEdGraphPin* InOldPin)
+	{
+		// custom pins are inputs
+		if (InOldPin->Direction != EEdGraphPinDirection::EGPD_Input)
+		{
+			return false;
+		}
+
+		// look for old pin's name-type into current pins
+		const int32 PinIndex = Pins.IndexOfByPredicate([&](const UEdGraphPin* Pin)
+		{
+			return Pin->GetFName() == InOldPin->GetFName() && Pin->PinType == InOldPin->PinType;
+		});
+		
+		return PinIndex == INDEX_NONE;
+	};
+
+	// recreate pins if needed
+	if (InOldPins)
+	{
+		for (const UEdGraphPin* OldPin : *InOldPins)
+		{
+			if (bNeedsCreation(OldPin))
+			{
+				RecreateGoalPin(OldPin);
+			}
+		}
+	}
+}
+
+void UAnimGraphNode_IKRig::CreateCustomPinsFromValidAsset()
+{
+	// pin's creation function
+	auto CreateGoalPin = [this]( const uint32 InGoalIndex,
+								 const FName& InPropertyName,
+								 const FEdGraphPinType InPinType)
+	{
+		const FName& GoalName = Node.Goals[InGoalIndex].Name;
+		const uint32 GoalHash = GetTypeHash(GoalName);
+		
+		// pin's name is based on the property name within the FIKRigGoal structure + the name's hash value as a number
+		FName PinName = InPropertyName;
+		PinName.SetNumber(GoalHash);
+		
+		UEdGraphPin* NewPin = CreatePin(EEdGraphPinDirection::EGPD_Input, InPinType, PinName);
+		
+		// pin's pretty name is the "GoalName_InPropertyName"
+		NewPin->PinFriendlyName = FText::FromName(GetGoalSubPropertyPinPrettyName(GoalName, InPropertyName));
+
+		// default value
+		SetPinDefaultValue(NewPin, InPropertyName);
+	};
+
+	static const FName PC_Struct(TEXT("struct"));
+	
+	// position property
+	static FEdGraphPinType PositionPinType;
+	PositionPinType.PinCategory = PC_Struct;
+	PositionPinType.PinSubCategoryObject = TBaseStructure<FVector>::Get();
+
+	// rotation property
+	static FEdGraphPinType RotationPinType;
+	RotationPinType.PinCategory = PC_Struct;
+	RotationPinType.PinSubCategoryObject = TBaseStructure<FRotator>::Get();
+	
+	// alpha property
+	static const FName PC_Float(TEXT("float"));
+	static FEdGraphPinType AlphaPinType;
+	AlphaPinType.PinCategory = PC_Float;
+
+	// create pins
+	const TArray<UIKRigEffectorGoal*>& AssetGoals = Node.RigDefinitionAsset->GetGoalArray();
+	const TArray<FIKRigGoal>& Goals = Node.Goals;
+	for (int32 GoalIndex = 0; GoalIndex < Goals.Num(); GoalIndex++)
+	{
+		const FIKRigGoal& Goal = Node.Goals[GoalIndex];
+		
+		const int32 AssetGoalIndex = AssetGoals.IndexOfByPredicate([&Goal](const UIKRigEffectorGoal* InAssetGoal)
+		{
+			return Goal.Name == InAssetGoal->GoalName;
+		});
+
+		if (AssetGoalIndex == INDEX_NONE)
+		{
+			continue;
+		}
+
+		const UIKRigEffectorGoal* AssetGoal = AssetGoals[AssetGoalIndex];
+		if (Goal.TransformSource == EIKRigGoalTransformSource::Manual)
+		{
+			// position
+			if (AssetGoal->bExposePosition)
+			{
+				CreateGoalPin(GoalIndex, GET_MEMBER_NAME_CHECKED(FIKRigGoal, Position), PositionPinType);
+				CreateGoalPin(GoalIndex, GET_MEMBER_NAME_CHECKED(FIKRigGoal, PositionAlpha), AlphaPinType);
+			}
+
+			// rotation
+			if (AssetGoal->bExposeRotation)
+			{
+				CreateGoalPin(GoalIndex, GET_MEMBER_NAME_CHECKED(FIKRigGoal, Rotation), RotationPinType);
+				CreateGoalPin(GoalIndex, GET_MEMBER_NAME_CHECKED(FIKRigGoal, RotationAlpha), AlphaPinType);
+			}
+		}
+		else if (Goal.TransformSource == EIKRigGoalTransformSource::Bone)
+		{
+			// position
+			if (AssetGoal->bExposePosition)
+			{
+				CreateGoalPin(GoalIndex, GET_MEMBER_NAME_CHECKED(FIKRigGoal, PositionAlpha), AlphaPinType);
+			}
+
+			// rotation
+			if (AssetGoal->bExposeRotation)
+			{
+				CreateGoalPin(GoalIndex, GET_MEMBER_NAME_CHECKED(FIKRigGoal, RotationAlpha), AlphaPinType);
+			}
+		}
+	}
+}
+
+void UAnimGraphNode_IKRig::PostLoad()
+{
+	Super::PostLoad();
+
+	// update goals name if needed
+	if (IsValid(Node.RigDefinitionAsset))
+	{
+		//NOTE needed?
+		const TArray<UIKRigEffectorGoal*>& AssetGoals = Node.RigDefinitionAsset->GetGoalArray();
+	
+		const int32 NumAssetGoals = AssetGoals.Num();
+		const int32 NumNodeGoals = Node.Goals.Num();
+		if (NumAssetGoals == NumNodeGoals)
+		{
+			for (int32 Index = 0; Index < NumNodeGoals; Index++)
+			{
+				FName& GoalName = Node.Goals[Index].Name;
+				if (GoalName.IsNone())
+				{
+					GoalName = AssetGoals[Index]->GoalName;
+				}
+			}
+		}
+
+		// listened to changes within the asset / goals
+		BindPropertyChanges();
+	}
+}
+
+void UAnimGraphNode_IKRig::CustomizeDetails(IDetailLayoutBuilder& DetailBuilder)
+{
+	Super::CustomizeDetails(DetailBuilder);
+
+	// Do not allow multi-selection
+	if (DetailBuilder.GetSelectedObjects().Num() > 1)
+	{
+		return;
+	}
+
+	// Add goals customization
+	const TSharedRef<IPropertyHandle> NodePropHandle = DetailBuilder.GetProperty(GET_MEMBER_NAME_CHECKED(UAnimGraphNode_IKRig, Node), GetClass());
+	if (NodePropHandle->IsValidHandle())
+	{
+		TSharedRef<FIKRigGoalArrayLayout> InputArgumentGroup = MakeShareable(new FIKRigGoalArrayLayout(NodePropHandle));
+
+		IDetailCategoryBuilder& GoalsCategoryBuilder = DetailBuilder.EditCategory(GET_MEMBER_NAME_CHECKED(FAnimNode_IKRig, Goals));
+		GoalsCategoryBuilder.AddCustomBuilder(InputArgumentGroup);
+	}
+
+	// Handle property changed notification
+	const FSimpleDelegate OnValueChanged = FSimpleDelegate::CreateLambda([&DetailBuilder]()
+	{
+		DetailBuilder.ForceRefreshDetails();
+	});
+	
+	const TSharedRef<IPropertyHandle> AssetHandle = DetailBuilder.GetProperty(TEXT("Node.RigDefinitionAsset"), GetClass());
+	if (AssetHandle->IsValidHandle())
+	{
+		AssetHandle->SetOnPropertyValueChanged(OnValueChanged);
+	}
+	
+	const TSharedRef<IPropertyHandle> GoalHandle = DetailBuilder.GetProperty(TEXT("Node.Goals"), GetClass());
+	if (AssetHandle->IsValidHandle())
+	{
+		GoalHandle->SetOnChildPropertyValueChanged(OnValueChanged);
+	}
+}
+
+void UAnimGraphNode_IKRig::BindPropertyChanges()
+{
+	// already bound
+	if (OnAssetPropertyChangedHandle.IsValid())
+	{
+		return;	
+	}
+
+	// listen to the rig definition asset
+	using FPropertyChangedDelegate = FCoreUObjectDelegates::FOnObjectPropertyChanged::FDelegate;
+	const FPropertyChangedDelegate OnPropertyChangedDelegate =
+				FPropertyChangedDelegate::CreateUObject(this, &UAnimGraphNode_IKRig::OnPropertyChanged);
+	
+	OnAssetPropertyChangedHandle = FCoreUObjectDelegates::OnObjectPropertyChanged.Add(OnPropertyChangedDelegate);
+}
+
+void UAnimGraphNode_IKRig::OnPropertyChanged(UObject* ObjectBeingModified, FPropertyChangedEvent& PropertyChangedEvent)
+{
+	if (NeedsUpdate(ObjectBeingModified, PropertyChangedEvent))
+	{
+		UpdateGoalsFromAsset();
+	}
+}
+
+bool UAnimGraphNode_IKRig::NeedsUpdate(UObject* ObjectBeingModified, FPropertyChangedEvent& PropertyChangedEvent) const
+{
+	if (!ObjectBeingModified)
+	{
+		return false;
+	}
+	
+	if (!IsValid(Node.RigDefinitionAsset))
+	{
+		return false;
+	}
+
+	const FName PropertyName = PropertyChangedEvent.GetPropertyName();
+	if (PropertyName == NAME_None)
+	{
+		return false;
+	}
+
+	// something has changed within the asset
+	if (ObjectBeingModified == Node.RigDefinitionAsset)
+	{
+		// we can't use GET_MEMBER_NAME_CHECKED as Goals is a private property
+		static const TArray<FName> AssetWatchedProperties({ TEXT("Goals") } );
+
+		const bool bNeedsUpdate = AssetWatchedProperties.Contains(PropertyName);
+		return bNeedsUpdate;
+	}
+
+	// check whether this is a goal and if it belongs to the current asset
+	if (UIKRigEffectorGoal* GoalBeingModified = Cast<UIKRigEffectorGoal>(ObjectBeingModified))
+	{
+		static const TArray<FName> GoalWatchedProperties({	GET_MEMBER_NAME_CHECKED(UIKRigEffectorGoal, GoalName),
+															GET_MEMBER_NAME_CHECKED(UIKRigEffectorGoal, bExposePosition),
+															GET_MEMBER_NAME_CHECKED(UIKRigEffectorGoal, bExposeRotation)} );
+		
+		const TArray<UIKRigEffectorGoal*>& AssetGoals = Node.RigDefinitionAsset->GetGoalArray();
+		const bool bNeedsUpdate = AssetGoals.Contains(GoalBeingModified) && GoalWatchedProperties.Contains(PropertyName);
+		return bNeedsUpdate;
+	}
+	
+	return false;
+}
+
+void UAnimGraphNode_IKRig::UpdateGoalsFromAsset()
+{
+	TArray<FIKRigGoal> OldGoals = Node.Goals;
+	Node.Goals.Empty();
+
+	if (IsValid(Node.RigDefinitionAsset))
+	{
+		const TArray<UIKRigEffectorGoal*>& AssetGoals = Node.RigDefinitionAsset->GetGoalArray();
+		for (const UIKRigEffectorGoal* AssetGoal: AssetGoals)
+		{
+			if (AssetGoal->bExposePosition || AssetGoal->bExposeRotation)
+			{
+				const int32 OldGoalIndex = OldGoals.IndexOfByPredicate([&AssetGoal](const FIKRigGoal& OldGoal)
+				{
+					return AssetGoal->GoalName == OldGoal.Name;
+				});
+
+				if (OldGoalIndex != INDEX_NONE)
+				{
+					Node.Goals.Add(OldGoals[OldGoalIndex]);
+				}
+				else
+				{
+					Node.Goals.Emplace(AssetGoal->GoalName);
+				}
+			}
+		}
+	}
+	
+	ReconstructNode();
+}
+
 #undef LOCTEXT_NAMESPACE
diff --git a/Engine/Plugins/Animation/IKRig/Source/IKRigDeveloper/Public/AnimGraphNode_IKRig.h b/Engine/Plugins/Animation/IKRig/Source/IKRigDeveloper/Public/AnimGraphNode_IKRig.h
--- a/Engine/Plugins/Animation/IKRig/Source/IKRigDeveloper/Public/AnimGraphNode_IKRig.h
+++ b/Engine/Plugins/Animation/IKRig/Source/IKRigDeveloper/Public/AnimGraphNode_IKRig.h
@@ -3,17 +3,88 @@
 #pragma once
 
 #include "CoreMinimal.h"
-#include "UObject/ObjectMacros.h"
-#include "AnimGraphNode_Base.h"
+#include "AnimGraphNode_CustomProperty.h"
 #include "AnimNodes/AnimNode_IKRig.h"
+#include "IDetailCustomNodeBuilder.h"
+
 #include "AnimGraphNode_IKRig.generated.h"
 
 class FPrimitiveDrawInterface;
 class USkeletalMeshComponent;
 
+/////////////////////////////////////////////////////
+// FIKRigGoalLayout
+
+class FIKRigGoalLayout : public IDetailCustomNodeBuilder, public TSharedFromThis<FIKRigGoalLayout>
+{
+public:
+	
+	FIKRigGoalLayout(TSharedPtr<IPropertyHandle> InGoalPropHandle,
+					 const bool InExposePosition,
+					 const bool InExposeRotation)
+		: GoalPropHandle(InGoalPropHandle)
+		, bExposePosition(InExposePosition)
+		, bExposeRotation(InExposeRotation)
+	{}
+
+	/** IDetailCustomNodeBuilder Interface*/
+	virtual void GenerateHeaderRowContent(FDetailWidgetRow& InOutGoalRow) override;
+	virtual void GenerateChildContent(IDetailChildrenBuilder& InOutChildrenBuilder) override;
+	virtual void Tick(float DeltaTime) override {}
+	virtual bool RequiresTick() const override { return false; }
+	virtual FName GetName() const override;
+	virtual bool InitiallyCollapsed() const override { return true; }
+
+	static FName GetGoalName(TSharedPtr<IPropertyHandle> InGoalHandle);
+	
+private:
+
+	EIKRigGoalTransformSource GetTransformSource() const;
+
+	const struct FReferenceSkeleton& GetReferenceSkeleton() const;
+	TSharedPtr<IPropertyHandle> GetBoneNameHandle() const;
+	void OnBoneSelectionChanged(FName Name) const;
+	FName GetSelectedBone(bool& bMultipleValues) const;
+
+	TSharedRef<SWidget> CreateManualValueWidget() const;
+	TSharedRef<SWidget> CreateBoneValueWidget() const;
+	TSharedRef<SWidget> CreateValueWidget() const;
+	
+	TSharedPtr<IPropertyHandle> GoalPropHandle = nullptr;
+	bool bExposePosition = false;
+	bool bExposeRotation = false;
+};
+
+/////////////////////////////////////////////////////
+// FIKRigGoalArrayLayout
+
+class FIKRigGoalArrayLayout : public IDetailCustomNodeBuilder, public TSharedFromThis<FIKRigGoalArrayLayout>
+{
+public:
+	
+	FIKRigGoalArrayLayout(TSharedPtr<IPropertyHandle> InNodePropHandle)
+		: NodePropHandle(InNodePropHandle)
+	{}
+	
+	virtual ~FIKRigGoalArrayLayout() {}
+
+	/** IDetailCustomNodeBuilder Interface*/
+	virtual void SetOnRebuildChildren(FSimpleDelegate InOnRegenerateChildren) override {}
+	virtual void GenerateHeaderRowContent(FDetailWidgetRow& NodeRow) override {}
+	virtual void GenerateChildContent(IDetailChildrenBuilder& ChildrenBuilder) override;
+	virtual void Tick(float DeltaTime) override {}
+	virtual bool RequiresTick() const override { return false; }
+	virtual FName GetName() const override { return NAME_None; }
+	virtual bool InitiallyCollapsed() const override { return false; }
+
+private:
+	
+	TSharedPtr<IPropertyHandle> NodePropHandle;
+};
+
 // Editor node for IKRig 
 UCLASS()
-class IKRIGDEVELOPER_API UAnimGraphNode_IKRig : public UAnimGraphNode_Base
+class IKRIGDEVELOPER_API UAnimGraphNode_IKRig : public UAnimGraphNode_CustomProperty
 {
 	GENERATED_BODY()
 
@@ -21,8 +92,15 @@ class IKRIGDEVELOPER_API UAnimGraphNode_IKRig : public UAnimGraphNode_Base
 	FAnimNode_IKRig Node;
 
 public:
+
+	UAnimGraphNode_IKRig() = default;
+	virtual ~UAnimGraphNode_IKRig();
+	
 	// UEdGraphNode interface
 	virtual FText GetNodeTitle(ENodeTitleType::Type TitleType) const override;
+	virtual void PostEditChangeProperty(struct FPropertyChangedEvent& PropertyChangedEvent) override;
+	virtual void CreateCustomPins(TArray<UEdGraphPin*>* InOldPins) override;
+	virtual void CustomizeDetails(IDetailLayoutBuilder& DetailBuilder) override;
 	// End of UEdGraphNode interface
 
 	// UAnimGraphNode_Base interface
@@ -30,4 +108,36 @@ public:
 	virtual void Draw(FPrimitiveDrawInterface* PDI, USkeletalMeshComponent* PreviewSkelMeshComp) const override;
 	virtual void ValidateAnimNodeDuringCompilation(USkeleton* ForSkeleton, FCompilerResultsLog& MessageLog) override;
 	// End of UAnimGraphNode_Base interface
+
+	// Begin UObject Interface.
+	virtual void PostLoad() override;
+	// End UObject Interface.
+
+	virtual bool NeedsToSpecifyValidTargetClass() const override { return false; }
+	
+protected:
+	
+	virtual FAnimNode_CustomProperty* GetCustomPropertyNode() override { return &Node; }
+	virtual const FAnimNode_CustomProperty* GetCustomPropertyNode() const override { return &Node; }
+
+private:
+
+	// set pin's default value based on the FIKRigGoal default struct
+	static void SetPinDefaultValue(UEdGraphPin* InPin, const FName& InPropertyName);
+
+	// create custom pins from exposed goals as defined in the rig definition asset 
+	void CreateCustomPinsFromValidAsset();
+
+	// recreate custom pins from old pins if the rig definition asset is not completely loaded
+	void CreateCustomPinsFromUnloadedAsset(TArray<UEdGraphPin*>* InOldPins);
+
+	// Handle to the registered delegate
+	FDelegateHandle OnAssetPropertyChangedHandle;
+
+	// Global callback to anticipate on changes to the asset / goals
+	bool NeedsUpdate(UObject* ObjectBeingModified, FPropertyChangedEvent& PropertyChangedEvent) const;
+	void OnPropertyChanged(UObject* ObjectBeingModified, FPropertyChangedEvent& PropertyChangedEvent);
+	void BindPropertyChanges();
+
+	void UpdateGoalsFromAsset();
 };
diff --git a/Engine/Plugins/Animation/IKRig/Source/IKRigEditor/Private/RigEditor/IKRigController.cpp b/Engine/Plugins/Animation/IKRig/Source/IKRigEditor/Private/RigEditor/IKRigController.cpp
--- a/Engine/Plugins/Animation/IKRig/Source/IKRigEditor/Private/RigEditor/IKRigController.cpp
+++ b/Engine/Plugins/Animation/IKRig/Source/IKRigEditor/Private/RigEditor/IKRigController.cpp
@@ -610,8 +610,9 @@ UIKRigEffectorGoal* UIKRigController::AddNewGoal(const FName& GoalName, const FN
 	// set initial transform
 	NewGoal->InitialTransform = GetRefPoseTransformOfBone(NewGoal->BoneName);
 	NewGoal->CurrentTransform = NewGoal->InitialTransform;
-
+ 
 	BroadcastNeedsReinitialized();
+	BroadcastGoalsChange();
 	
 	return NewGoal;
 }
@@ -648,6 +649,7 @@ bool UIKRigController::RemoveGoal(const FName& GoalName) const
 	}
 
 	BroadcastNeedsReinitialized();
+	BroadcastGoalsChange();
 	
 	return true;
 }
@@ -695,6 +697,7 @@ FName UIKRigController::RenameGoal(const FName& OldName, const FName& PotentialN
 	}
 
 	BroadcastNeedsReinitialized();
+	BroadcastGoalsChange();
 
 	return NewName;
 }
@@ -927,4 +930,17 @@ FName UIKRigController::GetGoalName(const int32& GoalIndex) const
 	return Asset->Goals[GoalIndex]->GoalName;
 }
 
+void UIKRigController::BroadcastGoalsChange() const
+{
+	if (Asset)
+	{
+		static const FName GoalsPropName = GET_MEMBER_NAME_CHECKED(UIKRigDefinition, Goals);
+		if (FProperty* GoalProperty = UIKRigDefinition::StaticClass()->FindPropertyByName(GoalsPropName) )
+		{
+			FPropertyChangedEvent GoalPropertyChangedEvent(GoalProperty);
+			FCoreUObjectDelegates::OnObjectPropertyChanged.Broadcast(Asset, GoalPropertyChangedEvent);
+		}
+	}
+}
+
 #undef LOCTEXT_NAMESPACE
diff --git a/Engine/Plugins/Animation/IKRig/Source/IKRigEditor/Public/RigEditor/IKRigController.h b/Engine/Plugins/Animation/IKRig/Source/IKRigEditor/Public/RigEditor/IKRigController.h
--- a/Engine/Plugins/Animation/IKRig/Source/IKRigEditor/Public/RigEditor/IKRigController.h
+++ b/Engine/Plugins/Animation/IKRig/Source/IKRigEditor/Public/RigEditor/IKRigController.h
@@ -200,5 +200,8 @@ private:
 	/** Lazy-generated map of Controllers to IK Rig Assets. Avoids duplicate controllers. */
 	static TMap<UIKRigDefinition*, UIKRigController*> AssetToControllerMap;
 
+	// broadcast changes within the asset goals array
+	void BroadcastGoalsChange() const;
+	
 	friend class UIKRigDefinition;
 };
