Retargeter debug: goals were not updated properly when exported animations (based on MH feedback)
added a minor change regarding tranform copy/references

diff --git a/Engine/Plugins/Animation/IKRig/Source/IKRig/Private/Retargeter/IKRetargetProcessor.cpp b/Engine/Plugins/Animation/IKRig/Source/IKRig/Private/Retargeter/IKRetargetProcessor.cpp
--- a/Engine/Plugins/Animation/IKRig/Source/IKRig/Private/Retargeter/IKRetargetProcessor.cpp
+++ b/Engine/Plugins/Animation/IKRig/Source/IKRig/Private/Retargeter/IKRetargetProcessor.cpp
@@ -54,7 +54,7 @@ int32 FRetargetSkeleton::FindBoneIndexByName(const FName InName) const
 
 void FRetargetSkeleton::UpdateGlobalTransformsBelowBone(
 	const int32 StartBoneIndex,
-	TArray<FTransform>& InLocalPose,
+	const TArray<FTransform>& InLocalPose,
 	TArray<FTransform>& OutGlobalPose) const
 {
 	check(BoneNames.IsValidIndex(StartBoneIndex));
@@ -70,7 +70,7 @@ void FRetargetSkeleton::UpdateGlobalTransformsBelowBone(
 void FRetargetSkeleton::UpdateLocalTransformsBelowBone(
 	const int32 StartBoneIndex,
 	TArray<FTransform>& OutLocalPose,
-	TArray<FTransform>& InGlobalPose) const
+	const TArray<FTransform>& InGlobalPose) const
 {
 	check(BoneNames.IsValidIndex(StartBoneIndex));
 	check(BoneNames.Num() == OutLocalPose.Num());
@@ -92,23 +92,23 @@ void FRetargetSkeleton::UpdateGlobalTransformOfSingleBone(
 	{
 		return; // root always in global space
 	}
-	const FTransform ChildLocalTransform = InLocalPose[BoneIndex];
-	const FTransform ParentGlobalTransform = OutGlobalPose[ParentIndex];
+	const FTransform& ChildLocalTransform = InLocalPose[BoneIndex];
+	const FTransform& ParentGlobalTransform = OutGlobalPose[ParentIndex];
 	OutGlobalPose[BoneIndex] = ChildLocalTransform * ParentGlobalTransform;
 }
 
 void FRetargetSkeleton::UpdateLocalTransformOfSingleBone(
 	const int32 BoneIndex,
 	TArray<FTransform>& OutLocalPose,
-	TArray<FTransform>& InGlobalPose) const
+	const TArray<FTransform>& InGlobalPose) const
 {
 	const int32 ParentIndex = ParentIndices[BoneIndex];
 	if (ParentIndex == INDEX_NONE)
 	{
 		return; // root always in global space
 	}
-	const FTransform ChildGlobalTransform = InGlobalPose[BoneIndex];
-	const FTransform ParentGlobalTransform = InGlobalPose[ParentIndex];
+	const FTransform& ChildGlobalTransform = InGlobalPose[BoneIndex];
+	const FTransform& ParentGlobalTransform = InGlobalPose[ParentIndex];
 	OutLocalPose[BoneIndex] = ChildGlobalTransform.GetRelativeTransform(ParentGlobalTransform);
 }
 
diff --git a/Engine/Plugins/Animation/IKRig/Source/IKRig/Public/Retargeter/IKRetargetProcessor.h b/Engine/Plugins/Animation/IKRig/Source/IKRig/Public/Retargeter/IKRetargetProcessor.h
--- a/Engine/Plugins/Animation/IKRig/Source/IKRig/Public/Retargeter/IKRetargetProcessor.h
+++ b/Engine/Plugins/Animation/IKRig/Source/IKRig/Public/Retargeter/IKRetargetProcessor.h
@@ -38,13 +38,13 @@ struct IKRIG_API FRetargetSkeleton
 
 	void UpdateGlobalTransformsBelowBone(
 		const int32 StartBoneIndex,
-		TArray<FTransform>& InLocalPose,
+		const TArray<FTransform>& InLocalPose,
 		TArray<FTransform>& OutGlobalPose) const;
 
 	void UpdateLocalTransformsBelowBone(
 		const int32 StartBoneIndex,
 		TArray<FTransform>& OutLocalPose,
-		TArray<FTransform>& InGlobalPose) const;
+		const TArray<FTransform>& InGlobalPose) const;
 	
 	void UpdateGlobalTransformOfSingleBone(
 		const int32 BoneIndex,
@@ -54,7 +54,7 @@ struct IKRIG_API FRetargetSkeleton
 	void UpdateLocalTransformOfSingleBone(
 		const int32 BoneIndex,
 		TArray<FTransform>& OutLocalPose,
-		TArray<FTransform>& InGlobalPose) const;
+		const TArray<FTransform>& InGlobalPose) const;
 
 	void GetChildrenIndices(const int32 BoneIndex, TArray<int32>& OutChildren) const;
 };
diff --git a/Engine/Plugins/Animation/IKRig/Source/IKRigEditor/Private/RetargetEditor/IKRetargetBatchOperation.cpp b/Engine/Plugins/Animation/IKRig/Source/IKRigEditor/Private/RetargetEditor/IKRetargetBatchOperation.cpp
--- a/Engine/Plugins/Animation/IKRig/Source/IKRigEditor/Private/RetargetEditor/IKRetargetBatchOperation.cpp
+++ b/Engine/Plugins/Animation/IKRig/Source/IKRigEditor/Private/RetargetEditor/IKRetargetBatchOperation.cpp
@@ -254,8 +254,11 @@ void FIKRetargetBatchOperation::ConvertAnimation(
 				SourceComponentPose.Add(BonePose);
 			}
 
+			// update goals 
+			Processor->CopyTargetIKRigSettingsFromAsset();
+			
 			// run the retarget
-			TArray<FTransform>& TargetComponentPose = Processor->RunRetargeter(SourceComponentPose);
+			const TArray<FTransform>& TargetComponentPose = Processor->RunRetargeter(SourceComponentPose);
 
 			// convert to a local-space pose
 			TArray<FTransform> TargetLocalPose = TargetComponentPose;
