Fixed source bone names transform extraction

diff --git a/Engine/Plugins/Animation/IKRig/Source/IKRig/Public/Retargeter/IKRetargetProcessor.h b/Engine/Plugins/Animation/IKRig/Source/IKRig/Public/Retargeter/IKRetargetProcessor.h
--- a/Engine/Plugins/Animation/IKRig/Source/IKRig/Public/Retargeter/IKRetargetProcessor.h
+++ b/Engine/Plugins/Animation/IKRig/Source/IKRig/Public/Retargeter/IKRetargetProcessor.h
@@ -354,6 +354,9 @@ public:
 	/** Get index of the root bone of the target skeleton. */
 	const int32 GetTargetSkeletonRootBone() const { return RootRetargeter.Target.BoneIndex; };
 
+	/** Get read-only access to the source skeleton. */
+	const FRetargetSkeleton& GetSourceSkeleton() const { return SourceSkeleton; };
+	
 	/** Get whether this processor is ready to call RunRetargeter() and generate new poses. */
 	bool IsInitialized() const { return bIsInitialized; };
 
diff --git a/Engine/Plugins/Animation/IKRig/Source/IKRigEditor/Private/RetargetEditor/IKRetargetBatchOperation.cpp b/Engine/Plugins/Animation/IKRig/Source/IKRigEditor/Private/RetargetEditor/IKRetargetBatchOperation.cpp
--- a/Engine/Plugins/Animation/IKRig/Source/IKRigEditor/Private/RetargetEditor/IKRetargetBatchOperation.cpp
+++ b/Engine/Plugins/Animation/IKRig/Source/IKRigEditor/Private/RetargetEditor/IKRetargetBatchOperation.cpp
@@ -211,6 +211,23 @@ void FIKRetargetBatchOperation::ConvertAnimation(
 		return;
 	}
 
+	// target skeleton data
+	const FTargetSkeleton& TargetSkeleton = Processor->GetTargetSkeleton();
+	const TArray<FName>& TargetBoneNames = TargetSkeleton.BoneNames;
+	const int32 NumTargetBones = TargetBoneNames.Num();
+
+	// allocate target keyframe data
+	TArray<FRawAnimSequenceTrack> BoneTracks;
+	BoneTracks.SetNumZeroed(NumTargetBones);
+
+	// source skeleton data
+	const FRetargetSkeleton& SourceSkeleton = Processor->GetSourceSkeleton();
+	const TArray<FName>& SourceBoneNames = SourceSkeleton.BoneNames;
+	const int32 NumSourceBones = SourceBoneNames.Num();
+
+	TArray<FTransform> SourceComponentPose;
+	SourceComponentPose.SetNum(NumSourceBones);
+	
 	// for each pair of source / target animation sequences
 	for (TPair<UAnimationAsset*, UAnimationAsset*>& Pair : DuplicatedAnimAssets)
 	{
@@ -234,24 +251,27 @@ void FIKRetargetBatchOperation::ConvertAnimation(
 		// number of frames in this animation
 		const int32 NumFrames = SourceSequence->GetNumberOfSampledKeys();
 
-		// make space for the target keyframe data
-		const int32 NumTargetBones = Processor->GetTargetSkeleton().BoneNames.Num();
-		TArray<FRawAnimSequenceTrack> BoneTracks;
-		BoneTracks.SetNumZeroed(NumTargetBones);
-
+		// BoneTracks arrays allocation
+		for (int32 TargetBoneIndex=0; TargetBoneIndex<NumTargetBones; ++TargetBoneIndex)
+		{
+			BoneTracks[TargetBoneIndex].PosKeys.SetNum(NumFrames);
+			BoneTracks[TargetBoneIndex].RotKeys.SetNum(NumFrames);
+			BoneTracks[TargetBoneIndex].ScaleKeys.SetNum(NumFrames);
+		}
+		
 		// retarget each frame's pose from source to target
 		for (int32 FrameIndex=0; FrameIndex<NumFrames; ++FrameIndex)
 		{
 			// get the source global pose
 			FAnimPose SourcePoseAtFrame;
 			UAnimPoseExtensions::GetAnimPoseAtFrame(SourceSequence, FrameIndex, FAnimPoseEvaluationOptions(), SourcePoseAtFrame);
-			TArray<FName> BoneNames;
-			UAnimPoseExtensions::GetBoneNames(SourcePoseAtFrame,BoneNames);
-			TArray<FTransform> SourceComponentPose;
-			for (const FName& BoneName : BoneNames)
+
+			// we don't use UAnimPoseExtensions::GetBoneNames as the sequence can store bones that only exist on the
+			// skeleton, but not on the current mesh. This results in indices discrepancy
+			for (int32 BoneIndex = 0; BoneIndex < NumSourceBones; BoneIndex++)
 			{
-				FTransform BonePose = UAnimPoseExtensions::GetBonePose(SourcePoseAtFrame, BoneName, EAnimPoseSpaces::World);
-				SourceComponentPose.Add(BonePose);
+				const FName& BoneName = SourceBoneNames[BoneIndex];
+				SourceComponentPose[BoneIndex] = UAnimPoseExtensions::GetBonePose(SourcePoseAtFrame, BoneName, EAnimPoseSpaces::World);
 			}
 
 			// update goals 
@@ -262,14 +282,18 @@ void FIKRetargetBatchOperation::ConvertAnimation(
 
 			// convert to a local-space pose
 			TArray<FTransform> TargetLocalPose = TargetComponentPose;
-			Processor->GetTargetSkeleton().UpdateLocalTransformsBelowBone(0,TargetLocalPose, TargetComponentPose);
+			TargetSkeleton.UpdateLocalTransformsBelowBone(0,TargetLocalPose, TargetComponentPose);
 
 			// store key data for each bone
-			for (int32 TargetBoneIndex=0; TargetBoneIndex<TargetLocalPose.Num(); ++TargetBoneIndex)
+			for (int32 TargetBoneIndex=0; TargetBoneIndex<NumTargetBones; ++TargetBoneIndex)
 			{
-				BoneTracks[TargetBoneIndex].PosKeys.Add(FVector3f(TargetLocalPose[TargetBoneIndex].GetLocation()));
-				BoneTracks[TargetBoneIndex].RotKeys.Add(FQuat4f(TargetLocalPose[TargetBoneIndex].GetRotation()));
-				BoneTracks[TargetBoneIndex].ScaleKeys.Add(FVector4f(TargetLocalPose[TargetBoneIndex].GetScale3D()));
+				const FTransform& LocalPose = TargetLocalPose[TargetBoneIndex];
+				
+				FRawAnimSequenceTrack& BoneTrack = BoneTracks[TargetBoneIndex];
+				
+				BoneTrack.PosKeys[FrameIndex] = FVector3f(LocalPose.GetLocation());
+				BoneTrack.RotKeys[FrameIndex] = FQuat4f(LocalPose.GetRotation());
+				BoneTrack.ScaleKeys[FrameIndex] = FVector4f(LocalPose.GetScale3D());
 			}
 		}
 
@@ -277,8 +301,9 @@ void FIKRetargetBatchOperation::ConvertAnimation(
 		const bool bShouldTransact = false;
 		for (int32 TargetBoneIndex=0; TargetBoneIndex<NumTargetBones; ++TargetBoneIndex)
 		{
-			FName TargetBoneName = Processor->GetTargetSkeleton().BoneNames[TargetBoneIndex];
-			FRawAnimSequenceTrack& RawTrack = BoneTracks[TargetBoneIndex];
+			const FName& TargetBoneName = TargetBoneNames[TargetBoneIndex];
+
+			const FRawAnimSequenceTrack& RawTrack = BoneTracks[TargetBoneIndex];
 			TargetSeqController.AddBoneTrack(TargetBoneName, bShouldTransact);
 			TargetSeqController.SetBoneTrackKeys(TargetBoneName, RawTrack.PosKeys, RawTrack.RotKeys, RawTrack.ScaleKeys);
 		}
