disable rotation limit and end bone fwd axis when useless

diff --git a/Engine/Plugins/Animation/IKRig/Source/IKRig/Public/Solvers/IKRig_LimbSolver.h b/Engine/Plugins/Animation/IKRig/Source/IKRig/Public/Solvers/IKRig_LimbSolver.h
--- a/Engine/Plugins/Animation/IKRig/Source/IKRig/Public/Solvers/IKRig_LimbSolver.h
+++ b/Engine/Plugins/Animation/IKRig/Source/IKRig/Public/Solvers/IKRig_LimbSolver.h
@@ -58,7 +58,7 @@ public:
 
 	/** Only used if bEnableRotationLimit is enabled. Prevents the leg from folding onto itself,
 	* and forces at least this angle between Parent and Child bone. */
-	UPROPERTY(EditAnywhere, Category = "Limb IK Settings|FABRIK|Limits", meta = (ClampMin = "0.0", ClampMax = "90.0", UIMin = "0.0", UIMax = "90.0"))
+	UPROPERTY(EditAnywhere, Category = "Limb IK Settings|FABRIK|Limits", meta = (EditCondition="bEnableLimit", ClampMin = "0.0", ClampMax = "90.0", UIMin = "0.0", UIMax = "90.0"))
 	float MinRotationAngle = 15.f;
 
 	/** Pull averaging only has a visual impact when we have more than 2 bones (3 links). */
@@ -80,7 +80,7 @@ public:
 	bool bEnableTwistCorrection = false;
 	
 	/** Forward Axis for Foot bone. */
-	UPROPERTY(EditAnywhere, Category = "Limb IK Settings|Twist")
+	UPROPERTY(EditAnywhere, Category = "Limb IK Settings|Twist", meta = (EditCondition="bEnableTwistCorrection"))
 	TEnumAsByte<EAxis::Type> EndBoneForwardAxis = EAxis::Y;
 	
 	/** UIKRigSolver interface */
diff --git a/Engine/Plugins/Animation/IKRig/Source/IKRig/Public/Solvers/LimbSolver.h b/Engine/Plugins/Animation/IKRig/Source/IKRig/Public/Solvers/LimbSolver.h
--- a/Engine/Plugins/Animation/IKRig/Source/IKRig/Public/Solvers/LimbSolver.h
+++ b/Engine/Plugins/Animation/IKRig/Source/IKRig/Public/Solvers/LimbSolver.h
@@ -32,7 +32,7 @@ struct IKRIG_API FLimbSolverSettings
 
 	/** Only used if bEnableRotationLimit is enabled. Prevents the leg from folding onto itself,
 	* and forces at least this angle between Parent and Child bone. */
-	UPROPERTY(EditAnywhere, Category = "Limb Solver Settings|FABRIK|Limits", meta = (ClampMin = "0.0", ClampMax = "90.0", UIMin = "0.0", UIMax = "90.0"))
+	UPROPERTY(EditAnywhere, Category = "Limb Solver Settings|FABRIK|Limits", meta = (EditCondition="bEnableLimit", ClampMin = "0.0", ClampMax = "90.0", UIMin = "0.0", UIMax = "90.0"))
 	float MinRotationAngle = 15.f;
 
 	/** Pull averaging only has a visual impact when we have more than 2 bones (3 links). */
@@ -54,7 +54,7 @@ struct IKRIG_API FLimbSolverSettings
 	bool bEnableTwistCorrection = false;
 	
 	/** Forward Axis for Foot bone. */
-	UPROPERTY(EditAnywhere, Category = "Limb Solver Settings|Twist")
+	UPROPERTY(EditAnywhere, Category = "Limb Solver Settings|Twist", meta = (EditCondition="bEnableTwistCorrection"))
 	TEnumAsByte<EAxis::Type> EndBoneForwardAxis = EAxis::Y;
 };
 
