Add Alpha pin on IKRig AnimNode

diff --git a/Engine/Plugins/Animation/IKRig/Source/IKRig/Private/AnimNodes/AnimNode_IKRig.cpp b/Engine/Plugins/Animation/IKRig/Source/IKRig/Private/AnimNodes/AnimNode_IKRig.cpp
--- a/Engine/Plugins/Animation/IKRig/Source/IKRig/Private/AnimNodes/AnimNode_IKRig.cpp
+++ b/Engine/Plugins/Animation/IKRig/Source/IKRig/Private/AnimNodes/AnimNode_IKRig.cpp
@@ -9,6 +9,13 @@
 #include "Animation/AnimInstanceProxy.h"
 #include "Algo/ForEach.h"
 
+FAnimNode_IKRig::FAnimNode_IKRig()
+	: AlphaInputType(EAnimAlphaInputType::Float)
+	, bAlphaBoolEnabled(true)
+	, Alpha(1.0f)
+	, ActualAlpha(0.f)
+{}
+
 void FAnimNode_IKRig::Evaluate_AnyThread(FPoseContext& Output) 
 {
 	DECLARE_SCOPE_HIERARCHICAL_COUNTER_FUNC()
@@ -32,6 +39,11 @@ void FAnimNode_IKRig::Evaluate_AnyThread(FPoseContext& Output)
 		return;
 	}
 
+	if (!FAnimWeight::IsRelevant(ActualAlpha))
+	{
+		return;
+	}
+
 	// copy input pose to solver stack
 	CopyInputPoseToSolver(Output.Pose);
 	// update target goal transforms
@@ -118,7 +130,7 @@ void FAnimNode_IKRig::CopyOutputPoseToAnimGraph(FCompactPose& OutputPose)
 			// input pose (in local space).
 			if (*Index != -1)
 			{
-				OutputPose[CPIndex] = IKRigSkeleton.CurrentPoseLocal[*Index];	
+				OutputPose[CPIndex].BlendWith(IKRigSkeleton.CurrentPoseLocal[*Index], ActualAlpha);
 			}
 		}
 	}
@@ -155,7 +167,32 @@ void FAnimNode_IKRig::Initialize_AnyThread(const FAnimationInitializeContext& Co
 
 void FAnimNode_IKRig::Update_AnyThread(const FAnimationUpdateContext& Context)
 {
-	GetEvaluateGraphExposedInputs().Execute(Context);
+	ActualAlpha = 0.f;
+	// if (IsLODEnabled(Context.AnimInstanceProxy))
+	{
+		GetEvaluateGraphExposedInputs().Execute(Context);
+
+		// alpha handlers
+		switch (AlphaInputType)
+		{
+		case EAnimAlphaInputType::Float : 
+			ActualAlpha = AlphaScaleBias.ApplyTo(AlphaScaleBiasClamp.ApplyTo(Alpha, Context.GetDeltaTime()));
+			break;
+		case EAnimAlphaInputType::Bool :
+			ActualAlpha = AlphaBoolBlend.ApplyTo(bAlphaBoolEnabled, Context.GetDeltaTime());
+			break;
+		case EAnimAlphaInputType::Curve :
+			if (UAnimInstance* AnimInstance = Cast<UAnimInstance>(Context.AnimInstanceProxy->GetAnimInstanceObject()))
+			{
+				ActualAlpha = AlphaScaleBiasClamp.ApplyTo(AnimInstance->GetCurveValue(AlphaCurveName), Context.GetDeltaTime());
+			}
+			break;
+		};
+
+		// Make sure Alpha is clamped between 0 and 1.
+		ActualAlpha = FMath::Clamp<float>(ActualAlpha, 0.f, 1.f);
+	}
+	
 	FAnimNode_Base::Update_AnyThread(Context);
 	Source.Update(Context);
 }
@@ -296,10 +333,10 @@ void FAnimNode_IKRig::InitializeProperties(const UObject* InSourceInstance, UCla
 			// while evaluating
 			const FName& GoalPropertyName = DestPropertyNames[PropIndex];
 
-			// we use the goal name's hash value as described in UAnimGraphNode_IKRig::CreateCustomPinsFromValidAsset
+			// find the right goal
 			const int32 GoalIndex = Goals.IndexOfByPredicate([&GoalPropertyName](const FIKRigGoal& InGoal)
 			{
-				return GetTypeHash(InGoal.Name) == GoalPropertyName.GetNumber();
+			 	return GoalPropertyName.ToString().EndsWith(InGoal.Name.ToString());
 			});
 
 			if (Goals.IsValidIndex(GoalIndex))
diff --git a/Engine/Plugins/Animation/IKRig/Source/IKRig/Public/AnimNodes/AnimNode_IKRig.h b/Engine/Plugins/Animation/IKRig/Source/IKRig/Public/AnimNodes/AnimNode_IKRig.h
--- a/Engine/Plugins/Animation/IKRig/Source/IKRig/Public/AnimNodes/AnimNode_IKRig.h
+++ b/Engine/Plugins/Animation/IKRig/Source/IKRig/Public/AnimNodes/AnimNode_IKRig.h
@@ -6,6 +6,8 @@
 
 #include "IKRigDataTypes.h"
 #include "Animation/AnimNode_CustomProperty.h"
+#include "Animation/InputScaleBias.h"
+
 #include "AnimNode_IKRig.generated.h"
 
 class UIKRigProcessor;
@@ -18,6 +20,8 @@ struct IKRIG_API FAnimNode_IKRig : public FAnimNode_CustomProperty
 {
 	GENERATED_BODY()
 
+	FAnimNode_IKRig();
+	
 	/** The input pose to start the IK solve relative to. */
 	UPROPERTY(EditAnywhere, Category = Links)
 	FPoseLink Source;
@@ -49,6 +53,29 @@ struct IKRIG_API FAnimNode_IKRig : public FAnimNode_CustomProperty
 	
 #endif
 
+	/** alpha value handler **/
+	UPROPERTY(EditAnywhere, BlueprintReadWrite, Category = Alpha)
+	EAnimAlphaInputType AlphaInputType;
+
+	UPROPERTY(EditAnywhere, BlueprintReadWrite, Category = Alpha, meta = (PinShownByDefault, DisplayName = "bEnabled"))
+	bool bAlphaBoolEnabled;
+
+	// Current strength of the skeletal control
+	UPROPERTY(EditAnywhere, BlueprintReadWrite, Category = Alpha, meta = (PinShownByDefault))
+	float Alpha;
+
+	UPROPERTY(EditAnywhere, BlueprintReadWrite, Category = Alpha)
+	FInputScaleBias AlphaScaleBias;
+
+	UPROPERTY(EditAnywhere, BlueprintReadWrite, Category = Alpha, meta = (DisplayName = "Blend Settings"))
+	FInputAlphaBoolBlend AlphaBoolBlend;
+
+	UPROPERTY(EditAnywhere, BlueprintReadWrite, Category = Alpha, meta = (PinShownByDefault))
+	FName AlphaCurveName;
+
+	UPROPERTY(EditAnywhere, BlueprintReadWrite, Category = Alpha)
+	FInputScaleBiasClamp AlphaScaleBiasClamp;
+	
 private:
 
 	/** IK Rig runtime processor */
@@ -64,6 +91,9 @@ private:
 	/** Cached functions used to update goals using custom properties to avoid looking for then when evaluating */
 	using UpdateFunction = std::function<void(const UObject*)>;
 	TArray<UpdateFunction> UpdateFunctions;
+
+	UPROPERTY(Transient)
+	float ActualAlpha;
 	
 public:
 
diff --git a/Engine/Plugins/Animation/IKRig/Source/IKRigDeveloper/Private/AnimGraphNode_IKRig.cpp b/Engine/Plugins/Animation/IKRig/Source/IKRigDeveloper/Private/AnimGraphNode_IKRig.cpp
--- a/Engine/Plugins/Animation/IKRig/Source/IKRigDeveloper/Private/AnimGraphNode_IKRig.cpp
+++ b/Engine/Plugins/Animation/IKRig/Source/IKRigDeveloper/Private/AnimGraphNode_IKRig.cpp
@@ -362,6 +362,44 @@ void UAnimGraphNode_IKRig::PostEditChangeProperty(struct FPropertyChangedEvent&
 		return;
 	}
 	
+	if (PropertyName == GET_MEMBER_NAME_STRING_CHECKED(FAnimNode_IKRig, AlphaInputType))
+	{
+		FScopedTransaction Transaction(LOCTEXT("ChangeAlphaInputType", "Change Alpha Input Type"));
+		Modify();
+
+		// Break links to pins going away
+		for (int32 PinIndex = 0; PinIndex < Pins.Num(); ++PinIndex)
+		{
+			UEdGraphPin* Pin = Pins[PinIndex];
+			if (Pin->PinName == GET_MEMBER_NAME_STRING_CHECKED(FAnimNode_IKRig, Alpha))
+			{
+				if (Node.AlphaInputType != EAnimAlphaInputType::Float)
+				{
+					Pin->BreakAllPinLinks();
+					PropertyBindings.Remove(Pin->PinName);
+				}
+			}
+			else if (Pin->PinName == GET_MEMBER_NAME_STRING_CHECKED(FAnimNode_IKRig, bAlphaBoolEnabled))
+			{
+				if (Node.AlphaInputType != EAnimAlphaInputType::Bool)
+				{
+					Pin->BreakAllPinLinks();
+					PropertyBindings.Remove(Pin->PinName);
+				}
+			}
+			else if (Pin->PinName == GET_MEMBER_NAME_STRING_CHECKED(FAnimNode_IKRig, AlphaCurveName))
+			{
+				if (Node.AlphaInputType != EAnimAlphaInputType::Curve)
+				{
+					Pin->BreakAllPinLinks();
+					PropertyBindings.Remove(Pin->PinName);
+				}
+			}
+		}
+
+		ReconstructNode();
+	}
+
 	if (PropertyName == GET_MEMBER_NAME_STRING_CHECKED(FIKRigGoal, TransformSource))
 	{
 		ReconstructNode();
@@ -376,13 +414,18 @@ FName GetGoalSubPropertyPinPrettyName(const FName& InGoalName, const FName& InPr
 	return *FString::Printf(TEXT("%s_%s"), *InGoalName.ToString(), *InPropertyName.ToString());
 }
 
+FName GetGoalSubPropertyPinName(const FName& InGoalName, const FName& InPropertyName)
+{
+	return *FString::Printf(TEXT("%s_%s"), *InPropertyName.ToString(), *InGoalName.ToString());
+}
+
 void UAnimGraphNode_IKRig::CreateCustomPins(TArray<UEdGraphPin*>* InOldPins)
 {
 	if (!IsValid(Node.RigDefinitionAsset))
 	{
 		return;
 	}
-
+	
 	// the asset is not completely loaded so we'll use the old pins to sustain the current set of custom pins
 	if (Node.RigDefinitionAsset->HasAllFlags(RF_NeedPostLoad)) 
 	{
@@ -471,13 +514,8 @@ void UAnimGraphNode_IKRig::CreateCustomPinsFromValidAsset()
 								 const FEdGraphPinType InPinType)
 	{
 		const FName& GoalName = Node.Goals[InGoalIndex].Name;
-		const uint32 GoalHash = GetTypeHash(GoalName);
-		
-		// pin's name is based on the property name within the FIKRigGoal structure + the name's hash value as a number
-		FName PinName = InPropertyName;
-		PinName.SetNumber(GoalHash);
-		
-		UEdGraphPin* NewPin = CreatePin(EEdGraphPinDirection::EGPD_Input, InPinType, PinName);
+				
+		UEdGraphPin* NewPin = CreatePin(EEdGraphPinDirection::EGPD_Input, InPinType, GetGoalSubPropertyPinName(GoalName, InPropertyName));
 		
 		// pin's pretty name is the "GoalName_InPropertyName"
 		NewPin->PinFriendlyName = FText::FromName(GetGoalSubPropertyPinPrettyName(GoalName, InPropertyName));
@@ -620,6 +658,59 @@ void UAnimGraphNode_IKRig::CustomizeDetails(IDetailLayoutBuilder& DetailBuilder)
 	{
 		GoalHandle->SetOnChildPropertyValueChanged(OnValueChanged);
 	}
+
+	// alpha customization
+	if (Node.AlphaInputType != EAnimAlphaInputType::Bool)
+	{
+		DetailBuilder.HideProperty(NodePropHandle->GetChildHandle(GET_MEMBER_NAME_CHECKED(FAnimNode_IKRig, bAlphaBoolEnabled)));
+		DetailBuilder.HideProperty(NodePropHandle->GetChildHandle(GET_MEMBER_NAME_CHECKED(FAnimNode_IKRig, AlphaBoolBlend)));
+	}
+
+	if (Node.AlphaInputType != EAnimAlphaInputType::Float)
+	{
+		DetailBuilder.HideProperty(NodePropHandle->GetChildHandle(GET_MEMBER_NAME_CHECKED(FAnimNode_IKRig, Alpha)));
+		DetailBuilder.HideProperty(NodePropHandle->GetChildHandle(GET_MEMBER_NAME_CHECKED(FAnimNode_IKRig, AlphaScaleBias)));
+	}
+
+	if (Node.AlphaInputType != EAnimAlphaInputType::Curve)
+	{
+		DetailBuilder.HideProperty(NodePropHandle->GetChildHandle(GET_MEMBER_NAME_CHECKED(FAnimNode_IKRig, AlphaCurveName)));
+	}
+
+	if (Node.AlphaInputType != EAnimAlphaInputType::Float && (Node.AlphaInputType != EAnimAlphaInputType::Curve))
+	{
+		DetailBuilder.HideProperty(NodePropHandle->GetChildHandle(GET_MEMBER_NAME_CHECKED(FAnimNode_IKRig, AlphaScaleBiasClamp)));
+	}
+}
+
+void UAnimGraphNode_IKRig::CustomizePinData(UEdGraphPin* Pin, FName SourcePropertyName, int32 ArrayIndex) const
+{
+	Super::CustomizePinData(Pin, SourcePropertyName, ArrayIndex);
+
+	if (Pin->PinName == GET_MEMBER_NAME_STRING_CHECKED(FAnimNode_IKRig, Alpha))
+	{
+		Pin->bHidden = (Node.AlphaInputType != EAnimAlphaInputType::Float);
+
+		if (!Pin->bHidden)
+		{
+			Pin->PinFriendlyName = Node.AlphaScaleBias.GetFriendlyName(Node.AlphaScaleBiasClamp.GetFriendlyName(Pin->PinFriendlyName));
+		}
+	}
+
+	if (Pin->PinName == GET_MEMBER_NAME_STRING_CHECKED(FAnimNode_IKRig, bAlphaBoolEnabled))
+	{
+		Pin->bHidden = (Node.AlphaInputType != EAnimAlphaInputType::Bool);
+	}
+
+	if (Pin->PinName == GET_MEMBER_NAME_STRING_CHECKED(FAnimNode_IKRig, AlphaCurveName))
+	{
+		Pin->bHidden = (Node.AlphaInputType != EAnimAlphaInputType::Curve);
+
+		if (!Pin->bHidden)
+		{
+			Pin->PinFriendlyName = Node.AlphaScaleBiasClamp.GetFriendlyName(Pin->PinFriendlyName);
+		}
+	}
 }
 
 void UAnimGraphNode_IKRig::BindPropertyChanges()
diff --git a/Engine/Plugins/Animation/IKRig/Source/IKRigDeveloper/Public/AnimGraphNode_IKRig.h b/Engine/Plugins/Animation/IKRig/Source/IKRigDeveloper/Public/AnimGraphNode_IKRig.h
--- a/Engine/Plugins/Animation/IKRig/Source/IKRigDeveloper/Public/AnimGraphNode_IKRig.h
+++ b/Engine/Plugins/Animation/IKRig/Source/IKRigDeveloper/Public/AnimGraphNode_IKRig.h
@@ -101,6 +101,7 @@ public:
 	virtual void PostEditChangeProperty(struct FPropertyChangedEvent& PropertyChangedEvent) override;
 	virtual void CreateCustomPins(TArray<UEdGraphPin*>* InOldPins) override;
 	virtual void CustomizeDetails(IDetailLayoutBuilder& DetailBuilder) override;
+	virtual void CustomizePinData(UEdGraphPin* Pin, FName SourcePropertyName, int32 ArrayIndex) const override;
 	// End of UEdGraphNode interface
 
 	// UAnimGraphNode_Base interface
