fixed missing include

diff --git a/Engine/Plugins/Animation/IKRig/Source/IKRigDeveloper/Private/AnimGraphNode_IKRig.cpp b/Engine/Plugins/Animation/IKRig/Source/IKRigDeveloper/Private/AnimGraphNode_IKRig.cpp
--- a/Engine/Plugins/Animation/IKRig/Source/IKRigDeveloper/Private/AnimGraphNode_IKRig.cpp
+++ b/Engine/Plugins/Animation/IKRig/Source/IKRigDeveloper/Private/AnimGraphNode_IKRig.cpp
@@ -6,7 +6,7 @@
 #include "IKRigSolver.h"
 #include "Kismet2/CompilerResultsLog.h"
 #include "AnimationGraphSchema.h"
-
+#include "ScopedTransaction.h"
 #include "DetailLayoutBuilder.h"
 #include "DetailWidgetRow.h"
 #include "IDetailChildrenBuilder.h"
@@ -362,6 +362,44 @@ void UAnimGraphNode_IKRig::PostEditChangeProperty(struct FPropertyChangedEvent&
 		return;
 	}
 	
+	if (PropertyName == GET_MEMBER_NAME_STRING_CHECKED(FAnimNode_IKRig, AlphaInputType))
+	{
+		FScopedTransaction Transaction(LOCTEXT("ChangeAlphaInputType", "Change Alpha Input Type"));
+		Modify();
+
+		// Break links to pins going away
+		for (int32 PinIndex = 0; PinIndex < Pins.Num(); ++PinIndex)
+		{
+			UEdGraphPin* Pin = Pins[PinIndex];
+			if (Pin->PinName == GET_MEMBER_NAME_STRING_CHECKED(FAnimNode_IKRig, Alpha))
+			{
+				if (Node.AlphaInputType != EAnimAlphaInputType::Float)
+				{
+					Pin->BreakAllPinLinks();
+					PropertyBindings.Remove(Pin->PinName);
+				}
+			}
+			else if (Pin->PinName == GET_MEMBER_NAME_STRING_CHECKED(FAnimNode_IKRig, bAlphaBoolEnabled))
+			{
+				if (Node.AlphaInputType != EAnimAlphaInputType::Bool)
+				{
+					Pin->BreakAllPinLinks();
+					PropertyBindings.Remove(Pin->PinName);
+				}
+			}
+			else if (Pin->PinName == GET_MEMBER_NAME_STRING_CHECKED(FAnimNode_IKRig, AlphaCurveName))
+			{
+				if (Node.AlphaInputType != EAnimAlphaInputType::Curve)
+				{
+					Pin->BreakAllPinLinks();
+					PropertyBindings.Remove(Pin->PinName);
+				}
+			}
+		}
+
+		ReconstructNode();
+	}
+
 	if (PropertyName == GET_MEMBER_NAME_STRING_CHECKED(FIKRigGoal, TransformSource))
 	{
 		ReconstructNode();
@@ -376,13 +414,18 @@ FName GetGoalSubPropertyPinPrettyName(const FName& InGoalName, const FName& InPr
 	return *FString::Printf(TEXT("%s_%s"), *InGoalName.ToString(), *InPropertyName.ToString());
 }
 
+FName GetGoalSubPropertyPinName(const FName& InGoalName, const FName& InPropertyName)
+{
+	return *FString::Printf(TEXT("%s_%s"), *InPropertyName.ToString(), *InGoalName.ToString());
+}
+
 void UAnimGraphNode_IKRig::CreateCustomPins(TArray<UEdGraphPin*>* InOldPins)
 {
 	if (!IsValid(Node.RigDefinitionAsset))
 	{
 		return;
 	}
-
+	
 	// the asset is not completely loaded so we'll use the old pins to sustain the current set of custom pins
 	if (Node.RigDefinitionAsset->HasAllFlags(RF_NeedPostLoad)) 
 	{
@@ -471,13 +514,8 @@ void UAnimGraphNode_IKRig::CreateCustomPinsFromValidAsset()
 								 const FEdGraphPinType InPinType)
 	{
 		const FName& GoalName = Node.Goals[InGoalIndex].Name;
-		const uint32 GoalHash = GetTypeHash(GoalName);
-		
-		// pin's name is based on the property name within the FIKRigGoal structure + the name's hash value as a number
-		FName PinName = InPropertyName;
-		PinName.SetNumber(GoalHash);
-		
-		UEdGraphPin* NewPin = CreatePin(EEdGraphPinDirection::EGPD_Input, InPinType, PinName);
+				
+		UEdGraphPin* NewPin = CreatePin(EEdGraphPinDirection::EGPD_Input, InPinType, GetGoalSubPropertyPinName(GoalName, InPropertyName));
 		
 		// pin's pretty name is the "GoalName_InPropertyName"
 		NewPin->PinFriendlyName = FText::FromName(GetGoalSubPropertyPinPrettyName(GoalName, InPropertyName));
@@ -620,6 +658,59 @@ void UAnimGraphNode_IKRig::CustomizeDetails(IDetailLayoutBuilder& DetailBuilder)
 	{
 		GoalHandle->SetOnChildPropertyValueChanged(OnValueChanged);
 	}
+
+	// alpha customization
+	if (Node.AlphaInputType != EAnimAlphaInputType::Bool)
+	{
+		DetailBuilder.HideProperty(NodePropHandle->GetChildHandle(GET_MEMBER_NAME_CHECKED(FAnimNode_IKRig, bAlphaBoolEnabled)));
+		DetailBuilder.HideProperty(NodePropHandle->GetChildHandle(GET_MEMBER_NAME_CHECKED(FAnimNode_IKRig, AlphaBoolBlend)));
+	}
+
+	if (Node.AlphaInputType != EAnimAlphaInputType::Float)
+	{
+		DetailBuilder.HideProperty(NodePropHandle->GetChildHandle(GET_MEMBER_NAME_CHECKED(FAnimNode_IKRig, Alpha)));
+		DetailBuilder.HideProperty(NodePropHandle->GetChildHandle(GET_MEMBER_NAME_CHECKED(FAnimNode_IKRig, AlphaScaleBias)));
+	}
+
+	if (Node.AlphaInputType != EAnimAlphaInputType::Curve)
+	{
+		DetailBuilder.HideProperty(NodePropHandle->GetChildHandle(GET_MEMBER_NAME_CHECKED(FAnimNode_IKRig, AlphaCurveName)));
+	}
+
+	if (Node.AlphaInputType != EAnimAlphaInputType::Float && (Node.AlphaInputType != EAnimAlphaInputType::Curve))
+	{
+		DetailBuilder.HideProperty(NodePropHandle->GetChildHandle(GET_MEMBER_NAME_CHECKED(FAnimNode_IKRig, AlphaScaleBiasClamp)));
+	}
+}
+
+void UAnimGraphNode_IKRig::CustomizePinData(UEdGraphPin* Pin, FName SourcePropertyName, int32 ArrayIndex) const
+{
+	Super::CustomizePinData(Pin, SourcePropertyName, ArrayIndex);
+
+	if (Pin->PinName == GET_MEMBER_NAME_STRING_CHECKED(FAnimNode_IKRig, Alpha))
+	{
+		Pin->bHidden = (Node.AlphaInputType != EAnimAlphaInputType::Float);
+
+		if (!Pin->bHidden)
+		{
+			Pin->PinFriendlyName = Node.AlphaScaleBias.GetFriendlyName(Node.AlphaScaleBiasClamp.GetFriendlyName(Pin->PinFriendlyName));
+		}
+	}
+
+	if (Pin->PinName == GET_MEMBER_NAME_STRING_CHECKED(FAnimNode_IKRig, bAlphaBoolEnabled))
+	{
+		Pin->bHidden = (Node.AlphaInputType != EAnimAlphaInputType::Bool);
+	}
+
+	if (Pin->PinName == GET_MEMBER_NAME_STRING_CHECKED(FAnimNode_IKRig, AlphaCurveName))
+	{
+		Pin->bHidden = (Node.AlphaInputType != EAnimAlphaInputType::Curve);
+
+		if (!Pin->bHidden)
+		{
+			Pin->PinFriendlyName = Node.AlphaScaleBiasClamp.GetFriendlyName(Pin->PinFriendlyName);
+		}
+	}
 }
 
 void UAnimGraphNode_IKRig::BindPropertyChanges()
