ease bones, goals and chains selection in IKRig / IKRetargeter using filters

* changed FIKRigRetargetChainSettings::StartBone / EndBone to be FBoneReference instead of FName then rely on SBoneSelectionWidget for selection
* changed retarget chain goals and target chain mapping to use SSearchableComboBox

diff --git a/Engine/Plugins/Animation/IKRig/Source/IKRig/Private/Retargeter/IKRetargetProcessor.cpp b/Engine/Plugins/Animation/IKRig/Source/IKRig/Private/Retargeter/IKRetargetProcessor.cpp
--- a/Engine/Plugins/Animation/IKRig/Source/IKRig/Private/Retargeter/IKRetargetProcessor.cpp
+++ b/Engine/Plugins/Animation/IKRig/Source/IKRig/Private/Retargeter/IKRetargetProcessor.cpp
@@ -243,8 +243,8 @@ FResolvedBoneChain::FResolvedBoneChain(
 	TArray<int32> &OutBoneIndices)
 {
 	// validate start and end bones exist
-	const int32 StartIndex = Skeleton.FindBoneIndexByName(BoneChain.StartBone);
-	const int32 EndIndex = Skeleton.FindBoneIndexByName(BoneChain.EndBone);
+	const int32 StartIndex = Skeleton.FindBoneIndexByName(BoneChain.StartBone.BoneName);
+	const int32 EndIndex = Skeleton.FindBoneIndexByName(BoneChain.EndBone.BoneName);
 	bFoundStartBone = StartIndex != INDEX_NONE;
 	bFoundEndBone = EndIndex != INDEX_NONE;
 	
@@ -642,21 +642,21 @@ bool FRetargetChainPair::ValidateBoneChainWithSkeletalMesh(
 	if (!ResolvedChain.bFoundStartBone)
 	{
 		UE_LOG(LogTemp, Warning, TEXT("IK Retargeter bone chain, %s, could not find start bone, %s in mesh %s"),
-            *BoneChain.ChainName.ToString(), *BoneChain.StartBone.ToString(), *RetargetSkeleton.SkeletalMesh->GetName());
+            *BoneChain.ChainName.ToString(), *BoneChain.StartBone.BoneName.ToString(), *RetargetSkeleton.SkeletalMesh->GetName());
 	}
 	
 	// warn if END bone not found
 	if (!ResolvedChain.bFoundEndBone)
 	{
 		UE_LOG(LogTemp, Warning, TEXT("IK Retargeter bone chain, %s, could not find end bone, %s in mesh %s"),
-            *BoneChain.ChainName.ToString(), *BoneChain.EndBone.ToString(), *RetargetSkeleton.SkeletalMesh->GetName());
+            *BoneChain.ChainName.ToString(), *BoneChain.EndBone.BoneName.ToString(), *RetargetSkeleton.SkeletalMesh->GetName());
 	}
 
 	// warn if END bone was not a child of START bone
 	if (!ResolvedChain.bEndIsChildOfStart)
 	{
 		UE_LOG(LogTemp, Warning, TEXT("IK Retargeter bone chain, %s, end bone, '%s' was not a child of the start bone '%s'."),
-			*BoneChain.ChainName.ToString(), *BoneChain.EndBone.ToString(), *BoneChain.StartBone.ToString());
+			*BoneChain.ChainName.ToString(), *BoneChain.EndBone.BoneName.ToString(), *BoneChain.StartBone.BoneName.ToString());
 	}
 	
 	return ResolvedChain.IsValid();
diff --git a/Engine/Plugins/Animation/IKRig/Source/IKRig/Public/IKRigDefinition.h b/Engine/Plugins/Animation/IKRig/Source/IKRig/Public/IKRigDefinition.h
--- a/Engine/Plugins/Animation/IKRig/Source/IKRig/Public/IKRigDefinition.h
+++ b/Engine/Plugins/Animation/IKRig/Source/IKRig/Public/IKRigDefinition.h
@@ -2,7 +2,7 @@
 
 #pragma once
 
-#include "CoreMinimal.h"
+#include "BoneContainer.h"
 #include "UObject/Object.h"
 #include "Interfaces/Interface_PreviewMeshProvider.h"
 #include "IKRigSkeleton.h"
@@ -130,20 +130,21 @@ struct IKRIG_API FBoneChain
 
 	FBoneChain() = default;
 	
-	FBoneChain(FName InName, FName InStartBone, FName InEndBone)
-	: ChainName(InName),
-	StartBone(InStartBone),
-	EndBone(InEndBone),
-	IKGoalName(NAME_None){}
+	FBoneChain(FName InName, const FName& InStartBone, const FName& InEndBone)
+		: ChainName(InName)
+		, StartBone(InStartBone)
+		, EndBone(InEndBone)
+		, IKGoalName(NAME_None)
+	{}
 
 	UPROPERTY(EditAnywhere, Category = BoneChain)
 	FName ChainName;
 
 	UPROPERTY(EditAnywhere, Category = BoneChain)
-	FName StartBone;
+	FBoneReference StartBone;
 
 	UPROPERTY(EditAnywhere, Category = BoneChain)
-	FName EndBone;
+	FBoneReference EndBone;
 	
 	UPROPERTY(EditAnywhere, Category = IK)
 	FName IKGoalName;
diff --git a/Engine/Plugins/Animation/IKRig/Source/IKRigDeveloper/Private/AnimGraphNode_RetargetPoseFromMesh.cpp b/Engine/Plugins/Animation/IKRig/Source/IKRigDeveloper/Private/AnimGraphNode_RetargetPoseFromMesh.cpp
--- a/Engine/Plugins/Animation/IKRig/Source/IKRigDeveloper/Private/AnimGraphNode_RetargetPoseFromMesh.cpp
+++ b/Engine/Plugins/Animation/IKRig/Source/IKRigDeveloper/Private/AnimGraphNode_RetargetPoseFromMesh.cpp
@@ -83,12 +83,12 @@ void UAnimGraphNode_RetargetPoseFromMesh::ValidateAnimNodeDuringCompilation(USke
 	const TArray<FBoneChain> &TargetBoneChains = Node.IKRetargeterAsset->GetTargetIKRig()->GetRetargetChains();
     for (const FBoneChain &Chain : TargetBoneChains)
     {
-        if (RefSkel.FindBoneIndex(Chain.StartBone) == INDEX_NONE)
+        if (RefSkel.FindBoneIndex(Chain.StartBone.BoneName) == INDEX_NONE)
         {
         	MessageLog.Warning(*LOCTEXT("StartBoneNotFound", "@@ - Start Bone in target IK Rig Bone Chain not found.").ToString(), this);
         }
 
-    	if (RefSkel.FindBoneIndex(Chain.EndBone) == INDEX_NONE)
+    	if (RefSkel.FindBoneIndex(Chain.EndBone.BoneName) == INDEX_NONE)
     	{
     		MessageLog.Warning(*LOCTEXT("EndBoneNotFound", "@@ - End Bone in target IK Rig Bone Chain not found.").ToString(), this);
     	}
diff --git a/Engine/Plugins/Animation/IKRig/Source/IKRigEditor/Private/RetargetEditor/SIKRetargetChainMapList.cpp b/Engine/Plugins/Animation/IKRig/Source/IKRigEditor/Private/RetargetEditor/SIKRetargetChainMapList.cpp
--- a/Engine/Plugins/Animation/IKRig/Source/IKRigEditor/Private/RetargetEditor/SIKRetargetChainMapList.cpp
+++ b/Engine/Plugins/Animation/IKRig/Source/IKRigEditor/Private/RetargetEditor/SIKRetargetChainMapList.cpp
@@ -13,6 +13,7 @@
 #include "Retargeter/IKRetargeter.h"
 #include "RigEditor/IKRigEditorStyle.h"
 #include "Widgets/Input/SComboBox.h"
+#include "SSearchableComboBox.h"
 
 #define LOCTEXT_NAMESPACE "SIKRigRetargetChains"
 
@@ -81,7 +82,7 @@ TSharedRef< SWidget > SIKRetargetChainMapRow::GenerateWidgetForColumn( const FNa
 		.VAlign(VAlign_Center)
 		.Padding(3.0f, 1.0f)
 		[
-			SNew(SComboBox<TSharedPtr<FString>>)
+			SNew(SSearchableComboBox)
 			.OptionsSource(&SourceChainOptions)
 			.OnGenerateWidget_Lambda([](TSharedPtr<FString> InItem)
 			{
diff --git a/Engine/Plugins/Animation/IKRig/Source/IKRigEditor/Private/RigEditor/IKRigController.cpp b/Engine/Plugins/Animation/IKRig/Source/IKRigEditor/Private/RigEditor/IKRigController.cpp
--- a/Engine/Plugins/Animation/IKRig/Source/IKRigEditor/Private/RigEditor/IKRigController.cpp
+++ b/Engine/Plugins/Animation/IKRig/Source/IKRigEditor/Private/RigEditor/IKRigController.cpp
@@ -275,7 +275,7 @@ FName UIKRigController::GetRetargetChainStartBone(const FName& ChainName) const
 		return NAME_None;
 	}
 	
-	return Chain->StartBone;
+	return Chain->StartBone.BoneName;
 }
 
 FName UIKRigController::GetRetargetChainEndBone(const FName& ChainName) const
@@ -287,7 +287,7 @@ FName UIKRigController::GetRetargetChainEndBone(const FName& ChainName) const
 		return NAME_None;
 	}
 	
-	return Chain->EndBone;
+	return Chain->EndBone.BoneName;
 }
 
 const TArray<FBoneChain>& UIKRigController::GetRetargetChains() const
@@ -319,8 +319,8 @@ void UIKRigController::SortRetargetChains() const
 {
 	Asset->RetargetDefinition.BoneChains.Sort([this](const FBoneChain& A, const FBoneChain& B)
 	{
-		const int32 IndexA = Asset->Skeleton.GetBoneIndexFromName(A.StartBone);
-		const int32 IndexB = Asset->Skeleton.GetBoneIndexFromName(B.StartBone);
+		const int32 IndexA = Asset->Skeleton.GetBoneIndexFromName(A.StartBone.BoneName);
+		const int32 IndexB = Asset->Skeleton.GetBoneIndexFromName(B.StartBone.BoneName);
 		return IndexA < IndexB;
 	});
 }
diff --git a/Engine/Plugins/Animation/IKRig/Source/IKRigEditor/Private/RigEditor/SIKRigRetargetChainList.cpp b/Engine/Plugins/Animation/IKRig/Source/IKRigEditor/Private/RigEditor/SIKRigRetargetChainList.cpp
--- a/Engine/Plugins/Animation/IKRig/Source/IKRigEditor/Private/RigEditor/SIKRigRetargetChainList.cpp
+++ b/Engine/Plugins/Animation/IKRig/Source/IKRigEditor/Private/RigEditor/SIKRigRetargetChainList.cpp
@@ -14,6 +14,8 @@
 #include "Widgets/Input/SEditableTextBox.h"
 #include "Widgets/Text/SInlineEditableTextBlock.h"
 #include "SPositiveActionButton.h"
+#include "SSearchableComboBox.h"
+#include "BoneSelectionWidget.h"
 
 #define LOCTEXT_NAMESPACE "SIKRigRetargetChains"
 
@@ -82,14 +84,10 @@ TSharedRef<SWidget> SIKRigRetargetChainRow::GenerateWidgetForColumn(const FName&
 		.VAlign(VAlign_Center)
 		.Padding(3.0f, 1.0f)
 		[
-			SNew(SComboBox<FName>)
-			.OptionsSource(&ChainList.Pin()->EditorController.Pin()->AssetController->GetIKRigSkeleton().BoneNames)
-			.OnGenerateWidget(this, &SIKRigRetargetChainRow::MakeBoneComboEntryWidget)
-			.OnSelectionChanged(this, &SIKRigRetargetChainRow::OnStartBoneComboSelectionChanged)
-			[
-				SNew(STextBlock)
-				.Text(this, &SIKRigRetargetChainRow::GetStartBoneName)
-			]
+			SNew(SBoneSelectionWidget)
+			.OnBoneSelectionChanged(this, &SIKRigRetargetChainRow::OnStartBoneComboSelectionChanged)
+			.OnGetSelectedBone(this, &SIKRigRetargetChainRow::GetStartBoneName)
+			.OnGetReferenceSkeleton(this, &SIKRigRetargetChainRow::GetReferenceSkeleton)
 		];
 		return StartWidget;
 	}
@@ -104,14 +102,10 @@ TSharedRef<SWidget> SIKRigRetargetChainRow::GenerateWidgetForColumn(const FName&
 		.VAlign(VAlign_Center)
 		.Padding(3.0f, 1.0f)
 		[
-			SNew(SComboBox<FName>)
-			.OptionsSource(&ChainList.Pin()->EditorController.Pin()->AssetController->GetIKRigSkeleton().BoneNames)
-			.OnGenerateWidget(this, &SIKRigRetargetChainRow::MakeBoneComboEntryWidget)
-			.OnSelectionChanged(this, &SIKRigRetargetChainRow::OnEndBoneComboSelectionChanged)
-			[
-				SNew(STextBlock)
-				.Text(this, &SIKRigRetargetChainRow::GetEndBoneName)
-			]
+			SNew(SBoneSelectionWidget)
+			.OnBoneSelectionChanged(this, &SIKRigRetargetChainRow::OnEndBoneComboSelectionChanged)
+			.OnGetSelectedBone(this, &SIKRigRetargetChainRow::GetEndBoneName)
+			.OnGetReferenceSkeleton(this, &SIKRigRetargetChainRow::GetReferenceSkeleton)
 		];
 		return EndWidget;
 	}
@@ -127,7 +121,7 @@ TSharedRef<SWidget> SIKRigRetargetChainRow::GenerateWidgetForColumn(const FName&
 		.VAlign(VAlign_Center)
 		.Padding(3.0f, 1.0f)
 		[
-			SNew(SComboBox<TSharedPtr<FString>>)
+			SNew(SSearchableComboBox)
 			.OptionsSource(&GoalOptions)
 			.OnGenerateWidget(this, &SIKRigRetargetChainRow::MakeGoalComboEntryWidget)
 			.OnSelectionChanged(this, &SIKRigRetargetChainRow::OnGoalComboSelectionChanged)
@@ -140,17 +134,12 @@ TSharedRef<SWidget> SIKRigRetargetChainRow::GenerateWidgetForColumn(const FName&
 	}
 }
 
-TSharedRef<SWidget> SIKRigRetargetChainRow::MakeBoneComboEntryWidget(FName InItem) const
-{
-	return SNew(STextBlock).Text(FText::FromName(InItem));
-}
-
 TSharedRef<SWidget> SIKRigRetargetChainRow::MakeGoalComboEntryWidget(TSharedPtr<FString> InItem) const
 {
 	return SNew(STextBlock).Text(FText::FromString(*InItem.Get()));
 }
 
-void SIKRigRetargetChainRow::OnStartBoneComboSelectionChanged(FName InName, ESelectInfo::Type SelectInfo)
+void SIKRigRetargetChainRow::OnStartBoneComboSelectionChanged(FName InName) const
 {
 	const TSharedPtr<FIKRigEditorController> Controller = ChainList.Pin()->EditorController.Pin();
 	if (!Controller.IsValid())
@@ -162,28 +151,28 @@ void SIKRigRetargetChainRow::OnStartBoneComboSelectionChanged(FName InName, ESel
 	ChainList.Pin()->RefreshView();
 }
 
-FText SIKRigRetargetChainRow::GetStartBoneName() const
+FName SIKRigRetargetChainRow::GetStartBoneName(bool& bMultipleValues) const
 {
 	const TSharedPtr<FIKRigEditorController> Controller = ChainList.Pin()->EditorController.Pin();
 	if (!Controller.IsValid())
 	{
-		return FText::GetEmpty();
+		return NAME_None;
 	}
-	
-	const FName StartBoneName = Controller->AssetController->GetRetargetChainStartBone(ChainElement.Pin()->ChainName);
-	return FText::FromName(StartBoneName);
+
+	bMultipleValues = false;
+	return Controller->AssetController->GetRetargetChainStartBone(ChainElement.Pin()->ChainName);
 }
 
-FText SIKRigRetargetChainRow::GetEndBoneName() const
+FName SIKRigRetargetChainRow::GetEndBoneName(bool& bMultipleValues) const
 {	
 	const TSharedPtr<FIKRigEditorController> Controller = ChainList.Pin()->EditorController.Pin();
 	if (!Controller.IsValid())
 	{
-		return FText::GetEmpty();
+		return NAME_None;
 	}
-	
-	const FName EndBoneName = Controller->AssetController->GetRetargetChainEndBone(ChainElement.Pin()->ChainName);
-	return FText::FromName(EndBoneName);
+
+	bMultipleValues = false;
+	return Controller->AssetController->GetRetargetChainEndBone(ChainElement.Pin()->ChainName);
 }
 
 FText SIKRigRetargetChainRow::GetGoalName() const
@@ -198,7 +187,7 @@ FText SIKRigRetargetChainRow::GetGoalName() const
 	return FText::FromName(GoalName);
 }
 
-void SIKRigRetargetChainRow::OnEndBoneComboSelectionChanged(FName InName, ESelectInfo::Type SelectInfo)
+void SIKRigRetargetChainRow::OnEndBoneComboSelectionChanged(FName InName) const
 {
 	const TSharedPtr<FIKRigEditorController> Controller = ChainList.Pin()->EditorController.Pin();
 	if (!Controller.IsValid())
@@ -236,6 +225,25 @@ void SIKRigRetargetChainRow::OnRenameChain(const FText& InText, ETextCommit::Typ
 	ChainList.Pin()->RefreshView();
 }
 
+const FReferenceSkeleton& SIKRigRetargetChainRow::GetReferenceSkeleton() const
+{
+	static const FReferenceSkeleton DummySkeleton;
+	
+	const TSharedPtr<FIKRigEditorController> Controller = ChainList.Pin()->EditorController.Pin();
+	if (!Controller.IsValid())
+	{
+		return DummySkeleton; 
+	}
+
+	USkeleton* Skeleton = Controller->AssetController->GetSkeleton();
+	if (Skeleton == nullptr)
+	{
+		return DummySkeleton;
+	}
+	
+	return Skeleton->GetReferenceSkeleton();
+}
+
 void SIKRigRetargetChainList::Construct(const FArguments& InArgs, TSharedRef<FIKRigEditorController> InEditorController)
 {
 	EditorController = InEditorController;
diff --git a/Engine/Plugins/Animation/IKRig/Source/IKRigEditor/Public/RigEditor/SIKRigRetargetChainList.h b/Engine/Plugins/Animation/IKRig/Source/IKRigEditor/Public/RigEditor/SIKRigRetargetChainList.h
--- a/Engine/Plugins/Animation/IKRig/Source/IKRigEditor/Public/RigEditor/SIKRigRetargetChainList.h
+++ b/Engine/Plugins/Animation/IKRig/Source/IKRigEditor/Public/RigEditor/SIKRigRetargetChainList.h
@@ -61,21 +61,22 @@ public:
 	/** Overridden from SMultiColumnTableRow.  Generates a widget for this column of the table row. */
 	virtual TSharedRef<SWidget> GenerateWidgetForColumn(const FName& ColumnName) override;
 
-	TSharedRef<SWidget> MakeBoneComboEntryWidget(FName InItem) const;
+private:
+
+	const FReferenceSkeleton& GetReferenceSkeleton() const;
+
 	TSharedRef<SWidget> MakeGoalComboEntryWidget(TSharedPtr<FString> InItem) const;
 	
-	void OnStartBoneComboSelectionChanged(FName InName, ESelectInfo::Type SelectInfo);
-	void OnEndBoneComboSelectionChanged(FName InName, ESelectInfo::Type SelectInfo);
+	void OnStartBoneComboSelectionChanged(FName InName) const;
+	void OnEndBoneComboSelectionChanged(FName InName) const;
 	void OnGoalComboSelectionChanged(TSharedPtr<FString> InGoalName, ESelectInfo::Type SelectInfo);
 	
-	FText GetStartBoneName() const;
-	FText GetEndBoneName() const;
+	FName GetStartBoneName(bool& bMultipleValues) const;
+	FName GetEndBoneName(bool& bMultipleValues) const;
 	FText GetGoalName() const;
-	
-	void OnRenameChain(const FText& InText, ETextCommit::Type ) const;
-
-private:
 
+	void OnRenameChain(const FText& InText, ETextCommit::Type ) const;
+	
 	TArray<TSharedPtr<FString>> GoalOptions;
 	TWeakPtr<FRetargetChainElement> ChainElement;
 	TWeakPtr<SIKRigRetargetChainList> ChainList;
@@ -90,10 +91,11 @@ struct FIKRigRetargetChainSettings
 
 	FIKRigRetargetChainSettings(){};
 	
-	FIKRigRetargetChainSettings(FName InChainName, FName InStartBone, FName InEndBone)
-	:	ChainName(InChainName),
-		StartBone(InStartBone),
-		EndBone(InEndBone){}
+	FIKRigRetargetChainSettings(FName InChainName, const FName& InStartBone, const FName& InEndBone)
+		: ChainName(InChainName)
+		, StartBone(InStartBone)
+		, EndBone(InEndBone)
+	{}
 
 	UPROPERTY(EditAnywhere, Category = "Bone Chain")
 	FName ChainName;
