IKRig: removed pluralized references to solvers in menu since you can't multi-select them anymore.

diff --git a/Engine/Plugins/Animation/IKRig/Source/IKRigEditor/Private/RigEditor/IKRigSkeletonCommands.cpp b/Engine/Plugins/Animation/IKRig/Source/IKRigEditor/Private/RigEditor/IKRigSkeletonCommands.cpp
--- a/Engine/Plugins/Animation/IKRig/Source/IKRigEditor/Private/RigEditor/IKRigSkeletonCommands.cpp
+++ b/Engine/Plugins/Animation/IKRig/Source/IKRigEditor/Private/RigEditor/IKRigSkeletonCommands.cpp
@@ -8,12 +8,12 @@ void FIKRigSkeletonCommands::RegisterCommands()
 {
 	UI_COMMAND(NewGoal, "New IK Goal", "Add new IK goal at the selected bone.", EUserInterfaceActionType::Button, FInputChord(EKeys::N, EModifierKey::Control));
 	UI_COMMAND(DeleteElement, "Delete Goal or Effector", "Delete the selected IK Goal, Effector or Bone Settings.", EUserInterfaceActionType::Button, FInputChord(EKeys::Delete), FInputChord(EKeys::BackSpace));
-	UI_COMMAND(ConnectGoalToSolvers, "Connect Goal to Selected Solvers", "Make the selected goal an effector in the selected solvers.", EUserInterfaceActionType::Button, FInputChord());
-	UI_COMMAND(DisconnectGoalFromSolvers, "Disconnect Goal from Selected Solvers", "Remove effectors from the selected solvers that use this goal.", EUserInterfaceActionType::Button, FInputChord());
-	UI_COMMAND(SetRootBoneOnSolvers, "Set Root Bone on Selected Solvers", "Set the Root Bone setting on the selected solvers.", EUserInterfaceActionType::Button, FInputChord());
-	UI_COMMAND(SetEndBoneOnSolvers, "Set End Bone on Selected Solvers", "Set End Bone on Selected Solvers.", EUserInterfaceActionType::Button, FInputChord());
-	UI_COMMAND(AddBoneSettings, "Add Settings to Selected Bone", "Apply settings to the selected bone for all selected solvers (defined per-solver for limits, stiffness etc).", EUserInterfaceActionType::Button, FInputChord());
-	UI_COMMAND(RemoveBoneSettings, "Remove Settings on Selected Bone", "Remove all settings on the selected bone in all selected solvers.", EUserInterfaceActionType::Button, FInputChord());
+	UI_COMMAND(ConnectGoalToSolver, "Connect Goal to Selected Solver", "Connect this goal to the selected solver.", EUserInterfaceActionType::Button, FInputChord());
+	UI_COMMAND(DisconnectGoalFromSolver, "Disconnect Goal from Selected Solver", "Disconnect this goal from the selected solver.", EUserInterfaceActionType::Button, FInputChord());
+	UI_COMMAND(SetRootBoneOnSolver, "Set Root Bone on Selected Solver", "Set the Root Bone setting on the selected solver.", EUserInterfaceActionType::Button, FInputChord());
+	UI_COMMAND(SetEndBoneOnSolver, "Set End Bone on Selected Solver", "Set End Bone on selected solver.", EUserInterfaceActionType::Button, FInputChord());
+	UI_COMMAND(AddBoneSettings, "Add Settings to Selected Bone", "Apply settings to the selected bone for the selected solver (defined per-solver for limits, stiffness etc).", EUserInterfaceActionType::Button, FInputChord());
+	UI_COMMAND(RemoveBoneSettings, "Remove Settings on Selected Bone", "Remove all settings on the selected bone for the selected solver.", EUserInterfaceActionType::Button, FInputChord());
 	UI_COMMAND(ExcludeBone, "Exclude Selected Bone From Solve", "Ignore bone in all solvers.", EUserInterfaceActionType::Button, FInputChord());
 	UI_COMMAND(IncludeBone, "Include Selected Bone In Solve", "Include bone in all solvers.", EUserInterfaceActionType::Button, FInputChord());
 	UI_COMMAND(NewRetargetChain, "New Retarget Chain from Selected Bones", "Create a new retarget bone chain from the selected bones.", EUserInterfaceActionType::Button, FInputChord());
diff --git a/Engine/Plugins/Animation/IKRig/Source/IKRigEditor/Private/RigEditor/SIKRigSkeleton.cpp b/Engine/Plugins/Animation/IKRig/Source/IKRigEditor/Private/RigEditor/SIKRigSkeleton.cpp
--- a/Engine/Plugins/Animation/IKRig/Source/IKRigEditor/Private/RigEditor/SIKRigSkeleton.cpp
+++ b/Engine/Plugins/Animation/IKRig/Source/IKRigEditor/Private/RigEditor/SIKRigSkeleton.cpp
@@ -404,19 +404,19 @@ void SIKRigSkeleton::BindCommands()
         FExecuteAction::CreateSP(this, &SIKRigSkeleton::HandleDeleteElement),
         FCanExecuteAction::CreateSP(this, &SIKRigSkeleton::CanDeleteElement));
 
-	CommandList->MapAction(Commands.ConnectGoalToSolvers,
+	CommandList->MapAction(Commands.ConnectGoalToSolver,
         FExecuteAction::CreateSP(this, &SIKRigSkeleton::HandleConnectGoalToSolver),
         FCanExecuteAction::CreateSP(this, &SIKRigSkeleton::CanConnectGoalToSolvers));
 
-	CommandList->MapAction(Commands.DisconnectGoalFromSolvers,
+	CommandList->MapAction(Commands.DisconnectGoalFromSolver,
         FExecuteAction::CreateSP(this, &SIKRigSkeleton::HandleDisconnectGoalFromSolver),
         FCanExecuteAction::CreateSP(this, &SIKRigSkeleton::CanDisconnectGoalFromSolvers));
 
-	CommandList->MapAction(Commands.SetRootBoneOnSolvers,
+	CommandList->MapAction(Commands.SetRootBoneOnSolver,
         FExecuteAction::CreateSP(this, &SIKRigSkeleton::HandleSetRootBoneOnSolvers),
         FCanExecuteAction::CreateSP(this, &SIKRigSkeleton::CanSetRootBoneOnSolvers));
 
-	CommandList->MapAction(Commands.SetEndBoneOnSolvers,
+	CommandList->MapAction(Commands.SetEndBoneOnSolver,
 	FExecuteAction::CreateSP(this, &SIKRigSkeleton::HandleSetEndBoneOnSolvers),
 	FCanExecuteAction::CreateSP(this, &SIKRigSkeleton::CanSetEndBoneOnSolvers),
 	FIsActionChecked(),
@@ -470,13 +470,13 @@ void SIKRigSkeleton::FillContextMenu(FMenuBuilder& MenuBuilder)
 	MenuBuilder.EndSection();
 	
 	MenuBuilder.BeginSection("ConnectGoals", LOCTEXT("ConnectGoalOperations", "Connect Goals To Solvers"));
-	MenuBuilder.AddMenuEntry(Actions.ConnectGoalToSolvers);
-	MenuBuilder.AddMenuEntry(Actions.DisconnectGoalFromSolvers);
+	MenuBuilder.AddMenuEntry(Actions.ConnectGoalToSolver);
+	MenuBuilder.AddMenuEntry(Actions.DisconnectGoalFromSolver);
 	MenuBuilder.EndSection();
 
 	MenuBuilder.BeginSection("BoneSettings", LOCTEXT("BoneSettingsOperations", "Bone Settings"));
-	MenuBuilder.AddMenuEntry(Actions.SetRootBoneOnSolvers);
-	MenuBuilder.AddMenuEntry(Actions.SetEndBoneOnSolvers);
+	MenuBuilder.AddMenuEntry(Actions.SetRootBoneOnSolver);
+	MenuBuilder.AddMenuEntry(Actions.SetEndBoneOnSolver);
 	MenuBuilder.AddMenuEntry(Actions.AddBoneSettings);
 	MenuBuilder.AddMenuEntry(Actions.RemoveBoneSettings);
 	MenuBuilder.EndSection();
diff --git a/Engine/Plugins/Animation/IKRig/Source/IKRigEditor/Public/RigEditor/IKRigSkeletonCommands.h b/Engine/Plugins/Animation/IKRig/Source/IKRigEditor/Public/RigEditor/IKRigSkeletonCommands.h
--- a/Engine/Plugins/Animation/IKRig/Source/IKRigEditor/Public/RigEditor/IKRigSkeletonCommands.h
+++ b/Engine/Plugins/Animation/IKRig/Source/IKRigEditor/Public/RigEditor/IKRigSkeletonCommands.h
@@ -23,15 +23,15 @@ public:
 	/** delete selected element */
 	TSharedPtr< FUICommandInfo > DeleteElement;
 	/** add goal to selected solvers */
-	TSharedPtr< FUICommandInfo > ConnectGoalToSolvers;
+	TSharedPtr< FUICommandInfo > ConnectGoalToSolver;
 	/** remove goal to selected solvers */
-	TSharedPtr< FUICommandInfo > DisconnectGoalFromSolvers;
+	TSharedPtr< FUICommandInfo > DisconnectGoalFromSolver;
 
 	/** set root bone on selected solvers */
-	TSharedPtr< FUICommandInfo > SetRootBoneOnSolvers;
+	TSharedPtr< FUICommandInfo > SetRootBoneOnSolver;
 
 	/** set end bone on selected solvers (Pole Solvers only)*/
-	TSharedPtr< FUICommandInfo > SetEndBoneOnSolvers;
+	TSharedPtr< FUICommandInfo > SetEndBoneOnSolver;
 	
 	/** add bone setting to bone */
 	TSharedPtr< FUICommandInfo > AddBoneSettings;
