Fixed non-unity builds.

diff --git a/Engine/Plugins/Animation/IKRig/Source/IKRigEditor/Private/RigEditor/SIKRigRetargetChainList.cpp b/Engine/Plugins/Animation/IKRig/Source/IKRigEditor/Private/RigEditor/SIKRigRetargetChainList.cpp
--- a/Engine/Plugins/Animation/IKRig/Source/IKRigEditor/Private/RigEditor/SIKRigRetargetChainList.cpp
+++ b/Engine/Plugins/Animation/IKRig/Source/IKRigEditor/Private/RigEditor/SIKRigRetargetChainList.cpp
@@ -16,6 +16,7 @@
 #include "SPositiveActionButton.h"
 #include "SSearchableComboBox.h"
 #include "BoneSelectionWidget.h"
+#include "Animation/Skeleton.h"
 
 #define LOCTEXT_NAMESPACE "SIKRigRetargetChains"
 
