IK Rig: Added buttons to create and delete retarget chains. Fixed bug with retarget chains not being highlighted correctly in viewport.

diff --git a/Engine/Plugins/Animation/IKRig/Source/IKRigEditor/Private/RigEditor/IKRigController.cpp b/Engine/Plugins/Animation/IKRig/Source/IKRigEditor/Private/RigEditor/IKRigController.cpp
--- a/Engine/Plugins/Animation/IKRig/Source/IKRigEditor/Private/RigEditor/IKRigController.cpp
+++ b/Engine/Plugins/Animation/IKRig/Source/IKRigEditor/Private/RigEditor/IKRigController.cpp
@@ -143,20 +143,16 @@ bool UIKRigController::DoesBoneHaveSettings(const FName& BoneName) const
 	return false;
 }
 
-bool UIKRigController::AddRetargetChain(const FName& ChainName, const FName& StartBone, const FName& EndBone) const
+void UIKRigController::AddRetargetChain(const FName& ChainName, const FName& StartBone, const FName& EndBone) const
 {
-	if (Asset->GetRetargetChainByName(ChainName))
-	{
-		return false; // bone chain already exists with that name
-	}
+	FName UniqueChainName = GetUniqueRetargetChainName(ChainName);
 
 	FScopedTransaction Transaction(LOCTEXT("AddRetargetChain_Label", "Add Retarget Chain"));
 	Asset->Modify();
 	
-	Asset->RetargetDefinition.BoneChains.Emplace(ChainName, StartBone, EndBone);
+	Asset->RetargetDefinition.BoneChains.Emplace(UniqueChainName, StartBone, EndBone);
 	SortRetargetChains();
 	BroadcastNeedsReinitialized();
-	return true;
 }
 
 bool UIKRigController::RemoveRetargetChain(const FName& ChainName) const
@@ -321,10 +317,89 @@ void UIKRigController::SortRetargetChains() const
 	{
 		const int32 IndexA = Asset->Skeleton.GetBoneIndexFromName(A.StartBone.BoneName);
 		const int32 IndexB = Asset->Skeleton.GetBoneIndexFromName(B.StartBone.BoneName);
+		if (IndexA == IndexB)
+		{
+			// fallback to sorting alphabetically
+			return A.ChainName.LexicalLess(B.ChainName);
+		}
 		return IndexA < IndexB;
 	});
 }
 
+FName UIKRigController::GetUniqueRetargetChainName(const FName& NameToMakeUnique) const
+{
+	if (!Asset->GetRetargetChainByName(NameToMakeUnique))
+	{
+		return NameToMakeUnique; // name is already unique
+	}
+	
+	auto IsNameBeingUsed = [this](const FName& NameToTry)->bool
+	{
+		for (const FBoneChain& Chain : Asset->RetargetDefinition.BoneChains)
+		{
+			if (Chain.ChainName == NameToTry)
+			{
+				return true;
+			}
+		}
+		return false;
+	};
+
+	FString ChainName = NameToMakeUnique.ToString();
+	int8 SuffixInt = 0;
+	while (true)
+	{
+		const FName ChainNameToTry = FName(*FString::Format(TEXT("{0}_{1}"), {ChainName, FString::FromInt(SuffixInt)}));
+		if (!IsNameBeingUsed(ChainNameToTry))
+		{
+			return ChainNameToTry;
+		}
+		++SuffixInt;
+	}
+}
+
+bool UIKRigController::ValidateChain(const FName& ChainName, TSet<int32>& OutChainIndices) const
+{
+	const FBoneChain* Chain = Asset->GetRetargetChainByName(ChainName);
+	if (!Chain)
+	{
+		return false; // chain doesn't exist
+	}
+
+	const FIKRigSkeleton &Skeleton = GetIKRigSkeleton();
+	const int32 StartBoneIndex = Skeleton.GetBoneIndexFromName(Chain->StartBone.BoneName);
+	const int32 EndBoneIndex = Skeleton.GetBoneIndexFromName(Chain->EndBone.BoneName);
+	
+	if (StartBoneIndex == INDEX_NONE)
+	{
+		return false; // no start bone specified, not allowed
+	}
+
+	if (EndBoneIndex == INDEX_NONE)
+	{
+		OutChainIndices.Add(StartBoneIndex);
+		return true; // no end bone specified, this is a single bone "chain" which is fine
+	}
+
+	// this chain has a start AND an end bone so we must verify that end bone is child of start bone
+	int32 NextBoneIndex = EndBoneIndex;
+	while (true)
+	{
+		OutChainIndices.Add(NextBoneIndex);
+		if (StartBoneIndex == NextBoneIndex)
+		{
+			return true;
+		}
+		
+		NextBoneIndex = Skeleton.GetParentIndex(NextBoneIndex);
+		if (NextBoneIndex == INDEX_NONE)
+		{
+			// oops, we walked all the way past the root without finding the start bone
+			return false;
+		}
+	}
+}
+
 // -------------------------------------------------------
 // SKELETON
 //
diff --git a/Engine/Plugins/Animation/IKRig/Source/IKRigEditor/Private/RigEditor/IKRigEditMode.cpp b/Engine/Plugins/Animation/IKRig/Source/IKRigEditor/Private/RigEditor/IKRigEditMode.cpp
--- a/Engine/Plugins/Animation/IKRig/Source/IKRigEditor/Private/RigEditor/IKRigEditMode.cpp
+++ b/Engine/Plugins/Animation/IKRig/Source/IKRigEditor/Private/RigEditor/IKRigEditMode.cpp
@@ -273,18 +273,7 @@ void FIKRigEditMode::GetAffectedBones(
 			{
 				return;
 			}
-			UIKRigController* AssetController = Controller->AssetController;
-			const FName StartBone = AssetController->GetRetargetChainStartBone(SelectedChainName);
-			const FName EndBone = AssetController->GetRetargetChainEndBone(SelectedChainName);
-			const int32 ChainStartIndex = Skeleton.GetBoneIndexFromName(StartBone);
-			const int32 ChainEndIndex = Skeleton.GetBoneIndexFromName(EndBone);
-			OutSelectedBones.Add(ChainStartIndex);
-			OutAffectedBones.Add(ChainStartIndex);
-			for (int32 BoneIndex=ChainStartIndex; BoneIndex<=ChainEndIndex; ++BoneIndex)
-			{
-				OutSelectedBones.Add(BoneIndex);
-				OutAffectedBones.Add(BoneIndex);
-			}
+			Controller->AssetController->ValidateChain(SelectedChainName, OutSelectedBones);
 		}
 		break;
 		
diff --git a/Engine/Plugins/Animation/IKRig/Source/IKRigEditor/Private/RigEditor/SIKRigRetargetChainList.cpp b/Engine/Plugins/Animation/IKRig/Source/IKRigEditor/Private/RigEditor/SIKRigRetargetChainList.cpp
--- a/Engine/Plugins/Animation/IKRig/Source/IKRigEditor/Private/RigEditor/SIKRigRetargetChainList.cpp
+++ b/Engine/Plugins/Animation/IKRig/Source/IKRigEditor/Private/RigEditor/SIKRigRetargetChainList.cpp
@@ -130,6 +130,35 @@ TSharedRef<SWidget> SIKRigRetargetChainRow::GenerateWidgetForColumn(const FName&
 				SNew(STextBlock)
 				.Text(this, &SIKRigRetargetChainRow::GetGoalName)
 			]
+		]
+		+ SHorizontalBox::Slot()
+		.AutoWidth()
+		.HAlign(HAlign_Right)
+		.VAlign(VAlign_Center)
+		.Padding(3)
+		[
+			SNew(SButton)
+			.ToolTipText(LOCTEXT("DeleteChain", "Remove retarget bone chain from list."))
+			.OnClicked_Lambda([this]() -> FReply
+			{
+				const TSharedPtr<FIKRigEditorController> Controller = ChainList.Pin()->EditorController.Pin();
+				if (!Controller.IsValid())
+				{
+					return FReply::Unhandled();
+				}
+
+				UIKRigController* AssetController = Controller->AssetController;
+				AssetController->RemoveRetargetChain(ChainElement.Pin()->ChainName);
+
+				ChainList.Pin()->RefreshView();
+				return FReply::Handled();
+			})
+			.Content()
+			[
+				SNew(SImage)
+				.Image(FAppStyle::Get().GetBrush("Icons.Delete"))
+				.ColorAndOpacity(FSlateColor::UseForeground())
+			]
 		];
 		return GoalWidget;
 	}
@@ -264,7 +293,7 @@ void SIKRigRetargetChainList::Construct(const FArguments& InArgs, TSharedRef<FIK
 			.AutoWidth()
 			.HAlign(HAlign_Left)
 			.VAlign(VAlign_Center)
-			.Padding(3.0f, 0.0f)
+			.Padding(3.0f, 3.0f)
 			[
 				SNew(STextBlock)
 				.Text(LOCTEXT("RetargetRootLabel", "Retarget Root:"))
@@ -275,7 +304,7 @@ void SIKRigRetargetChainList::Construct(const FArguments& InArgs, TSharedRef<FIK
 			.AutoWidth()
 			.HAlign(HAlign_Left)
 			.VAlign(VAlign_Center)
-			.Padding(3.0f, 0.0f)
+			.Padding(3.0f, 3.0f)
 			[
 				SAssignNew(RetargetRootTextBox, SEditableTextBox)
 				.Text(FText::FromName(InEditorController->AssetController->GetRetargetRoot()))
@@ -283,6 +312,26 @@ void SIKRigRetargetChainList::Construct(const FArguments& InArgs, TSharedRef<FIK
 				.IsReadOnly(true)
 			]
         ]
+        
+		+ SVerticalBox::Slot()
+		.AutoHeight()
+		.VAlign(VAlign_Top)
+		.HAlign(HAlign_Left)
+		[
+			SNew(SPositiveActionButton)
+			.Icon(FAppStyle::Get().GetBrush("Icons.Plus"))
+			.Text(LOCTEXT("AddNewChainLabel", "Add New Chain"))
+			.ToolTipText(LOCTEXT("AddNewChainToolTip", "Add a new retarget bone chain."))
+			.OnClicked_Lambda([this]()
+			{
+				const UIKRigController* Controller = EditorController.Pin()->AssetController;
+				static FText NewChainText = LOCTEXT("NewRetargetChainLabel", "NewRetargetChain");
+				static FName NewChainName = FName(*NewChainText.ToString());
+				Controller->AddRetargetChain(NewChainName, NAME_None, NAME_None);
+				RefreshView();
+				return FReply::Handled();
+			})
+		]
 
         +SVerticalBox::Slot()
 		[
diff --git a/Engine/Plugins/Animation/IKRig/Source/IKRigEditor/Public/RigEditor/IKRigController.h b/Engine/Plugins/Animation/IKRig/Source/IKRigEditor/Public/RigEditor/IKRigController.h
--- a/Engine/Plugins/Animation/IKRig/Source/IKRigEditor/Public/RigEditor/IKRigController.h
+++ b/Engine/Plugins/Animation/IKRig/Source/IKRigEditor/Public/RigEditor/IKRigController.h
@@ -142,7 +142,7 @@ public:
 	 * 
 	 */
 	/** Add a Chain with the given Name and Start/End bones. Returns true if a new Chain was created. */
-	bool AddRetargetChain(const FName& ChainName, const FName& StartBone, const FName& EndBone) const;
+	void AddRetargetChain(const FName& ChainName, const FName& StartBone, const FName& EndBone) const;
 	/** Remove a Chain with the given name. Returns true if a Chain was removed. */
 	bool RemoveRetargetChain(const FName& ChainName) const;
 	/** Renamed the given Chain. Returns the new name (same as old if unsuccessful). */
@@ -167,6 +167,10 @@ public:
 	FName GetRetargetRoot() const;
 	/** Sorts the Chains from Root to tip based on the Start Bone of each Chain. */
 	void SortRetargetChains() const;
+	/** Make unique name for a retargeting bone chain. Adds a numbered suffix to make it unique.*/
+	FName GetUniqueRetargetChainName(const FName& NameToMakeUnique) const;
+	/** Returns true if this is a valid chain. Produces array of bone indices between start and end (inclusive). */
+	bool ValidateChain(const FName& ChainName, TSet<int32>& OutChainIndices) const;
 	/** END retarget chains */
 
 	// force all currently connected processors to reinitialize using latest asset state
