Notify the retargeter for chains being removed from the source/target ik rig

diff --git a/Engine/Plugins/Animation/IKRig/Source/IKRigEditor/Private/RetargetEditor/IKRetargetEditorController.cpp b/Engine/Plugins/Animation/IKRig/Source/IKRigEditor/Private/RetargetEditor/IKRetargetEditorController.cpp
--- a/Engine/Plugins/Animation/IKRig/Source/IKRigEditor/Private/RetargetEditor/IKRetargetEditorController.cpp
+++ b/Engine/Plugins/Animation/IKRig/Source/IKRigEditor/Private/RetargetEditor/IKRetargetEditorController.cpp
@@ -39,6 +39,7 @@ void FIKRetargetEditorController::BindToIKRigAsset(UIKRigDefinition* InIKRig) co
 	{
 		Controller->OnIKRigNeedsInitialized().AddSP(this, &FIKRetargetEditorController::OnIKRigNeedsInitialized);
 		Controller->OnRetargetChainRenamed().AddSP(this, &FIKRetargetEditorController::OnRetargetChainRenamed);
+		Controller->OnRetargetChainRemoved().AddSP(this, &FIKRetargetEditorController::OnRetargetChainRemoved);
 	}
 }
 
@@ -67,6 +68,13 @@ void FIKRetargetEditorController::OnRetargetChainRenamed(UIKRigDefinition* Modif
 	AssetController->OnRetargetChainRenamed(ModifiedIKRig, OldName, NewName);
 }
 
+void FIKRetargetEditorController::OnRetargetChainRemoved(UIKRigDefinition* ModifiedIKRig, const FName& InChainRemoved) const
+{
+	check(ModifiedIKRig)
+	AssetController->OnRetargetChainRemoved(ModifiedIKRig, InChainRemoved);
+	RefreshAllViews();
+}
+
 void FIKRetargetEditorController::OnRetargeterNeedsInitialized(const UIKRetargeter* Retargeter) const
 {
 	// force edit pose mode to be off
diff --git a/Engine/Plugins/Animation/IKRig/Source/IKRigEditor/Private/RetargetEditor/IKRetargeterController.cpp b/Engine/Plugins/Animation/IKRig/Source/IKRigEditor/Private/RetargetEditor/IKRetargeterController.cpp
--- a/Engine/Plugins/Animation/IKRig/Source/IKRigEditor/Private/RetargetEditor/IKRetargeterController.cpp
+++ b/Engine/Plugins/Animation/IKRig/Source/IKRigEditor/Private/RetargetEditor/IKRetargeterController.cpp
@@ -277,6 +277,40 @@ void UIKRetargeterController::OnRetargetChainRenamed(UIKRigDefinition* IKRig, FN
 	}
 }
 
+void UIKRetargeterController::OnRetargetChainRemoved(UIKRigDefinition* IKRig, const FName& InChainRemoved) const
+{
+	const bool bIsSourceRig = IKRig == Asset->SourceIKRigAsset;
+	check(bIsSourceRig || IKRig == Asset->TargetIKRigAsset)
+
+	// set source chain name to NONE if it has been deleted 
+	if (bIsSourceRig)
+	{
+		for (FRetargetChainMap& ChainMap : Asset->ChainMapping)
+		{
+			FName& SourceChainToUpdate = ChainMap.SourceChain;
+			if (SourceChainToUpdate == InChainRemoved)
+			{
+				SourceChainToUpdate = NAME_None;
+				BroadcastNeedsReinitialized();
+				return;
+			}
+		}
+		return;
+	}
+	
+	// remove target mapping if the target chain has been removed
+	const int32 ChainIndex = Asset->ChainMapping.IndexOfByPredicate([&InChainRemoved](const FRetargetChainMap& ChainMap)
+	{
+		return ChainMap.TargetChain == InChainRemoved;
+	});
+	
+	if (ChainIndex != INDEX_NONE)
+	{
+		Asset->ChainMapping.RemoveAt(ChainIndex);
+		BroadcastNeedsReinitialized();
+	}
+}
+
 void UIKRetargeterController::SetSourceChainForTargetChain(FName TargetChain, FName SourceChainToMapTo)
 {
 	FScopedTransaction Transaction(LOCTEXT("SetRetargetChainSource", "Set Retarget Chain Source"));
diff --git a/Engine/Plugins/Animation/IKRig/Source/IKRigEditor/Private/RigEditor/IKRigController.cpp b/Engine/Plugins/Animation/IKRig/Source/IKRigEditor/Private/RigEditor/IKRigController.cpp
--- a/Engine/Plugins/Animation/IKRig/Source/IKRigEditor/Private/RigEditor/IKRigController.cpp
+++ b/Engine/Plugins/Animation/IKRig/Source/IKRigEditor/Private/RigEditor/IKRigController.cpp
@@ -172,6 +172,7 @@ bool UIKRigController::RemoveRetargetChain(const FName& ChainName) const
 	
 	if (Asset->RetargetDefinition.BoneChains.RemoveAll(Pred) > 0)
 	{
+		RetargetChainRemoved.Broadcast(Asset, ChainName);
 		SortRetargetChains();
 		BroadcastNeedsReinitialized();
 		return true;
diff --git a/Engine/Plugins/Animation/IKRig/Source/IKRigEditor/Public/RetargetEditor/IKRetargetEditorController.h b/Engine/Plugins/Animation/IKRig/Source/IKRigEditor/Public/RetargetEditor/IKRetargetEditorController.h
--- a/Engine/Plugins/Animation/IKRig/Source/IKRigEditor/Public/RetargetEditor/IKRetargetEditorController.h
+++ b/Engine/Plugins/Animation/IKRig/Source/IKRigEditor/Public/RetargetEditor/IKRetargetEditorController.h
@@ -29,6 +29,8 @@ public:
 	void OnIKRigNeedsInitialized(UIKRigDefinition* ModifiedIKRig) const;
 	/** callback when IK Rig asset's retarget chain has been renamed */
 	void OnRetargetChainRenamed(UIKRigDefinition* ModifiedIKRig, FName OldName, FName NewName) const;
+	/** callback when IK Rig asset's retarget chain has been removed */
+	void OnRetargetChainRemoved(UIKRigDefinition* ModifiedIKRig, const FName& InChainRemoved) const;
 	/** callback when IK Retargeter asset requires reinitialization */
 	void OnRetargeterNeedsInitialized(const UIKRetargeter* Retargeter) const;
 	
diff --git a/Engine/Plugins/Animation/IKRig/Source/IKRigEditor/Public/RetargetEditor/IKRetargeterController.h b/Engine/Plugins/Animation/IKRig/Source/IKRigEditor/Public/RetargetEditor/IKRetargeterController.h
--- a/Engine/Plugins/Animation/IKRig/Source/IKRigEditor/Public/RetargetEditor/IKRetargeterController.h
+++ b/Engine/Plugins/Animation/IKRig/Source/IKRigEditor/Public/RetargetEditor/IKRetargeterController.h
@@ -59,6 +59,8 @@ public:
 	void AutoMapChains();
 	/** Callback when IK Rig chain is renamed. Retains existing mappings using the new name */
 	void OnRetargetChainRenamed(UIKRigDefinition* IKRig, FName OldChainName, FName NewChainName) const;
+	/** Callback when IK Rig chain is removed. */
+	void OnRetargetChainRemoved(UIKRigDefinition* IKRig, const FName& InChainRemoved) const;
 	/** Set the source chain to map to a given target chain */
 	void SetSourceChainForTargetChain(FName TargetChain, FName SourceChainToMapTo);
 	/** Get the source chain mapped to a given target chain */
diff --git a/Engine/Plugins/Animation/IKRig/Source/IKRigEditor/Public/RigEditor/IKRigController.h b/Engine/Plugins/Animation/IKRig/Source/IKRigEditor/Public/RigEditor/IKRigController.h
--- a/Engine/Plugins/Animation/IKRig/Source/IKRigEditor/Public/RigEditor/IKRigController.h
+++ b/Engine/Plugins/Animation/IKRig/Source/IKRigEditor/Public/RigEditor/IKRigController.h
@@ -189,11 +189,16 @@ private:
 	DECLARE_MULTICAST_DELEGATE_ThreeParams(FOnRetargetChainRenamed, UIKRigDefinition*, FName /*old name*/, FName /*new name*/);
 	FOnRetargetChainRenamed RetargetChainRenamed;
 
+	/** Called whenever a retarget chain is removed.*/
+	DECLARE_MULTICAST_DELEGATE_TwoParams(FOnRetargetChainRemoved, UIKRigDefinition*, const FName& /*chain name*/);
+	FOnRetargetChainRemoved RetargetChainRemoved;
+
 public:
 	
 	FOnIKRigNeedsInitialized& OnIKRigNeedsInitialized(){ return IKRigNeedsInitialized; };
 
 	FOnRetargetChainRenamed& OnRetargetChainRenamed(){ return RetargetChainRenamed; };
+	FOnRetargetChainRemoved& OnRetargetChainRemoved(){ return RetargetChainRemoved; };
 
 private:
 
