Update details panel when navigating through the hierarchy

diff --git a/Engine/Plugins/Animation/IKRig/Source/IKRigEditor/Private/RigEditor/IKRigEditorController.cpp b/Engine/Plugins/Animation/IKRig/Source/IKRigEditor/Private/RigEditor/IKRigEditorController.cpp
--- a/Engine/Plugins/Animation/IKRig/Source/IKRigEditor/Private/RigEditor/IKRigEditorController.cpp
+++ b/Engine/Plugins/Animation/IKRig/Source/IKRigEditor/Private/RigEditor/IKRigEditorController.cpp
@@ -371,6 +371,42 @@ void FIKRigEditorController::ShowEmptyDetails()
 	DetailsView->SetObject(AssetController->GetAsset());
 }
 
+void FIKRigEditorController::ShowDetailsForElements(const TArray<TSharedPtr<FIKRigTreeElement>>& InItems)
+{
+	// TODO
+	// this function is currently only dealing with displaying a single element in the details panel (the last one)
+	// as this is the current state of the details panel in the IKRig editor. However, it should be able to handle
+	// multiple items. This is a new feature that has to be done that's why this function's interface takes an array
+	// of items as parameter for future implementation.
+	
+	if (InItems.Num() > 0)
+	{
+		const TSharedPtr<FIKRigTreeElement>& LastItem = InItems.Last();
+		switch (LastItem->ElementType)
+		{
+			case IKRigTreeElementType::BONE:
+				ShowDetailsForBone(LastItem->BoneName);
+				break;
+			
+			case IKRigTreeElementType::GOAL:
+				ShowDetailsForGoal(LastItem->GoalName);
+				break;
+			
+			case IKRigTreeElementType::SOLVERGOAL:
+				ShowDetailsForGoalSettings(LastItem->SolverGoalName, LastItem->SolverGoalIndex);
+				break;
+			
+			case IKRigTreeElementType::BONE_SETTINGS:
+				ShowDetailsForBoneSettings(LastItem->BoneSettingBoneName, LastItem->BoneSettingsSolverIndex);
+				break;
+			
+			default:
+				ensure(false);
+				break;
+		}
+	}
+}
+
 void FIKRigEditorController::AddNewRetargetChain(const FName ChainName, const FName StartBone, const FName EndBone)
 {
 	FIKRigRetargetChainSettings Settings(ChainName, StartBone, EndBone);
diff --git a/Engine/Plugins/Animation/IKRig/Source/IKRigEditor/Private/RigEditor/SIKRigSkeleton.cpp b/Engine/Plugins/Animation/IKRig/Source/IKRigEditor/Private/RigEditor/SIKRigSkeleton.cpp
--- a/Engine/Plugins/Animation/IKRig/Source/IKRigEditor/Private/RigEditor/SIKRigSkeleton.cpp
+++ b/Engine/Plugins/Animation/IKRig/Source/IKRigEditor/Private/RigEditor/SIKRigSkeleton.cpp
@@ -1370,6 +1370,17 @@ void SIKRigSkeleton::HandleGetChildrenForTree(
 
 void SIKRigSkeleton::OnSelectionChanged(TSharedPtr<FIKRigTreeElement> Selection, ESelectInfo::Type SelectInfo)
 {
+	const TSharedPtr<FIKRigEditorController> Controller = EditorController.Pin();
+	if (!Controller.IsValid())
+	{
+		return;
+	}
+
+	// update details view
+	const TArray<TSharedPtr<FIKRigTreeElement>> SelectedItems = TreeView->GetSelectedItems();
+	Controller->ShowDetailsForElements(SelectedItems);
+
+	// NOTE: we may want to set the last selected item here
 }
 
 TSharedPtr<SWidget> SIKRigSkeleton::CreateContextMenu()
@@ -1388,24 +1399,6 @@ void SIKRigSkeleton::OnItemClicked(TSharedPtr<FIKRigTreeElement> InItem)
 		return;
 	}
 	
-	// update details view
-	if (InItem->ElementType == IKRigTreeElementType::BONE)
-	{
-		Controller->ShowDetailsForBone(InItem->BoneName);
-	}
-	else if (InItem->ElementType == IKRigTreeElementType::GOAL)
-	{
-		Controller->ShowDetailsForGoal(InItem->GoalName);
-	}
-	else if (InItem->ElementType == IKRigTreeElementType::SOLVERGOAL)
-	{
-		Controller->ShowDetailsForGoalSettings(InItem->SolverGoalName, InItem->SolverGoalIndex);
-	}
-	else if (InItem->ElementType == IKRigTreeElementType::BONE_SETTINGS)
-	{
-		Controller->ShowDetailsForBoneSettings(InItem->BoneSettingBoneName, InItem->BoneSettingsSolverIndex);
-	}
-
 	// to rename an item, you have to select it first, then click on it again within a time limit (slow double click)
 	const bool ClickedOnSameItem = TreeView->LastSelected.Pin().Get() == InItem.Get();
 	const uint32 CurrentCycles = FPlatformTime::Cycles();
diff --git a/Engine/Plugins/Animation/IKRig/Source/IKRigEditor/Public/RigEditor/IKRigEditorController.h b/Engine/Plugins/Animation/IKRig/Source/IKRigEditor/Public/RigEditor/IKRigEditorController.h
--- a/Engine/Plugins/Animation/IKRig/Source/IKRigEditor/Public/RigEditor/IKRigEditorController.h
+++ b/Engine/Plugins/Animation/IKRig/Source/IKRigEditor/Public/RigEditor/IKRigEditorController.h
@@ -113,6 +113,8 @@ public:
 	void ShowDetailsForSolver(const int32 SolverIndex);
 	/** show nothing in details view */
 	void ShowEmptyDetails();
+	/** show selected items in details view */
+	void ShowDetailsForElements(const TArray<TSharedPtr<FIKRigTreeElement>>& InItems);
 
 	/** create a new retarget chain */
 	void AddNewRetargetChain(const FName ChainName, const FName StartBone, const FName EndBone);
