Added undo/redo for Full Body IK bone settings

diff --git a/Engine/Plugins/Animation/IKRig/Source/IKRig/Public/Solvers/IKRig_PBIKSolver.h b/Engine/Plugins/Animation/IKRig/Source/IKRig/Public/Solvers/IKRig_PBIKSolver.h
--- a/Engine/Plugins/Animation/IKRig/Source/IKRig/Public/Solvers/IKRig_PBIKSolver.h
+++ b/Engine/Plugins/Animation/IKRig/Source/IKRig/Public/Solvers/IKRig_PBIKSolver.h
@@ -68,7 +68,10 @@ public:
 		X(EPBIKLimitType::Free),
 		Y(EPBIKLimitType::Free),
 		Z(EPBIKLimitType::Free),
-		PreferredAngles(FVector::ZeroVector){}
+		PreferredAngles(FVector::ZeroVector)
+	{
+		SetFlags(RF_Transactional);
+	}
 
 	/** The bone these settings are applied to. */
 	UPROPERTY(VisibleAnywhere, Category = Bone, meta = (Constant, CustomWidget = "BoneName"))
