Fixed crash with irrelevant bone index

diff --git a/Engine/Plugins/Animation/IKRig/Source/IKRig/Private/Solvers/IKRig_SetTransform.cpp b/Engine/Plugins/Animation/IKRig/Source/IKRig/Private/Solvers/IKRig_SetTransform.cpp
--- a/Engine/Plugins/Animation/IKRig/Source/IKRig/Private/Solvers/IKRig_SetTransform.cpp
+++ b/Engine/Plugins/Animation/IKRig/Source/IKRig/Private/Solvers/IKRig_SetTransform.cpp
@@ -18,6 +18,12 @@ void UIKRig_SetTransform::Initialize(const FIKRigSkeleton& IKRigSkeleton)
 
 void UIKRig_SetTransform::Solve(FIKRigSkeleton& IKRigSkeleton, const FIKRigGoalContainer& Goals)
 {
+	// BoneIndex is irrelevant
+	if (BoneIndex == INDEX_NONE)
+	{
+		return;
+	}
+	
 	const FIKRigGoal* InGoal = Goals.FindGoalByName(Goal);
 	if (!InGoal)
 	{
