IKRig Editor widget adaptions

diff --git a/Engine/Plugins/Animation/IKRig/Source/IKRig/IKRig.Build.cs b/Engine/Plugins/Animation/IKRig/Source/IKRig/IKRig.Build.cs
--- a/Engine/Plugins/Animation/IKRig/Source/IKRig/IKRig.Build.cs
+++ b/Engine/Plugins/Animation/IKRig/Source/IKRig/IKRig.Build.cs
@@ -31,5 +31,14 @@ public class IKRig : ModuleRules
 			{
 				// ... add any modules that your module loads dynamically here ...
 			});
+		
+		if (Target.bBuildEditor == true)
+		{
+			PrivateDependencyModuleNames.AddRange(
+				new string[]
+				{
+					"ApplicationCore",
+				});
+		}
 	}
 }
diff --git a/Engine/Plugins/Animation/IKRig/Source/IKRig/Private/IKRigDefinition.cpp b/Engine/Plugins/Animation/IKRig/Source/IKRig/Private/IKRigDefinition.cpp
--- a/Engine/Plugins/Animation/IKRig/Source/IKRig/Private/IKRigDefinition.cpp
+++ b/Engine/Plugins/Animation/IKRig/Source/IKRig/Private/IKRigDefinition.cpp
@@ -5,6 +5,12 @@
 #include "IKRigObjectVersion.h"
 #include "Engine/SkeletalMesh.h"
 
+#if WITH_EDITOR
+
+#include "HAL/PlatformApplicationMisc.h"
+
+#endif
+
 #define LOCTEXT_NAMESPACE	"IKRigDefinition"
 
 #if WITH_EDITOR
@@ -76,8 +82,128 @@ void UIKRigEffectorGoal::OnNumericValueChanged(ESlateTransformComponent::Type Co
 	}
 }
 
-bool UIKRigEffectorGoal::TransformDiffersFromDefault(TSharedPtr<IPropertyHandle> PropertyHandle,
-	ESlateTransformComponent::Type Component) const
+void UIKRigEffectorGoal::OnCopyToClipboard(ESlateTransformComponent::Type Component, EIKRigTransformType::Type TransformType)
+{
+	const FTransform Xfo = TransformType == EIKRigTransformType::Current ? CurrentTransform : InitialTransform;
+	
+	FString Content;
+	switch(Component)
+	{
+	case ESlateTransformComponent::Location:
+		{
+			const FVector Data = Xfo.GetLocation();
+			TBaseStructure<FVector>::Get()->ExportText(Content, &Data, &Data, nullptr, PPF_None, nullptr);
+			break;
+		}
+	case ESlateTransformComponent::Rotation:
+		{
+			const FRotator Data = Xfo.Rotator();
+			TBaseStructure<FRotator>::Get()->ExportText(Content, &Data, &Data, nullptr, PPF_None, nullptr);
+			break;
+		}
+	case ESlateTransformComponent::Scale:
+		{
+			const FVector Data = Xfo.GetScale3D();
+			TBaseStructure<FVector>::Get()->ExportText(Content, &Data, &Data, nullptr, PPF_None, nullptr);
+			break;
+		}
+	case ESlateTransformComponent::Max:
+	default:
+		{
+			TBaseStructure<FTransform>::Get()->ExportText(Content, &Xfo, &Xfo, nullptr, PPF_None, nullptr);
+			break;
+		}
+	}
+
+	if(!Content.IsEmpty())
+	{
+		FPlatformApplicationMisc::ClipboardCopy(*Content);
+	}
+}
+
+void UIKRigEffectorGoal::OnPasteFromClipboard(ESlateTransformComponent::Type Component, EIKRigTransformType::Type TransformType)
+{
+	FString Content;
+	FPlatformApplicationMisc::ClipboardPaste(Content);
+
+	if(Content.IsEmpty())
+	{
+		return;
+	}
+
+	FTransform& Xfo = TransformType == EIKRigTransformType::Current ? CurrentTransform : InitialTransform;
+	
+	Modify();
+
+	class FIKRigEffectorGoalErrorPipe : public FOutputDevice
+	{
+	public:
+
+		int32 NumErrors;
+
+		FIKRigEffectorGoalErrorPipe()
+			: FOutputDevice()
+			, NumErrors(0)
+		{
+		}
+
+		virtual void Serialize(const TCHAR* V, ELogVerbosity::Type Verbosity, const class FName& Category) override
+		{
+			NumErrors++;
+		}
+	};
+
+	FIKRigEffectorGoalErrorPipe ErrorPipe;
+				
+	switch(Component)
+	{
+		case ESlateTransformComponent::Location:
+		{
+			FVector Data = Xfo.GetLocation();
+			TBaseStructure<FVector>::Get()->ImportText(*Content, &Data, nullptr, PPF_None, &ErrorPipe, TBaseStructure<FVector>::Get()->GetName(), true);
+			if(ErrorPipe.NumErrors == 0)
+			{
+				Xfo.SetLocation(Data);
+			}
+			break;
+		}
+		case ESlateTransformComponent::Rotation:
+		{
+			FRotator Data = Xfo.Rotator();
+			TBaseStructure<FRotator>::Get()->ImportText(*Content, &Data, nullptr, PPF_None, &ErrorPipe, TBaseStructure<FRotator>::Get()->GetName(), true);
+			if(ErrorPipe.NumErrors == 0)
+			{
+				Xfo.SetRotation(FQuat(Data));
+			}
+			break;
+		}
+		case ESlateTransformComponent::Scale:
+		{
+			FVector Data = Xfo.GetScale3D();
+			TBaseStructure<FVector>::Get()->ImportText(*Content, &Data, nullptr, PPF_None, &ErrorPipe, TBaseStructure<FVector>::Get()->GetName(), true);
+			if(ErrorPipe.NumErrors == 0)
+			{
+				Xfo.SetScale3D(Data);
+			}
+			break;
+		}
+		case ESlateTransformComponent::Max:
+		default:
+		{
+			FTransform Data = Xfo;
+			TBaseStructure<FTransform>::Get()->ImportText(*Content, &Data, nullptr, PPF_None, &ErrorPipe, TBaseStructure<FTransform>::Get()->GetName(), true);
+			if(ErrorPipe.NumErrors == 0)
+			{
+				Xfo = Data;
+			}
+			break;
+		}
+	}
+}
+
+bool UIKRigEffectorGoal::TransformDiffersFromDefault(
+	ESlateTransformComponent::Type Component,
+	TSharedPtr<IPropertyHandle> PropertyHandle) const
 {
 	if(PropertyHandle->GetProperty()->GetFName() == GET_MEMBER_NAME_CHECKED(UIKRigEffectorGoal, CurrentTransform))
 	{
@@ -100,8 +226,9 @@ bool UIKRigEffectorGoal::TransformDiffersFromDefault(TSharedPtr<IPropertyHandle>
 	return false;
 }
 
-void UIKRigEffectorGoal::ResetTransformToDefault(TSharedPtr<IPropertyHandle> PropertyHandle,
-	ESlateTransformComponent::Type Component)
+void UIKRigEffectorGoal::ResetTransformToDefault(
+	ESlateTransformComponent::Type Component,
+	TSharedPtr<IPropertyHandle> PropertyHandle)
 {
 	switch(Component)
 	{
diff --git a/Engine/Plugins/Animation/IKRig/Source/IKRig/Public/IKRigDefinition.h b/Engine/Plugins/Animation/IKRig/Source/IKRig/Public/IKRigDefinition.h
--- a/Engine/Plugins/Animation/IKRig/Source/IKRig/Public/IKRigDefinition.h
+++ b/Engine/Plugins/Animation/IKRig/Source/IKRig/Public/IKRigDefinition.h
@@ -117,8 +117,11 @@ public:
 		EIKRigTransformType::Type TransformType
 	);
 
-	bool TransformDiffersFromDefault(TSharedPtr<IPropertyHandle> PropertyHandle, ESlateTransformComponent::Type Component) const;
-	void ResetTransformToDefault(TSharedPtr<IPropertyHandle> PropertyHandle, ESlateTransformComponent::Type Component);
+	void OnCopyToClipboard(ESlateTransformComponent::Type Component, EIKRigTransformType::Type TransformType);
+	void OnPasteFromClipboard(ESlateTransformComponent::Type Component, EIKRigTransformType::Type TransformType);
+
+	bool TransformDiffersFromDefault(ESlateTransformComponent::Type Component, TSharedPtr<IPropertyHandle> PropertyHandle) const;
+	void ResetTransformToDefault(ESlateTransformComponent::Type Component, TSharedPtr<IPropertyHandle> PropertyHandle);
 	
 #endif
 };
diff --git a/Engine/Plugins/Animation/IKRig/Source/IKRigEditor/IKRigEditor.Build.cs b/Engine/Plugins/Animation/IKRig/Source/IKRigEditor/IKRigEditor.Build.cs
--- a/Engine/Plugins/Animation/IKRig/Source/IKRigEditor/IKRigEditor.Build.cs
+++ b/Engine/Plugins/Animation/IKRig/Source/IKRigEditor/IKRigEditor.Build.cs
@@ -47,6 +47,7 @@ namespace UnrealBuildTool.Rules
 					"ToolWidgets",
 					"AnimationCore",
                     "AnimationWidgets",
+                    "ApplicationCore",
 				}
             );
 
diff --git a/Engine/Plugins/Animation/IKRig/Source/IKRigEditor/Private/RigEditor/IKRigDetailCustomizations.cpp b/Engine/Plugins/Animation/IKRig/Source/IKRigEditor/Private/RigEditor/IKRigDetailCustomizations.cpp
--- a/Engine/Plugins/Animation/IKRig/Source/IKRigEditor/Private/RigEditor/IKRigDetailCustomizations.cpp
+++ b/Engine/Plugins/Animation/IKRig/Source/IKRigEditor/Private/RigEditor/IKRigDetailCustomizations.cpp
@@ -90,8 +90,14 @@ void FIKRigGenericDetailCustomization::CustomizeDetailsForClass<UIKRigBoneDetail
 		LOCTEXT("ReferenceBoneTransformTooltip", "The reference transform of the bone")
 	};
 
-	static TSharedRef<TArray<EIKRigTransformType::Type>> VisibleTransforms =
-		MakeShareable(new TArray<EIKRigTransformType::Type>({EIKRigTransformType::Current}));
+	const TArray<EIKRigTransformType::Type> TransformTypes =
+	{
+		EIKRigTransformType::Current,
+		EIKRigTransformType::Reference
+	};
+
+	static TAttribute<TArray<EIKRigTransformType::Type>> VisibleTransforms =
+		TArray<EIKRigTransformType::Type>({EIKRigTransformType::Current});
 
 	TArray<TSharedRef<IPropertyHandle>> Properties;
 	Properties.Add(DetailBuilder.GetProperty(GET_MEMBER_NAME_CHECKED(UIKRigBoneDetails, CurrentTransform)));
@@ -102,7 +108,13 @@ void FIKRigGenericDetailCustomization::CustomizeDetailsForClass<UIKRigBoneDetail
 		DetailBuilder.HideProperty(Property);
 	}
 
-	TSharedPtr<SSegmentedControl<EIKRigTransformType::Type>> TransformChoiceWidget;
+	TSharedPtr<SSegmentedControl<EIKRigTransformType::Type>> TransformChoiceWidget =
+		SSegmentedControl<EIKRigTransformType::Type>::Create(
+			TransformTypes,
+			ButtonLabels,
+			ButtonTooltips,
+			VisibleTransforms
+		);
 
 	DetailBuilder.EditCategory(TEXT("Selection")).SetSortOrder(1);
 
@@ -120,22 +132,7 @@ void FIKRigGenericDetailCustomization::CustomizeDetailsForClass<UIKRigBoneDetail
 		.HAlign(HAlign_Left)
 		.VAlign(VAlign_Center)
 		[
-			SAssignNew(TransformChoiceWidget, SSegmentedControl<EIKRigTransformType::Type>)
-			.SupportsMultiSelection(true)
-			.Values_Lambda([]()
-			{
-				return VisibleTransforms.Get();
-			})
-			.OnValuesChanged_Lambda([](TArray<EIKRigTransformType::Type> Values)
-			{
-				VisibleTransforms.Get() = Values;
-			})
-			+ SSegmentedControl<EIKRigTransformType::Type>::Slot(EIKRigTransformType::Current)
-			.Text(ButtonLabels[0])
-			.ToolTip(ButtonTooltips[0])
-			+ SSegmentedControl<EIKRigTransformType::Type>::Slot(EIKRigTransformType::Reference)
-			.Text(ButtonLabels[1])
-			.ToolTip(ButtonTooltips[1])
+			TransformChoiceWidget.ToSharedRef()
 		]
 	];
 
@@ -151,66 +148,22 @@ void FIKRigGenericDetailCustomization::CustomizeDetailsForClass<UIKRigBoneDetail
 	{
 		const EIKRigTransformType::Type TransformType = (EIKRigTransformType::Type)PropertyIndex; 
 
-		TransformWidgetArgs
-		.OnGetIsComponentRelative_UObject(BoneDetails, &UIKRigBoneDetails::IsComponentRelative, TransformType)
-		.OnIsComponentRelativeChanged_UObject(BoneDetails, &UIKRigBoneDetails::OnComponentRelativeChanged, TransformType)
-		.Transform_UObject(BoneDetails, &UIKRigBoneDetails::GetTransform, TransformType);
-		
-		IDetailGroup& Group = CategoryBuilder.AddGroup(*ButtonLabels[PropertyIndex].ToString(), ButtonLabels[PropertyIndex], false, true);
-		Group.HeaderRow()
-		.Visibility( TAttribute<EVisibility>::CreateLambda([TransformChoiceWidget, TransformType]() -> EVisibility
+		TransformWidgetArgs.OnGetIsComponentRelative_UObject(BoneDetails, &UIKRigBoneDetails::IsComponentRelative, TransformType);
+		TransformWidgetArgs.OnIsComponentRelativeChanged_UObject(BoneDetails, &UIKRigBoneDetails::OnComponentRelativeChanged, TransformType);
+		TransformWidgetArgs.Transform_UObject(BoneDetails, &UIKRigBoneDetails::GetTransform, TransformType);
+		TransformWidgetArgs.OnCopyToClipboard_UObject(BoneDetails, &UIKRigBoneDetails::OnCopyToClipboard, TransformType);
+		TransformWidgetArgs.OnPasteFromClipboard_UObject(BoneDetails, &UIKRigBoneDetails::OnPasteFromClipboard, TransformType);
+
+		TransformWidgetArgs.Visibility_Lambda([TransformChoiceWidget, TransformType]() -> EVisibility
 		{
 			return TransformChoiceWidget->HasValue(TransformType) ? EVisibility::Visible : EVisibility::Collapsed;
-		}))
-		.NameContent()
-		[
-			SNew(STextBlock)
-			.Font(IDetailLayoutBuilder::GetDetailFont())
-			.Text(ButtonLabels[PropertyIndex])
-			.ToolTipText(ButtonTooltips[PropertyIndex])
-		];
-
-		Group.AddPropertyRow(Properties[PropertyIndex])
-		.CustomWidget()
-		.NameContent()
-		.HAlign(HAlign_Fill)
-		[
-			SAdvancedTransformInputBox<FTransform>::ConstructLabel(TransformWidgetArgs, ESlateTransformComponent::Location)
-		]
-		.ValueContent()
-		.MinDesiredWidth(IKRigDetailCustomizationsConstants::ItemWidth * 3.0f)
-		.MaxDesiredWidth(IKRigDetailCustomizationsConstants::ItemWidth * 3.0f)
-		[
-			SAdvancedTransformInputBox<FTransform>::ConstructWidget(TransformWidgetArgs, ESlateTransformComponent::Location)
-		];
-
-		Group.AddPropertyRow(Properties[PropertyIndex])
-		.CustomWidget()
-		.NameContent()
-		.HAlign(HAlign_Fill)
-		[
-			SAdvancedTransformInputBox<FTransform>::ConstructLabel(TransformWidgetArgs, ESlateTransformComponent::Rotation)
-		]
-		.ValueContent()
-		.MinDesiredWidth(IKRigDetailCustomizationsConstants::ItemWidth * 3.0f)
-		.MaxDesiredWidth(IKRigDetailCustomizationsConstants::ItemWidth * 3.0f)
-		[
-			SAdvancedTransformInputBox<FTransform>::ConstructWidget(TransformWidgetArgs, ESlateTransformComponent::Rotation)
-		];
+		});
 
-		Group.AddPropertyRow(Properties[PropertyIndex])
-		.CustomWidget()
-		.NameContent()
-		.HAlign(HAlign_Fill)
-		[
-			SAdvancedTransformInputBox<FTransform>::ConstructLabel(TransformWidgetArgs, ESlateTransformComponent::Scale)
-		]
-		.ValueContent()
-		.MinDesiredWidth(IKRigDetailCustomizationsConstants::ItemWidth * 3.0f)
-		.MaxDesiredWidth(IKRigDetailCustomizationsConstants::ItemWidth * 3.0f)
-		[
-			SAdvancedTransformInputBox<FTransform>::ConstructWidget(TransformWidgetArgs, ESlateTransformComponent::Scale)
-		];
+		SAdvancedTransformInputBox<FTransform>::ConstructGroupedTransformRows(
+			CategoryBuilder, 
+			ButtonLabels[PropertyIndex], 
+			ButtonTooltips[PropertyIndex], 
+			TransformWidgetArgs);
 	}
 }
 
@@ -237,8 +190,14 @@ void FIKRigGenericDetailCustomization::CustomizeDetailsForClass<UIKRigEffectorGo
 		LOCTEXT("ReferenceGoalTransformTooltip", "The reference transform of the goal")
 	};
 
-	static TSharedRef<TArray<EIKRigTransformType::Type>> VisibleTransforms =
-		MakeShareable(new TArray<EIKRigTransformType::Type>({EIKRigTransformType::Current}));
+	const TArray<EIKRigTransformType::Type> TransformTypes =
+	{
+		EIKRigTransformType::Current,
+		EIKRigTransformType::Reference
+	};
+
+	static TAttribute<TArray<EIKRigTransformType::Type>> VisibleTransforms =
+		TArray<EIKRigTransformType::Type>({EIKRigTransformType::Current});
 
 	TArray<TSharedRef<IPropertyHandle>> Properties;
 	Properties.Add(DetailBuilder.GetProperty(GET_MEMBER_NAME_CHECKED(UIKRigEffectorGoal, CurrentTransform)));
@@ -249,7 +208,13 @@ void FIKRigGenericDetailCustomization::CustomizeDetailsForClass<UIKRigEffectorGo
 		DetailBuilder.HideProperty(Property);
 	}
 
-	TSharedPtr<SSegmentedControl<EIKRigTransformType::Type>> TransformChoiceWidget;
+	TSharedPtr<SSegmentedControl<EIKRigTransformType::Type>> TransformChoiceWidget =
+		SSegmentedControl<EIKRigTransformType::Type>::Create(
+			TransformTypes,
+			ButtonLabels,
+			ButtonTooltips,
+			VisibleTransforms
+		);
 
 	DetailBuilder.EditCategory(TEXT("Goal Settings")).SetSortOrder(1);
 	DetailBuilder.EditCategory(TEXT("Viewport Goal Settings")).SetSortOrder(3);
@@ -268,22 +233,7 @@ void FIKRigGenericDetailCustomization::CustomizeDetailsForClass<UIKRigEffectorGo
 		.HAlign(HAlign_Left)
 		.VAlign(VAlign_Center)
 		[
-			SAssignNew(TransformChoiceWidget, SSegmentedControl<EIKRigTransformType::Type>)
-			.SupportsMultiSelection(true)
-			.Values_Lambda([]()
-			{
-				return VisibleTransforms.Get();
-			})
-			.OnValuesChanged_Lambda([](TArray<EIKRigTransformType::Type> Values)
-			{
-				VisibleTransforms.Get() = Values;
-			})
-			+ SSegmentedControl<EIKRigTransformType::Type>::Slot(EIKRigTransformType::Current)
-			.Text(ButtonLabels[0])
-			.ToolTip(ButtonTooltips[0])
-			+ SSegmentedControl<EIKRigTransformType::Type>::Slot(EIKRigTransformType::Reference)
-			.Text(ButtonLabels[1])
-			.ToolTip(ButtonTooltips[1])
+			TransformChoiceWidget.ToSharedRef()
 		]
 	];
 
@@ -305,74 +255,28 @@ void FIKRigGenericDetailCustomization::CustomizeDetailsForClass<UIKRigEffectorGo
 			.IsEnabled(false)
 			.DisplayScaleLock(false);
 		}
-
-		TransformWidgetArgs
-		.OnGetNumericValue_UObject(EffectorGoal, &UIKRigEffectorGoal::GetNumericValue, TransformType);
 		
-		IDetailGroup& Group = CategoryBuilder.AddGroup(*ButtonLabels[PropertyIndex].ToString(), ButtonLabels[PropertyIndex], false, true);
-		Group.HeaderRow()
-		.Visibility( TAttribute<EVisibility>::CreateLambda([TransformChoiceWidget, TransformType]() -> EVisibility
+		TransformWidgetArgs.OnGetNumericValue_UObject(EffectorGoal, &UIKRigEffectorGoal::GetNumericValue, TransformType);
+		TransformWidgetArgs.OnNumericValueChanged_UObject(EffectorGoal, &UIKRigEffectorGoal::OnNumericValueChanged, ETextCommit::Default, TransformType);
+		TransformWidgetArgs.OnNumericValueCommitted_UObject(EffectorGoal, &UIKRigEffectorGoal::OnNumericValueChanged, TransformType);
+
+		TransformWidgetArgs.OnCopyToClipboard_UObject(EffectorGoal, &UIKRigEffectorGoal::OnCopyToClipboard, TransformType);
+		TransformWidgetArgs.OnPasteFromClipboard_UObject(EffectorGoal, &UIKRigEffectorGoal::OnPasteFromClipboard, TransformType);
+
+		const TSharedPtr<IPropertyHandle> PropertyHandle = Properties[PropertyIndex];
+		TransformWidgetArgs.DiffersFromDefault_UObject(EffectorGoal, &UIKRigEffectorGoal::TransformDiffersFromDefault, PropertyHandle);
+		TransformWidgetArgs.OnResetToDefault_UObject(EffectorGoal, &UIKRigEffectorGoal::ResetTransformToDefault, PropertyHandle);
+
+		TransformWidgetArgs.Visibility_Lambda([TransformChoiceWidget, TransformType]() -> EVisibility
 		{
 			return TransformChoiceWidget->HasValue(TransformType) ? EVisibility::Visible : EVisibility::Collapsed;
-		}))
-		.NameContent()
-		[
-			SNew(STextBlock)
-			.Font(IDetailLayoutBuilder::GetDetailFont())
-			.Text(ButtonLabels[PropertyIndex])
-			.ToolTipText(ButtonTooltips[PropertyIndex])
-		];
-
-		Group.AddPropertyRow(Properties[PropertyIndex])
-		.OverrideResetToDefault(FResetToDefaultOverride::Create(
-			FIsResetToDefaultVisible::CreateUObject(EffectorGoal, &UIKRigEffectorGoal::TransformDiffersFromDefault, ESlateTransformComponent::Location),
-			FResetToDefaultHandler::CreateUObject(EffectorGoal, &UIKRigEffectorGoal::ResetTransformToDefault, ESlateTransformComponent::Location)))
-		.CustomWidget()
-		.NameContent()
-		.HAlign(HAlign_Fill)
-		[
-			SAdvancedTransformInputBox<FTransform>::ConstructLabel(TransformWidgetArgs, ESlateTransformComponent::Location)
-		]
-		.ValueContent()
-		.MinDesiredWidth(IKRigDetailCustomizationsConstants::ItemWidth * 3.0f)
-		.MaxDesiredWidth(IKRigDetailCustomizationsConstants::ItemWidth * 3.0f)
-		[
-			SAdvancedTransformInputBox<FTransform>::ConstructWidget(TransformWidgetArgs, ESlateTransformComponent::Location)
-		];
-
-		Group.AddPropertyRow(Properties[PropertyIndex])
-		.OverrideResetToDefault(FResetToDefaultOverride::Create(
-			FIsResetToDefaultVisible::CreateUObject(EffectorGoal, &UIKRigEffectorGoal::TransformDiffersFromDefault, ESlateTransformComponent::Rotation),
-			FResetToDefaultHandler::CreateUObject(EffectorGoal, &UIKRigEffectorGoal::ResetTransformToDefault, ESlateTransformComponent::Rotation)))
-		.CustomWidget()
-		.NameContent()
-		.HAlign(HAlign_Fill)
-		[
-			SAdvancedTransformInputBox<FTransform>::ConstructLabel(TransformWidgetArgs, ESlateTransformComponent::Rotation)
-		]
-		.ValueContent()
-		.MinDesiredWidth(IKRigDetailCustomizationsConstants::ItemWidth * 3.0f)
-		.MaxDesiredWidth(IKRigDetailCustomizationsConstants::ItemWidth * 3.0f)
-		[
-			SAdvancedTransformInputBox<FTransform>::ConstructWidget(TransformWidgetArgs, ESlateTransformComponent::Rotation)
-		];
-
-		Group.AddPropertyRow(Properties[PropertyIndex])
-		.OverrideResetToDefault(FResetToDefaultOverride::Create(
-			FIsResetToDefaultVisible::CreateUObject(EffectorGoal, &UIKRigEffectorGoal::TransformDiffersFromDefault, ESlateTransformComponent::Scale),
-			FResetToDefaultHandler::CreateUObject(EffectorGoal, &UIKRigEffectorGoal::ResetTransformToDefault, ESlateTransformComponent::Scale)))
-		.CustomWidget()
-		.NameContent()
-		.HAlign(HAlign_Fill)
-		[
-			SAdvancedTransformInputBox<FTransform>::ConstructLabel(TransformWidgetArgs, ESlateTransformComponent::Scale)
-		]
-		.ValueContent()
-		.MinDesiredWidth(IKRigDetailCustomizationsConstants::ItemWidth * 3.0f)
-		.MaxDesiredWidth(IKRigDetailCustomizationsConstants::ItemWidth * 3.0f)
-		[
-			SAdvancedTransformInputBox<FTransform>::ConstructWidget(TransformWidgetArgs, ESlateTransformComponent::Scale)
-		];
+		});
+
+		SAdvancedTransformInputBox<FTransform>::ConstructGroupedTransformRows(
+			CategoryBuilder, 
+			ButtonLabels[PropertyIndex], 
+			ButtonTooltips[PropertyIndex], 
+			TransformWidgetArgs);
 	}
 }
 
diff --git a/Engine/Plugins/Animation/IKRig/Source/IKRigEditor/Private/RigEditor/IKRigEditorController.cpp b/Engine/Plugins/Animation/IKRig/Source/IKRigEditor/Private/RigEditor/IKRigEditorController.cpp
--- a/Engine/Plugins/Animation/IKRig/Source/IKRigEditor/Private/RigEditor/IKRigEditorController.cpp
+++ b/Engine/Plugins/Animation/IKRig/Source/IKRigEditor/Private/RigEditor/IKRigEditorController.cpp
@@ -12,6 +12,12 @@
 #include "RigEditor/SIKRigSolverStack.h"
 #include "RigEditor/IKRigAnimInstance.h"
 
+#if WITH_EDITOR
+
+#include "HAL/PlatformApplicationMisc.h"
+
+#endif
+
 #define LOCTEXT_NAMESPACE "IKRigEditorController"
 
 bool UIKRigBoneDetails::CurrentTransformRelative[3] = {true, true, true};
@@ -104,6 +110,60 @@ void UIKRigBoneDetails::OnComponentRelativeChanged(
 	}
 }
 
+#if WITH_EDITOR
+
+void UIKRigBoneDetails::OnCopyToClipboard(ESlateTransformComponent::Type Component, EIKRigTransformType::Type TransformType)
+{
+	TOptional<FTransform> Optional = GetTransform(TransformType);
+	if(!Optional.IsSet())
+	{
+		return;
+	}
+
+	const FTransform Xfo = Optional.GetValue();
+	
+	FString Content;
+	switch(Component)
+	{
+	case ESlateTransformComponent::Location:
+		{
+			const FVector Data = Xfo.GetLocation();
+			TBaseStructure<FVector>::Get()->ExportText(Content, &Data, &Data, nullptr, PPF_None, nullptr);
+			break;
+		}
+	case ESlateTransformComponent::Rotation:
+		{
+			const FRotator Data = Xfo.Rotator();
+			TBaseStructure<FRotator>::Get()->ExportText(Content, &Data, &Data, nullptr, PPF_None, nullptr);
+			break;
+		}
+	case ESlateTransformComponent::Scale:
+		{
+			const FVector Data = Xfo.GetScale3D();
+			TBaseStructure<FVector>::Get()->ExportText(Content, &Data, &Data, nullptr, PPF_None, nullptr);
+			break;
+		}
+	case ESlateTransformComponent::Max:
+	default:
+		{
+			TBaseStructure<FTransform>::Get()->ExportText(Content, &Xfo, &Xfo, nullptr, PPF_None, nullptr);
+			break;
+		}
+	}
+
+	if(!Content.IsEmpty())
+	{
+		FPlatformApplicationMisc::ClipboardCopy(*Content);
+	}
+}
+
+void UIKRigBoneDetails::OnPasteFromClipboard(ESlateTransformComponent::Type Component, EIKRigTransformType::Type TransformType)
+{
+	// paste is not supported yet.
+}
+
+#endif
+
 void FIKRigEditorController::Initialize(TSharedPtr<FIKRigEditorToolkit> Toolkit, UIKRigDefinition* IKRigAsset)
 {
 	EditorToolkit = Toolkit;
diff --git a/Engine/Plugins/Animation/IKRig/Source/IKRigEditor/Public/RigEditor/IKRigEditorController.h b/Engine/Plugins/Animation/IKRig/Source/IKRigEditor/Public/RigEditor/IKRigEditorController.h
--- a/Engine/Plugins/Animation/IKRig/Source/IKRigEditor/Public/RigEditor/IKRigEditorController.h
+++ b/Engine/Plugins/Animation/IKRig/Source/IKRigEditor/Public/RigEditor/IKRigEditorController.h
@@ -58,6 +58,13 @@ public:
 	bool IsComponentRelative(ESlateTransformComponent::Type Component, EIKRigTransformType::Type TransformType) const;
 	void OnComponentRelativeChanged(ESlateTransformComponent::Type Component, bool bIsRelative, EIKRigTransformType::Type TransformType);
 
+#if WITH_EDITOR
+
+	void OnCopyToClipboard(ESlateTransformComponent::Type Component, EIKRigTransformType::Type TransformType);
+	void OnPasteFromClipboard(ESlateTransformComponent::Type Component, EIKRigTransformType::Type TransformType);
+
+#endif
+
 private:
 	
 	static bool CurrentTransformRelative[3];
