Keep the Edit Pose mode activated when undoing actions

diff --git a/Engine/Plugins/Animation/IKRig/Source/IKRigEditor/Private/RetargetEditor/IKRetargetEditor.cpp b/Engine/Plugins/Animation/IKRig/Source/IKRigEditor/Private/RetargetEditor/IKRetargetEditor.cpp
--- a/Engine/Plugins/Animation/IKRig/Source/IKRigEditor/Private/RetargetEditor/IKRetargetEditor.cpp
+++ b/Engine/Plugins/Animation/IKRig/Source/IKRigEditor/Private/RetargetEditor/IKRetargetEditor.cpp
@@ -266,8 +266,17 @@ TStatId FIKRetargetEditor::GetStatId() const
 
 void FIKRetargetEditor::PostUndo(bool bSuccess)
 {
+	const bool WasEditing = EditorController->IsEditingPose();
+	
 	EditorController->AssetController->BroadcastNeedsReinitialized();
 	EditorController->RefreshAllViews();
+
+	// restore pose mode state to avoid stepping out of the edition when undoing things
+	// note that BroadcastNeedsReinitialized will unset it in FIKRetargetEditorController::OnRetargeterNeedsInitialized
+	if (WasEditing)
+	{
+		EditorController->HandleEditPose();
+	}
 }
 
 void FIKRetargetEditor::PostRedo(bool bSuccess)
