fixed crash with wrong start/end indices

diff --git a/Engine/Plugins/Animation/IKRig/Source/IKRig/Private/Retargeter/IKRetargetProcessor.cpp b/Engine/Plugins/Animation/IKRig/Source/IKRig/Private/Retargeter/IKRetargetProcessor.cpp
--- a/Engine/Plugins/Animation/IKRig/Source/IKRig/Private/Retargeter/IKRetargetProcessor.cpp
+++ b/Engine/Plugins/Animation/IKRig/Source/IKRig/Private/Retargeter/IKRetargetProcessor.cpp
@@ -247,34 +247,37 @@ FResolvedBoneChain::FResolvedBoneChain(
 	const int32 EndIndex = Skeleton.FindBoneIndexByName(BoneChain.EndBone.BoneName);
 	bFoundStartBone = StartIndex != INDEX_NONE;
 	bFoundEndBone = EndIndex != INDEX_NONE;
-	
-	// validate end bone is child of start bone
-	bEndIsChildOfStart = true;
-	// record all bones in chain while walking up the hierarchy (tip to root of chain)
-	OutBoneIndices.Reset();
-	int32 NextBoneIndex = EndIndex;
-	while (true)
+
+	if (bFoundStartBone && bFoundEndBone)
 	{
-		OutBoneIndices.Add(NextBoneIndex);
-		
-		if (NextBoneIndex == StartIndex)
+		// validate end bone is child of start bone
+		bEndIsChildOfStart = true;
+		// record all bones in chain while walking up the hierarchy (tip to root of chain)
+		OutBoneIndices.Reset();
+		int32 NextBoneIndex = EndIndex;
+		while (true)
 		{
-			break;
-		}
+			OutBoneIndices.Add(NextBoneIndex);
+		
+			if (NextBoneIndex == StartIndex)
+			{
+				break;
+			}
 
-		NextBoneIndex = Skeleton.GetParentIndex(NextBoneIndex);
+			NextBoneIndex = Skeleton.GetParentIndex(NextBoneIndex);
 		
-		if (NextBoneIndex == 0)
-		{
-			// oops, we walked all the way up to the root without finding the start bone
-			bEndIsChildOfStart = false;
-			OutBoneIndices.Reset();
-			return;
+			if (NextBoneIndex == 0)
+			{
+				// oops, we walked all the way up to the root without finding the start bone
+				bEndIsChildOfStart = false;
+				OutBoneIndices.Reset();
+				return;
+			}
 		}
-	}
 
-	// reverse the indices (we want root to tip order)
-	Algo::Reverse(OutBoneIndices);
+		// reverse the indices (we want root to tip order)
+		Algo::Reverse(OutBoneIndices);
+	}
 }
 
 void FTargetSkeleton::SetBoneIsRetargeted(const int32 BoneIndex, const bool IsRetargeted)
