clear chain mapping when resetting the target asset

diff --git a/Engine/Plugins/Animation/IKRig/Source/IKRigEditor/Private/RetargetEditor/IKRetargeterController.cpp b/Engine/Plugins/Animation/IKRig/Source/IKRigEditor/Private/RetargetEditor/IKRetargeterController.cpp
--- a/Engine/Plugins/Animation/IKRig/Source/IKRigEditor/Private/RetargetEditor/IKRetargeterController.cpp
+++ b/Engine/Plugins/Animation/IKRig/Source/IKRigEditor/Private/RetargetEditor/IKRetargeterController.cpp
@@ -104,7 +104,8 @@ void UIKRetargeterController::CleanChainMapping()
 {
 	if (!IsValid(Asset->TargetIKRigAsset))
 	{
-		// don't clean chain mappings, in case user is replacing with IK Rig asset that has some valid mappings
+		// empty the chain mapping
+		Asset->ChainMapping.Empty();
 		return;
 	}
 	
