skeleton hierarchy: maintain selection expanded states as much as possible

diff --git a/Engine/Plugins/Animation/IKRig/Source/IKRigEditor/Private/RigEditor/SIKRigSkeleton.cpp b/Engine/Plugins/Animation/IKRig/Source/IKRigEditor/Private/RigEditor/SIKRigSkeleton.cpp
--- a/Engine/Plugins/Animation/IKRig/Source/IKRigEditor/Private/RigEditor/SIKRigSkeleton.cpp
+++ b/Engine/Plugins/Animation/IKRig/Source/IKRigEditor/Private/RigEditor/SIKRigSkeleton.cpp
@@ -186,7 +186,7 @@ void SIKRigSkeletonItem::OnNameCommitted(const FText& InText, ETextCommit::Type
 	{
 		return; 
 	}
-
+	
 	const FText OldText = WeakRigTreeElement.Pin()->Key;
 	const FName OldName = WeakRigTreeElement.Pin()->GoalName;
 	const FName PotentialNewName = FName(InText.ToString());
@@ -197,8 +197,8 @@ void SIKRigSkeletonItem::OnNameCommitted(const FText& InText, ETextCommit::Type
 		WeakRigTreeElement.Pin()->GoalName = NewName;
 	}
 	
-	Controller->SkeletonView->ReplaceItemInSelection(OldText, WeakRigTreeElement.Pin()->Key);
 	Controller->RefreshAllViews();
+	Controller->SkeletonView->ReplaceItemInSelection(OldText, WeakRigTreeElement.Pin()->Key);
 }
 
 FText SIKRigSkeletonItem::GetName() const
@@ -354,12 +354,12 @@ void SIKRigSkeleton::ReplaceItemInSelection(const FText& OldName, const FText& N
 	for (const TSharedPtr<FIKRigTreeElement>& Item : AllElements)
 	{
 		// remove old selection
-		if (Item->Key.CompareTo(OldName))
+		if (Item->Key.EqualTo(OldName))
 		{
 			TreeView->SetItemSelection(Item, false, ESelectInfo::Direct);
 		}
 		// add new selection
-		if (Item->Key.CompareTo(NewName))
+		if (Item->Key.EqualTo(NewName))
 		{
 			TreeView->SetItemSelection(Item, true, ESelectInfo::Direct);
 		}
@@ -1226,9 +1226,9 @@ void SIKRigSkeleton::RefreshTreeView(bool IsInitialSetup)
 		return;
 	}
 	
-	// save expansion state
-	TreeView->SaveAndClearSparseItemInfos();
-
+	// save expansion and selection state
+	TreeView->SaveAndClearState();
+	
 	// reset all tree items
 	RootElements.Reset();
 	AllElements.Reset();
@@ -1243,7 +1243,8 @@ void SIKRigSkeleton::RefreshTreeView(bool IsInitialSetup)
 	}
 
 	// get all goals
-	TArray<UIKRigEffectorGoal*> Goals = AssetController->GetAllGoals();
+	const TArray<UIKRigEffectorGoal*>& Goals = AssetController->GetAllGoals();
+	
 	// get all solvers
 	const TArray<UIKRigSolver*>& Solvers = AssetController->GetSolverArray();
 	// record bone element indices
@@ -1336,12 +1337,6 @@ void SIKRigSkeleton::RefreshTreeView(bool IsInitialSetup)
 		BoneTreeElement->Parent = ParentBoneTreeElement;
 	}
 
-	// restore expansion state
-	for (const TSharedPtr<FIKRigTreeElement>& Element : AllElements)
-	{
-		TreeView->RestoreSparseItemInfos(Element);
-	}
-
 	// expand all elements upon the initial construction of the tree
 	if (IsInitialSetup)
 	{
@@ -1350,6 +1345,14 @@ void SIKRigSkeleton::RefreshTreeView(bool IsInitialSetup)
 			SetExpansionRecursive(RootElement, false, true);
 		}
 	}
+	else
+	{
+		// restore expansion and selection state
+		for (const TSharedPtr<FIKRigTreeElement>& Element : AllElements)
+		{
+			TreeView->RestoreState(Element);
+		}
+	}
 	
 	TreeView->RequestTreeRefresh();
 }
diff --git a/Engine/Plugins/Animation/IKRig/Source/IKRigEditor/Public/RigEditor/SIKRigSkeleton.h b/Engine/Plugins/Animation/IKRig/Source/IKRigEditor/Public/RigEditor/SIKRigSkeleton.h
--- a/Engine/Plugins/Animation/IKRig/Source/IKRigEditor/Public/RigEditor/SIKRigSkeleton.h
+++ b/Engine/Plugins/Animation/IKRig/Source/IKRigEditor/Public/RigEditor/SIKRigSkeleton.h
@@ -100,6 +100,26 @@ public:
 		return Reply;
 	}
 
+	/** Save a snapshot of items expansion and selection state */
+	void SaveAndClearState()
+	{
+		SaveAndClearSparseItemInfos();
+		SaveAndClearSelection();
+	}
+	
+	/** Restore items expansion and selection state from the saved snapshot after tree reconstruction */
+	void RestoreState(const TSharedPtr<FIKRigTreeElement>& ItemPtr)
+	{
+		RestoreSparseItemInfos(ItemPtr);
+		RestoreSelection(ItemPtr);
+	}
+
+	/** slow double-click rename state*/
+	uint32 LastClickCycles = 0;
+	TWeakPtr<FIKRigTreeElement> LastSelected;
+
+private:
+
 	/** Save a snapshot of the internal map that tracks item expansion before tree reconstruction */
 	void SaveAndClearSparseItemInfos()
 	{
@@ -108,7 +128,7 @@ public:
 	}
 
 	/** Restore the expansion infos map from the saved snapshot after tree reconstruction */
-	void RestoreSparseItemInfos(TSharedPtr<FIKRigTreeElement> ItemPtr)
+	void RestoreSparseItemInfos(const TSharedPtr<FIKRigTreeElement>& ItemPtr)
 	{
 		for (const TTuple<TSharedPtr<FIKRigTreeElement>, FSparseItemInfo>& Pair : OldSparseItemInfos)
 		{
@@ -116,19 +136,40 @@ public:
 			{
 				// the SparseItemInfos now reference the new element, but keep the same expansion state
 				SparseItemInfos.Add(ItemPtr, Pair.Value);
-				break;
+				return;
 			}
 		}
-	}
-
-	/** slow double-click rename state*/
-	uint32 LastClickCycles = 0;
-	TWeakPtr<FIKRigTreeElement> LastSelected;
 
-private:
+		// set default state as expanded if not found
+		SparseItemInfos.Add(ItemPtr, FSparseItemInfo(true, false));
+	}
+	
+	/** Save a snapshot of the internal set that tracks item selection before tree reconstruction */
+	void SaveAndClearSelection()
+	{
+		OldSelectedItems = SelectedItems;
+		ClearSelection();
+	}
+	
+	/** Restore the selection from the saved snapshot after tree reconstruction */
+	void RestoreSelection(const TSharedPtr<FIKRigTreeElement>& ItemPtr)
+	{
+		for (const TSharedPtr<FIKRigTreeElement>& OldItem : OldSelectedItems)
+		{
+			if (OldItem->Key.EqualTo(ItemPtr->Key))
+			{
+				// select the new element
+				SetItemSelection(ItemPtr, true, ESelectInfo::Direct);
+				return;
+			}
+		}
+	}
 	
 	/** A temporary snapshot of the SparseItemInfos in STreeView, used during SIKRigSkeleton::RefreshTreeView() */
 	TSparseItemMap OldSparseItemInfos;
+
+	/** A temporary snapshot of the SelectedItems in SListView, used during SIKRigSkeleton::RefreshTreeView() */
+	TItemSet OldSelectedItems;
 };
 
 class SIKRigSkeleton : public SCompoundWidget, public FEditorUndoClient
