ensure that the target chain is safe to use when initializing

diff --git a/Engine/Plugins/Animation/IKRig/Source/IKRig/Private/Retargeter/IKRetargetProcessor.cpp b/Engine/Plugins/Animation/IKRig/Source/IKRig/Private/Retargeter/IKRetargetProcessor.cpp
--- a/Engine/Plugins/Animation/IKRig/Source/IKRig/Private/Retargeter/IKRetargetProcessor.cpp
+++ b/Engine/Plugins/Animation/IKRig/Source/IKRig/Private/Retargeter/IKRetargetProcessor.cpp
@@ -938,8 +938,12 @@ bool UIKRetargetProcessor::InitializeBoneChainPairs()
 	{
 		// get target bone chain
 		const FBoneChain* TargetBoneChain = RetargeterAsset->GetTargetIKRig()->GetRetargetChainByName(ChainMap.TargetChain);
-		ensureMsgf(TargetBoneChain, TEXT("IK Retargeter has mapping for a target bone chain that doesn't exist: %s"), *ChainMap.TargetChain.ToString());
-
+		if (!TargetBoneChain)
+		{
+			UE_LOG(LogTemp, Error, TEXT("IK Retargeter missing target bone chain: %s. Please update the mapping."), *ChainMap.TargetChain.ToString());
+			continue;
+		}
+		
 		// user opted to not map this to anything, we don't need to spam a warning about it
 		if (ChainMap.SourceChain == NAME_None)
 		{
