- Fixing crash when setting export folder in batch retarget.
- Brings folder selection dialog to front of screen.

diff --git a/Engine/Plugins/Animation/IKRig/Source/IKRigEditor/Private/RetargetEditor/SRetargetAnimAssetsWindow.cpp b/Engine/Plugins/Animation/IKRig/Source/IKRigEditor/Private/RetargetEditor/SRetargetAnimAssetsWindow.cpp
--- a/Engine/Plugins/Animation/IKRig/Source/IKRigEditor/Private/RetargetEditor/SRetargetAnimAssetsWindow.cpp
+++ b/Engine/Plugins/Animation/IKRig/Source/IKRigEditor/Private/RetargetEditor/SRetargetAnimAssetsWindow.cpp
@@ -21,7 +21,6 @@
 
 void SSelectExportPathDialog::Construct(const FArguments& InArgs)
 {
-	ParentWindowPtr = InArgs._ParentWindow;
 	AssetPath = FText::FromString(FPackageName::GetLongPackagePath(InArgs._DefaultAssetPath.ToString()));
 
 	if(AssetPath.IsEmpty())
@@ -40,7 +39,7 @@ void SSelectExportPathDialog::Construct(const FArguments& InArgs)
 		.Title(LOCTEXT("SSelectExportPathDialog_Title", "Select Export Path"))
 		.SupportsMinimize(false)
 		.SupportsMaximize(false)
-		//.SizingRule( ESizingRule::Autosized )
+		.IsTopmostWindow(true)
 		.ClientSize(FVector2D(450, 450))
 		[
 			SNew(SVerticalBox)
@@ -792,8 +791,7 @@ FText SRetargetAnimAssetsWindow::GetFolderPath() const
 FReply SRetargetAnimAssetsWindow::GetExportFolder()
 {
 	TSharedRef<SSelectExportPathDialog> Dialog = SNew(SSelectExportPathDialog)
-	.DefaultAssetPath(FText::FromString(BatchContext.FolderPath))
-	.ParentWindow(DialogWindow);
+	.DefaultAssetPath(FText::FromString(BatchContext.FolderPath));
 	
 	if(Dialog->ShowModal() != EAppReturnType::Cancel)
 	{
diff --git a/Engine/Plugins/Animation/IKRig/Source/IKRigEditor/Public/RetargetEditor/SRetargetAnimAssetsWindow.h b/Engine/Plugins/Animation/IKRig/Source/IKRigEditor/Public/RetargetEditor/SRetargetAnimAssetsWindow.h
--- a/Engine/Plugins/Animation/IKRig/Source/IKRigEditor/Public/RetargetEditor/SRetargetAnimAssetsWindow.h
+++ b/Engine/Plugins/Animation/IKRig/Source/IKRigEditor/Public/RetargetEditor/SRetargetAnimAssetsWindow.h
@@ -14,15 +14,8 @@ class UIKRetargeter;
 class SSelectExportPathDialog: public SWindow
 {
 public:
-	SLATE_BEGIN_ARGS(SSelectExportPathDialog)
-	{
-	}
-
+	SLATE_BEGIN_ARGS(SSelectExportPathDialog){}
 	SLATE_ARGUMENT(FText, DefaultAssetPath)
-
-	/** A reference to the parent window */
-	SLATE_ARGUMENT(TSharedPtr<SWindow>, ParentWindow)
-	
 	SLATE_END_ARGS()
 
 	SSelectExportPathDialog() : UserResponse(EAppReturnType::Cancel)
