fix wrong LOD management within the retargeter anim node

diff --git a/Engine/Plugins/Animation/IKRig/Source/IKRig/Private/AnimNodes/AnimNode_RetargetPoseFromMesh.cpp b/Engine/Plugins/Animation/IKRig/Source/IKRig/Private/AnimNodes/AnimNode_RetargetPoseFromMesh.cpp
--- a/Engine/Plugins/Animation/IKRig/Source/IKRig/Private/AnimNodes/AnimNode_RetargetPoseFromMesh.cpp
+++ b/Engine/Plugins/Animation/IKRig/Source/IKRig/Private/AnimNodes/AnimNode_RetargetPoseFromMesh.cpp
@@ -16,6 +16,40 @@ void FAnimNode_RetargetPoseFromMesh::Initialize_AnyThread(const FAnimationInitia
 void FAnimNode_RetargetPoseFromMesh::CacheBones_AnyThread(const FAnimationCacheBonesContext& Context)
 {
 	DECLARE_SCOPE_HIERARCHICAL_COUNTER_ANIMNODE(CacheBones_AnyThread)
+
+	const FBoneContainer& RequiredBones = Context.AnimInstanceProxy->GetRequiredBones();
+	if (!RequiredBones.IsValid())
+	{
+		return;
+	}
+
+	const bool bIsRetargeterReady = IKRetargeterAsset && Processor && Processor->IsInitialized();
+	if (!bIsRetargeterReady)
+	{
+		return;
+	}
+
+	// rebuild mapping
+	RequiredToTargetBoneMapping.Reset();
+
+	const FReferenceSkeleton& RefSkeleton = RequiredBones.GetReferenceSkeleton();
+	const FTargetSkeleton& TargetSkeleton = Processor->GetTargetSkeleton();
+	
+	const TArray<FBoneIndexType>& RequiredBonesArray = RequiredBones.GetBoneIndicesArray();
+	for (int32 Index = 0; Index < RequiredBonesArray.Num(); ++Index)
+	{
+		const FBoneIndexType ReqBoneIndex = RequiredBonesArray[Index]; 
+		if (ReqBoneIndex != INDEX_NONE)
+		{
+			const FName Name = RefSkeleton.GetBoneName(ReqBoneIndex);
+			const int32 TargetBoneIndex = TargetSkeleton.FindBoneIndexByName(Name);
+			if (TargetBoneIndex != INDEX_NONE)
+			{
+				// store require bone to target bone indices
+				RequiredToTargetBoneMapping.Emplace(Index, TargetBoneIndex);
+			}
+		}
+	}
 }
 
 void FAnimNode_RetargetPoseFromMesh::Update_AnyThread(const FAnimationUpdateContext& Context)
@@ -61,9 +95,15 @@ void FAnimNode_RetargetPoseFromMesh::Evaluate_AnyThread(FPoseContext& Output)
 	// copy pose back
 	FCSPose<FCompactPose> ComponentPose;
 	ComponentPose.InitPose(Output.Pose);
-	for (FCompactPoseBoneIndex CompactBoneIndex : Output.Pose.ForEachBoneIndex())
+	const FCompactPose& CompactPose = ComponentPose.GetPose();
+	for (const TPair<int32, int32>& Pair : RequiredToTargetBoneMapping)
 	{
-		ComponentPose.SetComponentSpaceTransform(CompactBoneIndex, RetargetedPose[CompactBoneIndex.GetInt()]);
+		const FCompactPoseBoneIndex CompactBoneIndex(Pair.Key);
+		if (CompactPose.IsValidIndex(CompactBoneIndex))
+		{
+			const int32 TargetBoneIndex = Pair.Value;
+			ComponentPose.SetComponentSpaceTransform(CompactBoneIndex, RetargetedPose[TargetBoneIndex]);
+		}
 	}
 
 	// convert to local space
diff --git a/Engine/Plugins/Animation/IKRig/Source/IKRig/Public/AnimNodes/AnimNode_RetargetPoseFromMesh.h b/Engine/Plugins/Animation/IKRig/Source/IKRig/Public/AnimNodes/AnimNode_RetargetPoseFromMesh.h
--- a/Engine/Plugins/Animation/IKRig/Source/IKRig/Public/AnimNodes/AnimNode_RetargetPoseFromMesh.h
+++ b/Engine/Plugins/Animation/IKRig/Source/IKRig/Public/AnimNodes/AnimNode_RetargetPoseFromMesh.h
@@ -64,4 +64,7 @@ private:
 
 	// cached transforms, copied on the game thread
 	TArray<FTransform> SourceMeshComponentSpaceBoneTransforms;
+
+	// mapping from required bones to actual bones within the target skeleton
+	TArray< TPair<int32, int32> > RequiredToTargetBoneMapping;
 };
