use SkeletalMesh::RefSkeleton instead of USkeleton::ReferenceSkeleton to reflect the same hierachy as the IKRig asset

diff --git a/Engine/Plugins/Animation/IKRig/Source/IKRigEditor/Private/RigEditor/SIKRigRetargetChainList.cpp b/Engine/Plugins/Animation/IKRig/Source/IKRigEditor/Private/RigEditor/SIKRigRetargetChainList.cpp
--- a/Engine/Plugins/Animation/IKRig/Source/IKRigEditor/Private/RigEditor/SIKRigRetargetChainList.cpp
+++ b/Engine/Plugins/Animation/IKRig/Source/IKRigEditor/Private/RigEditor/SIKRigRetargetChainList.cpp
@@ -16,7 +16,7 @@
 #include "SPositiveActionButton.h"
 #include "SSearchableComboBox.h"
 #include "BoneSelectionWidget.h"
-#include "Animation/Skeleton.h"
+#include "Engine/SkeletalMesh.h"
 
 #define LOCTEXT_NAMESPACE "SIKRigRetargetChains"
 
@@ -265,13 +265,13 @@ const FReferenceSkeleton& SIKRigRetargetChainRow::GetReferenceSkeleton() const
 		return DummySkeleton; 
 	}
 
-	USkeleton* Skeleton = Controller->AssetController->GetSkeleton();
-	if (Skeleton == nullptr)
+	USkeletalMesh* SkeletalMesh = Controller->AssetController->GetSkeletalMesh();
+	if (SkeletalMesh == nullptr)
 	{
 		return DummySkeleton;
 	}
-	
-	return Skeleton->GetReferenceSkeleton();
+
+	return SkeletalMesh->GetRefSkeleton();
 }
 
 void SIKRigRetargetChainList::Construct(const FArguments& InArgs, TSharedRef<FIKRigEditorController> InEditorController)
