Add ik goals to the processor even when they're not exposed in the IKRig asset. They are set as additive goals with a null offset to keep the input animation unchanged

diff --git a/Engine/Plugins/Animation/IKRig/Source/IKRigDeveloper/Private/AnimGraphNode_IKRig.cpp b/Engine/Plugins/Animation/IKRig/Source/IKRigDeveloper/Private/AnimGraphNode_IKRig.cpp
--- a/Engine/Plugins/Animation/IKRig/Source/IKRigDeveloper/Private/AnimGraphNode_IKRig.cpp
+++ b/Engine/Plugins/Animation/IKRig/Source/IKRigDeveloper/Private/AnimGraphNode_IKRig.cpp
@@ -344,10 +344,8 @@ void UAnimGraphNode_IKRig::PostEditChangeProperty(struct FPropertyChangedEvent&
 			const TArray<UIKRigEffectorGoal*>& AssetGoals = Node.RigDefinitionAsset->GetGoalArray();
 			for (const UIKRigEffectorGoal* AssetGoal: AssetGoals)
 			{
-				if (AssetGoal->bExposePosition || AssetGoal->bExposeRotation)
-				{
-					Node.Goals.Emplace(AssetGoal->GoalName);
-				}
+				const int32 GoalIndex = Node.Goals.Emplace(AssetGoal->GoalName);
+				SetupGoal(AssetGoal, Node.Goals[GoalIndex]);
 			}
 
 			BindPropertyChanges();
@@ -613,6 +611,7 @@ void UAnimGraphNode_IKRig::PostLoad()
 				{
 					GoalName = AssetGoals[Index]->GoalName;
 				}
+				SetupGoal(AssetGoals[Index], Node.Goals[Index]);
 			}
 		}
 
@@ -790,26 +789,46 @@ void UAnimGraphNode_IKRig::UpdateGoalsFromAsset()
 		const TArray<UIKRigEffectorGoal*>& AssetGoals = Node.RigDefinitionAsset->GetGoalArray();
 		for (const UIKRigEffectorGoal* AssetGoal: AssetGoals)
 		{
-			if (AssetGoal->bExposePosition || AssetGoal->bExposeRotation)
+			const int32 OldGoalIndex = OldGoals.IndexOfByPredicate([&AssetGoal](const FIKRigGoal& OldGoal)
 			{
-				const int32 OldGoalIndex = OldGoals.IndexOfByPredicate([&AssetGoal](const FIKRigGoal& OldGoal)
-				{
-					return AssetGoal->GoalName == OldGoal.Name;
-				});
-
-				if (OldGoalIndex != INDEX_NONE)
-				{
-					Node.Goals.Add(OldGoals[OldGoalIndex]);
-				}
-				else
-				{
-					Node.Goals.Emplace(AssetGoal->GoalName);
-				}
+				return AssetGoal->GoalName == OldGoal.Name;
+			});
+			
+			if (OldGoalIndex != INDEX_NONE)
+			{
+				Node.Goals.Add(OldGoals[OldGoalIndex]);
+			}
+			else
+			{
+				Node.Goals.Emplace(AssetGoal->GoalName);
 			}
+			
+			SetupGoal(AssetGoal, Node.Goals.Last());
 		}
 	}
 	
 	ReconstructNode();
 }
 
+void UAnimGraphNode_IKRig::SetupGoal(const UIKRigEffectorGoal* InAssetGoal, FIKRigGoal& OutGoal)
+{
+	// setup hidden goals so that they keep the bone's translation / rotation
+	if (InAssetGoal->GoalName == OutGoal.Name)
+	{
+		if (!InAssetGoal->bExposePosition)
+		{
+			OutGoal.Position = FVector::Zero();
+			OutGoal.PositionAlpha = 0.f;
+			OutGoal.PositionSpace = EIKRigGoalSpace::Additive;
+		}
+
+		if (!InAssetGoal->bExposeRotation)
+		{
+			OutGoal.Rotation = FQuat::Identity.Rotator();
+			OutGoal.RotationAlpha = 0.f;
+			OutGoal.RotationSpace = EIKRigGoalSpace::Additive;
+		}
+	}
+}
+
 #undef LOCTEXT_NAMESPACE
diff --git a/Engine/Plugins/Animation/IKRig/Source/IKRigDeveloper/Public/AnimGraphNode_IKRig.h b/Engine/Plugins/Animation/IKRig/Source/IKRigDeveloper/Public/AnimGraphNode_IKRig.h
--- a/Engine/Plugins/Animation/IKRig/Source/IKRigDeveloper/Public/AnimGraphNode_IKRig.h
+++ b/Engine/Plugins/Animation/IKRig/Source/IKRigDeveloper/Public/AnimGraphNode_IKRig.h
@@ -140,5 +140,9 @@ private:
 	void OnPropertyChanged(UObject* ObjectBeingModified, FPropertyChangedEvent& PropertyChangedEvent);
 	void BindPropertyChanges();
 
+	// update the goals' array within the anim node based on the asset
 	void UpdateGoalsFromAsset();
+
+	// setup goal based on it's asset definition 
+	static void SetupGoal(const UIKRigEffectorGoal* InAssetGoal, FIKRigGoal& OutGoal);
 };
