duplicate and retarget sequences referenced within anim montage for a complete retarget export

diff --git a/Engine/Plugins/Animation/IKRig/Source/IKRigEditor/Private/RetargetEditor/IKRetargetBatchOperation.cpp b/Engine/Plugins/Animation/IKRig/Source/IKRigEditor/Private/RetargetEditor/IKRetargetBatchOperation.cpp
--- a/Engine/Plugins/Animation/IKRig/Source/IKRigEditor/Private/RetargetEditor/IKRetargetBatchOperation.cpp
+++ b/Engine/Plugins/Animation/IKRig/Source/IKRigEditor/Private/RetargetEditor/IKRetargetBatchOperation.cpp
@@ -23,7 +23,7 @@
 #include "Retargeter/IKRetargeter.h"
 #include "Retargeter/IKRetargetProcessor.h"
 #include "Widgets/Notifications/SNotificationList.h"
-
+#include "Animation/AnimMontage.h"
 
 #define LOCTEXT_NAMESPACE "RetargetBatchOperation"
 
@@ -39,6 +39,29 @@ int32 FIKRetargetBatchOperation::GenerateAssetLists(const FIKRetargetBatchOperat
 		if (UAnimationAsset* AnimAsset = Cast<UAnimationAsset>(Asset))
 		{
 			AnimationAssetsToRetarget.AddUnique(AnimAsset);
+
+			// sequences that are used within the montage need to be added as well to be duplicated. They will then
+			// be replaced in UAnimMontage::ReplaceReferredAnimations
+			if (UAnimMontage* AnimMontage = Cast<UAnimMontage>(AnimAsset))
+			{
+				// add segments
+				for (const FSlotAnimationTrack& Track: AnimMontage->SlotAnimTracks)
+				{
+					for (const FAnimSegment& Segment: Track.AnimTrack.AnimSegments)
+					{
+						if (Segment.IsValid() && Segment.AnimReference)
+						{
+							AnimationAssetsToRetarget.AddUnique(Segment.AnimReference);
+						}
+					}
+				}
+
+				// add preview pose
+				if (AnimMontage->PreviewBasePose)
+				{
+					AnimationAssetsToRetarget.AddUnique(AnimMontage->PreviewBasePose);
+				}
+			}
 		}
 		else if (UAnimBlueprint* AnimBlueprint = Cast<UAnimBlueprint>(Asset))
 		{
