add retargeting options to anim montage
fixed MaxWidth preventing the retargeter selection

diff --git a/Engine/Plugins/Animation/IKRig/Source/IKRigEditor/Private/RetargetEditor/AssetTypeActions_IKRetargeter.cpp b/Engine/Plugins/Animation/IKRig/Source/IKRigEditor/Private/RetargetEditor/AssetTypeActions_IKRetargeter.cpp
--- a/Engine/Plugins/Animation/IKRig/Source/IKRigEditor/Private/RetargetEditor/AssetTypeActions_IKRetargeter.cpp
+++ b/Engine/Plugins/Animation/IKRig/Source/IKRigEditor/Private/RetargetEditor/AssetTypeActions_IKRetargeter.cpp
@@ -47,16 +47,20 @@ UThumbnailInfo* FAssetTypeActions_IKRetargeter::GetThumbnailInfo(UObject* Asset)
 
 void FAssetTypeActions_IKRetargeter::ExtendAnimSequenceToolMenu()
 {
-	TArray<UToolMenu*> MenusToExtend;
-	MenusToExtend.Add(UToolMenus::Get()->ExtendMenu("ContentBrowser.AssetContextMenu.AnimSequence"));
-	MenusToExtend.Add(UToolMenus::Get()->ExtendMenu("ContentBrowser.AssetContextMenu.BlendSpace"));
-	MenusToExtend.Add(UToolMenus::Get()->ExtendMenu("ContentBrowser.AssetContextMenu.AimOffsetBlendSpace"));
-	MenusToExtend.Add(UToolMenus::Get()->ExtendMenu("ContentBrowser.AssetContextMenu.BlendSpace1D"));
-	MenusToExtend.Add(UToolMenus::Get()->ExtendMenu("ContentBrowser.AssetContextMenu.PoseAsset"));
-	MenusToExtend.Add(UToolMenus::Get()->ExtendMenu("ContentBrowser.AssetContextMenu.AnimBlueprint"));
+	static const TArray<FName> MenusToExtend 
+	{
+		"ContentBrowser.AssetContextMenu.AnimSequence",
+		"ContentBrowser.AssetContextMenu.BlendSpace",
+		"ContentBrowser.AssetContextMenu.AimOffsetBlendSpace",
+		"ContentBrowser.AssetContextMenu.BlendSpace1D",
+		"ContentBrowser.AssetContextMenu.PoseAsset",
+		"ContentBrowser.AssetContextMenu.AnimBlueprint",
+		"ContentBrowser.AssetContextMenu.AnimMontage"
+	};
 
-	for (UToolMenu* Menu : MenusToExtend)
+	for (const FName& MenuName : MenusToExtend)
 	{
+		UToolMenu* Menu = UToolMenus::Get()->ExtendMenu(MenuName);
 		if (Menu == nullptr)
 		{
 			continue;
diff --git a/Engine/Plugins/Animation/IKRig/Source/IKRigEditor/Private/RetargetEditor/SRetargetAnimAssetsWindow.cpp b/Engine/Plugins/Animation/IKRig/Source/IKRigEditor/Private/RetargetEditor/SRetargetAnimAssetsWindow.cpp
--- a/Engine/Plugins/Animation/IKRig/Source/IKRigEditor/Private/RetargetEditor/SRetargetAnimAssetsWindow.cpp
+++ b/Engine/Plugins/Animation/IKRig/Source/IKRigEditor/Private/RetargetEditor/SRetargetAnimAssetsWindow.cpp
@@ -650,7 +650,6 @@ void SRetargetAnimAssetsWindow::ShowWindow(TArray<UObject*> InSelectedAssets)
 		.SupportsMinimize(false)
 		.SupportsMaximize(false)
 		.HasCloseButton(true)
-		.MaxWidth(1024.0f)
 		.IsTopmostWindow(true)
 		.SizingRule(ESizingRule::Autosized);
 	
