IK Retargeter: Removing rotation from IK pass.

diff --git a/Engine/Plugins/Animation/IKRig/Source/IKRig/Private/Retargeter/IKRetargetProcessor.cpp b/Engine/Plugins/Animation/IKRig/Source/IKRig/Private/Retargeter/IKRetargetProcessor.cpp
--- a/Engine/Plugins/Animation/IKRig/Source/IKRig/Private/Retargeter/IKRetargetProcessor.cpp
+++ b/Engine/Plugins/Animation/IKRig/Source/IKRig/Private/Retargeter/IKRetargetProcessor.cpp
@@ -1124,7 +1124,7 @@ void UIKRetargetProcessor::RunIKRetarget(
 			OutIKGoal.EndEffectorPosition,
 			OutIKGoal.EndEffectorRotation,
 			1.0f,
-			1.0f,
+			0.0f,
 			EIKRigGoalSpace::Component,
 			EIKRigGoalSpace::Component);
 		IKRigProcessor->SetIKGoal(Goal);
