ensure we are not transacting when evaluating the processor as it might not have been reinitiaized yet.

diff --git a/Engine/Plugins/Animation/IKRig/Source/IKRig/Private/AnimNodes/AnimNode_IKRig.cpp b/Engine/Plugins/Animation/IKRig/Source/IKRig/Private/AnimNodes/AnimNode_IKRig.cpp
--- a/Engine/Plugins/Animation/IKRig/Source/IKRig/Private/AnimNodes/AnimNode_IKRig.cpp
+++ b/Engine/Plugins/Animation/IKRig/Source/IKRig/Private/AnimNodes/AnimNode_IKRig.cpp
@@ -89,10 +89,14 @@ void FAnimNode_IKRig::AssignGoalTargets()
 	// update goal transforms before solve
 	// these transforms can come from a few different sources, handled here...
 
+	#if WITH_EDITOR
 	// use the goal transforms from the source asset itself
 	// this is used to live preview results from the IK Rig editor
-	#if WITH_EDITOR
-	if (bDriveWithSourceAsset)
+	// NOTE: as the transaction when undoing/redoing can be called on the preview scene before the editor, the processor
+	// might not have been reinitialized, resulting in data being desynchronized. Thus, we must wait until the
+	// transaction has been fully processed.
+
+	if (bDriveWithSourceAsset && !GIsTransacting)
 	{
 		IKRigProcessor->CopyAllInputsFromSourceAssetAtRuntime(RigDefinitionAsset);
 		return;
