Stop Auto-Sorting Retarget Chains While Editing

diff --git a/Engine/Plugins/Animation/IKRig/Source/IKRig/Private/Retargeter/IKRetargetProcessor.cpp b/Engine/Plugins/Animation/IKRig/Source/IKRig/Private/Retargeter/IKRetargetProcessor.cpp
--- a/Engine/Plugins/Animation/IKRig/Source/IKRig/Private/Retargeter/IKRetargetProcessor.cpp
+++ b/Engine/Plugins/Animation/IKRig/Source/IKRig/Private/Retargeter/IKRetargetProcessor.cpp
@@ -931,9 +931,13 @@ bool UIKRetargetProcessor::InitializeRoots()
 
 bool UIKRetargetProcessor::InitializeBoneChainPairs()
 {
-	check(RetargeterAsset->GetSourceIKRig() && RetargeterAsset->GetTargetIKRig());
+	const UIKRigDefinition* TargetIKRig = RetargeterAsset->GetTargetIKRig();
+	const UIKRigDefinition* SourceIKRig = RetargeterAsset->GetSourceIKRig();
+	
+	check(SourceIKRig && TargetIKRig);
 
-	const TArray<FRetargetChainMap>& ChainMapping = RetargeterAsset->GetChainMapping();
+	// check that chains are available in both IKRig assets before sorting them based on StartBone index
+	const TArray<FRetargetChainMap>& ChainMapping = RetargeterAsset->GetChainMapping();	
 	for (const FRetargetChainMap& ChainMap : ChainMapping)
 	{
 		// get target bone chain
@@ -975,6 +979,22 @@ bool UIKRetargetProcessor::InitializeBoneChainPairs()
 		}
 	}
 
+	// sort the chains based on their StartBone's index
+	auto ChainsSorter = [this](const FRetargetChainPair& A, const FRetargetChainPair& B)
+	{
+		const int32 IndexA = A.TargetBoneIndices.Num() > 0 ? A.TargetBoneIndices[0] : INDEX_NONE;
+		const int32 IndexB = B.TargetBoneIndices.Num() > 0 ? B.TargetBoneIndices[0] : INDEX_NONE;
+		if (IndexA == IndexB)
+		{
+			// fallback to sorting alphabetically
+			return A.TargetBoneChainName.LexicalLess(B.TargetBoneChainName);
+		}
+		return IndexA < IndexB;
+	};
+	
+	ChainPairsFK.Sort(ChainsSorter);
+	ChainPairsIK.Sort(ChainsSorter);
+	
 	// record which bones in the target skeleton are being retargeted
 	for (const FRetargetChainPairFK& FKChainPair : ChainPairsFK)
 	{
diff --git a/Engine/Plugins/Animation/IKRig/Source/IKRigEditor/Private/RetargetEditor/IKRetargeterController.cpp b/Engine/Plugins/Animation/IKRig/Source/IKRigEditor/Private/RetargetEditor/IKRetargeterController.cpp
--- a/Engine/Plugins/Animation/IKRig/Source/IKRigEditor/Private/RetargetEditor/IKRetargeterController.cpp
+++ b/Engine/Plugins/Animation/IKRig/Source/IKRigEditor/Private/RetargetEditor/IKRetargeterController.cpp
@@ -155,23 +155,8 @@ void UIKRetargeterController::CleanChainMapping()
 		}
 	}
 
-	// enforce the same chain order as the target IK rig
-	Asset->ChainMapping.Sort([this](const FRetargetChainMap& A, const FRetargetChainMap& B)
-	{
-		const TArray<FBoneChain>& BoneChains = Asset->TargetIKRigAsset->GetRetargetChains();
-		
-		const int32 IndexA = BoneChains.IndexOfByPredicate([&A](const FBoneChain& Chain)
-		{
-			return A.TargetChain == Chain.ChainName;
-		});
-
-		const int32 IndexB = BoneChains.IndexOfByPredicate([&B](const FBoneChain& Chain)
-		{
-			return B.TargetChain == Chain.ChainName;
-		});
- 
-		return IndexA < IndexB;
-	});
+	// enforce the chain order based on the StartBone index
+	SortChainMapping();
 
 	BroadcastNeedsReinitialized();
 }
@@ -221,7 +206,7 @@ void UIKRetargeterController::CleanPoseList()
 	BroadcastNeedsReinitialized();
 }
 
-void UIKRetargeterController::AutoMapChains()
+void UIKRetargeterController::AutoMapChains() const
 {
 	TArray<FName> SourceChainNames;
 	GetSourceChainNames(SourceChainNames);
@@ -258,6 +243,9 @@ void UIKRetargeterController::AutoMapChains()
 		}
 	}
 
+	// sort them
+	SortChainMapping();
+
 	// force update with latest mapping
 	BroadcastNeedsReinitialized();
 }
@@ -487,4 +475,42 @@ FRetargetChainMap* UIKRetargeterController::GetChainMap(const FName& TargetChain
 	return nullptr;
 }
 
+void UIKRetargeterController::SortChainMapping() const
+{
+	Asset->ChainMapping.Sort([this](const FRetargetChainMap& A, const FRetargetChainMap& B)
+	{
+		const TArray<FBoneChain>& BoneChains = Asset->TargetIKRigAsset->GetRetargetChains();
+		const FIKRigSkeleton& TargetSkeleton = Asset->TargetIKRigAsset->Skeleton;
+
+		// look for chains
+		const int32 IndexA = BoneChains.IndexOfByPredicate([&A](const FBoneChain& Chain)
+		{
+			return A.TargetChain == Chain.ChainName;
+		});
+
+		const int32 IndexB = BoneChains.IndexOfByPredicate([&B](const FBoneChain& Chain)
+		{
+			return B.TargetChain == Chain.ChainName;
+		});
+
+		// compare their StartBone Index 
+		if (IndexA > INDEX_NONE && IndexB > INDEX_NONE)
+		{
+			const int32 StartBoneIndexA = TargetSkeleton.GetBoneIndexFromName(BoneChains[IndexA].StartBone.BoneName);
+			const int32 StartBoneIndexB = TargetSkeleton.GetBoneIndexFromName(BoneChains[IndexB].StartBone.BoneName);
+
+			if (StartBoneIndexA == StartBoneIndexB)
+			{
+				// fallback to sorting alphabetically
+				return BoneChains[IndexA].ChainName.LexicalLess(BoneChains[IndexB].ChainName);
+			}
+				
+			return StartBoneIndexA < StartBoneIndexB;	
+		}
+
+		// sort them according to the target ik rig if previously failed 
+		return IndexA < IndexB;
+	});
+}
+
 #undef LOCTEXT_NAMESPACE
diff --git a/Engine/Plugins/Animation/IKRig/Source/IKRigEditor/Private/RigEditor/IKRigController.cpp b/Engine/Plugins/Animation/IKRig/Source/IKRigEditor/Private/RigEditor/IKRigController.cpp
--- a/Engine/Plugins/Animation/IKRig/Source/IKRigEditor/Private/RigEditor/IKRigController.cpp
+++ b/Engine/Plugins/Animation/IKRig/Source/IKRigEditor/Private/RigEditor/IKRigController.cpp
@@ -151,7 +151,6 @@ void UIKRigController::AddRetargetChain(const FName& ChainName, const FName& Sta
 	Asset->Modify();
 	
 	Asset->RetargetDefinition.BoneChains.Emplace(UniqueChainName, StartBone, EndBone);
-	SortRetargetChains();
 	BroadcastNeedsReinitialized();
 }
 
@@ -173,7 +172,6 @@ bool UIKRigController::RemoveRetargetChain(const FName& ChainName) const
 	if (Asset->RetargetDefinition.BoneChains.RemoveAll(Pred) > 0)
 	{
 		RetargetChainRemoved.Broadcast(Asset, ChainName);
-		SortRetargetChains();
 		BroadcastNeedsReinitialized();
 		return true;
 	}
@@ -209,7 +207,6 @@ bool UIKRigController::SetRetargetChainStartBone(const FName& ChainName, const F
 		FScopedTransaction Transaction(LOCTEXT("SetRetargetChainStartBone_Label", "Set Retarget Chain Start Bone"));
 		Asset->Modify();
 		BoneChain->StartBone = StartBoneName;
-		SortRetargetChains();
 		BroadcastNeedsReinitialized();
 		return true;
 	}
diff --git a/Engine/Plugins/Animation/IKRig/Source/IKRigEditor/Private/RigEditor/SIKRigRetargetChainList.cpp b/Engine/Plugins/Animation/IKRig/Source/IKRigEditor/Private/RigEditor/SIKRigRetargetChainList.cpp
--- a/Engine/Plugins/Animation/IKRig/Source/IKRigEditor/Private/RigEditor/SIKRigRetargetChainList.cpp
+++ b/Engine/Plugins/Animation/IKRig/Source/IKRigEditor/Private/RigEditor/SIKRigRetargetChainList.cpp
@@ -341,6 +341,7 @@ void SIKRigRetargetChainList::Construct(const FArguments& InArgs, TSharedRef<FIK
 			.ListItemsSource( &ListViewItems )
 			.OnGenerateRow( this, &SIKRigRetargetChainList::MakeListRowWidget )
 			.OnMouseButtonClick(this, &SIKRigRetargetChainList::OnItemClicked)
+			.OnContextMenuOpening(this, &SIKRigRetargetChainList::CreateContextMenu)
 			.ItemHeight( 22.0f )
 			.HeaderRow
 			(
@@ -466,4 +467,36 @@ FReply SIKRigRetargetChainList::OnKeyDown(const FGeometry& MyGeometry, const FKe
 	return FReply::Unhandled();
 }
 
+TSharedPtr<SWidget> SIKRigRetargetChainList::CreateContextMenu()
+{
+	static constexpr bool CloseAfterSelection = true;
+	FMenuBuilder MenuBuilder(CloseAfterSelection, CommandList);
+
+	MenuBuilder.BeginSection("Chains", LOCTEXT("ChainsSection", "Chains"));
+
+	const FUIAction Action = FUIAction( FExecuteAction::CreateSP(this, &SIKRigRetargetChainList::SortChainList));
+	static const FText Label = LOCTEXT("SortChainsLabel", "Sort Chains");
+	static const FText Tooltip = LOCTEXT("SortChainsTooltip", "Sort chain list in hierarchical order. This does not affect the retargeting behavior.");
+	MenuBuilder.AddMenuEntry(Label, Tooltip, FSlateIcon(), Action);
+
+	MenuBuilder.EndSection();
+	
+	return MenuBuilder.MakeWidget();
+}
+
+void SIKRigRetargetChainList::SortChainList()
+{
+	const TSharedPtr<FIKRigEditorController> Controller = EditorController.Pin();
+	if (!Controller.IsValid())
+	{
+		return;
+	}
+	
+	if (const UIKRigController* AssetController = Controller->AssetController)
+	{
+		AssetController->SortRetargetChains();
+		RefreshView();
+	}
+}
+
 #undef LOCTEXT_NAMESPACE
diff --git a/Engine/Plugins/Animation/IKRig/Source/IKRigEditor/Public/RetargetEditor/IKRetargeterController.h b/Engine/Plugins/Animation/IKRig/Source/IKRigEditor/Public/RetargetEditor/IKRetargeterController.h
--- a/Engine/Plugins/Animation/IKRig/Source/IKRigEditor/Public/RetargetEditor/IKRetargeterController.h
+++ b/Engine/Plugins/Animation/IKRig/Source/IKRigEditor/Public/RetargetEditor/IKRetargeterController.h
@@ -56,7 +56,7 @@ public:
 	/** Remove invalid chain mappings (no longer existing in currently referenced source/target IK Rig assets) */
 	void CleanChainMapping();
 	/** Use fuzzy string search to find "best" Source chain to map to each Target chain */
-	void AutoMapChains();
+	void AutoMapChains() const;
 	/** Callback when IK Rig chain is renamed. Retains existing mappings using the new name */
 	void OnRetargetChainRenamed(UIKRigDefinition* IKRig, FName OldChainName, FName NewChainName) const;
 	/** Callback when IK Rig chain is removed. */
@@ -119,6 +119,9 @@ private:
 
 	FRetargetChainMap* GetChainMap(const FName& TargetChainName) const;
 
+	/** Sort the Asset ChainMapping based on the StartBone of the target chains. */
+	void SortChainMapping() const;
+
 	/** The actual asset that this Controller modifies. */
 	UIKRetargeter* Asset = nullptr;
 };
diff --git a/Engine/Plugins/Animation/IKRig/Source/IKRigEditor/Public/RigEditor/SIKRigRetargetChainList.h b/Engine/Plugins/Animation/IKRig/Source/IKRigEditor/Public/RigEditor/SIKRigRetargetChainList.h
--- a/Engine/Plugins/Animation/IKRig/Source/IKRigEditor/Public/RigEditor/SIKRigRetargetChainList.h
+++ b/Engine/Plugins/Animation/IKRig/Source/IKRigEditor/Public/RigEditor/SIKRigRetargetChainList.h
@@ -151,6 +151,12 @@ private:
 	virtual FReply OnKeyDown(const FGeometry& MyGeometry, const FKeyEvent& InKeyEvent);
 	/** END SWidget interface */
 
+	/** list view OnContextMenuOpening callback */
+	TSharedPtr< SWidget > CreateContextMenu();
+	
+	/** sort chain list callback */
+	void SortChainList();
+
 	friend SIKRigRetargetChainRow;
 	friend FIKRigEditorController;
 };
