fix wrong usage of removed SkeletonMeshComponent

diff --git a/Engine/Plugins/Animation/IKRig/Source/IKRigEditor/Private/RigEditor/IKRigEditorController.cpp b/Engine/Plugins/Animation/IKRig/Source/IKRigEditor/Private/RigEditor/IKRigEditorController.cpp
--- a/Engine/Plugins/Animation/IKRig/Source/IKRigEditor/Private/RigEditor/IKRigEditorController.cpp
+++ b/Engine/Plugins/Animation/IKRig/Source/IKRigEditor/Private/RigEditor/IKRigEditorController.cpp
@@ -33,7 +33,7 @@ TOptional<FTransform> UIKRigBoneDetails::GetTransform(EIKRigTransformType::Type
 	FTransform LocalTransform = FTransform::Identity;
 	FTransform GlobalTransform = FTransform::Identity;
 	const bool* IsRelative = nullptr;
-	
+
 	const int32 BoneIndex = AssetPtr->Skeleton.GetBoneIndexFromName(SelectedBone);
 	if(BoneIndex == INDEX_NONE)
 	{
@@ -45,8 +45,21 @@ TOptional<FTransform> UIKRigBoneDetails::GetTransform(EIKRigTransformType::Type
 		case EIKRigTransformType::Current:
 		{
 			IsRelative = CurrentTransformRelative;
-			GlobalTransform = AnimInstancePtr->GetSkelMeshComponent()->GetBoneTransform(BoneIndex);
-			LocalTransform = AnimInstancePtr->GetSkelMeshComponent()->GetBoneSpaceTransforms()[BoneIndex];
+			
+			USkeletalMeshComponent* SkeletalMeshComponent = AnimInstancePtr->GetSkelMeshComponent();
+			const bool IsSkelMeshValid = SkeletalMeshComponent != nullptr &&
+										SkeletalMeshComponent->SkeletalMesh != nullptr;
+			if (IsSkelMeshValid)
+			{
+				GlobalTransform = SkeletalMeshComponent->GetBoneTransform(BoneIndex);
+				const TArray<FTransform>& LocalTransforms = SkeletalMeshComponent->GetBoneSpaceTransforms();
+				LocalTransform = LocalTransforms.IsValidIndex(BoneIndex) ? LocalTransforms[BoneIndex] : FTransform::Identity;
+			}
+			else
+			{
+				GlobalTransform = AssetPtr->Skeleton.CurrentPoseGlobal[BoneIndex];
+				LocalTransform = AssetPtr->Skeleton.CurrentPoseLocal[BoneIndex];
+			}
 			break;
 		}
 		case EIKRigTransformType::Reference:
diff --git a/Engine/Plugins/Animation/IKRig/Source/IKRigEditor/Private/RigEditor/IKRigToolkit.cpp b/Engine/Plugins/Animation/IKRig/Source/IKRigEditor/Private/RigEditor/IKRigToolkit.cpp
--- a/Engine/Plugins/Animation/IKRig/Source/IKRigEditor/Private/RigEditor/IKRigToolkit.cpp
+++ b/Engine/Plugins/Animation/IKRig/Source/IKRigEditor/Private/RigEditor/IKRigToolkit.cpp
@@ -216,6 +216,13 @@ void FIKRigEditorToolkit::HandlePreviewMeshChanged(USkeletalMesh* InOldSkeletalM
 	{
 		return; // already set to this skeletal mesh
 	}
+
+	// we do not reset the current skeletal mesh to keep track of the last valid one but we still need to reinit 
+	if (!InNewSkeletalMesh)
+	{
+		EditorController->AssetController->BroadcastNeedsReinitialized();
+		return;
+	}
 	
 	// update asset with new skeletal mesh (will copy new skeleton data)
 	if (!EditorController->AssetController->SetSkeletalMesh(InNewSkeletalMesh))
