fixed transaction for limb and pole effectors + reordered pole solver's setup messages

diff --git a/Engine/Plugins/Animation/IKRig/Source/IKRig/Private/Solvers/IKRig_LimbSolver.cpp b/Engine/Plugins/Animation/IKRig/Source/IKRig/Private/Solvers/IKRig_LimbSolver.cpp
--- a/Engine/Plugins/Animation/IKRig/Source/IKRig/Private/Solvers/IKRig_LimbSolver.cpp
+++ b/Engine/Plugins/Animation/IKRig/Source/IKRig/Private/Solvers/IKRig_LimbSolver.cpp
@@ -182,6 +182,7 @@ bool UIKRig_LimbSolver::IsBoneAffectedBySolver(const FName& BoneName, const FIKR
 
 void UIKRig_LimbSolver::AddGoal(const UIKRigEffectorGoal& NewGoal)
 {
+	Effector->Modify();
 	Effector->GoalName = NewGoal.GoalName;
 	Effector->BoneName = NewGoal.BoneName;
 }
@@ -190,6 +191,7 @@ void UIKRig_LimbSolver::RemoveGoal(const FName& GoalName)
 {
 	if (Effector->GoalName == GoalName)
 	{
+		Effector->Modify();
 		Effector->GoalName = NAME_None;
 		Effector->BoneName = NAME_None;
 	}
diff --git a/Engine/Plugins/Animation/IKRig/Source/IKRig/Private/Solvers/IKRig_PoleSolver.cpp b/Engine/Plugins/Animation/IKRig/Source/IKRig/Private/Solvers/IKRig_PoleSolver.cpp
--- a/Engine/Plugins/Animation/IKRig/Source/IKRig/Private/Solvers/IKRig_PoleSolver.cpp
+++ b/Engine/Plugins/Animation/IKRig/Source/IKRig/Private/Solvers/IKRig_PoleSolver.cpp
@@ -158,21 +158,21 @@ FText UIKRig_PoleSolver::GetNiceName() const
 
 bool UIKRig_PoleSolver::GetWarningMessage(FText& OutWarningMessage) const
 {
-	if (Effector->GoalName == NAME_None)
+	if (RootName == NAME_None)
 	{
-		OutWarningMessage = LOCTEXT("MissingGoal", "Missing goal.");
+		OutWarningMessage = LOCTEXT("MissingRoot", "Missing root bone.");
 		return true;
 	}
-
+	
 	if (EndName == NAME_None)
 	{
 		OutWarningMessage = LOCTEXT("MissingEnd", "Missing end bone.");
 		return true;
 	}
 	
-	if (RootName == NAME_None)
+	if (Effector->GoalName == NAME_None)
 	{
-		OutWarningMessage = LOCTEXT("MissingRoot", "Missing root bone.");
+		OutWarningMessage = LOCTEXT("MissingGoal", "Missing goal.");
 		return true;
 	}
 	
@@ -187,6 +187,7 @@ bool UIKRig_PoleSolver::GetWarningMessage(FText& OutWarningMessage) const
 
 void UIKRig_PoleSolver::AddGoal(const UIKRigEffectorGoal& NewGoal)
 {
+	Effector->Modify();
 	Effector->GoalName = NewGoal.GoalName;
 	Effector->BoneName = NewGoal.BoneName;
 }
@@ -195,6 +196,7 @@ void UIKRig_PoleSolver::RemoveGoal(const FName& GoalName)
 {
 	if (Effector->GoalName == GoalName)
 	{
+		Effector->Modify();
 		Effector->GoalName = NAME_None;
 		Effector->BoneName = NAME_None;
 	}
@@ -204,6 +206,7 @@ void UIKRig_PoleSolver::RenameGoal(const FName& OldName, const FName& NewName)
 {
 	if (Effector->GoalName == OldName)
 	{
+		Effector->Modify();
 		Effector->GoalName = NewName;
 	}
 }
@@ -212,6 +215,7 @@ void UIKRig_PoleSolver::SetGoalBone(const FName& GoalName, const FName& NewBoneN
 {
 	if (Effector->GoalName == GoalName)
 	{
+		Effector->Modify();
 		Effector->BoneName = NewBoneName;
 	}
 }
diff --git a/Engine/Plugins/Animation/IKRig/Source/IKRig/Public/Solvers/IKRig_PoleSolver.h b/Engine/Plugins/Animation/IKRig/Source/IKRig/Public/Solvers/IKRig_PoleSolver.h
--- a/Engine/Plugins/Animation/IKRig/Source/IKRig/Public/Solvers/IKRig_PoleSolver.h
+++ b/Engine/Plugins/Animation/IKRig/Source/IKRig/Public/Solvers/IKRig_PoleSolver.h
@@ -14,6 +14,8 @@ class IKRIG_API UIKRig_PoleSolverEffector : public UObject
 
 public:
 
+	UIKRig_PoleSolverEffector() { SetFlags(RF_Transactional); }
+	
 	UPROPERTY(VisibleAnywhere, Category = "Pole Solver Effector")
 	FName GoalName = NAME_None;
 
