Fixing crash in IK Rig editor when skeleton view is initialized after the viewport.

diff --git a/Engine/Plugins/Animation/IKRig/Source/IKRigEditor/Private/RigEditor/IKRigEditMode.cpp b/Engine/Plugins/Animation/IKRig/Source/IKRigEditor/Private/RigEditor/IKRigEditMode.cpp
--- a/Engine/Plugins/Animation/IKRig/Source/IKRigEditor/Private/RigEditor/IKRigEditMode.cpp
+++ b/Engine/Plugins/Animation/IKRig/Source/IKRigEditor/Private/RigEditor/IKRigEditMode.cpp
@@ -28,6 +28,11 @@ bool FIKRigEditMode::GetCameraTarget(FSphere& OutTarget) const
 	{
 		return false; 
 	}
+
+	if (!Controller->SkeletonView.IsValid())
+	{
+		return false;
+	}
 	
 	// target union of selected goals and bones
 	TArray<FName> OutGoalNames, OutBoneNames;
@@ -100,6 +105,11 @@ void FIKRigEditMode::RenderGoals(FPrimitiveDrawInterface* PDI)
 	{
 		return; 
 	}
+
+	if (!Controller->SkeletonView.IsValid())
+	{
+		return;
+	}
 	
 	UIKRigController* AssetController = Controller->AssetController;
 	const UIKRigDefinition* IKRigAsset = AssetController->GetAsset();
@@ -289,6 +299,11 @@ bool FIKRigEditMode::AllowWidgetMove()
 	{
 		return false; 
 	}
+
+	if (!Controller->SkeletonView.IsValid())
+	{
+		return false;
+	}
 	
 	return Controller->SkeletonView->GetNumSelectedGoals() > 0;
 }
@@ -300,6 +315,11 @@ bool FIKRigEditMode::ShouldDrawWidget() const
 	{
 		return false; 
 	}
+
+	if (!Controller->SkeletonView.IsValid())
+	{
+		return false;
+	}
 	
 	return Controller->SkeletonView->GetNumSelectedGoals() > 0;
 }
@@ -311,6 +331,11 @@ bool FIKRigEditMode::UsesTransformWidget() const
 	{
 		return false; 
 	}
+
+	if (!Controller->SkeletonView.IsValid())
+	{
+		return false;
+	}
 	
 	return Controller->SkeletonView->GetNumSelectedGoals() > 0;
 }
@@ -322,6 +347,11 @@ bool FIKRigEditMode::UsesTransformWidget(UE::Widget::EWidgetMode CheckMode) cons
 	{
 		return false; 
 	}
+
+	if (!Controller->SkeletonView.IsValid())
+	{
+		return false;
+	}
 	
 	return Controller->SkeletonView->GetNumSelectedGoals() > 0;
 }
@@ -331,7 +361,12 @@ FVector FIKRigEditMode::GetWidgetLocation() const
 	const TSharedPtr<FIKRigEditorController> Controller = EditorController.Pin();
 	if (!Controller.IsValid())
 	{
-		return FVector::ZeroVector; 
+		return FVector::ZeroVector;
+	}
+
+	if (!Controller->SkeletonView.IsValid())
+	{
+		return FVector::ZeroVector;
 	}
 	
 	TArray<FName> OutGoalNames;
@@ -393,6 +428,11 @@ bool FIKRigEditMode::StartTracking(FEditorViewportClient* InViewportClient, FVie
 	{
 		return false; 
 	}
+
+	if (!Controller->SkeletonView.IsValid())
+	{
+		return false;
+	}
 	
 	TArray<FName> SelectedGoalNames;
 	Controller->SkeletonView->GetSelectedGoalNames(SelectedGoalNames);
@@ -441,6 +481,11 @@ bool FIKRigEditMode::InputDelta(FEditorViewportClient* InViewportClient, FViewpo
 	{
 		return false; 
 	}
+
+	if (!Controller->SkeletonView.IsValid())
+	{
+		return false;
+	}
 	
 	if (!Controller->bManipulatingGoals)
 	{
@@ -485,6 +530,11 @@ bool FIKRigEditMode::GetCustomDrawingCoordinateSystem(FMatrix& InMatrix, void* I
 	{
 		return false;
 	}
+
+	if (!Controller->SkeletonView.IsValid())
+	{
+		return false;
+	}
 	
 	TArray<FName> SelectedGoalNames;
 	Controller->SkeletonView->GetSelectedGoalNames(SelectedGoalNames);
