IK Rig: Fixing crash when skeleton view initialized after viewport

diff --git a/Engine/Plugins/Animation/IKRig/Source/IKRigEditor/Private/RigEditor/IKRigEditMode.cpp b/Engine/Plugins/Animation/IKRig/Source/IKRigEditor/Private/RigEditor/IKRigEditMode.cpp
--- a/Engine/Plugins/Animation/IKRig/Source/IKRigEditor/Private/RigEditor/IKRigEditMode.cpp
+++ b/Engine/Plugins/Animation/IKRig/Source/IKRigEditor/Private/RigEditor/IKRigEditMode.cpp
@@ -139,6 +139,11 @@ void FIKRigEditMode::RenderBones(FPrimitiveDrawInterface* PDI)
 		return; 
 	}
 
+	if (!Controller->SkeletonView.IsValid())
+	{
+		return;
+	}
+
 	UIKRigController* AssetController = Controller->AssetController;
 	const UIKRigDefinition* IKRigAsset = AssetController->GetAsset();
 	if (!IKRigAsset->DrawBones)
