remove reference to the skeletal mesh when closing

diff --git a/Engine/Plugins/Animation/IKRig/Source/IKRigEditor/Private/RigEditor/IKRigToolkit.cpp b/Engine/Plugins/Animation/IKRig/Source/IKRigEditor/Private/RigEditor/IKRigToolkit.cpp
--- a/Engine/Plugins/Animation/IKRig/Source/IKRigEditor/Private/RigEditor/IKRigToolkit.cpp
+++ b/Engine/Plugins/Animation/IKRig/Source/IKRigEditor/Private/RigEditor/IKRigToolkit.cpp
@@ -34,7 +34,11 @@ FIKRigEditorToolkit::FIKRigEditorToolkit()
 
 FIKRigEditorToolkit::~FIKRigEditorToolkit()
 {
-	EditorController->SkelMeshComponent->SetSkeletalMesh(nullptr);
+	if (PersonaToolkit.IsValid())
+	{
+		static constexpr bool bSetPreviewMeshInAsset = false;
+		PersonaToolkit->SetPreviewMesh(nullptr, bSetPreviewMeshInAsset);
+	}
 }
 
 void FIKRigEditorToolkit::InitAssetEditor(
@@ -49,17 +53,18 @@ void FIKRigEditorToolkit::InitAssetEditor(
 	FPersonaToolkitArgs PersonaToolkitArgs;
 	PersonaToolkitArgs.OnPreviewSceneCreated = FOnPreviewSceneCreated::FDelegate::CreateSP(this, &FIKRigEditorToolkit::HandlePreviewSceneCreated);
 	
-	FPersonaModule& PersonaModule = FModuleManager::LoadModuleChecked<FPersonaModule>("Persona");
+	const FPersonaModule& PersonaModule = FModuleManager::LoadModuleChecked<FPersonaModule>("Persona");
 	PersonaToolkit = PersonaModule.CreatePersonaToolkit(IKRigAsset, PersonaToolkitArgs);
 	
+	
 	// when/if preview mesh is changed, we need to reinitialize the anim instance
     PersonaToolkit->GetPreviewScene()->RegisterOnPreviewMeshChanged(FOnPreviewMeshChanged::CreateSP(this, &FIKRigEditorToolkit::HandlePreviewMeshChanged));
 
-	TSharedRef<IAssetFamily> AssetFamily = PersonaModule.CreatePersonaAssetFamily(IKRigAsset);
+	const TSharedRef<IAssetFamily> AssetFamily = PersonaModule.CreatePersonaAssetFamily(IKRigAsset);
 	AssetFamily->RecordAssetOpened(FAssetData(IKRigAsset));
 
-	const bool bCreateDefaultStandaloneMenu = true;
-	const bool bCreateDefaultToolbar = true;
+	static constexpr bool bCreateDefaultStandaloneMenu = true;
+	static constexpr bool bCreateDefaultToolbar = true;
 	FAssetEditorToolkit::InitAssetEditor(
 		Mode, 
 		InitToolkitHost, 
