set the details panel on selected solver

diff --git a/Engine/Plugins/Animation/IKRig/Source/IKRigEditor/Private/RigEditor/SIKRigSkeleton.cpp b/Engine/Plugins/Animation/IKRig/Source/IKRigEditor/Private/RigEditor/SIKRigSkeleton.cpp
--- a/Engine/Plugins/Animation/IKRig/Source/IKRigEditor/Private/RigEditor/SIKRigSkeleton.cpp
+++ b/Engine/Plugins/Animation/IKRig/Source/IKRigEditor/Private/RigEditor/SIKRigSkeleton.cpp
@@ -392,6 +392,11 @@ void SIKRigSkeleton::GetSelectedBoneChains(TArray<FIKRigSkeletonChain>& OutChain
 	return Skeleton.GetChainsInList(SelectedBones, OutChains);
 }
 
+bool SIKRigSkeleton::HasSelectedItems() const
+{
+	return TreeView->GetNumItemsSelected() > 0;
+}
+
 void SIKRigSkeleton::BindCommands()
 {
 	const FIKRigSkeletonCommands& Commands = FIKRigSkeletonCommands::Get();
diff --git a/Engine/Plugins/Animation/IKRig/Source/IKRigEditor/Private/RigEditor/SIKRigSolverStack.cpp b/Engine/Plugins/Animation/IKRig/Source/IKRigEditor/Private/RigEditor/SIKRigSolverStack.cpp
--- a/Engine/Plugins/Animation/IKRig/Source/IKRigEditor/Private/RigEditor/SIKRigSolverStack.cpp
+++ b/Engine/Plugins/Animation/IKRig/Source/IKRigEditor/Private/RigEditor/SIKRigSolverStack.cpp
@@ -449,7 +449,7 @@ void SIKRigSolverStack::OnItemClicked(TSharedPtr<FSolverStackElement> InItem)
 	EditorController.Pin()->SetLastSelectedType(EIKRigSelectionType::SolverStack);
 }
 
-void SIKRigSolverStack::ShowDetailsForItem(TSharedPtr<FSolverStackElement> InItem)
+void SIKRigSolverStack::ShowDetailsForItem(TSharedPtr<FSolverStackElement> InItem) const
 {
 	const TSharedPtr<FIKRigEditorController> Controller = EditorController.Pin();
 	if (!Controller.IsValid())
@@ -457,16 +457,24 @@ void SIKRigSolverStack::ShowDetailsForItem(TSharedPtr<FSolverStackElement> InIte
 		return; 
 	}
 
+	// update bones greyed out when not affected
+	Controller->SkeletonView->RefreshTreeView();
+
+	// the solver selection must be done after we rebuilt the skeleton tree has it keeps the selection
 	if (!InItem.IsValid())
 	{
-		Controller->ShowEmptyDetails();	
-	}else
+		// clear the details panel only if there's nothing selected in the skeleton view
+		if (Controller->SkeletonView->HasSelectedItems())
+		{
+			return;
+		}
+		
+		Controller->ShowEmptyDetails();
+	}
+	else
 	{
 		Controller->ShowDetailsForSolver(InItem.Get()->IndexInStack);
 	}
-
-	// update bones greyed out when not affected
-	Controller->SkeletonView->RefreshTreeView(); 
 }
 
 TOptional<EItemDropZone> SIKRigSolverStack::OnCanAcceptDrop(
diff --git a/Engine/Plugins/Animation/IKRig/Source/IKRigEditor/Public/RigEditor/SIKRigSkeleton.h b/Engine/Plugins/Animation/IKRig/Source/IKRigEditor/Public/RigEditor/SIKRigSkeleton.h
--- a/Engine/Plugins/Animation/IKRig/Source/IKRigEditor/Public/RigEditor/SIKRigSkeleton.h
+++ b/Engine/Plugins/Animation/IKRig/Source/IKRigEditor/Public/RigEditor/SIKRigSkeleton.h
@@ -195,6 +195,7 @@ public:
 		IKRigTreeElementType ItemType,
 		const bool bReplace);
 	void GetSelectedBoneChains(TArray<FIKRigSkeletonChain>& OutChains);
+	bool HasSelectedItems() const;
 	/** END selection state queries */
 
 private:
diff --git a/Engine/Plugins/Animation/IKRig/Source/IKRigEditor/Public/RigEditor/SIKRigSolverStack.h b/Engine/Plugins/Animation/IKRig/Source/IKRigEditor/Public/RigEditor/SIKRigSolverStack.h
--- a/Engine/Plugins/Animation/IKRig/Source/IKRigEditor/Public/RigEditor/SIKRigSolverStack.h
+++ b/Engine/Plugins/Animation/IKRig/Source/IKRigEditor/Public/RigEditor/SIKRigSolverStack.h
@@ -112,7 +112,7 @@ private:
 	void OnSelectionChanged(TSharedPtr<FSolverStackElement> InItem, ESelectInfo::Type SelectInfo);
 	/** when a solver is clicked on in the stack view */
 	void OnItemClicked(TSharedPtr<FSolverStackElement> InItem);
-	void ShowDetailsForItem(TSharedPtr<FSolverStackElement> InItem);
+	void ShowDetailsForItem(TSharedPtr<FSolverStackElement> InItem) const; 
 
 	/** list view generate row callback */
 	TSharedRef<ITableRow> MakeListRowWidget(TSharedPtr<FSolverStackElement> InElement, const TSharedRef<STableViewBase>& OwnerTable);
