IKRig: Fixing crash when user layout has editor tabs closed.

diff --git a/Engine/Plugins/Animation/IKRig/Source/IKRigEditor/Private/RigEditor/IKRigEditMode.cpp b/Engine/Plugins/Animation/IKRig/Source/IKRigEditor/Private/RigEditor/IKRigEditMode.cpp
--- a/Engine/Plugins/Animation/IKRig/Source/IKRigEditor/Private/RigEditor/IKRigEditMode.cpp
+++ b/Engine/Plugins/Animation/IKRig/Source/IKRigEditor/Private/RigEditor/IKRigEditMode.cpp
@@ -28,16 +28,11 @@ bool FIKRigEditMode::GetCameraTarget(FSphere& OutTarget) const
 	{
 		return false; 
 	}
-
-	if (!Controller->SkeletonView.IsValid())
-	{
-		return false;
-	}
 	
 	// target union of selected goals and bones
 	TArray<FName> OutGoalNames, OutBoneNames;
-	Controller->SkeletonView->GetSelectedGoalNames(OutGoalNames);
-	Controller->SkeletonView->GetSelectedBoneNames(OutBoneNames);
+	Controller->GetSelectedGoalNames(OutGoalNames);
+	Controller->GetSelectedBoneNames(OutBoneNames);
 	
 	if (!OutGoalNames.IsEmpty() || !OutBoneNames.IsEmpty())
 	{
@@ -105,13 +100,8 @@ void FIKRigEditMode::RenderGoals(FPrimitiveDrawInterface* PDI)
 	{
 		return; 
 	}
-
-	if (!Controller->SkeletonView.IsValid())
-	{
-		return;
-	}
 	
-	UIKRigController* AssetController = Controller->AssetController;
+	const UIKRigController* AssetController = Controller->AssetController;
 	const UIKRigDefinition* IKRigAsset = AssetController->GetAsset();
 	if (!IKRigAsset->DrawGoals)
 	{
@@ -121,7 +111,7 @@ void FIKRigEditMode::RenderGoals(FPrimitiveDrawInterface* PDI)
 	TArray<UIKRigEffectorGoal*> Goals = AssetController->GetAllGoals();
 	for (const UIKRigEffectorGoal* Goal : Goals)
 	{
-		const bool bIsSelected = Controller->SkeletonView->IsGoalSelected(Goal->GoalName);
+		const bool bIsSelected = Controller->IsGoalSelected(Goal->GoalName);
 		const float Size = IKRigAsset->GoalSize * Goal->SizeMultiplier;
 		const float Thickness = IKRigAsset->GoalThickness * Goal->ThicknessMultiplier;
 		PDI->SetHitProxy(new HIKRigEditorGoalProxy(Goal->GoalName));
@@ -139,12 +129,7 @@ void FIKRigEditMode::RenderBones(FPrimitiveDrawInterface* PDI)
 		return; 
 	}
 
-	if (!Controller->SkeletonView.IsValid())
-	{
-		return;
-	}
-
-	UIKRigController* AssetController = Controller->AssetController;
+	const UIKRigController* AssetController = Controller->AssetController;
 	const UIKRigDefinition* IKRigAsset = AssetController->GetAsset();
 	if (!IKRigAsset->DrawBones)
 	{
@@ -231,8 +216,7 @@ void FIKRigEditMode::GetAffectedBones(
 		{
 			// get selected bones
 			TArray<TSharedPtr<FIKRigTreeElement>> SelectedBoneItems;
-			SIKRigSkeleton* SkeletonView = Controller->SkeletonView.Get();
-			SkeletonView->GetSelectedBones(SelectedBoneItems);
+			Controller->GetSelectedBones(SelectedBoneItems);
 
 			// record indices of all selected bones
 			for (TSharedPtr<FIKRigTreeElement> SelectedBone: SelectedBoneItems)
@@ -283,7 +267,7 @@ void FIKRigEditMode::GetAffectedBones(
 		
 		case EIKRigSelectionType::RetargetChains:
 		{
-			const FName SelectedChainName = Controller->RetargetingView->GetSelectedChain();
+			const FName SelectedChainName = Controller->GetSelectedChain();
 			if (SelectedChainName == NAME_None)
 			{
 				return;
@@ -304,13 +288,8 @@ bool FIKRigEditMode::AllowWidgetMove()
 	{
 		return false; 
 	}
-
-	if (!Controller->SkeletonView.IsValid())
-	{
-		return false;
-	}
 	
-	return Controller->SkeletonView->GetNumSelectedGoals() > 0;
+	return Controller->GetNumSelectedGoals() > 0;
 }
 
 bool FIKRigEditMode::ShouldDrawWidget() const
@@ -320,13 +299,8 @@ bool FIKRigEditMode::ShouldDrawWidget() const
 	{
 		return false; 
 	}
-
-	if (!Controller->SkeletonView.IsValid())
-	{
-		return false;
-	}
 	
-	return Controller->SkeletonView->GetNumSelectedGoals() > 0;
+	return Controller->GetNumSelectedGoals() > 0;
 }
 
 bool FIKRigEditMode::UsesTransformWidget() const
@@ -336,13 +310,8 @@ bool FIKRigEditMode::UsesTransformWidget() const
 	{
 		return false; 
 	}
-
-	if (!Controller->SkeletonView.IsValid())
-	{
-		return false;
-	}
 	
-	return Controller->SkeletonView->GetNumSelectedGoals() > 0;
+	return Controller->GetNumSelectedGoals() > 0;
 }
 
 bool FIKRigEditMode::UsesTransformWidget(UE::Widget::EWidgetMode CheckMode) const
@@ -352,13 +321,8 @@ bool FIKRigEditMode::UsesTransformWidget(UE::Widget::EWidgetMode CheckMode) cons
 	{
 		return false; 
 	}
-
-	if (!Controller->SkeletonView.IsValid())
-	{
-		return false;
-	}
 	
-	return Controller->SkeletonView->GetNumSelectedGoals() > 0;
+	return Controller->GetNumSelectedGoals() > 0;
 }
 
 FVector FIKRigEditMode::GetWidgetLocation() const
@@ -368,14 +332,9 @@ FVector FIKRigEditMode::GetWidgetLocation() const
 	{
 		return FVector::ZeroVector;
 	}
-
-	if (!Controller->SkeletonView.IsValid())
-	{
-		return FVector::ZeroVector;
-	}
 	
 	TArray<FName> OutGoalNames;
-	Controller->SkeletonView->GetSelectedGoalNames(OutGoalNames);
+	Controller->GetSelectedGoalNames(OutGoalNames);
 	if (OutGoalNames.IsEmpty())
 	{
 		return FVector::ZeroVector; 
@@ -433,14 +392,9 @@ bool FIKRigEditMode::StartTracking(FEditorViewportClient* InViewportClient, FVie
 	{
 		return false; 
 	}
-
-	if (!Controller->SkeletonView.IsValid())
-	{
-		return false;
-	}
 	
 	TArray<FName> SelectedGoalNames;
-	Controller->SkeletonView->GetSelectedGoalNames(SelectedGoalNames);
+	Controller->GetSelectedGoalNames(SelectedGoalNames);
 	if (SelectedGoalNames.IsEmpty())
 	{
 		return false; // no goals selected to manipulate
@@ -486,11 +440,6 @@ bool FIKRigEditMode::InputDelta(FEditorViewportClient* InViewportClient, FViewpo
 	{
 		return false; 
 	}
-
-	if (!Controller->SkeletonView.IsValid())
-	{
-		return false;
-	}
 	
 	if (!Controller->bManipulatingGoals)
 	{
@@ -498,8 +447,8 @@ bool FIKRigEditMode::InputDelta(FEditorViewportClient* InViewportClient, FViewpo
 	}
 
 	TArray<FName> SelectedGoalNames;
-	Controller->SkeletonView->GetSelectedGoalNames(SelectedGoalNames);
-	UIKRigController* AssetController = Controller->AssetController;
+	Controller->GetSelectedGoalNames(SelectedGoalNames);
+	const UIKRigController* AssetController = Controller->AssetController;
 
 	// translate goals
 	if(InViewportClient->GetWidgetMode() == UE::Widget::WM_Translate)
@@ -535,14 +484,9 @@ bool FIKRigEditMode::GetCustomDrawingCoordinateSystem(FMatrix& InMatrix, void* I
 	{
 		return false;
 	}
-
-	if (!Controller->SkeletonView.IsValid())
-	{
-		return false;
-	}
 	
 	TArray<FName> SelectedGoalNames;
-	Controller->SkeletonView->GetSelectedGoalNames(SelectedGoalNames);
+	Controller->GetSelectedGoalNames(SelectedGoalNames);
 	if (SelectedGoalNames.IsEmpty())
 	{
 		return false; // nothing selected to manipulate
@@ -578,7 +522,7 @@ bool FIKRigEditMode::InputKey(FEditorViewportClient* ViewportClient, FViewport*
 		}
 	
 		TArray<FName> SelectedGoalNames;
-		Controller->SkeletonView->GetSelectedGoalNames(SelectedGoalNames);
+		Controller->GetSelectedGoalNames(SelectedGoalNames);
 		if (SelectedGoalNames.IsEmpty())
 		{
 			return false; // nothing selected to manipulate
diff --git a/Engine/Plugins/Animation/IKRig/Source/IKRigEditor/Private/RigEditor/IKRigEditorController.cpp b/Engine/Plugins/Animation/IKRig/Source/IKRigEditor/Private/RigEditor/IKRigEditorController.cpp
--- a/Engine/Plugins/Animation/IKRig/Source/IKRigEditor/Private/RigEditor/IKRigEditorController.cpp
+++ b/Engine/Plugins/Animation/IKRig/Source/IKRigEditor/Private/RigEditor/IKRigEditorController.cpp
@@ -217,10 +217,33 @@ void FIKRigEditorController::Reset() const
 
 void FIKRigEditorController::RefreshAllViews() const
 {
-	SolverStackView->RefreshStackView();
-	SkeletonView->RefreshTreeView();
-	DetailsView->ForceRefresh();
-	RetargetingView->RefreshView();
+	if (SolverStackView.IsValid())
+	{
+		SolverStackView->RefreshStackView();
+	}
+
+	if (SkeletonView.IsValid())
+	{
+		SkeletonView->RefreshTreeView();
+	}
+
+	if (DetailsView.IsValid())
+	{
+		DetailsView->ForceRefresh();
+	}
+
+	if (RetargetingView.IsValid())
+	{
+		RetargetingView->RefreshView();
+	}
+}
+
+void FIKRigEditorController::RefreshTreeView() const
+{
+	if (SkeletonView.IsValid())
+	{
+		SkeletonView->RefreshTreeView();
+	}
 }
 
 void FIKRigEditorController::AddNewGoals(const TArray<FName>& GoalNames, const TArray<FName>& BoneNames)
@@ -275,30 +298,50 @@ void FIKRigEditorController::AddNewGoals(const TArray<FName>& GoalNames, const T
 
 void FIKRigEditorController::ClearSelection()
 {
-	SkeletonView->TreeView->ClearSelection();
+	if (SkeletonView.IsValid())
+	{
+		SkeletonView->TreeView->ClearSelection();	
+	}
+	
 	ShowEmptyDetails();
 }
 
 void FIKRigEditorController::HandleGoalSelectedInViewport(const FName& GoalName, bool bReplace)
-{	
-	SkeletonView->AddSelectedItemFromViewport(GoalName, IKRigTreeElementType::GOAL, bReplace);
+{
+	if (SkeletonView.IsValid())
+	{
+		SkeletonView->AddSelectedItemFromViewport(GoalName, IKRigTreeElementType::GOAL, bReplace);
+	}
+	
 	ShowDetailsForGoal(GoalName);
 }
 
 void FIKRigEditorController::HandleBoneSelectedInViewport(const FName& BoneName, bool bReplace)
 {
-	SkeletonView->AddSelectedItemFromViewport(BoneName, IKRigTreeElementType::BONE, bReplace);
+	if (SkeletonView.IsValid())
+	{
+		SkeletonView->AddSelectedItemFromViewport(BoneName, IKRigTreeElementType::BONE, bReplace);
+	}
+	
 	ShowDetailsForBone(BoneName);
 }
 
 void FIKRigEditorController::GetSelectedSolvers(TArray<TSharedPtr<FSolverStackElement>>& OutSelectedSolvers)
 {
-	OutSelectedSolvers.Reset();
-	OutSelectedSolvers.Append(SolverStackView->ListView->GetSelectedItems());
+	if (SolverStackView.IsValid())
+	{
+		OutSelectedSolvers.Reset();
+		OutSelectedSolvers.Append(SolverStackView->ListView->GetSelectedItems());
+	}
 }
 
 int32 FIKRigEditorController::GetSelectedSolverIndex()
 {
+	if (!SolverStackView.IsValid())
+	{
+		return INDEX_NONE;
+	}
+	
 	TArray<TSharedPtr<FSolverStackElement>> SelectedSolvers = SolverStackView->ListView->GetSelectedItems();
 	if (SelectedSolvers.IsEmpty())
 	{
@@ -308,6 +351,75 @@ int32 FIKRigEditorController::GetSelectedSolverIndex()
 	return SelectedSolvers[0]->IndexInStack;
 }
 
+void FIKRigEditorController::GetSelectedGoalNames(TArray<FName>& OutGoalNames) const
+{
+	if (!SkeletonView.IsValid())
+	{
+		return;
+	}
+
+	SkeletonView->GetSelectedGoalNames(OutGoalNames);
+}
+
+int32 FIKRigEditorController::GetNumSelectedGoals() const
+{
+	if (!SkeletonView.IsValid())
+	{
+		return 0;
+	}
+
+	return SkeletonView->GetNumSelectedGoals();
+}
+
+void FIKRigEditorController::GetSelectedBoneNames(TArray<FName>& OutBoneNames) const
+{
+	if (!SkeletonView.IsValid())
+	{
+		return;
+	}
+
+	SkeletonView->GetSelectedBoneNames(OutBoneNames);
+}
+
+void FIKRigEditorController::GetSelectedBones(TArray<TSharedPtr<FIKRigTreeElement>>& OutBoneItems) const
+{
+	if (!SkeletonView.IsValid())
+	{
+		return;
+	}
+
+	SkeletonView->GetSelectedBones(OutBoneItems);
+}
+
+bool FIKRigEditorController::IsGoalSelected(const FName& GoalName) const
+{
+	if (!SkeletonView.IsValid())
+	{
+		return false;
+	}
+
+	return SkeletonView->IsGoalSelected(GoalName);
+}
+
+FName FIKRigEditorController::GetSelectedChain() const
+{
+	if (!RetargetingView.IsValid())
+	{
+		return NAME_None;
+	}
+
+	return RetargetingView->GetSelectedChain();
+}
+
+bool FIKRigEditorController::DoesSkeletonHaveSelectedItems() const
+{
+	if (!SkeletonView.IsValid())
+	{
+		return false;
+	}
+	return SkeletonView->HasSelectedItems();
+}
+
 bool FIKRigEditorController::PromptToAddSolver() const
 {
 	if (AssetController->GetNumSolvers() > 0)
@@ -334,7 +446,7 @@ bool FIKRigEditorController::PromptToAddSolver() const
 		return false; // cancel button pressed, or window closed
 	}
 
-	if (Settings.SolverType != nullptr)
+	if (Settings.SolverType != nullptr && SolverStackView.IsValid())
 	{
 		SolverStackView->AddNewSolver(Settings.SolverType);
 	}
diff --git a/Engine/Plugins/Animation/IKRig/Source/IKRigEditor/Private/RigEditor/IKRigToolkit.cpp b/Engine/Plugins/Animation/IKRig/Source/IKRigEditor/Private/RigEditor/IKRigToolkit.cpp
--- a/Engine/Plugins/Animation/IKRig/Source/IKRigEditor/Private/RigEditor/IKRigToolkit.cpp
+++ b/Engine/Plugins/Animation/IKRig/Source/IKRigEditor/Private/RigEditor/IKRigToolkit.cpp
@@ -250,18 +250,12 @@ void FIKRigEditorToolkit::HandlePreviewMeshChanged(USkeletalMesh* InOldSkeletalM
 	EditorController->RefreshAllViews();
 }
 
-void FIKRigEditorToolkit::HandleDetailsCreated(const TSharedRef<class IDetailsView>& InDetailsView)
+void FIKRigEditorToolkit::HandleDetailsCreated(const TSharedRef<class IDetailsView>& InDetailsView) const
 {
-	EditorController->DetailsView = InDetailsView;
-	EditorController->DetailsView->OnFinishedChangingProperties().AddSP(this, &FIKRigEditorToolkit::OnFinishedChangingDetails);
+	EditorController->SetDetailsView(InDetailsView);
 	EditorController->ShowEmptyDetails();
 }
 
-void FIKRigEditorToolkit::OnFinishedChangingDetails(
-    const FPropertyChangedEvent& PropertyChangedEvent)
-{
-}
-
 void FIKRigEditorToolkit::HandleReset()
 {
 	EditorController.Get().Reset();
diff --git a/Engine/Plugins/Animation/IKRig/Source/IKRigEditor/Private/RigEditor/SIKRigRetargetChainList.cpp b/Engine/Plugins/Animation/IKRig/Source/IKRigEditor/Private/RigEditor/SIKRigRetargetChainList.cpp
--- a/Engine/Plugins/Animation/IKRig/Source/IKRigEditor/Private/RigEditor/SIKRigRetargetChainList.cpp
+++ b/Engine/Plugins/Animation/IKRig/Source/IKRigEditor/Private/RigEditor/SIKRigRetargetChainList.cpp
@@ -277,7 +277,7 @@ const FReferenceSkeleton& SIKRigRetargetChainRow::GetReferenceSkeleton() const
 void SIKRigRetargetChainList::Construct(const FArguments& InArgs, TSharedRef<FIKRigEditorController> InEditorController)
 {
 	EditorController = InEditorController;
-	EditorController.Pin()->RetargetingView = SharedThis(this);
+	EditorController.Pin()->SetRetargetingView(SharedThis(this));
 	
 	CommandList = MakeShared<FUICommandList>();
 
diff --git a/Engine/Plugins/Animation/IKRig/Source/IKRigEditor/Private/RigEditor/SIKRigSkeleton.cpp b/Engine/Plugins/Animation/IKRig/Source/IKRigEditor/Private/RigEditor/SIKRigSkeleton.cpp
--- a/Engine/Plugins/Animation/IKRig/Source/IKRigEditor/Private/RigEditor/SIKRigSkeleton.cpp
+++ b/Engine/Plugins/Animation/IKRig/Source/IKRigEditor/Private/RigEditor/SIKRigSkeleton.cpp
@@ -51,6 +51,7 @@ void SIKRigSkeletonItem::Construct(
 {
 	WeakRigTreeElement = InRigTreeElement;
 	EditorController = InEditorController;
+	SkeletonView = InSkeleton;
 
 	// is this element affected by the selected solver?
 	bool bIsConnectedToSelectedSolver;
@@ -198,7 +199,7 @@ void SIKRigSkeletonItem::OnNameCommitted(const FText& InText, ETextCommit::Type
 	}
 	
 	Controller->RefreshAllViews();
-	Controller->SkeletonView->ReplaceItemInSelection(OldText, WeakRigTreeElement.Pin()->Key);
+	SkeletonView.Pin()->ReplaceItemInSelection(OldText, WeakRigTreeElement.Pin()->Key);
 }
 
 FText SIKRigSkeletonItem::GetName() const
@@ -228,7 +229,7 @@ TSharedPtr<SWidget> FIKRigSkeletonDragDropOp::GetDefaultDecorator() const
 void SIKRigSkeleton::Construct(const FArguments& InArgs, TSharedRef<FIKRigEditorController> InEditorController)
 {
 	EditorController = InEditorController;
-	EditorController.Pin()->SkeletonView = SharedThis(this);
+	EditorController.Pin()->SetSkeletonsView(SharedThis(this));
 	CommandList = MakeShared<FUICommandList>();
 	BindCommands();
 	
diff --git a/Engine/Plugins/Animation/IKRig/Source/IKRigEditor/Private/RigEditor/SIKRigSolverStack.cpp b/Engine/Plugins/Animation/IKRig/Source/IKRigEditor/Private/RigEditor/SIKRigSolverStack.cpp
--- a/Engine/Plugins/Animation/IKRig/Source/IKRigEditor/Private/RigEditor/SIKRigSolverStack.cpp
+++ b/Engine/Plugins/Animation/IKRig/Source/IKRigEditor/Private/RigEditor/SIKRigSolverStack.cpp
@@ -76,7 +76,10 @@ void SIKRigSolverStackItem::Construct(
 						bool bEnabled = true;
 						if (InSolverStack.IsValid() && InSolverStack->EditorController.IsValid())
 						{
-							bEnabled = InSolverStack->EditorController.Pin()->AssetController->GetSolver(InStackElement->IndexInStack)->IsEnabled();
+							if (const UIKRigSolver* Solver = InSolverStack->EditorController.Pin()->AssetController->GetSolver(InStackElement->IndexInStack))
+							{
+								bEnabled = Solver->IsEnabled();
+							}
 						}
 						return bEnabled ? ECheckBoxState::Checked : ECheckBoxState::Unchecked;
 					})
@@ -214,7 +217,7 @@ SIKRigSolverStack::~SIKRigSolverStack()
 void SIKRigSolverStack::Construct(const FArguments& InArgs, TSharedRef<FIKRigEditorController> InEditorController)
 {
 	EditorController = InEditorController;
-	EditorController.Pin()->SolverStackView = SharedThis(this);
+	EditorController.Pin()->SetSolverStackView(SharedThis(this));
 	
 	CommandList = MakeShared<FUICommandList>();
 
@@ -338,7 +341,7 @@ void SIKRigSolverStack::AddNewSolver(UClass* Class)
 	const int32 NewSolverIndex = AssetController->AddSolver(Class);
 	// update stack view
 	RefreshStackView();
-	Controller->SkeletonView->RefreshTreeView(); // update solver indices in effector items
+	Controller->RefreshTreeView(); // updates solver indices in tree items
 	// select it
 	ListView->SetSelection(ListViewItems[NewSolverIndex]);
 	// show details for it
@@ -358,10 +361,10 @@ void SIKRigSolverStack::DeleteSolver(TSharedPtr<FSolverStackElement> SolverToDel
 		return;
 	}
 
-	UIKRigController* AssetController = Controller->AssetController;
+	const UIKRigController* AssetController = Controller->AssetController;
 	AssetController->RemoveSolver(SolverToDelete->IndexInStack);
 	RefreshStackView();
-	Controller->SkeletonView->RefreshTreeView(); // update solver indices in effector items
+	Controller->RefreshTreeView();
 }
 
 void SIKRigSolverStack::RefreshStackView()
@@ -449,7 +452,7 @@ void SIKRigSolverStack::OnItemClicked(TSharedPtr<FSolverStackElement> InItem)
 	EditorController.Pin()->SetLastSelectedType(EIKRigSelectionType::SolverStack);
 }
 
-void SIKRigSolverStack::ShowDetailsForItem(TSharedPtr<FSolverStackElement> InItem) const
+void SIKRigSolverStack::ShowDetailsForItem(TSharedPtr<FSolverStackElement> InItem)
 {
 	const TSharedPtr<FIKRigEditorController> Controller = EditorController.Pin();
 	if (!Controller.IsValid())
@@ -458,13 +461,13 @@ void SIKRigSolverStack::ShowDetailsForItem(TSharedPtr<FSolverStackElement> InIte
 	}
 
 	// update bones greyed out when not affected
-	Controller->SkeletonView->RefreshTreeView();
+	Controller->RefreshTreeView();
 
 	// the solver selection must be done after we rebuilt the skeleton tree has it keeps the selection
 	if (!InItem.IsValid())
 	{
 		// clear the details panel only if there's nothing selected in the skeleton view
-		if (Controller->SkeletonView->HasSelectedItems())
+		if (Controller->DoesSkeletonHaveSelectedItems())
 		{
 			return;
 		}
@@ -516,7 +519,7 @@ FReply SIKRigSolverStack::OnAcceptDrop(
 	if (bWasReparented)
 	{
 		RefreshStackView();
-		Controller->SkeletonView->RefreshTreeView(); // update solver indices in effector items
+		Controller->RefreshTreeView(); // update solver indices in effector items
 	}
 	
 	return FReply::Handled();
diff --git a/Engine/Plugins/Animation/IKRig/Source/IKRigEditor/Public/RigEditor/IKRigEditorController.h b/Engine/Plugins/Animation/IKRig/Source/IKRigEditor/Public/RigEditor/IKRigEditorController.h
--- a/Engine/Plugins/Animation/IKRig/Source/IKRigEditor/Public/RigEditor/IKRigEditorController.h
+++ b/Engine/Plugins/Animation/IKRig/Source/IKRigEditor/Public/RigEditor/IKRigEditorController.h
@@ -94,11 +94,28 @@ public:
 	void Reset() const;
 	/** refresh all views */
 	void RefreshAllViews() const;
+	/** refresh just the skeleton tree view */
+	void RefreshTreeView() const;
 
 	/** return list of those solvers in the stack that are selected by user */
 	void GetSelectedSolvers(TArray<TSharedPtr<FSolverStackElement> >& OutSelectedSolvers);
 	/** get index of the first selected solver, return INDEX_None if nothing selected */
 	int32 GetSelectedSolverIndex();
+	/** get names of all goals that are selected */
+	void GetSelectedGoalNames(TArray<FName>& OutGoalNames) const;
+	/** return the number of selected goals */
+	int32 GetNumSelectedGoals() const;
+	/** get names of all bones that are selected */
+	void GetSelectedBoneNames(TArray<FName>& OutBoneNames) const;
+	/** get all bones that are selected */
+	void GetSelectedBones(TArray<TSharedPtr<FIKRigTreeElement>>& OutBoneItems) const;
+	/** returns true if Goal is currently selected */
+	bool IsGoalSelected(const FName& GoalName) const;
+	/** get name of the selected retargeting chain */
+	FName GetSelectedChain() const;
+	/** is anything selected in the skeleton view? */
+	bool DoesSkeletonHaveSelectedItems() const;
+	
 	/** right after importing a skeleton, we ask user what solver they want to use */
 	bool PromptToAddSolver() const;
 	/** determine if the element is connected to the selected solver */
@@ -122,6 +139,15 @@ public:
 	void ShowEmptyDetails();
 	/** show selected items in details view */
 	void ShowDetailsForElements(const TArray<TSharedPtr<FIKRigTreeElement>>& InItems);
+	
+	/** set details tab view */
+	void SetDetailsView(const TSharedPtr<class IDetailsView>& InDetailsView){ DetailsView = InDetailsView; };
+	/** set skeleton tab view */
+	void SetSkeletonsView(const TSharedPtr<SIKRigSkeleton>& InSkeletonView){ SkeletonView = InSkeletonView; };
+	/** set solver stack tab view */
+	void SetSolverStackView(const TSharedPtr<SIKRigSolverStack>& InSolverStackView){ SolverStackView = InSolverStackView; };
+	/** set retargeting tab view */
+	void SetRetargetingView(const TSharedPtr<SIKRigRetargetChainList>& InRetargetingView){ RetargetingView = InRetargetingView; };
 
 	/** create a new retarget chain */
 	void AddNewRetargetChain(const FName ChainName, const FName StartBone, const FName EndBone);
@@ -139,18 +165,6 @@ public:
 	UPROPERTY(transient, NonTransactional)
 	TWeakObjectPtr<class UIKRigAnimInstance> AnimInstance;
 
-	/** asset properties tab */
-	TSharedPtr<class IDetailsView> DetailsView;
-
-	/** the skeleton tree view */
-	TSharedPtr<SIKRigSkeleton> SkeletonView;
-	
-	/** the solver stack view */
-	TSharedPtr<SIKRigSolverStack> SolverStackView;
-
-	/** the solver stack view */
-	TSharedPtr<SIKRigRetargetChainList> RetargetingView;
-
 	/** the persona toolkit */
 	TWeakPtr<FIKRigEditorToolkit> EditorToolkit;
 
@@ -175,6 +189,18 @@ public:
 
 private:
 
+	/** asset properties tab */
+	TSharedPtr<class IDetailsView> DetailsView;
+
+	/** the skeleton tree view */
+	TSharedPtr<SIKRigSkeleton> SkeletonView;
+	
+	/** the solver stack view */
+	TSharedPtr<SIKRigSolverStack> SolverStackView;
+
+	/** the solver stack view */
+	TSharedPtr<SIKRigRetargetChainList> RetargetingView;
+
 	/** Initializes editor's solvers instances */
 	void InitializeSolvers() const;
 
diff --git a/Engine/Plugins/Animation/IKRig/Source/IKRigEditor/Public/RigEditor/IKRigToolkit.h b/Engine/Plugins/Animation/IKRig/Source/IKRigEditor/Public/RigEditor/IKRigToolkit.h
--- a/Engine/Plugins/Animation/IKRig/Source/IKRigEditor/Public/RigEditor/IKRigToolkit.h
+++ b/Engine/Plugins/Animation/IKRig/Source/IKRigEditor/Public/RigEditor/IKRigToolkit.h
@@ -96,8 +96,7 @@ private:
 	/** preview scene setup */
 	void HandlePreviewSceneCreated(const TSharedRef<IPersonaPreviewScene>& InPersonaPreviewScene);
 	void HandlePreviewMeshChanged(USkeletalMesh* InOldSkeletalMesh, USkeletalMesh* InNewSkeletalMesh);
-	void HandleDetailsCreated(const TSharedRef<class IDetailsView>& InDetailsView);
-	void OnFinishedChangingDetails(const FPropertyChangedEvent& PropertyChangedEvent);
+	void HandleDetailsCreated(const TSharedRef<class IDetailsView>& InDetailsView) const;
 	void SetupAnimInstance();
 	/** END preview scene setup */
 	
diff --git a/Engine/Plugins/Animation/IKRig/Source/IKRigEditor/Public/RigEditor/SIKRigSkeleton.h b/Engine/Plugins/Animation/IKRig/Source/IKRigEditor/Public/RigEditor/SIKRigSkeleton.h
--- a/Engine/Plugins/Animation/IKRig/Source/IKRigEditor/Public/RigEditor/SIKRigSkeleton.h
+++ b/Engine/Plugins/Animation/IKRig/Source/IKRigEditor/Public/RigEditor/SIKRigSkeleton.h
@@ -76,6 +76,7 @@ private:
 
 	TWeakPtr<FIKRigTreeElement> WeakRigTreeElement;
 	TWeakPtr<FIKRigEditorController> EditorController;
+	TWeakPtr<SIKRigSkeleton> SkeletonView;
 };
 
 class FIKRigSkeletonDragDropOp : public FDecoratedDragDropOp
diff --git a/Engine/Plugins/Animation/IKRig/Source/IKRigEditor/Public/RigEditor/SIKRigSolverStack.h b/Engine/Plugins/Animation/IKRig/Source/IKRigEditor/Public/RigEditor/SIKRigSolverStack.h
--- a/Engine/Plugins/Animation/IKRig/Source/IKRigEditor/Public/RigEditor/SIKRigSolverStack.h
+++ b/Engine/Plugins/Animation/IKRig/Source/IKRigEditor/Public/RigEditor/SIKRigSolverStack.h
@@ -112,7 +112,7 @@ private:
 	void OnSelectionChanged(TSharedPtr<FSolverStackElement> InItem, ESelectInfo::Type SelectInfo);
 	/** when a solver is clicked on in the stack view */
 	void OnItemClicked(TSharedPtr<FSolverStackElement> InItem);
-	void ShowDetailsForItem(TSharedPtr<FSolverStackElement> InItem) const; 
+	void ShowDetailsForItem(TSharedPtr<FSolverStackElement> InItem); 
 
 	/** list view generate row callback */
 	TSharedRef<ITableRow> MakeListRowWidget(TSharedPtr<FSolverStackElement> InElement, const TSharedRef<STableViewBase>& OwnerTable);
