fixed naming consistency

diff --git a/Engine/Plugins/Animation/IKRig/Source/IKRig/Public/Solvers/IKRig_LimbSolver.h b/Engine/Plugins/Animation/IKRig/Source/IKRig/Public/Solvers/IKRig_LimbSolver.h
--- a/Engine/Plugins/Animation/IKRig/Source/IKRig/Public/Solvers/IKRig_LimbSolver.h
+++ b/Engine/Plugins/Animation/IKRig/Source/IKRig/Public/Solvers/IKRig_LimbSolver.h
@@ -33,7 +33,7 @@ public:
 
 	UIKRig_LimbSolver();
 	
-	UPROPERTY(VisibleAnywhere, Category = "Limb IK Settings")
+	UPROPERTY(VisibleAnywhere, DisplayName = "Root Bone", Category = "Limb IK Settings")
 	FName RootName = NAME_None;
 
 	/** Precision (distance to the target) */
diff --git a/Engine/Plugins/Animation/IKRig/Source/IKRig/Public/Solvers/IKRig_PoleSolver.h b/Engine/Plugins/Animation/IKRig/Source/IKRig/Public/Solvers/IKRig_PoleSolver.h
--- a/Engine/Plugins/Animation/IKRig/Source/IKRig/Public/Solvers/IKRig_PoleSolver.h
+++ b/Engine/Plugins/Animation/IKRig/Source/IKRig/Public/Solvers/IKRig_PoleSolver.h
@@ -36,10 +36,10 @@ public:
 
 	UIKRig_PoleSolver();
 	
-	UPROPERTY(VisibleAnywhere, Category = "Pole Solver Settings")
+	UPROPERTY(VisibleAnywhere, DisplayName = "Root Bone", Category = "Pole Solver Settings")
 	FName RootName = NAME_None;
 
-	UPROPERTY(VisibleAnywhere, Category = "Pole Solver Settings")
+	UPROPERTY(VisibleAnywhere, DisplayName = "End Bone", Category = "Pole Solver Settings")
 	FName EndName = NAME_None;
 	
 	/** UIKRigSolver interface */
