IK Rig: shows nice names in solver dialog instead of class names.

diff --git a/Engine/Plugins/Animation/IKRig/Source/IKRigEditor/Private/RigEditor/IKRigEditorController.cpp b/Engine/Plugins/Animation/IKRig/Source/IKRigEditor/Private/RigEditor/IKRigEditorController.cpp
--- a/Engine/Plugins/Animation/IKRig/Source/IKRigEditor/Private/RigEditor/IKRigEditorController.cpp
+++ b/Engine/Plugins/Animation/IKRig/Source/IKRigEditor/Private/RigEditor/IKRigEditorController.cpp
@@ -427,15 +427,54 @@ bool FIKRigEditorController::PromptToAddSolver() const
 		return true;
 	}
 
-	// prompt user to add a default solver
-	FIKRigAddFirstSolverSettings Settings;
-	const TSharedPtr<FStructOnScope> StructToDisplay = MakeShareable(new FStructOnScope(FIKRigAddFirstSolverSettings::StaticStruct(), (uint8*)&Settings));
-	TSharedRef<SKismetInspector> KismetInspector = SNew(SKismetInspector);
-	KismetInspector->ShowSingleStruct(StructToDisplay);
+	TArray<TSharedPtr<FIKRigSolverTypeAndName>> SolverTypes;
+	for (TObjectIterator<UClass> ClassIt; ClassIt; ++ClassIt)
+	{
+		UClass* Class = *ClassIt;
+		if (!Class->IsNative())
+		{
+			continue;
+		}
+
+		if (!ClassIt->IsChildOf(UIKRigSolver::StaticClass()))
+		{
+			continue;
+		}
+
+		if (Class == UIKRigSolver::StaticClass())
+		{
+			continue; // skip base class
+		}
+
+		const UIKRigSolver* SolverCDO = CastChecked<UIKRigSolver>(Class->ClassDefaultObject);
+		TSharedPtr<FIKRigSolverTypeAndName> SolverType = MakeShared<FIKRigSolverTypeAndName>();
+		SolverType->NiceName = SolverCDO->GetNiceName();
+		SolverType->SolverType = TSubclassOf<UIKRigSolver>(Class);
+		SolverTypes.Add(SolverType);
+	}
+
+	TSharedPtr<FIKRigSolverTypeAndName> SelectedSolver = SolverTypes[0];
+	TSharedRef<SComboBox<TSharedPtr<FIKRigSolverTypeAndName>>> SolverOptionBox = SNew(SComboBox<TSharedPtr<FIKRigSolverTypeAndName>>)
+	.OptionsSource(&SolverTypes)
+	.OnGenerateWidget_Lambda([](TSharedPtr<FIKRigSolverTypeAndName> Item)
+	{
+		return SNew(STextBlock).Text(Item->NiceName);
+	})
+	.OnSelectionChanged_Lambda([&SelectedSolver](TSharedPtr<FIKRigSolverTypeAndName> Item, ESelectInfo::Type)
+	{
+		SelectedSolver = Item;
+	})
+	.Content()
+	[
+		SNew(STextBlock).Text_Lambda([&SelectedSolver]()
+			{
+			return SelectedSolver->NiceName;
+		})
+	];
 	
 	TSharedRef<SCustomDialog> AddSolverDialog = SNew(SCustomDialog)
 		.Title(FText(LOCTEXT("EditorController_IKRigFirstSolver", "Add Default Solver")))
-		.DialogContent(KismetInspector)
+		.DialogContent(SolverOptionBox)
 		.Buttons({
 			SCustomDialog::FButton(LOCTEXT("OK", "OK")),
 			SCustomDialog::FButton(LOCTEXT("Cancel", "Cancel"))
@@ -446,9 +485,9 @@ bool FIKRigEditorController::PromptToAddSolver() const
 		return false; // cancel button pressed, or window closed
 	}
 
-	if (Settings.SolverType != nullptr && SolverStackView.IsValid())
+	if (SelectedSolver->SolverType != nullptr && SolverStackView.IsValid())
 	{
-		SolverStackView->AddNewSolver(Settings.SolverType);
+		SolverStackView->AddNewSolver(SelectedSolver->SolverType);
 	}
 
 	return true;
diff --git a/Engine/Plugins/Animation/IKRig/Source/IKRigEditor/Public/RigEditor/IKRigEditorController.h b/Engine/Plugins/Animation/IKRig/Source/IKRigEditor/Public/RigEditor/IKRigEditorController.h
--- a/Engine/Plugins/Animation/IKRig/Source/IKRigEditor/Public/RigEditor/IKRigEditorController.h
+++ b/Engine/Plugins/Animation/IKRig/Source/IKRigEditor/Public/RigEditor/IKRigEditorController.h
@@ -116,8 +116,6 @@ public:
 	/** is anything selected in the skeleton view? */
 	bool DoesSkeletonHaveSelectedItems() const;
 	
-	/** right after importing a skeleton, we ask user what solver they want to use */
-	bool PromptToAddSolver() const;
 	/** determine if the element is connected to the selected solver */
 	bool IsElementConnectedToSolver(TSharedRef<FIKRigTreeElement> TreeElement, int32 SolverIndex);
 	/** determine if the element is connected to ANY solver */
@@ -188,6 +186,11 @@ public:
 	/** END FGCObject interface */
 
 private:
+	/** right after importing a skeleton, we ask user what solver they want to use */
+	bool PromptToAddSolver() const;
+
+	/** Initializes editor's solvers instances */
+	void InitializeSolvers() const;
 
 	/** asset properties tab */
 	TSharedPtr<class IDetailsView> DetailsView;
@@ -201,25 +204,14 @@ private:
 	/** the solver stack view */
 	TSharedPtr<SIKRigRetargetChainList> RetargetingView;
 
-	/** Initializes editor's solvers instances */
-	void InitializeSolvers() const;
-
 	EIKRigSelectionType LastSelectedType;
 
 	UPROPERTY()
 	TObjectPtr<UIKRigBoneDetails> BoneDetails;
 };
 
-/** only used for pop-up window to selected a first solver to add*/
-USTRUCT()
-struct FIKRigAddFirstSolverSettings
+struct FIKRigSolverTypeAndName
 {
-	GENERATED_BODY()
-
-	FIKRigAddFirstSolverSettings()
-        : SolverType(nullptr)
-	{}
-
-	UPROPERTY(EditAnywhere, Category = "Add Solver")
+	FText NiceName;
 	TSubclassOf<UIKRigSolver> SolverType;
 };
