fixed ill-formed pin's default value

diff --git a/Engine/Plugins/Animation/IKRig/Source/IKRigDeveloper/Private/AnimGraphNode_IKRig.cpp b/Engine/Plugins/Animation/IKRig/Source/IKRigDeveloper/Private/AnimGraphNode_IKRig.cpp
--- a/Engine/Plugins/Animation/IKRig/Source/IKRigDeveloper/Private/AnimGraphNode_IKRig.cpp
+++ b/Engine/Plugins/Animation/IKRig/Source/IKRigDeveloper/Private/AnimGraphNode_IKRig.cpp
@@ -437,18 +437,49 @@ void UAnimGraphNode_IKRig::CreateCustomPins(TArray<UEdGraphPin*>* InOldPins)
 
 void UAnimGraphNode_IKRig::SetPinDefaultValue(UEdGraphPin* InPin, const FName& InPropertyName)
 {
+	// default FIKRigGoal structure 
 	static FIKRigGoal DefaultGoal;
 	static const TSharedPtr<FStructOnScope> StructOnScope =
 							MakeShareable(new FStructOnScope(FIKRigGoal::StaticStruct(), (uint8*)&DefaultGoal));
-	const UAnimationGraphSchema* AnimGraphDefaultSchema = GetDefault<UAnimationGraphSchema>();
+
+	// NOTE: we bypass the ExportTextItem for FVector and FRotator structures resulting in invalid pin's default value
+	// and use custom functions. We can not use ToString() neither here since the vector/rotator pin control
+	// doesn't use the standard 'X=0,Y=0,Z=0' syntax. This is seen in several other places so this should factorized.
+	// We may not notice it's actually not working as the default value we set is generally a structure formed by 0s
+	// but it becomes obvious if trying with a Vector set to (1.0, 2.0, 3.0) for instance.
+	auto GetDefaultValue = [](const FProperty* InProperty)
+	{
+		const uint8* Memory = InProperty->ContainerPtrToValuePtr<uint8>(StructOnScope->GetStructMemory());
+		if (Memory)
+		{
+			if (const FStructProperty* StructProp = CastField<FStructProperty>(InProperty))
+			{
+				if (StructProp->Struct == TBaseStructure<FVector>::Get())
+				{
+					const FVector& Vector = *(FVector*)Memory;
+					return FString::Printf(TEXT("%3.3f,%3.3f,%3.3f"), Vector.X, Vector.Y, Vector.Z);
+				}
+			
+				if (StructProp->Struct == TBaseStructure<FRotator>::Get())
+				{
+					const FRotator& Rotator = *(FRotator*)Memory;
+					return FString::Printf(TEXT("%3.3f,%3.3f,%3.3f"), Rotator.Pitch, Rotator.Yaw, Rotator.Roll);
+				}
+			}
+		}
+
+		// fallback to default behaviour 
+		FString DefaultValue;
+		InProperty->ExportTextItem(DefaultValue, Memory, nullptr, nullptr, PPF_None);
+		return DefaultValue;
+	};
 	
-	// default value
+	// set pin's initial default value
+	const UAnimationGraphSchema* AnimGraphDefaultSchema = GetDefault<UAnimationGraphSchema>();
 	AnimGraphDefaultSchema->SetPinAutogeneratedDefaultValueBasedOnType(InPin);
 	if (const FProperty* Property = FIKRigGoal::StaticStruct()->FindPropertyByName(InPropertyName))
 	{
-		const uint8* Memory = Property->ContainerPtrToValuePtr<uint8>(StructOnScope->GetStructMemory());
-		FString DefaultValue;
-		Property->ExportTextItem(DefaultValue, Memory, nullptr, nullptr, PPF_None);
+		const FString DefaultValue = GetDefaultValue(Property);
 		if (!DefaultValue.IsEmpty())
 		{
 			AnimGraphDefaultSchema->TrySetDefaultValue(*InPin, DefaultValue);
