fixed TransformSource property widget not showing up when in actor component mode

diff --git a/Engine/Plugins/Animation/IKRig/Source/IKRigDeveloper/Private/AnimGraphNode_IKRig.cpp b/Engine/Plugins/Animation/IKRig/Source/IKRigDeveloper/Private/AnimGraphNode_IKRig.cpp
--- a/Engine/Plugins/Animation/IKRig/Source/IKRigDeveloper/Private/AnimGraphNode_IKRig.cpp
+++ b/Engine/Plugins/Animation/IKRig/Source/IKRigDeveloper/Private/AnimGraphNode_IKRig.cpp
@@ -21,13 +21,13 @@
 /////////////////////////////////////////////////////
 // FIKRigGoalLayout
 
-TSharedRef<SWidget> FIKRigGoalLayout::CreateManualValueWidget() const
+TSharedRef<SWidget> FIKRigGoalLayout::CreatePropertyWidget() const
 {
 	const TSharedPtr<IPropertyHandle> TransformSourceHandle = GoalPropHandle->GetChildHandle(GET_MEMBER_NAME_CHECKED(FIKRigGoal, TransformSource));
 
-	return SNew(SHorizontalBox)
-
 	// transform source combo box
+	return SNew(SHorizontalBox)
+	
 	+ SHorizontalBox::Slot()
 	.FillWidth(1.0f)
 	.HAlign(HAlign_Left)
@@ -39,7 +39,6 @@ TSharedRef<SWidget> FIKRigGoalLayout::CreateManualValueWidget() const
 TSharedRef<SWidget> FIKRigGoalLayout::CreateBoneValueWidget() const
 {
 	const TSharedPtr<IPropertyHandle> TransformSourceHandle = GoalPropHandle->GetChildHandle(GET_MEMBER_NAME_CHECKED(FIKRigGoal, TransformSource));
-	TSharedPtr<IPropertyHandle> SourceBoneHandle = GoalPropHandle->GetChildHandle(GET_MEMBER_NAME_CHECKED(FIKRigGoal, SourceBone));
 	
 	return SNew(SHorizontalBox)
 
@@ -67,17 +66,12 @@ TSharedRef<SWidget> FIKRigGoalLayout::CreateBoneValueWidget() const
 TSharedRef<SWidget> FIKRigGoalLayout::CreateValueWidget() const
 {
 	const EIKRigGoalTransformSource TransformSource = GetTransformSource();
-	if (TransformSource == EIKRigGoalTransformSource::Manual)
-	{
-		return CreateManualValueWidget();
-	}
-
 	if (TransformSource == EIKRigGoalTransformSource::Bone)
 	{
 		return CreateBoneValueWidget();
 	}
 
-	return SNullWidget::NullWidget;
+	return CreatePropertyWidget();
 }
 
 void FIKRigGoalLayout::GenerateHeaderRowContent(FDetailWidgetRow& InOutGoalRow)
diff --git a/Engine/Plugins/Animation/IKRig/Source/IKRigDeveloper/Public/AnimGraphNode_IKRig.h b/Engine/Plugins/Animation/IKRig/Source/IKRigDeveloper/Public/AnimGraphNode_IKRig.h
--- a/Engine/Plugins/Animation/IKRig/Source/IKRigDeveloper/Public/AnimGraphNode_IKRig.h
+++ b/Engine/Plugins/Animation/IKRig/Source/IKRigDeveloper/Public/AnimGraphNode_IKRig.h
@@ -46,7 +46,7 @@ private:
 	void OnBoneSelectionChanged(FName Name) const;
 	FName GetSelectedBone(bool& bMultipleValues) const;
 
-	TSharedRef<SWidget> CreateManualValueWidget() const;
+	TSharedRef<SWidget> CreatePropertyWidget() const;
 	TSharedRef<SWidget> CreateBoneValueWidget() const;
 	TSharedRef<SWidget> CreateValueWidget() const;
 	
