ensure we set the right retarget source and target skeletal mesh to the exported sequence

diff --git a/Engine/Plugins/Animation/IKRig/Source/IKRigEditor/Private/RetargetEditor/IKRetargetBatchOperation.cpp b/Engine/Plugins/Animation/IKRig/Source/IKRigEditor/Private/RetargetEditor/IKRetargetBatchOperation.cpp
--- a/Engine/Plugins/Animation/IKRig/Source/IKRigEditor/Private/RetargetEditor/IKRetargetBatchOperation.cpp
+++ b/Engine/Plugins/Animation/IKRig/Source/IKRigEditor/Private/RetargetEditor/IKRetargetBatchOperation.cpp
@@ -168,6 +168,56 @@ void FIKRetargetBatchOperation::RetargetAssets(
 	USkeleton* OldSkeleton = Context.SourceMesh->GetSkeleton();
 	USkeleton* NewSkeleton = Context.TargetMesh->GetSkeleton();
 
+	auto UpdateRetargetSource = [&Context](UAnimSequence* InAnimSequenceToRetarget)
+	{
+#if WITH_EDITORONLY_DATA
+
+		const bool bHasRetargetSource = InAnimSequenceToRetarget->RetargetSource != NAME_None;
+		const bool bHasRetargetSourceAsset = !InAnimSequenceToRetarget->RetargetSourceAsset.IsNull();
+		if (!bHasRetargetSource && !bHasRetargetSourceAsset)
+		{
+			return;
+		}
+		
+		USkeleton* NewSkeleton = Context.TargetMesh->GetSkeleton();
+		const FName TargetRetargetSource = NewSkeleton->GetRetargetSourceForMesh(Context.TargetMesh);
+
+		// NOTE: should we keep the source anim sequence RetargetSource / RetargetSourceAsset if we don't find any
+		// TargetRetargetSource or reset them to default values?
+		
+		// update RetargetSource
+		if (bHasRetargetSource)
+		{
+			InAnimSequenceToRetarget->RetargetSource = TargetRetargetSource;
+		}
+
+		// update RetargetSourceAsset
+		if (bHasRetargetSourceAsset)
+		{
+			const FReferencePose* PoseFound = NewSkeleton->AnimRetargetSources.Find(TargetRetargetSource);
+			if (PoseFound && PoseFound->SourceReferenceMesh)
+			{
+				InAnimSequenceToRetarget->RetargetSourceAsset = PoseFound->SourceReferenceMesh;
+
+				// make sure we update the RetargetSourceAssetReferencePose accordingly
+				const USkeletalMesh* NewReferenceMesh = InAnimSequenceToRetarget->RetargetSourceAsset.LoadSynchronous();
+				if (NewReferenceMesh)
+				{
+					FAnimationRuntime::MakeSkeletonRefPoseFromMesh(NewReferenceMesh, NewSkeleton, InAnimSequenceToRetarget->RetargetSourceAssetReferencePose);
+				}
+				else
+				{
+					InAnimSequenceToRetarget->RetargetSourceAssetReferencePose.Empty();
+				}
+				return;
+			}
+
+			InAnimSequenceToRetarget->RetargetSourceAsset.Reset();
+			InAnimSequenceToRetarget->RetargetSourceAssetReferencePose.Empty();
+		}
+#endif
+	};
+	
 	for (UAnimationAsset* AssetToRetarget : AnimationAssetsToRetarget)
 	{
 		// synchronize curves between old/new asset
@@ -179,11 +229,14 @@ void FIKRetargetBatchOperation::RetargetAssets(
 			// clear transform curves since those curves won't work in new skeleton
 			IAnimationDataController& Controller = AnimSequenceToRetarget->GetController();
 			Controller.RemoveAllCurvesOfType(ERawCurveTrackTypes::RCT_Transform);
+			// make sure we update the retarget source
+			UpdateRetargetSource(AnimSequenceToRetarget);
 		}
 
 		// replace references to other animation
 		AssetToRetarget->ReplaceReferredAnimations(RemappedAnimAssets);
 		AssetToRetarget->SetSkeleton(NewSkeleton);
+		AssetToRetarget->SetPreviewMesh(Context.TargetMesh);
 		AssetToRetarget->PostEditChange();
 		AssetToRetarget->MarkPackageDirty();
 	}
