do not reconstruct split struct's sub pins when laoding the asset

diff --git a/Engine/Plugins/Animation/IKRig/Source/IKRigDeveloper/Private/AnimGraphNode_IKRig.cpp b/Engine/Plugins/Animation/IKRig/Source/IKRigDeveloper/Private/AnimGraphNode_IKRig.cpp
--- a/Engine/Plugins/Animation/IKRig/Source/IKRigDeveloper/Private/AnimGraphNode_IKRig.cpp
+++ b/Engine/Plugins/Animation/IKRig/Source/IKRigDeveloper/Private/AnimGraphNode_IKRig.cpp
@@ -501,13 +501,19 @@ void UAnimGraphNode_IKRig::CreateCustomPinsFromUnloadedAsset(TArray<UEdGraphPin*
 	// ensure that this is a goal related pin
 	auto bNeedsCreation = [&](const UEdGraphPin* InOldPin)
 	{
-		// custom pins are inputs
+		// in our case, custom pins are all inputs
 		if (InOldPin->Direction != EEdGraphPinDirection::EGPD_Input)
 		{
 			return false;
 		}
 
-		// look for old pin's name-type into current pins
+		// do not rebuild sub pins as they will be treated after in UK2Node::RestoreSplitPins
+		if (InOldPin->ParentPin && InOldPin->ParentPin->bHidden)
+		{
+		 	return false;
+		}
+		
+		// look for old pin's name-type into current pins. At this stage, the default pins have already been created
 		const int32 PinIndex = Pins.IndexOfByPredicate([&](const UEdGraphPin* Pin)
 		{
 			return Pin->GetFName() == InOldPin->GetFName() && Pin->PinType == InOldPin->PinType;
