FROST-8496 - ABP_Player_v2.uasset: [Compiler] Can't connect pins: float is not compatible with Double Real.

A couple of BP reals changes were missed for IKRig and a new-ish function in RigVMTypeUtils.

For RigVMTypeUtils, we simply need to ensure that we're using the correct PC_Real category when dealing with float/double CPPTypes.

For IKRig, the alpha pin needs to be a real/double type. Additionally, when copying the property data to the FIKRigGoal instance, we need to treat the property's value as a double type. Since FIKRigGoal::Position et al. use the alpha as a normalized value, it likely doesn't need the precision of a double, so it's been left as is. As a result, we need to perform a narrowing conversion when copying over the data.

diff --git a/Engine/Plugins/Animation/IKRig/Source/IKRig/Private/AnimNodes/AnimNode_IKRig.cpp b/Engine/Plugins/Animation/IKRig/Source/IKRig/Private/AnimNodes/AnimNode_IKRig.cpp
--- a/Engine/Plugins/Animation/IKRig/Source/IKRig/Private/AnimNodes/AnimNode_IKRig.cpp
+++ b/Engine/Plugins/Animation/IKRig/Source/IKRig/Private/AnimNodes/AnimNode_IKRig.cpp
@@ -353,7 +353,7 @@ void FAnimNode_IKRig::InitializeProperties(const UObject* InSourceInstance, UCla
 				{
 					UpdateFunctions.Add([&Goal, SourceProperty](const UObject* InSourceInstance)
 					{
-						const float* AlphaValue = SourceProperty->ContainerPtrToValuePtr<float>(InSourceInstance);
+						const double* AlphaValue = SourceProperty->ContainerPtrToValuePtr<double>(InSourceInstance);
 						Goal.PositionAlpha = FMath::Clamp<float>(*AlphaValue, 0.f, 1.f);
 					});
 				}
@@ -361,7 +361,7 @@ void FAnimNode_IKRig::InitializeProperties(const UObject* InSourceInstance, UCla
 				{
 					UpdateFunctions.Add([&Goal, SourceProperty](const UObject* InSourceInstance)
 					{
-						const float* AlphaValue = SourceProperty->ContainerPtrToValuePtr<float>(InSourceInstance);
+						const double* AlphaValue = SourceProperty->ContainerPtrToValuePtr<double>(InSourceInstance);
 						Goal.RotationAlpha = FMath::Clamp<float>(*AlphaValue, 0.f, 1.f);
 					});
 				}
diff --git a/Engine/Plugins/Animation/IKRig/Source/IKRigDeveloper/Private/AnimGraphNode_IKRig.cpp b/Engine/Plugins/Animation/IKRig/Source/IKRigDeveloper/Private/AnimGraphNode_IKRig.cpp
--- a/Engine/Plugins/Animation/IKRig/Source/IKRigDeveloper/Private/AnimGraphNode_IKRig.cpp
+++ b/Engine/Plugins/Animation/IKRig/Source/IKRigDeveloper/Private/AnimGraphNode_IKRig.cpp
@@ -566,9 +566,16 @@ void UAnimGraphNode_IKRig::CreateCustomPinsFromValidAsset()
 	RotationPinType.PinSubCategoryObject = TBaseStructure<FRotator>::Get();
 	
 	// alpha property
+	static const FName PC_Real(TEXT("real"));
 	static const FName PC_Float(TEXT("float"));
+	static const FName PC_Double(TEXT("double"));
 	static FEdGraphPinType AlphaPinType;
+#if ENABLE_BLUEPRINT_REAL_NUMBERS
+	AlphaPinType.PinCategory = PC_Real;
+	AlphaPinType.PinSubCategory = PC_Double;
+#else
 	AlphaPinType.PinCategory = PC_Float;
+#endif
 
 	// create pins
 	const TArray<UIKRigEffectorGoal*>& AssetGoals = Node.RigDefinitionAsset->GetGoalArray();
diff --git a/Engine/Source/Developer/RigVMDeveloper/Private/RigVMTypeUtils.cpp b/Engine/Source/Developer/RigVMDeveloper/Private/RigVMTypeUtils.cpp
--- a/Engine/Source/Developer/RigVMDeveloper/Private/RigVMTypeUtils.cpp
+++ b/Engine/Source/Developer/RigVMDeveloper/Private/RigVMTypeUtils.cpp
@@ -300,6 +300,18 @@ FEdGraphPinType RigVMTypeUtils::PinTypeFromCPPType(const FString& InCPPType, UOb
 	{
 		PinType.PinCategory = UEdGraphSchema_K2::PC_Int;
 	}
+#if ENABLE_BLUEPRINT_REAL_NUMBERS
+	else if (BaseCPPType == TEXT("float"))
+	{
+		PinType.PinCategory = UEdGraphSchema_K2::PC_Real;
+		PinType.PinSubCategory = UEdGraphSchema_K2::PC_Float;
+	}
+	else if (BaseCPPType == TEXT("double"))
+	{
+		PinType.PinCategory = UEdGraphSchema_K2::PC_Real;
+		PinType.PinSubCategory = UEdGraphSchema_K2::PC_Double;
+	}
+#else
 	else if (BaseCPPType == TEXT("float"))
 	{
 		PinType.PinCategory = UEdGraphSchema_K2::PC_Float;
@@ -308,6 +320,7 @@ FEdGraphPinType RigVMTypeUtils::PinTypeFromCPPType(const FString& InCPPType, UOb
 	{
 		PinType.PinCategory = UEdGraphSchema_K2::PC_Double;
 	}
+#endif
 	else if (BaseCPPType == TEXT("FName"))
 	{
 		PinType.PinCategory = UEdGraphSchema_K2::PC_Name;
