IKRig: Fixed duplicated localization keys.

diff --git a/Engine/Plugins/Animation/IKRig/Source/IKRigEditor/Private/RetargetEditor/AssetTypeActions_IKRetargeter.cpp b/Engine/Plugins/Animation/IKRig/Source/IKRigEditor/Private/RetargetEditor/AssetTypeActions_IKRetargeter.cpp
--- a/Engine/Plugins/Animation/IKRig/Source/IKRigEditor/Private/RetargetEditor/AssetTypeActions_IKRetargeter.cpp
+++ b/Engine/Plugins/Animation/IKRig/Source/IKRigEditor/Private/RetargetEditor/AssetTypeActions_IKRetargeter.cpp
@@ -100,11 +100,10 @@ void FAssetTypeActions_IKRetargeter::CreateRetargetSubMenu(FToolMenuSection& InS
 	}
 
 	// change menu label if anim blueprint is selected
-	FText MenuLabel = FText(LOCTEXT("RetargetAnimation", "Duplicate and Retarget Animation Assets"));
-	if (Cast<UAnimBlueprint>(SelectedObjects[0]))
-	{
-		MenuLabel = FText(LOCTEXT("RetargetAnimation", "Duplicate and Retarget Animation Blueprint"));
-	}
+	static const FText AnimAssetLabel(LOCTEXT("RetargetAnimationAsset", "Duplicate and Retarget Animation Assets"));
+	static const FText AnimBlueprintLabel(LOCTEXT("RetargetAnimationBlueprint", "Duplicate and Retarget Animation Blueprint"));
+
+	const FText MenuLabel = Cast<UAnimBlueprint>(SelectedObjects[0]) == nullptr ? AnimAssetLabel : AnimBlueprintLabel;
 	
 	InSection.AddMenuEntry(
 		"IKRetargetToDifferentSkeleton",
