Added per-chain customization to the IK Retargeter to address foot sliding and other issues.

diff --git a/Engine/Plugins/Animation/IKRig/Source/IKRig/Private/AnimNodes/AnimNode_RetargetPoseFromMesh.cpp b/Engine/Plugins/Animation/IKRig/Source/IKRig/Private/AnimNodes/AnimNode_RetargetPoseFromMesh.cpp
--- a/Engine/Plugins/Animation/IKRig/Source/IKRig/Private/AnimNodes/AnimNode_RetargetPoseFromMesh.cpp
+++ b/Engine/Plugins/Animation/IKRig/Source/IKRig/Private/AnimNodes/AnimNode_RetargetPoseFromMesh.cpp
@@ -85,7 +85,7 @@ void FAnimNode_RetargetPoseFromMesh::Evaluate_AnyThread(FPoseContext& Output)
 	// NOTE: this copies goal targets as well, but these are overwritten by IK chain goals
 	if (bDriveTargetIKRigWithAsset)
 	{
-		Processor->CopyTargetIKRigSettingsFromAsset();
+		Processor->CopyAllSettingsFromAsset();
 	}
 #endif
 
diff --git a/Engine/Plugins/Animation/IKRig/Source/IKRig/Private/Retargeter/IKRetargetProcessor.cpp b/Engine/Plugins/Animation/IKRig/Source/IKRig/Private/Retargeter/IKRetargetProcessor.cpp
--- a/Engine/Plugins/Animation/IKRig/Source/IKRig/Private/Retargeter/IKRetargetProcessor.cpp
+++ b/Engine/Plugins/Animation/IKRig/Source/IKRig/Private/Retargeter/IKRetargetProcessor.cpp
@@ -360,6 +360,7 @@ void FChainEncoderFK::EncodePose(
 }
 
 void FChainDecoderFK::DecodePose(
+	const FRetargetChainSettings& Settings,
 	const TArray<int32>& TargetBoneIndices,
     const FChainEncoderFK& SourceChain,
     const FTargetSkeleton& TargetSkeleton,
@@ -377,6 +378,16 @@ void FChainDecoderFK::DecodePose(
 	// Otherwise the neck bones will remain at their location prior to the spine update.
 	UpdateIntermediateParents(TargetSkeleton,InOutGlobalPose);
 
+	// if FK retargeting has been disabled for this chain, then simply set it to the retarget pose
+	if (!Settings.CopyPoseUsingFK)
+	{
+		for (const int32& BoneIndex : TargetBoneIndices)
+		{
+			TargetSkeleton.UpdateGlobalTransformOfSingleBone(BoneIndex, TargetSkeleton.RetargetLocalPose, InOutGlobalPose);
+		}
+		return;
+	}
+
 	// now retarget the pose of each bone in the chain, copying from source to target
 	for (int32 ChainIndex=0; ChainIndex<TargetBoneIndices.Num(); ++ChainIndex)
 	{
@@ -406,7 +417,11 @@ void FChainDecoderFK::DecodePose(
 		const FQuat OutRotation = RotationDelta * TargetInitialRotation;
 
 		// calculate output POSITION (constant for now, maybe we support retargeting translation later?)
-		const FTransform& ParentGlobalTransform = InOutGlobalPose[ParentIndex];
+		FTransform ParentGlobalTransform = FTransform::Identity;
+		if (ParentIndex != INDEX_NONE)
+		{
+			ParentGlobalTransform = InOutGlobalPose[ParentIndex];
+		}
 		const FVector InitialLocalOffset = TargetSkeleton.RetargetLocalPose[BoneIndex].GetTranslation();
 		const FVector OutPosition = ParentGlobalTransform.TransformPosition(InitialLocalOffset);
 
@@ -516,6 +531,7 @@ bool FChainRetargeterIK::InitializeSource(
 	Source.BoneIndexC = BoneIndices.Last();
 	
 	const FTransform& End = SourceInitialGlobalPose[Source.BoneIndexC];
+	Source.PreviousEndPosition = End.GetTranslation();
 	Source.InitialEndPosition = End.GetTranslation();
 	Source.InitialEndRotation = End.GetRotation();
 
@@ -533,11 +549,11 @@ bool FChainRetargeterIK::InitializeSource(
 	return true;
 }
 
-void FChainRetargeterIK::EncodePose(const TArray<FTransform>& InputGlobalPose)
+void FChainRetargeterIK::EncodePose(const TArray<FTransform>& InSourceGlobalPose)
 {
-	const FVector A = InputGlobalPose[Source.BoneIndexA].GetTranslation();
+	const FVector A = InSourceGlobalPose[Source.BoneIndexA].GetTranslation();
 	//FVector B = InputGlobalPose[BoneIndexB].GetTranslation(); TODO use for pole vector 
-	const FVector C = InputGlobalPose[Source.BoneIndexC].GetTranslation();
+	const FVector C = InSourceGlobalPose[Source.BoneIndexC].GetTranslation();
 
     // get the normalized direction / length of the IK limb (how extended it is as percentage of original length)
     const FVector AC = C - A;
@@ -546,8 +562,10 @@ void FChainRetargeterIK::EncodePose(const TArray<FTransform>& InputGlobalPose)
 	AC.ToDirectionAndLength(ACDirection, ACLength);
 	const float NormalizedLimbLength = ACLength * Source.InvInitialLength;
 
+	Source.PreviousEndPosition = Source.CurrentEndPosition;
+	Source.CurrentEndPosition = C;
 	Source.CurrentEndDirectionNormalized = ACDirection * NormalizedLimbLength;
-	Source.CurrentEndRotation = InputGlobalPose[Source.BoneIndexC].GetRotation();
+	Source.CurrentEndRotation = InSourceGlobalPose[Source.BoneIndexC].GetRotation();
 	Source.CurrentHeightFromGroundNormalized = (C.Z - Source.InitialEndPosition.Z)  * Source.InvInitialLength;
 	Source.PoleVectorDirection = FVector::OneVector; // TBD
 }
@@ -563,7 +581,9 @@ bool FChainRetargeterIK::InitializeTarget(
 	}
 	
 	Target.BoneIndexA = BoneIndices[0];
+	Target.BoneIndexC = BoneIndices.Last();
 	const FTransform& Last = TargetInitialGlobalPose[BoneIndices.Last()];
+	Target.PrevEndPosition = Last.GetLocation();
 	Target.InitialEndPosition = Last.GetTranslation();
 	Target.InitialEndRotation = Last.GetRotation();
 	Target.InitialLength = (TargetInitialGlobalPose[Target.BoneIndexA].GetTranslation() - Last.GetTranslation()).Size();
@@ -578,19 +598,74 @@ bool FChainRetargeterIK::InitializeTarget(
 }
 	
 void FChainRetargeterIK::DecodePose(
+	const FRetargetChainSettings& Settings,
     const TArray<FTransform>& OutGlobalPose,
-    FDecodedIKChain& OutResults) const
+    FDecodedIKChain& OutResults)
 {
-	// calculate end effector position
-	const FVector Start = OutGlobalPose[Target.BoneIndexA].GetTranslation();
-	OutResults.EndEffectorPosition = Start + Source.CurrentEndDirectionNormalized * Target.InitialLength; // * LimbScale (todo)
-	// optionally factor in ground height
-	//const bool IsGrounded = false;
-	//if (IsGrounded) // todo expose this
-	//{
-	//	OutResults.EndEffectorPosition.Z = (EncodedChain.HeightFromGroundNormalized * Length) + EndPositionOrig.Z;
-	//}
+	FVector EndPosition;
+	if (Settings.DriveIKGoal)
+	{
+		// set position to length-scaled direction from source limb
+		const FVector Start = OutGlobalPose[Target.BoneIndexA].GetTranslation();
+		EndPosition = Start + (Source.CurrentEndDirectionNormalized * Target.InitialLength);
+
+		// blend to source location
+		if (Settings.BlendToSource > KINDA_SMALL_NUMBER)
+		{
+			const FVector Weight = Settings.BlendToSource * Settings.BlendToSourceWeights;
+			EndPosition.X = FMath::Lerp(EndPosition.X, Source.CurrentEndPosition.X, Weight.X);
+			EndPosition.Y = FMath::Lerp(EndPosition.Y, Source.CurrentEndPosition.Y, Weight.Y);
+			EndPosition.Z = FMath::Lerp(EndPosition.Z, Source.CurrentEndPosition.Z, Weight.Z);
+		}
 
+		// apply static offset
+		EndPosition += Settings.StaticOffset;
+
+		// apply extension
+		if (!FMath::IsNearlyEqual(Settings.Extension, 1.0f))
+		{
+			EndPosition = Start + (EndPosition - Start) * Settings.Extension;	
+		}
+
+		// match velocity
+		if (Settings.MatchSourceVelocity > KINDA_SMALL_NUMBER)
+		{
+			const FVector SourceVelocity = Source.CurrentEndPosition - Source.PreviousEndPosition;
+			const float SourceSpeed = SourceVelocity.Size();
+
+			const FVector TargetVelocity = EndPosition - Target.PrevEndPosition;
+			const float TargetSpeed = TargetVelocity.Size();
+
+			// if target is moving slowly enough, start matching velocity
+			if (TargetSpeed < Settings.TeleportVelocityThreshold)
+			{
+				// match target speed to source
+				const float BlendedSpeed = FMath::Lerp(TargetSpeed, SourceSpeed, Settings.MatchSourceVelocity);
+				EndPosition = Target.PrevEndPosition + TargetVelocity * BlendedSpeed;
+
+				// what would happen if we blended into and out of this state?
+				// TimeSinceTargetStartedMatchingVelocity
+				// WasMatchingVelocity
+			}
+		}
+		
+		// optionally factor in ground height
+		//const bool IsGrounded = false;
+		//if (IsGrounded) // todo expose this
+		//{
+		//	OutResults.EndEffectorPosition.Z = (EncodedChain.HeightFromGroundNormalized * Length) + EndPositionOrig.Z;
+		//}
+	}
+	else
+	{
+		// set goal location to the input coming from the previous retarget phase (FK if enabled)
+		EndPosition = OutGlobalPose[Target.BoneIndexC].GetTranslation();
+	}
+
+	// apply effector position
+	OutResults.EndEffectorPosition = EndPosition;
+	Target.PrevEndPosition = EndPosition;
+	
 	// calculate end effector rotation
 	const FQuat RotationDelta = Source.CurrentEndRotation * Source.InitialEndRotation.Inverse();
 	OutResults.EndEffectorRotation = RotationDelta * Target.InitialEndRotation;
@@ -600,6 +675,7 @@ void FChainRetargeterIK::DecodePose(
 }
 
 bool FRetargetChainPair::Initialize(
+	URetargetChainSettings* InSettings,
     const FBoneChain& SourceBoneChain,
     const FBoneChain& TargetBoneChain,
     const FRetargetSkeleton& SourceSkeleton,
@@ -623,7 +699,15 @@ bool FRetargetChainPair::Initialize(
 		return false;
     }
 
+	// ensure valid settings object
+	if (InSettings == nullptr)
+	{
+		UE_LOG(LogTemp, Warning, TEXT("IK Retargeter target bone chain, '%s', has null settings."), *TargetBoneChain.ChainName.ToString());
+		return false;
+	}
+
 	// store attributes of chain
+	Settings.CopySettingsFromAsset(InSettings);
 	SourceBoneChainName = SourceBoneChain.ChainName;
 	TargetBoneChainName = TargetBoneChain.ChainName;
 	
@@ -666,12 +750,13 @@ bool FRetargetChainPair::ValidateBoneChainWithSkeletalMesh(
 }
 
 bool FRetargetChainPairFK::Initialize(
+	URetargetChainSettings* InSettings,
 	const FBoneChain& SourceBoneChain,
 	const FBoneChain& TargetBoneChain,
 	const FRetargetSkeleton& SourceSkeleton,
 	const FTargetSkeleton& TargetSkeleton)
 {
-	const bool bChainInitialized = FRetargetChainPair::Initialize(SourceBoneChain, TargetBoneChain, SourceSkeleton, TargetSkeleton);
+	const bool bChainInitialized = FRetargetChainPair::Initialize(InSettings, SourceBoneChain, TargetBoneChain, SourceSkeleton, TargetSkeleton);
 	if (!bChainInitialized)
 	{
 		return false;
@@ -699,6 +784,7 @@ bool FRetargetChainPairFK::Initialize(
 }
 
 bool FRetargetChainPairIK::Initialize(
+	URetargetChainSettings* InSettings,
 	const FBoneChain& SourceBoneChain,
 	const FBoneChain& TargetBoneChain,
 	const FRetargetSkeleton& SourceSkeleton,
@@ -711,11 +797,11 @@ bool FRetargetChainPairIK::Initialize(
 		return false;
 	}
 
-	// store target IK goal
+	// store target IK goal name
 	IKGoalName = TargetBoneChain.IKGoalName;
 
 	// initialize bone chains
-	const bool bChainInitialized = FRetargetChainPair::Initialize(SourceBoneChain, TargetBoneChain, SourceSkeleton, TargetSkeleton);
+	const bool bChainInitialized = FRetargetChainPair::Initialize(InSettings, SourceBoneChain, TargetBoneChain, SourceSkeleton, TargetSkeleton);
 	if (!bChainInitialized)
 	{
 		return false;
@@ -937,43 +1023,41 @@ bool UIKRetargetProcessor::InitializeBoneChainPairs()
 	check(SourceIKRig && TargetIKRig);
 
 	// check that chains are available in both IKRig assets before sorting them based on StartBone index
-	const TArray<FRetargetChainMap>& ChainMapping = RetargeterAsset->GetChainMapping();	
-	for (const FRetargetChainMap& ChainMap : ChainMapping)
+	const TArray<TObjectPtr<URetargetChainSettings>>& ChainMapping = RetargeterAsset->GetAllChainSettings();	
+	for (URetargetChainSettings* ChainMap : ChainMapping)
 	{
 		// get target bone chain
-		const FBoneChain* TargetBoneChain = RetargeterAsset->GetTargetIKRig()->GetRetargetChainByName(ChainMap.TargetChain);
+		const FBoneChain* TargetBoneChain = RetargeterAsset->GetTargetIKRig()->GetRetargetChainByName(ChainMap->TargetChain);
 		if (!TargetBoneChain)
 		{
-			UE_LOG(LogTemp, Error, TEXT("IK Retargeter missing target bone chain: %s. Please update the mapping."), *ChainMap.TargetChain.ToString());
+			UE_LOG(LogTemp, Error, TEXT("IK Retargeter missing target bone chain: %s. Please update the mapping."), *ChainMap->TargetChain.ToString());
 			continue;
 		}
 		
 		// user opted to not map this to anything, we don't need to spam a warning about it
-		if (ChainMap.SourceChain == NAME_None)
+		if (ChainMap->SourceChain == NAME_None)
 		{
 			continue; 
 		}
 		
 		// get source bone chain
-		const FBoneChain* SourceBoneChain = RetargeterAsset->GetSourceIKRig()->GetRetargetChainByName(ChainMap.SourceChain);
+		const FBoneChain* SourceBoneChain = RetargeterAsset->GetSourceIKRig()->GetRetargetChainByName(ChainMap->SourceChain);
 		if (!SourceBoneChain)
 		{
-			UE_LOG(LogTemp, Error, TEXT("IK Retargeter missing source bone chain: %s"), *ChainMap.SourceChain.ToString());
+			UE_LOG(LogTemp, Error, TEXT("IK Retargeter missing source bone chain: %s"), *ChainMap->SourceChain.ToString());
 			continue;
 		}
 
 		// all chains are loaded as FK (giving IK better starting pose)
 		FRetargetChainPairFK ChainPair;
-		const bool bChainPairValid = ChainPair.Initialize(*SourceBoneChain, *TargetBoneChain, SourceSkeleton, TargetSkeleton);
-		if (bChainPairValid)
+		if (ChainPair.Initialize(ChainMap, *SourceBoneChain, *TargetBoneChain, SourceSkeleton, TargetSkeleton))
 		{
 			ChainPairsFK.Add(ChainPair);
 		}
 		
 		// load IK chain
 		FRetargetChainPairIK ChainPairIK;
-		const bool bChainPairIKValid = ChainPairIK.Initialize(*SourceBoneChain, *TargetBoneChain, SourceSkeleton, TargetSkeleton);
-		if (bChainPairIKValid)
+		if (ChainPairIK.Initialize(ChainMap, *SourceBoneChain, *TargetBoneChain, SourceSkeleton, TargetSkeleton))
 		{
 			ChainPairsIK.Add(ChainPairIK);
 		}
@@ -1112,13 +1196,13 @@ void UIKRetargetProcessor::RunFKRetarget(
     for (FRetargetChainPairFK& ChainPair : ChainPairsFK)
     {
     	ChainPair.FKEncoder.EncodePose(ChainPair.SourceBoneIndices, InGlobalTransforms);
-    	ChainPair.FKDecoder.DecodePose(ChainPair.TargetBoneIndices, ChainPair.FKEncoder, TargetSkeleton,OutGlobalTransforms);
+    	ChainPair.FKDecoder.DecodePose(ChainPair.Settings, ChainPair.TargetBoneIndices, ChainPair.FKEncoder, TargetSkeleton,OutGlobalTransforms);
     }
 }
 
 void UIKRetargetProcessor::RunIKRetarget(
-	const TArray<FTransform>& InGlobalPose,
-    TArray<FTransform>& OutGlobalPose)
+	const TArray<FTransform>& InSourceGlobalPose,
+    TArray<FTransform>& OutTargeGlobalPose)
 {
 	if (!IKRigProcessor->IsInitialized())
 	{
@@ -1134,10 +1218,10 @@ void UIKRetargetProcessor::RunIKRetarget(
 	for (FRetargetChainPairIK& ChainPair : ChainPairsIK)
 	{
 		// encode them all using the input pose
-		ChainPair.IKChainRetargeter.EncodePose(InGlobalPose);
+		ChainPair.IKChainRetargeter.EncodePose(InSourceGlobalPose);
 		// decode the IK goal and apply to IKRig
 		FDecodedIKChain OutIKGoal;
-		ChainPair.IKChainRetargeter.DecodePose(OutGlobalPose, OutIKGoal);
+		ChainPair.IKChainRetargeter.DecodePose(ChainPair.Settings, OutTargeGlobalPose, OutIKGoal);
 		// set the goal transform on the IK Rig
 		FIKRigGoal Goal = FIKRigGoal(
 			ChainPair.IKGoalName,
@@ -1151,11 +1235,11 @@ void UIKRetargetProcessor::RunIKRetarget(
 	}
 
 	// copy input pose to start IK solve from
-	IKRigProcessor->SetInputPoseGlobal(OutGlobalPose);
+	IKRigProcessor->SetInputPoseGlobal(OutTargeGlobalPose);
 	// run IK solve
 	IKRigProcessor->Solve();
 	// copy results of solve
-	IKRigProcessor->CopyOutputGlobalPoseToArray(OutGlobalPose);
+	IKRigProcessor->CopyOutputGlobalPoseToArray(OutTargeGlobalPose);
 }
 
 FTransform UIKRetargetProcessor::GetTargetBoneRetargetPoseGlobalTransform(const int32& TargetBoneIndex) const
@@ -1185,8 +1269,29 @@ void UIKRetargetProcessor::SetNeedsInitialized()
 	}
 }
 
-void UIKRetargetProcessor::CopyTargetIKRigSettingsFromAsset()
+void UIKRetargetProcessor::CopyAllSettingsFromAsset()
 {
 	IKRigProcessor->CopyAllInputsFromSourceAssetAtRuntime(RetargeterAsset->GetTargetIKRig());
+
+	// copy most recent settings from asset for each chain
+	const TArray<TObjectPtr<URetargetChainSettings>>& AllChainSettings = RetargeterAsset->GetAllChainSettings();
+	for (const TObjectPtr<URetargetChainSettings>& ChainSettings : AllChainSettings)
+	{
+		for (FRetargetChainPairFK& Chain : ChainPairsFK)
+		{
+			if (Chain.TargetBoneChainName == ChainSettings->TargetChain)
+			{
+				Chain.Settings.CopySettingsFromAsset(ChainSettings);
+			}
+		}
+		
+		for (FRetargetChainPairIK& Chain : ChainPairsIK)
+		{
+			if (Chain.TargetBoneChainName == ChainSettings->TargetChain)
+			{
+				Chain.Settings.CopySettingsFromAsset(ChainSettings);
+			}
+		}
+	}
 }
 #endif
diff --git a/Engine/Plugins/Animation/IKRig/Source/IKRig/Private/Retargeter/IKRetargeter.cpp b/Engine/Plugins/Animation/IKRig/Source/IKRig/Private/Retargeter/IKRetargeter.cpp
--- a/Engine/Plugins/Animation/IKRig/Source/IKRig/Private/Retargeter/IKRetargeter.cpp
+++ b/Engine/Plugins/Animation/IKRig/Source/IKRig/Private/Retargeter/IKRetargeter.cpp
@@ -10,6 +10,45 @@ const FName UIKRetargeter::GetTargetPreviewMeshPropertyName() { return GET_MEMBE
 #endif
 const FName UIKRetargeter::GetDefaultPoseName() { return DefaultPoseName; }
 
+void UIKRetargeter::PostLoad()
+{
+	Super::PostLoad();
+
+	// load deprecated chain mapping (pre UStruct to UObject refactor)
+	PRAGMA_DISABLE_DEPRECATION_WARNINGS
+	if (!ChainMapping_DEPRECATED.IsEmpty())
+	{
+		for (const FRetargetChainMap& OldChainMap : ChainMapping_DEPRECATED)
+		{
+			if (OldChainMap.TargetChain == NAME_None)
+			{
+				continue;
+			}
+			
+			TObjectPtr<URetargetChainSettings>* MatchingChain = ChainSettings.FindByPredicate([&](const URetargetChainSettings* Chain)
+			{
+				return Chain ? Chain->TargetChain == OldChainMap.TargetChain : false;
+			});
+			
+			if (MatchingChain)
+			{
+				(*MatchingChain)->SourceChain = OldChainMap.SourceChain;
+			}
+			else
+			{
+				TObjectPtr<URetargetChainSettings> NewChainMap = NewObject<URetargetChainSettings>();
+				NewChainMap->TargetChain = OldChainMap.TargetChain;
+				NewChainMap->SourceChain = OldChainMap.SourceChain;
+				ChainSettings.Add(NewChainMap);
+			}
+		}
+	}
+	PRAGMA_ENABLE_DEPRECATION_WARNINGS
+
+	// remove null settings
+	ChainSettings.Remove(nullptr);
+}
+
 void FIKRetargetPose::SetBoneRotationOffset(FName BoneName, FQuat RotationDelta, const FIKRigSkeleton& Skeleton)
 {
 	FQuat* RotOffset = BoneRotationOffsets.Find(BoneName);
diff --git a/Engine/Plugins/Animation/IKRig/Source/IKRig/Public/Retargeter/IKRetargetProcessor.h b/Engine/Plugins/Animation/IKRig/Source/IKRig/Public/Retargeter/IKRetargetProcessor.h
--- a/Engine/Plugins/Animation/IKRig/Source/IKRig/Public/Retargeter/IKRetargetProcessor.h
+++ b/Engine/Plugins/Animation/IKRig/Source/IKRig/Public/Retargeter/IKRetargetProcessor.h
@@ -3,8 +3,10 @@
 #pragma once
 
 #include "CoreMinimal.h"
+#include "IKRetargeter.h"
 #include "IKRetargetProcessor.generated.h"
 
+class URetargetChainSettings;
 class UIKRigDefinition;
 class UIKRigProcessor;
 struct FReferenceSkeleton;
@@ -141,6 +143,39 @@ struct FRootRetargeter
 	void DecodePose( TArray<FTransform> &OutTargetGlobalPose, const float StrideScale) const;
 };
 
+struct FRetargetChainSettings
+{
+	FName TargetChainName;
+	
+	bool CopyPoseUsingFK = true;
+
+	bool DriveIKGoal = true;
+	FVector StaticOffset;
+	float Extension = 1.0f;
+	float BlendToSource = 0.0f;
+	FVector BlendToSourceWeights = FVector::OneVector;
+	float MatchSourceVelocity = 0.0f;
+	float TeleportVelocityThreshold = 0.0f;
+
+public:
+	
+	void CopySettingsFromAsset(const URetargetChainSettings* AssetChainSettings)
+	{
+		TargetChainName = AssetChainSettings->TargetChain;
+		
+		CopyPoseUsingFK = AssetChainSettings->CopyPoseUsingFK;
+		
+		DriveIKGoal = AssetChainSettings->DriveIKGoal;
+		Extension = AssetChainSettings->Extension;
+		StaticOffset = AssetChainSettings->StaticOffset;
+		BlendToSource = AssetChainSettings->BlendToSource;
+		BlendToSourceWeights = AssetChainSettings->BlendToSourceWeights;
+		MatchSourceVelocity = AssetChainSettings->MatchSourceVelocity;
+		TeleportVelocityThreshold = AssetChainSettings->VelocityThreshold;
+	};
+};
+
+
 struct FChainFK
 {
 	TArray<FTransform> InitialBoneTransforms;
@@ -173,6 +208,7 @@ struct FChainDecoderFK : public FChainFK
 		const FTargetSkeleton& TargetSkeleton);
 	
 	void DecodePose(
+		const FRetargetChainSettings& Settings,
 		const TArray<int32>& TargetBoneIndices,
 		const FChainEncoderFK& SourceChain,
 		const FTargetSkeleton& TargetSkeleton,
@@ -216,6 +252,8 @@ struct FSourceChainIK
 	float InvInitialLength;
 
 	// results after encoding...
+	FVector PreviousEndPosition;
+	FVector CurrentEndPosition;
 	
 	FVector CurrentEndDirectionNormalized;
 	
@@ -229,12 +267,16 @@ struct FSourceChainIK
 struct FTargetChainIK
 {
 	int32 BoneIndexA;
+
+	int32 BoneIndexC;
 	
 	float InitialLength;
 	
 	FVector InitialEndPosition;
 	
 	FQuat InitialEndRotation;
+
+	FVector PrevEndPosition;
 };
 
 struct FChainRetargeterIK
@@ -249,11 +291,16 @@ struct FChainRetargeterIK
 
 	void EncodePose(const TArray<FTransform> &SourceInputGlobalPose);
 	
-	void DecodePose(const TArray<FTransform>& OutGlobalPose, FDecodedIKChain& OutResults) const;
+	void DecodePose(
+		const FRetargetChainSettings& Settings,
+		const TArray<FTransform>& OutGlobalPose,
+		FDecodedIKChain& OutResults);
 };
 
 struct FRetargetChainPair
 {
+	FRetargetChainSettings Settings;
+	
 	TArray<int32> SourceBoneIndices;
 	
 	TArray<int32> TargetBoneIndices;
@@ -262,9 +309,10 @@ struct FRetargetChainPair
 	
 	FName TargetBoneChainName;
 
-	virtual ~FRetargetChainPair(){};
+	virtual ~FRetargetChainPair() = default;
 	
 	virtual bool Initialize(
+		URetargetChainSettings* InSettings,
 		const FBoneChain& SourceBoneChain,
 		const FBoneChain& TargetBoneChain,
 		const FRetargetSkeleton& SourceSkeleton,
@@ -285,6 +333,7 @@ struct FRetargetChainPairFK : FRetargetChainPair
 	FChainDecoderFK FKDecoder;
 	
 	virtual bool Initialize(
+		URetargetChainSettings* InSettings,
         const FBoneChain& SourceBoneChain,
         const FBoneChain& TargetBoneChain,
         const FRetargetSkeleton& SourceSkeleton,
@@ -292,7 +341,7 @@ struct FRetargetChainPairFK : FRetargetChainPair
 };
 
 struct FRetargetChainPairIK : FRetargetChainPair
-{	
+{
 	FChainRetargeterIK IKChainRetargeter;
 	
 	FName IKGoalName;
@@ -300,6 +349,7 @@ struct FRetargetChainPairIK : FRetargetChainPair
 	FName PoleVectorGoalName;
 
 	virtual bool Initialize(
+		URetargetChainSettings* Settings,
         const FBoneChain& SourceBoneChain,
         const FBoneChain& TargetBoneChain,
         const FRetargetSkeleton& SourceSkeleton,
@@ -366,8 +416,8 @@ public:
 #if WITH_EDITOR
 	/** Set that this processor needs to be reinitialized. */
 	void SetNeedsInitialized();
-	/** During editor preview, drive the target IK Rig with the setting from it's source asset */
-	void CopyTargetIKRigSettingsFromAsset();
+	/** During editor preview, drive the target IK Rig with the settings from it's source asset */
+	void CopyAllSettingsFromAsset();
 #endif
 
 private:
@@ -422,5 +472,5 @@ private:
 	void RunFKRetarget(const TArray<FTransform>& InGlobalTransforms, TArray<FTransform>& OutGlobalTransforms);
 
 	/** Internal retarget phase for the IK chains. */
-	void RunIKRetarget(const TArray<FTransform>& InGlobalPose, TArray<FTransform>& OutGlobalPose);
+	void RunIKRetarget(const TArray<FTransform>& InSourceGlobalPose, TArray<FTransform>& OutTargeGlobalPose);
 };
diff --git a/Engine/Plugins/Animation/IKRig/Source/IKRig/Public/Retargeter/IKRetargeter.h b/Engine/Plugins/Animation/IKRig/Source/IKRig/Public/Retargeter/IKRetargeter.h
--- a/Engine/Plugins/Animation/IKRig/Source/IKRig/Public/Retargeter/IKRetargeter.h
+++ b/Engine/Plugins/Animation/IKRig/Source/IKRig/Public/Retargeter/IKRetargeter.h
@@ -8,38 +8,98 @@
 
 struct FIKRetargetPose;
 
-USTRUCT(Blueprintable)
+struct UE_DEPRECATED(5.1, "Use URetargetChainSettings instead.") FRetargetChainMap;
+USTRUCT()
 struct IKRIG_API FRetargetChainMap
 {
 	GENERATED_BODY()
 
-	FRetargetChainMap(){}
-	
+	FRetargetChainMap() = default;
 	FRetargetChainMap(const FName& TargetChain) : TargetChain(TargetChain){}
 	
-	/** The name of the chain on the Source IK Rig asset to copy animation from. */
-	UPROPERTY(EditAnywhere, BlueprintReadWrite, Category = Offsets)
+	UPROPERTY(EditAnywhere, Category = Offsets)
+	FName SourceChain = NAME_None;
+	
+	UPROPERTY(EditAnywhere, Category = Offsets)
+	FName TargetChain = NAME_None;
+};
+
+UCLASS()
+class IKRIG_API URetargetChainSettings: public UObject
+{
+	GENERATED_BODY()
+
+	public:
+	
+	URetargetChainSettings() = default;
+	
+	URetargetChainSettings(const FName& TargetChain) : TargetChain(TargetChain){}
+	
+	/** The chain on the Source IK Rig asset to copy animation FROM. */
+	UPROPERTY(VisibleAnywhere, Category = "Chain Mapping")
 	FName SourceChain = NAME_None;
 
-	/** The name of the chain on the Target IK Rig asset to copy animation from. */
-	UPROPERTY(EditAnywhere, BlueprintReadWrite, Category = Offsets)
+	/** The chain on the Target IK Rig asset to copy animation TO. */
+	UPROPERTY(VisibleAnywhere, Category = "Chain Mapping")
 	FName TargetChain = NAME_None;
 
-	/** Range -1 to 2. Default 0. Brings IK effector closer (-) or further (+) from origin of chain.
-	*  At -1 the end is on top of the origin. At +2 the end is fully extended twice the length of the chain. */
-	UPROPERTY(EditAnywhere, BlueprintReadWrite, Category = Offsets, meta = (ClampMin = "-1.0", ClampMax = "2.0", UIMin = "-1.0", UIMax = "2.0"))
-	float Extension = 0.0f;
+	/** Whether to copy the shape of the chain from the source skeleton (using retarget pose offsets). Default is true.
+	 * NOTE: This runs before the IK pass to copy the shape of the FK chain from the source skeleton. */
+	UPROPERTY(EditAnywhere, Category = "Retarget Passes")
+	bool CopyPoseUsingFK = true;
+
+	/** Whether to run IK on this chain (IK runs AFTER the FK pass). Default is true.
+	 * NOTE: This only has an effect if the chain has an IK Goal assigned to it in the Target IK Rig asset.
+	 * NOTE: If off, and this chain has an IK Goal, the IK will still be evaluated, but the Goal is set to the input bone location.*/
+	UPROPERTY(EditAnywhere, Category = "IK Adjustments")
+	bool DriveIKGoal = true;
+
+	/** Range 0 to 1. Default 0. Blends IK goal position from retargeted location (0) to source bone location (1).
+	*  At 0 the goal is placed at the retargeted location.
+	*  At 1 the goal is placed at the location of the source chain's end bone. */
+	UPROPERTY(EditAnywhere, Category = "IK Adjustments", meta = (ClampMin = "0.0", ClampMax = "1.0", UIMin = "0.0", UIMax = "1.0"))
+	float BlendToSource = 0.0f;
+	
+	/** Range 0 to 1. Default 1. Weight each axis separately when using Blend To Source.
+	*  At 0 the goal is placed at the retargeted location.
+	*  At 1 the goal is placed at the location of the source chain's end bone. */
+	UPROPERTY(EditAnywhere, Category = "IK Adjustments", meta = (ClampMin = "0.0", ClampMax = "1.0", UIMin = "0.0", UIMax = "1.0"))
+	FVector BlendToSourceWeights = FVector::OneVector;
+
+	/** Default 0, 0, 0. Apply static global-space offset to IK goal position. */
+	UPROPERTY(EditAnywhere, Category = "IK Adjustments")
+	FVector StaticOffset;
+	
+	/** Range 0 to 5. Default 1. Brings IK goal closer (-) or further (+) from origin of chain.
+	*  At 0 the effector is placed at the origin of the chain.
+	*  Values greater than 1 will stretch the chain beyond the retargeted length. */
+	UPROPERTY(EditAnywhere, Category = "IK Adjustments", meta = (ClampMin = "0.0", ClampMax = "5.0", UIMin = "0.1", UIMax = "2.0"))
+	float Extension = 1.0f;
+
+	/** Range 0 to 1. Default 0. Blends IK goal from retargeted velocity (0) to source bone velocity (1).
+	*  At 0 the goal is placed at the retargeted location.
+	*  At 1 the goal is placed at the retargeted location but clamped by the velocity of the source chain's end bone.
+	*  Values between 0 and 1 will blend the applied velocity between the retargeted amount and the source bone amount.*/
+	UPROPERTY(EditAnywhere, Category = "Experimental (May Change)", meta = (ClampMin = "0.0", ClampMax = "1.0", UIMin = "0.0", UIMax = "1.0"))
+	float MatchSourceVelocity = 0.0f;
+
+	/** Range 0 to 1000. Default 5. The maximum speed a source goal can be moving before being considered statically planted.
+	*  At 0 the goal is placed at the retargeted location.
+	*  At 1 the goal is placed at the retargeted location but clamped by the velocity of the source chain's end bone.
+	*  Values between 0 and 1 will blend the applied velocity between the retargeted amount and the source bone amount.*/
+	UPROPERTY(EditAnywhere, Category = "Experimental (May Change)", meta = (ClampMin = "0.0", ClampMax = "100.0", UIMin = "0.0", UIMax = "100.0"))
+	float VelocityThreshold = 5.0f;
+
+	/** Range 0 to 1. Default is 1.0. Blend IK effector at the end of this chain towards the original position
+	*  on the source skeleton (0.0) or the position on the retargeted target skeleton (1.0). */
+	//UPROPERTY(EditAnywhere, BlueprintReadWrite, Category = IkMode, meta = (ClampMin = "0.0", ClampMax = "1.0", UIMin = "0.0", UIMax = "1.0"))
+	//float IkToSourceOrTarget = 1.0f;
 
 	/** Range 0 to 1. Default 0. Allow the chain to stretch by translating to reach the IK goal locations.
 	*  At 0 the chain will not stretch at all. At 1 the chain will be allowed to stretch double it's full length to reach IK. */
 	//UPROPERTY(EditAnywhere, BlueprintReadWrite, Category = Stretch, meta = (ClampMin = "0.0", ClampMax = "1.0", UIMin = "0.0", UIMax = "1.0"))
 	//float StretchTolerance = 0.0f;
 
-	/** Range 0 to 1. Default is 1.0. Blend IK effector at the end of this chain towards the original position
-	 *  on the source skeleton (0.0) or the position on the retargeted target skeleton (1.0). */
-	//UPROPERTY(EditAnywhere, BlueprintReadWrite, Category = IkMode, meta = (ClampMin = "0.0", ClampMax = "1.0", UIMin = "0.0", UIMax = "1.0"))
-	//float IkToSourceOrTarget = 1.0f;
-	
 	/** When true, the source IK position is calculated relative to a source bone and applied relative to a target bone. */
 	//UPROPERTY(EditAnywhere, BlueprintReadWrite, Category = IkMode)
 	//bool bIkRelativeToSourceBone = false;
@@ -51,7 +111,7 @@ struct IKRIG_API FRetargetChainMap
 	/** A bone in the TARGET skeleton that the IK location is relative to.
 	 * This is usually the same bone as the source skeleton, but may have a different name in the target skeleton. */
 	//UPROPERTY(EditAnywhere, BlueprintReadWrite, Category = IkMode, meta = (EditCondition="bIkRelativeToSourceBone"))
-	//FName IkRelativeTargetBone;	
+	//FName IkRelativeTargetBone;
 };
 
 USTRUCT()
@@ -95,7 +155,7 @@ public:
 	 * WARNING: do not use for editing the data model. Use Controller class instead. */
 	UIKRigDefinition* GetTargetIKRigWriteable() const { return TargetIKRigAsset.Get(); };
 	/** Get read-only access to the chain mapping */
-	const TArray<FRetargetChainMap>& GetChainMapping() const { return ChainMapping; };
+	const TArray<TObjectPtr<URetargetChainSettings>> GetAllChainSettings() const { return ChainSettings; };
 	/** Get read-only access to a retarget pose */
 	const FIKRetargetPose* GetCurrentRetargetPose() const { return &RetargetPoses[CurrentRetargetPose]; };
 
@@ -111,8 +171,10 @@ public:
 	static const FName GetDefaultPoseName();
 
 #if WITH_EDITOR
-	bool IsInEditRetargetPoseMode() const { return bEditRetargetPoseMode; };
+	bool IsInEditRetargetPoseMode() const { return bEditRetargetPoseMode; }
 #endif
+	
+	virtual void PostLoad() override;
 
 private:
 
@@ -173,9 +235,15 @@ private:
 	UPROPERTY()
 	TMap<FName, FIKRetargetPose> RetargetPoses;
 
-	/** Mapping of chains to copy animation between source and target rigs.*/
+	/** (OLD VERSION) Mapping of chains to copy animation between source and target rigs.*/
+	PRAGMA_DISABLE_DEPRECATION_WARNINGS
+	UPROPERTY()
+	TArray<FRetargetChainMap> ChainMapping_DEPRECATED;
+	PRAGMA_ENABLE_DEPRECATION_WARNINGS
+	
+	/** Settings for how to map source chains to target chains.*/
 	UPROPERTY()
-	TArray<FRetargetChainMap> ChainMapping;
+	TArray<TObjectPtr<URetargetChainSettings>> ChainSettings;
 	
 	/** The set of retarget poses available as options for retargeting.*/
 	UPROPERTY()
diff --git a/Engine/Plugins/Animation/IKRig/Source/IKRigEditor/Private/RetargetEditor/IKRetargetBatchOperation.cpp b/Engine/Plugins/Animation/IKRig/Source/IKRigEditor/Private/RetargetEditor/IKRetargetBatchOperation.cpp
--- a/Engine/Plugins/Animation/IKRig/Source/IKRigEditor/Private/RetargetEditor/IKRetargetBatchOperation.cpp
+++ b/Engine/Plugins/Animation/IKRig/Source/IKRigEditor/Private/RetargetEditor/IKRetargetBatchOperation.cpp
@@ -351,7 +351,7 @@ void FIKRetargetBatchOperation::ConvertAnimation(
 			}
 
 			// update goals 
-			Processor->CopyTargetIKRigSettingsFromAsset();
+			Processor->CopyAllSettingsFromAsset();
 			
 			// run the retarget
 			const TArray<FTransform>& TargetComponentPose = Processor->RunRetargeter(SourceComponentPose);
diff --git a/Engine/Plugins/Animation/IKRig/Source/IKRigEditor/Private/RetargetEditor/IKRetargetEditMode.cpp b/Engine/Plugins/Animation/IKRig/Source/IKRigEditor/Private/RetargetEditor/IKRetargetEditMode.cpp
--- a/Engine/Plugins/Animation/IKRig/Source/IKRigEditor/Private/RetargetEditor/IKRetargetEditMode.cpp
+++ b/Engine/Plugins/Animation/IKRig/Source/IKRigEditor/Private/RetargetEditor/IKRetargetEditMode.cpp
@@ -220,6 +220,13 @@ bool FIKRetargetEditMode::HandleClick(FEditorViewportClient* InViewportClient, H
 			// clicking in empty space clears selection
 			const bool bReplaceSelection = true;
 			HandleBoneSelectedInViewport(NAME_None, bReplaceSelection);
+
+			// show asset in details view
+			const TSharedPtr<FIKRetargetEditorController> Controller = EditorController.Pin();
+			if (Controller.IsValid())
+			{
+				Controller->DetailsView->SetObject(Controller->AssetController->GetAsset());
+			}
 		}
 	}
 	
diff --git a/Engine/Plugins/Animation/IKRig/Source/IKRigEditor/Private/RetargetEditor/IKRetargeterController.cpp b/Engine/Plugins/Animation/IKRig/Source/IKRigEditor/Private/RetargetEditor/IKRetargeterController.cpp
--- a/Engine/Plugins/Animation/IKRig/Source/IKRigEditor/Private/RetargetEditor/IKRetargeterController.cpp
+++ b/Engine/Plugins/Animation/IKRig/Source/IKRigEditor/Private/RetargetEditor/IKRetargeterController.cpp
@@ -105,7 +105,7 @@ void UIKRetargeterController::CleanChainMapping()
 	if (!IsValid(Asset->TargetIKRigAsset))
 	{
 		// empty the chain mapping
-		Asset->ChainMapping.Empty();
+		Asset->ChainSettings.Empty();
 		return;
 	}
 	
@@ -114,32 +114,34 @@ void UIKRetargeterController::CleanChainMapping()
 
 	// remove all target chains that are no longer in the target IK rig asset
 	TArray<FName> TargetChainsToRemove;
-	for (FRetargetChainMap& ChainMap : Asset->ChainMapping)
+	for (const URetargetChainSettings* ChainMap : Asset->ChainSettings)
 	{
-		if (!TargetChainNames.Contains(ChainMap.TargetChain))
+		if (!TargetChainNames.Contains(ChainMap->TargetChain))
 		{
-			TargetChainsToRemove.Add(ChainMap.TargetChain);
+			TargetChainsToRemove.Add(ChainMap->TargetChain);
 		}
 	}
 	for (FName TargetChainToRemove : TargetChainsToRemove)
 	{
-		Asset->ChainMapping.RemoveAll([&TargetChainToRemove](FRetargetChainMap& Element)
+		Asset->ChainSettings.RemoveAll([&TargetChainToRemove](const URetargetChainSettings* Element)
 		{
-			return Element.TargetChain == TargetChainToRemove;
+			return Element->TargetChain == TargetChainToRemove;
 		});
 	}
 
 	// add a mapping for each chain that is in the target IK rig (if it doesn't have one already)
 	for (FName TargetChainName : TargetChainNames)
 	{
-		const bool HasChain = Asset->ChainMapping.ContainsByPredicate([&TargetChainName](FRetargetChainMap& Element)
+		const bool HasChain = Asset->ChainSettings.ContainsByPredicate([&TargetChainName](const URetargetChainSettings* Element)
 		{
-			return Element.TargetChain == TargetChainName;
+			return Element->TargetChain == TargetChainName;
 		});
 		
 		if (!HasChain)
 		{
-			Asset->ChainMapping.Add(FRetargetChainMap(TargetChainName));
+			TObjectPtr<URetargetChainSettings> ChainMap = NewObject<URetargetChainSettings>();
+			ChainMap->TargetChain = TargetChainName;
+			Asset->ChainSettings.Add(ChainMap);
 		}
 	}
 
@@ -147,11 +149,11 @@ void UIKRetargeterController::CleanChainMapping()
 	GetSourceChainNames(SourceChainNames);
 	
 	// reset any sources that are no longer present to "None"
-	for (FRetargetChainMap& ChainMap : Asset->ChainMapping)
+	for (URetargetChainSettings* ChainMap : Asset->ChainSettings)
 	{
-		if (!SourceChainNames.Contains(ChainMap.SourceChain))
+		if (!SourceChainNames.Contains(ChainMap->SourceChain))
 		{
-			ChainMap.SourceChain = NAME_None;
+			ChainMap->SourceChain = NAME_None;
 		}
 	}
 
@@ -212,15 +214,15 @@ void UIKRetargeterController::AutoMapChains() const
 	GetSourceChainNames(SourceChainNames);
 	
 	// auto-map any chains that have no value using a fuzzy string search
-	for (FRetargetChainMap& ChainMap : Asset->ChainMapping)
+	for (URetargetChainSettings* ChainMap : Asset->ChainSettings)
 	{
-		if (ChainMap.SourceChain != NAME_None)
+		if (ChainMap->SourceChain != NAME_None)
 		{
 			continue; // already set by user
 		}
 
 		// find "best match" automatically as a convenience for the user
-		FString TargetNameLowerCase = ChainMap.TargetChain.ToString().ToLower();
+		FString TargetNameLowerCase = ChainMap->TargetChain.ToString().ToLower();
 		float HighestScore = 0.2f;
 		int32 HighestScoreIndex = -1;
 		for (int32 ChainIndex=0; ChainIndex<SourceChainNames.Num(); ++ChainIndex)
@@ -239,7 +241,7 @@ void UIKRetargeterController::AutoMapChains() const
 		// apply source if any decent matches were found
 		if (SourceChainNames.IsValidIndex(HighestScoreIndex))
 		{
-			ChainMap.SourceChain = SourceChainNames[HighestScoreIndex];
+			ChainMap->SourceChain = SourceChainNames[HighestScoreIndex];
 		}
 	}
 
@@ -254,9 +256,9 @@ void UIKRetargeterController::OnRetargetChainRenamed(UIKRigDefinition* IKRig, FN
 {
 	const bool bIsSourceRig = IKRig == Asset->SourceIKRigAsset;
 	check(bIsSourceRig || IKRig == Asset->TargetIKRigAsset)
-	for (FRetargetChainMap& ChainMap : Asset->ChainMapping)
+	for (URetargetChainSettings* ChainMap : Asset->ChainSettings)
 	{
-		FName& ChainNameToUpdate = bIsSourceRig ? ChainMap.SourceChain : ChainMap.TargetChain;
+		FName& ChainNameToUpdate = bIsSourceRig ? ChainMap->SourceChain : ChainMap->TargetChain;
 		if (ChainNameToUpdate == OldChainName)
 		{
 			ChainNameToUpdate = NewChainName;
@@ -274,12 +276,11 @@ void UIKRetargeterController::OnRetargetChainRemoved(UIKRigDefinition* IKRig, co
 	// set source chain name to NONE if it has been deleted 
 	if (bIsSourceRig)
 	{
-		for (FRetargetChainMap& ChainMap : Asset->ChainMapping)
+		for (URetargetChainSettings* ChainMap : Asset->ChainSettings)
 		{
-			FName& SourceChainToUpdate = ChainMap.SourceChain;
-			if (SourceChainToUpdate == InChainRemoved)
+			if (ChainMap->SourceChain == InChainRemoved)
 			{
-				SourceChainToUpdate = NAME_None;
+				ChainMap->SourceChain = NAME_None;
 				BroadcastNeedsReinitialized();
 				return;
 			}
@@ -288,39 +289,31 @@ void UIKRetargeterController::OnRetargetChainRemoved(UIKRigDefinition* IKRig, co
 	}
 	
 	// remove target mapping if the target chain has been removed
-	const int32 ChainIndex = Asset->ChainMapping.IndexOfByPredicate([&InChainRemoved](const FRetargetChainMap& ChainMap)
+	const int32 ChainIndex = Asset->ChainSettings.IndexOfByPredicate([&InChainRemoved](const URetargetChainSettings* ChainMap)
 	{
-		return ChainMap.TargetChain == InChainRemoved;
+		return ChainMap->TargetChain == InChainRemoved;
 	});
 	
 	if (ChainIndex != INDEX_NONE)
 	{
-		Asset->ChainMapping.RemoveAt(ChainIndex);
+		Asset->ChainSettings.RemoveAt(ChainIndex);
 		BroadcastNeedsReinitialized();
 	}
 }
 
-void UIKRetargeterController::SetSourceChainForTargetChain(FName TargetChain, FName SourceChainToMapTo)
+void UIKRetargeterController::SetSourceChainForTargetChain(URetargetChainSettings* ChainMap, FName SourceChainToMapTo) const
 {
 	FScopedTransaction Transaction(LOCTEXT("SetRetargetChainSource", "Set Retarget Chain Source"));
 	Asset->Modify();
 	
-	FRetargetChainMap* ChainMap = GetChainMap(TargetChain);
 	check(ChainMap)
 	ChainMap->SourceChain = SourceChainToMapTo;
 	BroadcastNeedsReinitialized();
 }
 
-FName UIKRetargeterController::GetSourceChainForTargetChain(FName TargetChain)
+const TArray<TObjectPtr<URetargetChainSettings>>& UIKRetargeterController::GetChainMappings() const
 {
-	FRetargetChainMap* ChainMap = GetChainMap(TargetChain);
-	check(ChainMap)
-	return ChainMap->SourceChain;
-}
-
-const TArray<FRetargetChainMap>& UIKRetargeterController::GetChainMappings()
-{
-	return Asset->ChainMapping;
+	return Asset->ChainSettings;
 }
 
 USkeleton* UIKRetargeterController::GetSourceSkeletonAsset() const
@@ -462,13 +455,13 @@ FName UIKRetargeterController::MakePoseNameUnique(FName PoseName) const
 	return UniqueName;
 }
 
-FRetargetChainMap* UIKRetargeterController::GetChainMap(const FName& TargetChainName) const
+URetargetChainSettings* UIKRetargeterController::GetChainMap(const FName& TargetChainName) const
 {
-	for (FRetargetChainMap& ChainMap : Asset->ChainMapping)
+	for (URetargetChainSettings* ChainMap : Asset->ChainSettings)
 	{
-		if (ChainMap.TargetChain == TargetChainName)
+		if (ChainMap->TargetChain == TargetChainName)
 		{
-			return &ChainMap;
+			return ChainMap;
 		}
 	}
 
@@ -477,7 +470,7 @@ FRetargetChainMap* UIKRetargeterController::GetChainMap(const FName& TargetChain
 
 void UIKRetargeterController::SortChainMapping() const
 {
-	Asset->ChainMapping.Sort([this](const FRetargetChainMap& A, const FRetargetChainMap& B)
+	Asset->ChainSettings.Sort([this](const URetargetChainSettings& A, const URetargetChainSettings& B)
 	{
 		const TArray<FBoneChain>& BoneChains = Asset->TargetIKRigAsset->GetRetargetChains();
 		const FIKRigSkeleton& TargetSkeleton = Asset->TargetIKRigAsset->Skeleton;
diff --git a/Engine/Plugins/Animation/IKRig/Source/IKRigEditor/Private/RetargetEditor/SIKRetargetChainMapList.cpp b/Engine/Plugins/Animation/IKRig/Source/IKRigEditor/Private/RetargetEditor/SIKRetargetChainMapList.cpp
--- a/Engine/Plugins/Animation/IKRig/Source/IKRigEditor/Private/RetargetEditor/SIKRetargetChainMapList.cpp
+++ b/Engine/Plugins/Animation/IKRig/Source/IKRigEditor/Private/RetargetEditor/SIKRetargetChainMapList.cpp
@@ -67,7 +67,7 @@ TSharedRef< SWidget > SIKRetargetChainMapRow::GenerateWidgetForColumn( const FNa
 		.Padding(3.0f, 1.0f)
 		[
 			SNew(STextBlock)
-			.Text(FText::FromName(ChainMapElement.Pin()->TargetChainName))
+			.Text(FText::FromName(ChainMapElement.Pin()->ChainMap->TargetChain))
 			.Font(FEditorStyle::GetFontStyle(TEXT("BoldFont")))
 		];
 		return NewWidget;
@@ -105,10 +105,9 @@ void SIKRetargetChainMapRow::OnSourceChainComboSelectionChanged(TSharedPtr<FStri
 	{
 		return; 
 	}
-	
-	const FName TargetChainName = ChainMapElement.Pin()->TargetChainName;
+
 	const FName SourceChainName = FName(*InName.Get());
-	RetargeterController->SetSourceChainForTargetChain(TargetChainName, SourceChainName);
+	RetargeterController->SetSourceChainForTargetChain(ChainMapElement.Pin()->ChainMap.Get(), SourceChainName);
 }
 
 FText SIKRetargetChainMapRow::GetSourceChainName() const
@@ -118,10 +117,8 @@ FText SIKRetargetChainMapRow::GetSourceChainName() const
 	{
 		return FText::FromName(NAME_None); 
 	}
-	
-	const FName TargetChainName = ChainMapElement.Pin()->TargetChainName;
-	const FName SourceChainName = RetargeterController->GetSourceChainForTargetChain(TargetChainName);
-	return FText::FromName(SourceChainName);
+
+	return FText::FromName(ChainMapElement.Pin()->ChainMap->SourceChain);
 }
 
 void SIKRetargetChainMapList::Construct(
@@ -203,6 +200,11 @@ void SIKRetargetChainMapList::Construct(
 			.ListItemsSource( &ListViewItems )
 			.OnGenerateRow( this, &SIKRetargetChainMapList::MakeListRowWidget )
 			.OnMouseButtonClick(this, &SIKRetargetChainMapList::OnItemClicked)
+			.OnSelectionChanged_Lambda([this] (TSharedPtr<FRetargetChainMapElement> NewValue, ESelectInfo::Type SelectInfo)
+			{
+				OnSelectionChanged();
+			})
+
 			.ItemHeight( 22.0f )
 			.HeaderRow
 			(
@@ -279,10 +281,10 @@ void SIKRetargetChainMapList::RefreshView()
 	
 	// refresh list of chains
 	ListViewItems.Reset();
-	const TArray<FRetargetChainMap>& ChainMappings = RetargeterController->GetChainMappings();
-	for (const FRetargetChainMap& ChainMap : ChainMappings)
+	const TArray<TObjectPtr<URetargetChainSettings>>& ChainMappings = RetargeterController->GetChainMappings();
+	for (const TObjectPtr<URetargetChainSettings> ChainMap : ChainMappings)
 	{
-		TSharedPtr<FRetargetChainMapElement> ChainItem = FRetargetChainMapElement::Make(ChainMap.TargetChain);
+		TSharedPtr<FRetargetChainMapElement> ChainItem = FRetargetChainMapElement::Make(ChainMap);
 		ListViewItems.Add(ChainItem);
 	}
 
@@ -301,7 +303,36 @@ TSharedRef<ITableRow> SIKRetargetChainMapList::MakeListRowWidget(
 
 void SIKRetargetChainMapList::OnItemClicked(TSharedPtr<FRetargetChainMapElement> InItem)
 {
-	// TODO highlight chain in skeleton view
+	OnSelectionChanged();
+}
+
+void SIKRetargetChainMapList::OnSelectionChanged()
+{
+	const TSharedPtr<FIKRetargetEditorController> Controller = EditorController.Pin();
+	if (!Controller.IsValid())
+	{
+		return;
+	}
+	
+	// update selected settings
+	SelectedChainSettings.Empty();
+	TArray<TSharedPtr<FRetargetChainMapElement>> SelectedItems = ListView.Get()->GetSelectedItems();
+	for (const TSharedPtr<FRetargetChainMapElement>& Item : SelectedItems)
+	{
+		SelectedChainSettings.Add(Item->ChainMap.Get());
+	}
+
+	// selection cleared
+	if (SelectedChainSettings.IsEmpty())
+	{
+		// show asset settings in the details view
+		Controller->DetailsView->SetObject(Controller->AssetController->GetAsset());
+	}
+	else
+	{
+		// show chain settings in the details view
+		Controller->DetailsView->SetObjects(SelectedChainSettings);
+	}
 }
 
 EVisibility SIKRetargetChainMapList::IsAutoMapButtonVisible() const
diff --git a/Engine/Plugins/Animation/IKRig/Source/IKRigEditor/Public/RetargetEditor/IKRetargeterController.h b/Engine/Plugins/Animation/IKRig/Source/IKRigEditor/Public/RetargetEditor/IKRetargeterController.h
--- a/Engine/Plugins/Animation/IKRig/Source/IKRigEditor/Public/RetargetEditor/IKRetargeterController.h
+++ b/Engine/Plugins/Animation/IKRig/Source/IKRigEditor/Public/RetargetEditor/IKRetargeterController.h
@@ -9,7 +9,7 @@
 
 #include "IKRetargeterController.generated.h"
 
-struct FRetargetChainMap;
+class URetargetChainSettings;
 class UIKRigDefinition;
 class UIKRetargeter;
 
@@ -62,11 +62,9 @@ public:
 	/** Callback when IK Rig chain is removed. */
 	void OnRetargetChainRemoved(UIKRigDefinition* IKRig, const FName& InChainRemoved) const;
 	/** Set the source chain to map to a given target chain */
-	void SetSourceChainForTargetChain(FName TargetChain, FName SourceChainToMapTo);
-	/** Get the source chain mapped to a given target chain */
-	FName GetSourceChainForTargetChain(FName TargetChain);
+	void SetSourceChainForTargetChain(URetargetChainSettings* ChainMap, FName SourceChainToMapTo) const;
 	/** Get read-only access to the list of chain mappings */
-	const TArray<FRetargetChainMap>& GetChainMappings();
+	const TArray<TObjectPtr<URetargetChainSettings>>& GetChainMappings() const;
 	/** END RETARGET CHAIN MAPPING */
 
 	/** RETARGET POSE EDITING
@@ -117,7 +115,7 @@ public:
 	
 private:
 
-	FRetargetChainMap* GetChainMap(const FName& TargetChainName) const;
+	URetargetChainSettings* GetChainMap(const FName& TargetChainName) const;
 
 	/** Sort the Asset ChainMapping based on the StartBone of the target chains. */
 	void SortChainMapping() const;
diff --git a/Engine/Plugins/Animation/IKRig/Source/IKRigEditor/Public/RetargetEditor/SIKRetargetChainMapList.h b/Engine/Plugins/Animation/IKRig/Source/IKRigEditor/Public/RetargetEditor/SIKRetargetChainMapList.h
--- a/Engine/Plugins/Animation/IKRig/Source/IKRigEditor/Public/RetargetEditor/SIKRetargetChainMapList.h
+++ b/Engine/Plugins/Animation/IKRig/Source/IKRigEditor/Public/RetargetEditor/SIKRetargetChainMapList.h
@@ -26,20 +26,20 @@ public:
         TSharedRef<FRetargetChainMapElement> InStackElement,
         TSharedPtr<SIKRetargetChainMapList> InChainList);
 
-	static TSharedRef<FRetargetChainMapElement> Make(FName InTargetChain)
+	static TSharedRef<FRetargetChainMapElement> Make(TObjectPtr<URetargetChainSettings> InChainMap)
 	{
-		return MakeShareable(new FRetargetChainMapElement(InTargetChain));
+		return MakeShareable(new FRetargetChainMapElement(InChainMap));
 	}
 
-	FName TargetChainName;
+	TWeakObjectPtr<URetargetChainSettings> ChainMap;
 
 private:
 	
 	/** Hidden constructor, always use Make above */
-	FRetargetChainMapElement(FName InTargetChain) : TargetChainName(InTargetChain) {}
+	FRetargetChainMapElement(TObjectPtr<URetargetChainSettings> InChainMap) : ChainMap(InChainMap) {}
 
 	/** Hidden constructor, always use Make above */
-	FRetargetChainMapElement() {}
+	FRetargetChainMapElement() = default;
 };
 
 typedef TSharedPtr< FRetargetChainMapElement > FRetargetChainMapElementPtr;
@@ -103,6 +103,10 @@ private:
 	/** when a chain is clicked on in the table view */
 	void OnItemClicked(TSharedPtr<FRetargetChainMapElement> InItem);
 
+	/** update selection / details view */
+	void OnSelectionChanged();
+	TArray<UObject*> SelectedChainSettings;
+
 	/** auto-map chain button*/
 	EVisibility IsAutoMapButtonVisible() const;
 	FReply OnAutoMapButtonClicked() const;
