Allow renaming of goals with case change

diff --git a/Engine/Plugins/Animation/IKRig/Source/IKRigEditor/Private/RigEditor/IKRigController.cpp b/Engine/Plugins/Animation/IKRig/Source/IKRigEditor/Private/RigEditor/IKRigController.cpp
--- a/Engine/Plugins/Animation/IKRig/Source/IKRigEditor/Private/RigEditor/IKRigController.cpp
+++ b/Engine/Plugins/Animation/IKRig/Source/IKRigEditor/Private/RigEditor/IKRigController.cpp
@@ -729,21 +729,30 @@ bool UIKRigController::RemoveGoal(const FName& GoalName) const
 
 FName UIKRigController::RenameGoal(const FName& OldName, const FName& PotentialNewName) const
 {
+	const int32 GoalIndex = GetGoalIndex(OldName);
+	if (GoalIndex == INDEX_NONE)
+	{
+		return NAME_None; // can't rename goal we don't have
+	}
+	
 	// sanitize the potential new name
 	FString CleanName = PotentialNewName.ToString();
 	SanitizeGoalName(CleanName);
 	const FName NewName = FName(CleanName);
-
-	// validate new name
-	const int32 ExistingGoalIndex = GetGoalIndex(NewName);
+	// check if this goal already exists (case sensitive)
+	int32 ExistingGoalIndex = GetGoalIndex(NewName, ENameCase::CaseSensitive);
 	if (ExistingGoalIndex != INDEX_NONE)
 	{
 		return NAME_None; // name already in use, can't use that
 	}
-	const int32 GoalIndex = GetGoalIndex(OldName);
-	if (GoalIndex == INDEX_NONE)
+	// check if this goal already exists (case insensitive, it might just be renamed with a different case)
+	ExistingGoalIndex = GetGoalIndex(NewName);
+	if (ExistingGoalIndex != INDEX_NONE)
 	{
-		return NAME_None; // can't rename goal we don't have
+		if (!NewName.IsEqual(OldName, ENameCase::IgnoreCase))
+		{
+			return NAME_None; // name already in use, can't use that
+		}
 	}
 	
 	FScopedTransaction Transaction(LOCTEXT("RenameGoal_Label", "Rename Goal"));
@@ -973,24 +982,20 @@ void UIKRigController::SanitizeGoalName(FString& InOutName)
 		}
 	}
 
-	const int32 MaxNameLength = 20;
+	// FIXME magic numbers should actually mean something
+	static constexpr int32 MaxNameLength = 100;
 	if (InOutName.Len() > MaxNameLength)
 	{
 		InOutName.LeftChopInline(InOutName.Len() - MaxNameLength);
 	}
 }
 
-int32 UIKRigController::GetGoalIndex(const FName& GoalName) const
+int32 UIKRigController::GetGoalIndex(const FName& InGoalName, const ENameCase CompareMethod) const
 {
-	for (int32 i=0; i<Asset->Goals.Num(); ++i)
-	{	
-		if (Asset->Goals[i]->GoalName == GoalName)
-		{
-			return i;
-		}
-	}
-
-	return INDEX_NONE;
+	return Asset->Goals.IndexOfByPredicate([&](const TObjectPtr<UIKRigEffectorGoal>& Goal)
+	{
+		return Goal->GoalName.IsEqual(InGoalName, CompareMethod); 
+	});
 }
 
 FName UIKRigController::GetGoalName(const int32& GoalIndex) const
diff --git a/Engine/Plugins/Animation/IKRig/Source/IKRigEditor/Public/RigEditor/IKRigController.h b/Engine/Plugins/Animation/IKRig/Source/IKRigEditor/Public/RigEditor/IKRigController.h
--- a/Engine/Plugins/Animation/IKRig/Source/IKRigEditor/Public/RigEditor/IKRigController.h
+++ b/Engine/Plugins/Animation/IKRig/Source/IKRigEditor/Public/RigEditor/IKRigController.h
@@ -99,7 +99,7 @@ public:
 	/** Returns true if the given Goal is connected to the given Solver. False otherwise. */
 	bool IsGoalConnectedToSolver(const FName& GoalName, int32 SolverIndex) const;
 	/** Get the index of the given Goal in the list of Goals. */
-	int32 GetGoalIndex(const FName& GoalName) const;
+	int32 GetGoalIndex(const FName& InGoalName, const ENameCase CompareMethod = ENameCase::IgnoreCase) const;
 	/** Get the name of Goal at the given index. */
 	FName GetGoalName(const int32& GoalIndex) const;
 	/** Get read-only access to the list of Goals. */
