Fix chain mapping serialization issue + added transactional flags

diff --git a/Engine/Plugins/Animation/IKRig/Source/IKRig/Private/Retargeter/IKRetargeter.cpp b/Engine/Plugins/Animation/IKRig/Source/IKRig/Private/Retargeter/IKRetargeter.cpp
--- a/Engine/Plugins/Animation/IKRig/Source/IKRig/Private/Retargeter/IKRetargeter.cpp
+++ b/Engine/Plugins/Animation/IKRig/Source/IKRig/Private/Retargeter/IKRetargeter.cpp
@@ -36,7 +36,7 @@ void UIKRetargeter::PostLoad()
 			}
 			else
 			{
-				TObjectPtr<URetargetChainSettings> NewChainMap = NewObject<URetargetChainSettings>();
+				TObjectPtr<URetargetChainSettings> NewChainMap = NewObject<URetargetChainSettings>(this, URetargetChainSettings::StaticClass(), NAME_None, RF_Transactional);
 				NewChainMap->TargetChain = OldChainMap.TargetChain;
 				NewChainMap->SourceChain = OldChainMap.SourceChain;
 				ChainSettings.Add(NewChainMap);
diff --git a/Engine/Plugins/Animation/IKRig/Source/IKRigEditor/Private/RetargetEditor/IKRetargeterController.cpp b/Engine/Plugins/Animation/IKRig/Source/IKRigEditor/Private/RetargetEditor/IKRetargeterController.cpp
--- a/Engine/Plugins/Animation/IKRig/Source/IKRigEditor/Private/RetargetEditor/IKRetargeterController.cpp
+++ b/Engine/Plugins/Animation/IKRig/Source/IKRigEditor/Private/RetargetEditor/IKRetargeterController.cpp
@@ -139,7 +139,7 @@ void UIKRetargeterController::CleanChainMapping()
 		
 		if (!HasChain)
 		{
-			TObjectPtr<URetargetChainSettings> ChainMap = NewObject<URetargetChainSettings>();
+			TObjectPtr<URetargetChainSettings> ChainMap = NewObject<URetargetChainSettings>(Asset, URetargetChainSettings::StaticClass(), NAME_None, RF_Transactional);
 			ChainMap->TargetChain = TargetChainName;
 			Asset->ChainSettings.Add(ChainMap);
 		}
