Ensure that the end bone index is superior or equal to start bone

diff --git a/Engine/Plugins/Animation/IKRig/Source/IKRig/Private/Retargeter/IKRetargetProcessor.cpp b/Engine/Plugins/Animation/IKRig/Source/IKRig/Private/Retargeter/IKRetargetProcessor.cpp
--- a/Engine/Plugins/Animation/IKRig/Source/IKRig/Private/Retargeter/IKRetargetProcessor.cpp
+++ b/Engine/Plugins/Animation/IKRig/Source/IKRig/Private/Retargeter/IKRetargetProcessor.cpp
@@ -237,46 +237,50 @@ void FTargetSkeleton::UpdateGlobalTransformsAllNonRetargetedBones(TArray<FTransf
 	}
 }
 
-FResolvedBoneChain::FResolvedBoneChain(
-	const FBoneChain& BoneChain,
-	const FRetargetSkeleton& Skeleton,
-	TArray<int32> &OutBoneIndices)
+FResolvedBoneChain::FResolvedBoneChain( const FBoneChain& BoneChain, const FRetargetSkeleton& Skeleton,
+										TArray<int32>& OutBoneIndices)
 {
-	// validate start and end bones exist
+	// validate start and end bones exist and are not the root
 	const int32 StartIndex = Skeleton.FindBoneIndexByName(BoneChain.StartBone.BoneName);
 	const int32 EndIndex = Skeleton.FindBoneIndexByName(BoneChain.EndBone.BoneName);
-	bFoundStartBone = StartIndex != INDEX_NONE;
-	bFoundEndBone = EndIndex != INDEX_NONE;
+	bFoundStartBone = StartIndex > INDEX_NONE;
+	bFoundEndBone = EndIndex > INDEX_NONE;
 
-	if (bFoundStartBone && bFoundEndBone)
+	// no need to build the chain if start/end indices are wrong 
+	const bool bIsWellFormed = bFoundStartBone && bFoundEndBone && EndIndex >= StartIndex;
+	if (bIsWellFormed)
 	{
-		// validate end bone is child of start bone
-		bEndIsChildOfStart = true;
-		// record all bones in chain while walking up the hierarchy (tip to root of chain)
-		OutBoneIndices.Reset();
-		int32 NextBoneIndex = EndIndex;
-		while (true)
+		// init array with end bone 
+		OutBoneIndices = {EndIndex};
+
+		// if only one bone in the chain
+		if (EndIndex == StartIndex)
 		{
-			OutBoneIndices.Add(NextBoneIndex);
-		
-			if (NextBoneIndex == StartIndex)
-			{
-				break;
-			}
+			// validate end bone is child of start bone ?
+			bEndIsStartOrChildOfStart = true;
+			return;
+		}
 
-			NextBoneIndex = Skeleton.GetParentIndex(NextBoneIndex);
-		
-			if (NextBoneIndex == 0)
-			{
-				// oops, we walked all the way up to the root without finding the start bone
-				bEndIsChildOfStart = false;
-				OutBoneIndices.Reset();
-				return;
-			}
+		// record all bones in chain while walking up the hierarchy (tip to root of chain)
+		int32 ParentIndex = Skeleton.GetParentIndex(EndIndex);
+		while (ParentIndex > INDEX_NONE && ParentIndex >= StartIndex)
+		{
+			OutBoneIndices.Add(ParentIndex);
+			ParentIndex = Skeleton.GetParentIndex(ParentIndex);
 		}
 
-		// reverse the indices (we want root to tip order)
-		Algo::Reverse(OutBoneIndices);
+		// if we walked up till the start bone
+		if (OutBoneIndices.Last() == StartIndex)
+		{
+			// validate end bone is child of start bone
+			bEndIsStartOrChildOfStart = true;
+			// reverse the indices (we want root to tip order)
+			Algo::Reverse(OutBoneIndices);
+			return;
+		}
+      
+		// oops, we walked all the way up without finding the start bone
+		OutBoneIndices.Reset();
 	}
 }
 
@@ -740,7 +744,7 @@ bool FRetargetChainPair::ValidateBoneChainWithSkeletalMesh(
 	}
 
 	// warn if END bone was not a child of START bone
-	if (!ResolvedChain.bEndIsChildOfStart)
+	if (!ResolvedChain.bEndIsStartOrChildOfStart)
 	{
 		UE_LOG(LogTemp, Warning, TEXT("IK Retargeter bone chain, %s, end bone, '%s' was not a child of the start bone '%s'."),
 			*BoneChain.ChainName.ToString(), *BoneChain.EndBone.BoneName.ToString(), *BoneChain.StartBone.BoneName.ToString());
diff --git a/Engine/Plugins/Animation/IKRig/Source/IKRig/Public/Retargeter/IKRetargetProcessor.h b/Engine/Plugins/Animation/IKRig/Source/IKRig/Public/Retargeter/IKRetargetProcessor.h
--- a/Engine/Plugins/Animation/IKRig/Source/IKRig/Public/Retargeter/IKRetargetProcessor.h
+++ b/Engine/Plugins/Animation/IKRig/Source/IKRig/Public/Retargeter/IKRetargetProcessor.h
@@ -95,12 +95,12 @@ struct IKRIG_API FResolvedBoneChain
 	bool bFoundStartBone = false;
 	/* Does the END bone exist in the skeleton? */
 	bool bFoundEndBone = false;
-	/* Is the END bone a child of the START bone? */
-	bool bEndIsChildOfStart = false;
+	/* Is the END bone equals or a child of the START bone? */
+	bool bEndIsStartOrChildOfStart  = false;
 
 	bool IsValid() const
 	{
-		return bFoundStartBone && bFoundEndBone && bEndIsChildOfStart;
+		return bFoundStartBone && bFoundEndBone && bEndIsStartOrChildOfStart;
 	}
 };
 
