Adding ability to retarget root motion.

diff --git a/Engine/Plugins/Animation/IKRig/Source/IKRig/Private/Retargeter/IKRetargetProcessor.cpp b/Engine/Plugins/Animation/IKRig/Source/IKRig/Private/Retargeter/IKRetargetProcessor.cpp
--- a/Engine/Plugins/Animation/IKRig/Source/IKRig/Private/Retargeter/IKRetargetProcessor.cpp
+++ b/Engine/Plugins/Animation/IKRig/Source/IKRig/Private/Retargeter/IKRetargetProcessor.cpp
@@ -364,6 +364,7 @@ void FChainEncoderFK::EncodePose(
 }
 
 void FChainDecoderFK::DecodePose(
+	const FRootRetargeter& RootRetargeter,
 	const FRetargetChainSettings& Settings,
 	const TArray<int32>& TargetBoneIndices,
     const FChainEncoderFK& SourceChain,
@@ -426,8 +427,28 @@ void FChainDecoderFK::DecodePose(
 		{
 			ParentGlobalTransform = InOutGlobalPose[ParentIndex];
 		}
-		const FVector InitialLocalOffset = TargetSkeleton.RetargetLocalPose[BoneIndex].GetTranslation();
-		const FVector OutPosition = ParentGlobalTransform.TransformPosition(InitialLocalOffset);
+		FVector OutPosition;
+		switch (Settings.TranslationMode)
+		{
+			case ERetargetTranslationMode::None:
+				{
+					const FVector InitialLocalOffset = TargetSkeleton.RetargetLocalPose[BoneIndex].GetTranslation();
+					OutPosition = ParentGlobalTransform.TransformPosition(InitialLocalOffset);
+				}
+				break;
+			case ERetargetTranslationMode::GloballyScaled:
+				{
+					const FVector ScaledTranslation = SourceCurrentTransform.GetTranslation() * RootRetargeter.GlobalScale;
+					OutPosition = ScaledTranslation * Settings.TranslationMultiplier;
+				}
+				break;
+			case ERetargetTranslationMode::Absolute:
+				OutPosition = SourceCurrentTransform.GetTranslation() * Settings.TranslationMultiplier;
+				break;
+			default:
+				checkNoEntry();
+				break;
+		}
 
 		// calculate output SCALE
 		const FVector SourceCurrentScale = SourceCurrentTransform.GetScale3D();
@@ -1016,7 +1037,19 @@ bool UIKRetargetProcessor::InitializeRoots()
             *TargetRootBoneName.ToString(), *TargetSkeleton.SkeletalMesh->GetName());
 	}
 
-	return bRootEncoderInit && bRootDecoderInit;
+	const bool bWasInitialized = bRootEncoderInit && bRootDecoderInit;
+
+	// store global scale value
+	if (bWasInitialized)
+	{
+		RootRetargeter.GlobalScale = RootRetargeter.Source.InitialHeightInverse * RootRetargeter.Target.InitialHeight;
+	}
+	else
+	{
+		RootRetargeter.GlobalScale = 1.0f;
+	}
+	
+	return bWasInitialized;
 }
 
 bool UIKRetargetProcessor::InitializeBoneChainPairs()
@@ -1200,7 +1233,13 @@ void UIKRetargetProcessor::RunFKRetarget(
     for (FRetargetChainPairFK& ChainPair : ChainPairsFK)
     {
     	ChainPair.FKEncoder.EncodePose(ChainPair.SourceBoneIndices, InGlobalTransforms);
-    	ChainPair.FKDecoder.DecodePose(ChainPair.Settings, ChainPair.TargetBoneIndices, ChainPair.FKEncoder, TargetSkeleton,OutGlobalTransforms);
+    	ChainPair.FKDecoder.DecodePose(
+    		RootRetargeter,
+    		ChainPair.Settings,
+    		ChainPair.TargetBoneIndices,
+    		ChainPair.FKEncoder,
+    		TargetSkeleton,
+    		OutGlobalTransforms);
     }
 }
 
diff --git a/Engine/Plugins/Animation/IKRig/Source/IKRig/Public/Retargeter/IKRetargetProcessor.h b/Engine/Plugins/Animation/IKRig/Source/IKRig/Public/Retargeter/IKRetargetProcessor.h
--- a/Engine/Plugins/Animation/IKRig/Source/IKRig/Public/Retargeter/IKRetargetProcessor.h
+++ b/Engine/Plugins/Animation/IKRig/Source/IKRig/Public/Retargeter/IKRetargetProcessor.h
@@ -131,6 +131,8 @@ struct FRootRetargeter
 	FRootSource Source;
 	
 	FRootTarget Target;
+	
+	float GlobalScale = 1.0f;
 
 	void Reset();
 	
@@ -148,6 +150,8 @@ struct FRetargetChainSettings
 	FName TargetChainName;
 	
 	bool CopyPoseUsingFK = true;
+	ERetargetTranslationMode TranslationMode;
+	float TranslationMultiplier = 1.0f;
 
 	bool DriveIKGoal = true;
 	FVector StaticOffset;
@@ -164,6 +168,8 @@ public:
 		TargetChainName = AssetChainSettings->TargetChain;
 		
 		CopyPoseUsingFK = AssetChainSettings->CopyPoseUsingFK;
+		TranslationMode = AssetChainSettings->TranslationMode;
+		TranslationMultiplier = AssetChainSettings->TranslationMultiplier;
 		
 		DriveIKGoal = AssetChainSettings->DriveIKGoal;
 		Extension = AssetChainSettings->Extension;
@@ -208,6 +214,7 @@ struct FChainDecoderFK : public FChainFK
 		const FTargetSkeleton& TargetSkeleton);
 	
 	void DecodePose(
+		const FRootRetargeter& RootRetargeter,
 		const FRetargetChainSettings& Settings,
 		const TArray<int32>& TargetBoneIndices,
 		const FChainEncoderFK& SourceChain,
diff --git a/Engine/Plugins/Animation/IKRig/Source/IKRig/Public/Retargeter/IKRetargeter.h b/Engine/Plugins/Animation/IKRig/Source/IKRig/Public/Retargeter/IKRetargeter.h
--- a/Engine/Plugins/Animation/IKRig/Source/IKRig/Public/Retargeter/IKRetargeter.h
+++ b/Engine/Plugins/Animation/IKRig/Source/IKRig/Public/Retargeter/IKRetargeter.h
@@ -24,6 +24,14 @@ struct IKRIG_API FRetargetChainMap
 	FName TargetChain = NAME_None;
 };
 
+UENUM(BlueprintType)
+enum class ERetargetTranslationMode : uint8
+{
+	None			UMETA(DisplayName = "None"),
+	GloballyScaled	UMETA(DisplayName = "Globally Scaled"),
+	Absolute		UMETA(DisplayName = "Absolute"),
+};
+
 UCLASS()
 class IKRIG_API URetargetChainSettings: public UObject
 {
@@ -45,9 +53,22 @@ class IKRIG_API URetargetChainSettings: public UObject
 
 	/** Whether to copy the shape of the chain from the source skeleton (using retarget pose offsets). Default is true.
 	 * NOTE: This runs before the IK pass to copy the shape of the FK chain from the source skeleton. */
-	UPROPERTY(EditAnywhere, Category = "Retarget Passes")
+	UPROPERTY(EditAnywhere, Category = "FK Adjustments")
 	bool CopyPoseUsingFK = true;
 
+	/** Determines how translation is copied from the source chain to the target chain. Default is None.
+	* None: Translation of target bones are left unmodified from the retarget pose.
+	* Globally Scaled: Translation of target bone is set to the source bone offset multiplied by the global scale of the skeleton (determined by the relative height difference between retarget root bones).
+	* Absolute: Translation of target bone is set to the absolute position of the source bone. */
+	UPROPERTY(EditAnywhere, Category = "FK Adjustments")
+	ERetargetTranslationMode TranslationMode;
+
+	/** Range +/- inf. Default 1. Scales the amount of translation that is applied. Exact behavior depends on the Translation Mode.
+	*  In None Mode, this parameter has no effect.
+	*  In Globally Scaled and Absolute modes, the translation offset is scaled by this parameter.*/
+	UPROPERTY(EditAnywhere, Category = "FK Adjustments", meta = (UIMin = "0.0", UIMax = "1.0"))
+	float TranslationMultiplier = 1.0f;
+
 	/** Whether to run IK on this chain (IK runs AFTER the FK pass). Default is true.
 	 * NOTE: This only has an effect if the chain has an IK Goal assigned to it in the Target IK Rig asset.
 	 * NOTE: If off, and this chain has an IK Goal, the IK will still be evaluated, but the Goal is set to the input bone location.*/
