Fixed bug, exported retarget produces different results than editor preview.

diff --git a/Engine/Plugins/Animation/IKRig/Source/IKRigEditor/Private/RetargetEditor/IKRetargetBatchOperation.cpp b/Engine/Plugins/Animation/IKRig/Source/IKRigEditor/Private/RetargetEditor/IKRetargetBatchOperation.cpp
--- a/Engine/Plugins/Animation/IKRig/Source/IKRigEditor/Private/RetargetEditor/IKRetargetBatchOperation.cpp
+++ b/Engine/Plugins/Animation/IKRig/Source/IKRigEditor/Private/RetargetEditor/IKRetargetBatchOperation.cpp
@@ -334,13 +334,17 @@ void FIKRetargetBatchOperation::ConvertAnimation(
 			BoneTracks[TargetBoneIndex].RotKeys.SetNum(NumFrames);
 			BoneTracks[TargetBoneIndex].ScaleKeys.SetNum(NumFrames);
 		}
+
+		// ensure we evaluate the source animation using the skeletal mesh proportions that were evaluated in the viewport
+		FAnimPoseEvaluationOptions EvaluationOptions = FAnimPoseEvaluationOptions();
+		EvaluationOptions.OptionalSkeletalMesh = SourceSkeleton.SkeletalMesh;
 		
 		// retarget each frame's pose from source to target
 		for (int32 FrameIndex=0; FrameIndex<NumFrames; ++FrameIndex)
 		{
 			// get the source global pose
 			FAnimPose SourcePoseAtFrame;
-			UAnimPoseExtensions::GetAnimPoseAtFrame(SourceSequence, FrameIndex, FAnimPoseEvaluationOptions(), SourcePoseAtFrame);
+			UAnimPoseExtensions::GetAnimPoseAtFrame(SourceSequence, FrameIndex, EvaluationOptions, SourcePoseAtFrame);
 
 			// we don't use UAnimPoseExtensions::GetBoneNames as the sequence can store bones that only exist on the
 			// skeleton, but not on the current mesh. This results in indices discrepancy
diff --git a/Engine/Source/Editor/AnimationBlueprintLibrary/Private/AnimPose.cpp b/Engine/Source/Editor/AnimationBlueprintLibrary/Private/AnimPose.cpp
--- a/Engine/Source/Editor/AnimationBlueprintLibrary/Private/AnimPose.cpp
+++ b/Engine/Source/Editor/AnimationBlueprintLibrary/Private/AnimPose.cpp
@@ -487,8 +487,19 @@ void UAnimPoseExtensions::GetAnimPoseAtTime(const UAnimSequenceBase* AnimationSe
 				RequiredBoneIndexArray[BoneIndex] = BoneIndex;
 			}
 
+			// asset to use for retarget proportions (can be either USkeletalMesh or USkeleton)
+			UObject* AssetToUse;
+			if (EvaluationOptions.OptionalSkeletalMesh)
+			{
+				AssetToUse = CastChecked<UObject>(EvaluationOptions.OptionalSkeletalMesh);
+			}
+			else
+			{
+				AssetToUse = CastChecked<UObject>(AnimationSequenceBase->GetSkeleton());
+			}
+
 			FBoneContainer RequiredBones;
-			RequiredBones.InitializeTo(RequiredBoneIndexArray, FCurveEvaluationOption(false), *AnimationSequenceBase->GetSkeleton());
+			RequiredBones.InitializeTo(RequiredBoneIndexArray, FCurveEvaluationOption(false), *AssetToUse);
 			
 			RequiredBones.SetUseRAWData(EvaluationOptions.EvaluationType == EAnimDataEvalType::Raw);
 			RequiredBones.SetUseSourceData(EvaluationOptions.EvaluationType == EAnimDataEvalType::Source);
diff --git a/Engine/Source/Editor/AnimationBlueprintLibrary/Public/AnimPose.h b/Engine/Source/Editor/AnimationBlueprintLibrary/Public/AnimPose.h
--- a/Engine/Source/Editor/AnimationBlueprintLibrary/Public/AnimPose.h
+++ b/Engine/Source/Editor/AnimationBlueprintLibrary/Public/AnimPose.h
@@ -48,7 +48,11 @@ struct ANIMATIONBLUEPRINTLIBRARY_API FAnimPoseEvaluationOptions
 
 	// Whether or not to extract root motion values
 	UPROPERTY(EditAnywhere, AdvancedDisplay, BlueprintReadWrite, Category="Animation|Pose")
-	bool bExtractRootMotion = false;	
+	bool bExtractRootMotion = false;
+
+	// Optional skeletal mesh with proportions to use when evaluating a pose
+	UPROPERTY(EditAnywhere, AdvancedDisplay, BlueprintReadWrite, Category="Animation|Pose")
+	TObjectPtr<USkeletalMesh> OptionalSkeletalMesh = nullptr;
 };
 
 /** Script friendly representation of an evaluated animation bone pose */
