CIS Build Error: ICM: IKRetargetBatchOperation.cpp(52): error C4996: 'FAnimSegment::AnimReference'
#fix deprecation patch-up

diff --git a/Engine/Plugins/Animation/IKRig/Source/IKRigEditor/Private/RetargetEditor/IKRetargetBatchOperation.cpp b/Engine/Plugins/Animation/IKRig/Source/IKRigEditor/Private/RetargetEditor/IKRetargetBatchOperation.cpp
--- a/Engine/Plugins/Animation/IKRig/Source/IKRigEditor/Private/RetargetEditor/IKRetargetBatchOperation.cpp
+++ b/Engine/Plugins/Animation/IKRig/Source/IKRigEditor/Private/RetargetEditor/IKRetargetBatchOperation.cpp
@@ -49,9 +49,9 @@ int32 FIKRetargetBatchOperation::GenerateAssetLists(const FIKRetargetBatchOperat
 				{
 					for (const FAnimSegment& Segment: Track.AnimTrack.AnimSegments)
 					{
-						if (Segment.IsValid() && Segment.AnimReference)
+						if (Segment.IsValid() && Segment.GetAnimReference())
 						{
-							AnimationAssetsToRetarget.AddUnique(Segment.AnimReference);
+							AnimationAssetsToRetarget.AddUnique(Segment.GetAnimReference());
 						}
 					}
 				}
