Adding rotation mode options to retarget chains to help address slight errors in finger chains.

diff --git a/Engine/Plugins/Animation/IKRig/Source/IKRig/Private/Retargeter/IKRetargetProcessor.cpp b/Engine/Plugins/Animation/IKRig/Source/IKRig/Private/Retargeter/IKRetargetProcessor.cpp
--- a/Engine/Plugins/Animation/IKRig/Source/IKRig/Private/Retargeter/IKRetargetProcessor.cpp
+++ b/Engine/Plugins/Animation/IKRig/Source/IKRig/Private/Retargeter/IKRetargetProcessor.cpp
@@ -112,6 +112,20 @@ void FRetargetSkeleton::UpdateLocalTransformOfSingleBone(
 	OutLocalPose[BoneIndex] = ChildGlobalTransform.GetRelativeTransform(ParentGlobalTransform);
 }
 
+FTransform FRetargetSkeleton::GetGlobalRefPoseOfSingleBone(
+	const int32 BoneIndex,
+	const TArray<FTransform>& InGlobalPose) const
+{
+	const int32 ParentIndex = ParentIndices[BoneIndex];
+	if (ParentIndex == INDEX_NONE)
+	{
+		return RetargetLocalPose[BoneIndex]; // root always in global space
+	}
+	const FTransform& ChildLocalTransform = RetargetLocalPose[BoneIndex];
+	const FTransform& ParentGlobalTransform = InGlobalPose[ParentIndex];
+	return ChildLocalTransform * ParentGlobalTransform;
+}
+
 void FRetargetSkeleton::GetChildrenIndices(const int32 BoneIndex, TArray<int32>& OutChildren) const
 {
 	for (int32 ChildBoneIndex=0; ChildBoneIndex<ParentIndices.Num(); ++ChildBoneIndex)
@@ -290,23 +304,38 @@ void FTargetSkeleton::SetBoneIsRetargeted(const int32 BoneIndex, const bool IsRe
 	IsBoneRetargeted[BoneIndex] = IsRetargeted;
 }
 
-bool FChainFK::Initialize(const TArray<int32>& BoneIndices, const TArray<FTransform>& InitialGlobalPose)
+bool FChainFK::Initialize(
+	const FRetargetSkeleton& Skeleton,
+	const TArray<int32>& BoneIndices,
+	const TArray<FTransform>& InitialGlobalPose)
 {
 	check(!BoneIndices.IsEmpty());
 
 	// store all the initial bone transforms in the bone chain
-	InitialBoneTransforms.Reset();
+	InitialGlobalTransforms.Reset();
 	for (int32 Index=0; Index < BoneIndices.Num(); ++Index)
 	{
 		const int32 BoneIndex = BoneIndices[Index];
 		if (ensure(InitialGlobalPose.IsValidIndex(BoneIndex)))
 		{
-			InitialBoneTransforms.Emplace(InitialGlobalPose[BoneIndex]);
+			InitialGlobalTransforms.Emplace(InitialGlobalPose[BoneIndex]);
 		}
 	}
 
 	// initialize storage for current bones
-	CurrentBoneTransforms = InitialBoneTransforms;
+	CurrentGlobalTransforms = InitialGlobalTransforms;
+
+	// get the local space of the chain in retarget pose
+	InitialLocalTransforms.SetNum(InitialGlobalTransforms.Num());
+	FillTransformsWithLocalSpaceOfChain(Skeleton, InitialGlobalPose, BoneIndices, InitialLocalTransforms);
+
+	// store chain parent data
+	ChainParentBoneIndex = Skeleton.GetParentIndex(BoneIndices[0]);
+	ChainParentInitialGlobalTransform = FTransform::Identity;
+	if (ChainParentBoneIndex != INDEX_NONE)
+	{
+		ChainParentInitialGlobalTransform = InitialGlobalPose[ChainParentBoneIndex];
+	}
 
 	// calculate parameter of each bone, normalized by the length of the bone chain
 	return CalculateBoneParameters();
@@ -317,7 +346,7 @@ bool FChainFK::CalculateBoneParameters()
 	Params.Reset();
 	
 	// special case, a single-bone chain
-	if (InitialBoneTransforms.Num() == 1)
+	if (InitialGlobalTransforms.Num() == 1)
 	{
 		Params.Add(1.0f);
 		return true;
@@ -327,9 +356,9 @@ bool FChainFK::CalculateBoneParameters()
 	TArray<float> BoneDistances;
 	float TotalChainLength = 0.0f;
 	BoneDistances.Add(0.0f);
-	for (int32 i=1; i<InitialBoneTransforms.Num(); ++i)
+	for (int32 i=1; i<InitialGlobalTransforms.Num(); ++i)
 	{
-		TotalChainLength += (InitialBoneTransforms[i].GetTranslation() - InitialBoneTransforms[i-1].GetTranslation()).Size();
+		TotalChainLength += (InitialGlobalTransforms[i].GetTranslation() - InitialGlobalTransforms[i-1].GetTranslation()).Size();
 		BoneDistances.Add(TotalChainLength);
 	}
 
@@ -341,7 +370,7 @@ bool FChainFK::CalculateBoneParameters()
 	}
 
 	// calc each bone's param along length
-	for (int32 i=0; i<InitialBoneTransforms.Num(); ++i)
+	for (int32 i=0; i<InitialGlobalTransforms.Num(); ++i)
 	{
 		Params.Add(BoneDistances[i] / TotalChainLength); 
 	}
@@ -349,17 +378,91 @@ bool FChainFK::CalculateBoneParameters()
 	return true;
 }
 
+void FChainFK::FillTransformsWithLocalSpaceOfChain(
+	const FRetargetSkeleton& Skeleton,
+	const TArray<FTransform>& InGlobalPose,
+	const TArray<int32>& BoneIndices,
+	TArray<FTransform>& OutLocalTransforms)
+{
+	check(BoneIndices.Num() == OutLocalTransforms.Num())
+	
+	for (int32 ChainIndex=0; ChainIndex<BoneIndices.Num(); ++ChainIndex)
+	{
+		const int32 BoneIndex = BoneIndices[ChainIndex];
+		const int32 ParentIndex = Skeleton.GetParentIndex(BoneIndex);
+		if (ParentIndex == INDEX_NONE)
+		{
+			// root is always in "global" space
+			OutLocalTransforms[ChainIndex] = InGlobalPose[BoneIndex];
+			continue;
+		}
+
+		const FTransform& ChildGlobalTransform = InGlobalPose[BoneIndex];
+		const FTransform& ParentGlobalTransform = InGlobalPose[ParentIndex];
+		OutLocalTransforms[ChainIndex] = ChildGlobalTransform.GetRelativeTransform(ParentGlobalTransform);
+	}
+}
+
+void FChainFK::PutCurrentTransformsInRefPose(
+	const TArray<int32>& BoneIndices,
+	const FRetargetSkeleton& Skeleton,
+	const TArray<FTransform>& InCurrentGlobalPose)
+{
+	// update chain current transforms to the retarget pose in global space
+	for (int32 ChainIndex=0; ChainIndex<BoneIndices.Num(); ++ChainIndex)
+	{
+		// update first bone in chain based on the incoming parent
+		if (ChainIndex == 0)
+		{
+			const int32 BoneIndex = BoneIndices[ChainIndex];
+			CurrentGlobalTransforms[ChainIndex] = Skeleton.GetGlobalRefPoseOfSingleBone(BoneIndex, InCurrentGlobalPose);
+		}
+		else
+		{
+			// all subsequent bones in chain are based on previous parent
+			const int32 BoneIndex = BoneIndices[ChainIndex];
+			const FTransform& ParentGlobalTransform = CurrentGlobalTransforms[ChainIndex-1];
+			const FTransform& ChildLocalTransform = Skeleton.RetargetLocalPose[BoneIndex];
+			CurrentGlobalTransforms[ChainIndex] = ChildLocalTransform * ParentGlobalTransform;
+		}
+	}
+}
+
 void FChainEncoderFK::EncodePose(
+	const FRetargetSkeleton& SourceSkeleton,
 	const TArray<int32>& SourceBoneIndices,
-    const TArray<FTransform> &InputGlobalPose)
+    const TArray<FTransform> &InSourceGlobalPose)
 {
-	check(SourceBoneIndices.Num() == CurrentBoneTransforms.Num());
+	check(SourceBoneIndices.Num() == CurrentGlobalTransforms.Num());
 	
-	// copy the input pose for the chain
+	// copy the global input pose for the chain
 	for (int32 ChainIndex=0; ChainIndex<SourceBoneIndices.Num(); ++ChainIndex)
 	{
 		const int32 BoneIndex = SourceBoneIndices[ChainIndex];
-		CurrentBoneTransforms[ChainIndex] = InputGlobalPose[BoneIndex];
+		CurrentGlobalTransforms[ChainIndex] = InSourceGlobalPose[BoneIndex];
+	}
+
+	CurrentLocalTransforms.SetNum(SourceBoneIndices.Num());
+	FillTransformsWithLocalSpaceOfChain(SourceSkeleton, InSourceGlobalPose, SourceBoneIndices, CurrentLocalTransforms);
+
+	if (ChainParentBoneIndex != INDEX_NONE)
+	{
+		ChainParentCurrentGlobalTransform = InSourceGlobalPose[ChainParentBoneIndex];
+	}
+}
+
+void FChainEncoderFK::TransformCurrentChainTransforms(const FTransform& NewParentTransform)
+{
+	for (int32 ChainIndex=0; ChainIndex<CurrentGlobalTransforms.Num(); ++ChainIndex)
+	{
+		if (ChainIndex == 0)
+		{
+			CurrentGlobalTransforms[ChainIndex] = CurrentLocalTransforms[ChainIndex] * NewParentTransform;
+		}
+		else
+		{
+			CurrentGlobalTransforms[ChainIndex] = CurrentLocalTransforms[ChainIndex] * CurrentGlobalTransforms[ChainIndex-1];
+		}
 	}
 }
 
@@ -367,11 +470,11 @@ void FChainDecoderFK::DecodePose(
 	const FRootRetargeter& RootRetargeter,
 	const FRetargetChainSettings& Settings,
 	const TArray<int32>& TargetBoneIndices,
-    const FChainEncoderFK& SourceChain,
+    FChainEncoderFK& SourceChain,
     const FTargetSkeleton& TargetSkeleton,
     TArray<FTransform> &InOutGlobalPose)
 {
-	check(TargetBoneIndices.Num() == CurrentBoneTransforms.Num());
+	check(TargetBoneIndices.Num() == CurrentGlobalTransforms.Num());
 	check(TargetBoneIndices.Num() == Params.Num());
 
 	// Before setting this chain pose, we need to ensure that any
@@ -383,46 +486,113 @@ void FChainDecoderFK::DecodePose(
 	// Otherwise the neck bones will remain at their location prior to the spine update.
 	UpdateIntermediateParents(TargetSkeleton,InOutGlobalPose);
 
+	// transform entire source chain from it's root to match target's current root orientation (maintaining offset from retarget pose)
+	// this ensures children are retargeted in a "local" manner free from skewing that will happen if source and target
+	// become misaligned as can happen if parent chains were not retargeted
+	FTransform SourceChainParentInitialDelta = SourceChain.ChainParentInitialGlobalTransform.GetRelativeTransform(ChainParentInitialGlobalTransform);
+	FTransform TargetChainParentCurrentGlobalTransform = ChainParentBoneIndex == INDEX_NONE ? FTransform::Identity : InOutGlobalPose[ChainParentBoneIndex]; 
+	FTransform SourceChainParentTransform = SourceChainParentInitialDelta * TargetChainParentCurrentGlobalTransform;
+
+	// apply delta to the source chain's current transforms before transferring rotations to the target
+	SourceChain.TransformCurrentChainTransforms(SourceChainParentTransform);
+
 	// if FK retargeting has been disabled for this chain, then simply set it to the retarget pose
 	if (!Settings.CopyPoseUsingFK)
 	{
-		for (const int32& BoneIndex : TargetBoneIndices)
+		// put the chain in the global ref pose (globally rotated by parent bone in it's currently retargeted state)
+		PutCurrentTransformsInRefPose(TargetBoneIndices, TargetSkeleton, InOutGlobalPose);
+		
+		for (int32 ChainIndex=0; ChainIndex<TargetBoneIndices.Num(); ++ChainIndex)
 		{
-			TargetSkeleton.UpdateGlobalTransformOfSingleBone(BoneIndex, TargetSkeleton.RetargetLocalPose, InOutGlobalPose);
+			const int32 BoneIndex = TargetBoneIndices[ChainIndex];
+			InOutGlobalPose[BoneIndex] = CurrentGlobalTransforms[ChainIndex];
 		}
+
 		return;
 	}
 
+	const int32 NumBonesInSourceChain = SourceChain.CurrentGlobalTransforms.Num();
+	const int32 NumBonesInTargetChain = TargetBoneIndices.Num();
+	const int32 TargetStartIndex = FMath::Max(0, NumBonesInTargetChain - NumBonesInSourceChain);
+	const int32 SourceStartIndex = FMath::Max(0,NumBonesInSourceChain - NumBonesInTargetChain);
+
 	// now retarget the pose of each bone in the chain, copying from source to target
 	for (int32 ChainIndex=0; ChainIndex<TargetBoneIndices.Num(); ++ChainIndex)
 	{
-		// get the initial and current transform of source chain at param
-		// this is the interpolated transform along the chain
-		const float Param = Params[ChainIndex];
-		FTransform SourceCurrentTransform = GetTransformAtParam(
-			SourceChain.CurrentBoneTransforms,
-			SourceChain.Params,
-			Param);
-		FTransform SourceInitialTransform = GetTransformAtParam(
-			SourceChain.InitialBoneTransforms,
-			SourceChain.Params,
-			Param);
-
 		const int32 BoneIndex = TargetBoneIndices[ChainIndex];
-		const int32 ParentIndex = TargetSkeleton.ParentIndices[BoneIndex];
-		
-		// the initial transform on the target chain
-		const FTransform& TargetInitialTransform = InitialBoneTransforms[ChainIndex];
+		const FTransform& TargetInitialTransform = InitialGlobalTransforms[ChainIndex];
+		FTransform SourceCurrentTransform;
+		FTransform SourceInitialTransform;
 
-		// calculate output ROTATION
+		// get source current / initial transforms for this bone
+		switch (Settings.RotationMode)
+		{
+			case ERetargetRotationMode::Interpolated:
+			{
+				// get the initial and current transform of source chain at param
+				// this is the interpolated transform along the chain
+				const float Param = Params[ChainIndex];
+					
+				SourceCurrentTransform = GetTransformAtParam(
+					SourceChain.CurrentGlobalTransforms,
+					SourceChain.Params,
+					Param);
+
+				SourceInitialTransform = GetTransformAtParam(
+					SourceChain.InitialGlobalTransforms,
+					SourceChain.Params,
+					Param);
+			}
+			break;
+			case ERetargetRotationMode::OneToOne:
+			{
+				if (ChainIndex < NumBonesInSourceChain)
+				{
+					SourceCurrentTransform = SourceChain.CurrentGlobalTransforms[ChainIndex];
+					SourceInitialTransform = SourceChain.InitialGlobalTransforms[ChainIndex];
+				}else
+				{
+					SourceCurrentTransform = SourceChain.CurrentGlobalTransforms.Last();
+					SourceInitialTransform = SourceChain.InitialGlobalTransforms.Last();
+				}
+			}
+			break;
+			case ERetargetRotationMode::OneToOneReversed:
+			{
+				if (ChainIndex < TargetStartIndex)
+				{
+					SourceCurrentTransform = SourceChain.InitialGlobalTransforms[0];
+					SourceInitialTransform = SourceChain.InitialGlobalTransforms[0];
+				}
+				else
+				{
+					const int32 SourceChainIndex = SourceStartIndex + (ChainIndex - TargetStartIndex);
+					SourceCurrentTransform = SourceChain.CurrentGlobalTransforms[SourceChainIndex];
+					SourceInitialTransform = SourceChain.InitialGlobalTransforms[SourceChainIndex];
+				}
+			}
+			break;
+			case ERetargetRotationMode::None:
+			{
+				SourceCurrentTransform = SourceChain.InitialGlobalTransforms.Last();
+				SourceInitialTransform = SourceChain.InitialGlobalTransforms.Last();
+			}
+			break;
+			default:
+				checkNoEntry();
+			break;
+		}
+		
+		// apply rotation offset to the initial target rotation
 		const FQuat SourceCurrentRotation = SourceCurrentTransform.GetRotation();
 		const FQuat SourceInitialRotation = SourceInitialTransform.GetRotation();
 		const FQuat RotationDelta = SourceCurrentRotation * SourceInitialRotation.Inverse();
 		const FQuat TargetInitialRotation = TargetInitialTransform.GetRotation();
 		const FQuat OutRotation = RotationDelta * TargetInitialRotation;
 
-		// calculate output POSITION (constant for now, maybe we support retargeting translation later?)
+		// calculate output POSITION based on translation mode setting
 		FTransform ParentGlobalTransform = FTransform::Identity;
+		const int32 ParentIndex = TargetSkeleton.ParentIndices[BoneIndex];
 		if (ParentIndex != INDEX_NONE)
 		{
 			ParentGlobalTransform = InOutGlobalPose[ParentIndex];
@@ -438,12 +608,11 @@ void FChainDecoderFK::DecodePose(
 				break;
 			case ERetargetTranslationMode::GloballyScaled:
 				{
-					const FVector ScaledTranslation = SourceCurrentTransform.GetTranslation() * RootRetargeter.GlobalScale;
-					OutPosition = ScaledTranslation * Settings.TranslationMultiplier;
+					OutPosition = SourceCurrentTransform.GetTranslation() * RootRetargeter.GlobalScale; // todo expose global scale modifier
 				}
 				break;
 			case ERetargetTranslationMode::Absolute:
-				OutPosition = SourceCurrentTransform.GetTranslation() * Settings.TranslationMultiplier;
+				OutPosition = SourceCurrentTransform.GetTranslation();
 				break;
 			default:
 				checkNoEntry();
@@ -457,8 +626,32 @@ void FChainDecoderFK::DecodePose(
 		const FVector OutScale = SourceCurrentScale + (TargetInitialScale - SourceInitialScale);
 		
 		// apply output transform
-		FTransform OutTransform = FTransform(OutRotation, OutPosition, OutScale);
-		InOutGlobalPose[BoneIndex] = OutTransform;
+		InOutGlobalPose[BoneIndex] = FTransform(OutRotation, OutPosition, OutScale);
+	}
+
+	// apply final blending between retarget pose of chain and newly retargeted pose
+	// blend must be done in local space, so we do it in a separate loop after full chain pose is generated
+	// (skipped if the alphas are not near 1.0)
+	if (!FMath::IsNearlyEqual(Settings.RotationAlpha, 1.0f) || !FMath::IsNearlyEqual(Settings.TranslationAlpha, 1.0f))
+	{
+		TArray<FTransform> NewLocalTransforms;
+		NewLocalTransforms.SetNum(InitialLocalTransforms.Num());
+		FillTransformsWithLocalSpaceOfChain(TargetSkeleton, InOutGlobalPose, TargetBoneIndices, NewLocalTransforms);
+
+		for (int32 ChainIndex=0; ChainIndex<InitialLocalTransforms.Num(); ++ChainIndex)
+		{
+			// blend between current local pose and initial local pose
+			FTransform& NewLocalTransform = NewLocalTransforms[ChainIndex];
+			const FTransform& RefPoseLocalTransform = InitialLocalTransforms[ChainIndex];
+			NewLocalTransform.SetTranslation(FMath::Lerp(RefPoseLocalTransform.GetTranslation(), NewLocalTransform.GetTranslation(), Settings.TranslationAlpha));
+			NewLocalTransform.SetRotation(FQuat::FastLerp(RefPoseLocalTransform.GetRotation(), NewLocalTransform.GetRotation(), Settings.RotationAlpha).GetNormalized());
+
+			// put blended transforms back in global space and store in final output pose
+			const int32 BoneIndex = TargetBoneIndices[ChainIndex];
+			const int32 ParentIndex = TargetSkeleton.ParentIndices[BoneIndex];
+			const FTransform& ParentGlobalTransform = ParentIndex == INDEX_NONE ? FTransform::Identity : InOutGlobalPose[ParentIndex];
+			InOutGlobalPose[BoneIndex] = NewLocalTransform * ParentGlobalTransform;
+		}
 	}
 }
 
@@ -788,7 +981,7 @@ bool FRetargetChainPairFK::Initialize(
 	}
 
 	// initialize SOURCE FK chain encoder with retarget pose
-	const bool bFKEncoderInitialized = FKEncoder.Initialize(SourceBoneIndices, SourceSkeleton.RetargetGlobalPose);
+	const bool bFKEncoderInitialized = FKEncoder.Initialize(SourceSkeleton, SourceBoneIndices, SourceSkeleton.RetargetGlobalPose);
 	if (!bFKEncoderInitialized)
 	{
 		UE_LOG(LogTemp, Warning, TEXT("IK Retargeter failed to initialize FK encoder, '%s', on Skeletal Mesh: '%s'"),
@@ -797,7 +990,7 @@ bool FRetargetChainPairFK::Initialize(
 	}
 
 	// initialize TARGET FK chain decoder with retarget pose
-	const bool bFKDecoderInitialized = FKDecoder.Initialize(TargetBoneIndices, TargetSkeleton.RetargetGlobalPose);
+	const bool bFKDecoderInitialized = FKDecoder.Initialize(TargetSkeleton, TargetBoneIndices, TargetSkeleton.RetargetGlobalPose);
 	if (!bFKDecoderInitialized)
 	{
 		UE_LOG(LogTemp, Warning, TEXT("IK Retargeter failed to initialize FK decoder, '%s', on Skeletal Mesh: '%s'"),
@@ -1037,19 +1230,7 @@ bool UIKRetargetProcessor::InitializeRoots()
             *TargetRootBoneName.ToString(), *TargetSkeleton.SkeletalMesh->GetName());
 	}
 
-	const bool bWasInitialized = bRootEncoderInit && bRootDecoderInit;
-
-	// store global scale value
-	if (bWasInitialized)
-	{
-		RootRetargeter.GlobalScale = RootRetargeter.Source.InitialHeightInverse * RootRetargeter.Target.InitialHeight;
-	}
-	else
-	{
-		RootRetargeter.GlobalScale = 1.0f;
-	}
-	
-	return bWasInitialized;
+	return bRootEncoderInit && bRootDecoderInit;
 }
 
 bool UIKRetargetProcessor::InitializeBoneChainPairs()
@@ -1230,17 +1411,21 @@ void UIKRetargetProcessor::RunFKRetarget(
     TArray<FTransform>& OutGlobalTransforms)
 {
 	// spin through chains and encode/decode them all using the input pose
-    for (FRetargetChainPairFK& ChainPair : ChainPairsFK)
-    {
-    	ChainPair.FKEncoder.EncodePose(ChainPair.SourceBoneIndices, InGlobalTransforms);
-    	ChainPair.FKDecoder.DecodePose(
-    		RootRetargeter,
-    		ChainPair.Settings,
-    		ChainPair.TargetBoneIndices,
-    		ChainPair.FKEncoder,
-    		TargetSkeleton,
-    		OutGlobalTransforms);
-    }
+	for (FRetargetChainPairFK& ChainPair : ChainPairsFK)
+	{
+		ChainPair.FKEncoder.EncodePose(
+			SourceSkeleton,
+			ChainPair.SourceBoneIndices,
+			InGlobalTransforms);
+		
+		ChainPair.FKDecoder.DecodePose(
+			RootRetargeter,
+			ChainPair.Settings,
+			ChainPair.TargetBoneIndices,
+			ChainPair.FKEncoder,
+			TargetSkeleton,
+			OutGlobalTransforms);
+	}
 }
 
 void UIKRetargetProcessor::RunIKRetarget(
diff --git a/Engine/Plugins/Animation/IKRig/Source/IKRig/Public/Retargeter/IKRetargetProcessor.h b/Engine/Plugins/Animation/IKRig/Source/IKRig/Public/Retargeter/IKRetargetProcessor.h
--- a/Engine/Plugins/Animation/IKRig/Source/IKRig/Public/Retargeter/IKRetargetProcessor.h
+++ b/Engine/Plugins/Animation/IKRig/Source/IKRig/Public/Retargeter/IKRetargetProcessor.h
@@ -58,6 +58,10 @@ struct IKRIG_API FRetargetSkeleton
 		TArray<FTransform>& OutLocalPose,
 		const TArray<FTransform>& InGlobalPose) const;
 
+	FTransform GetGlobalRefPoseOfSingleBone(
+		const int32 BoneIndex,
+		const TArray<FTransform>& InGlobalPose) const;
+
 	void GetChildrenIndices(const int32 BoneIndex, TArray<int32>& OutChildren) const;
 };
 
@@ -131,7 +135,7 @@ struct FRootRetargeter
 	FRootSource Source;
 	
 	FRootTarget Target;
-	
+
 	float GlobalScale = 1.0f;
 
 	void Reset();
@@ -148,10 +152,12 @@ struct FRootRetargeter
 struct FRetargetChainSettings
 {
 	FName TargetChainName;
-	
+
 	bool CopyPoseUsingFK = true;
+	ERetargetRotationMode RotationMode;
+	float RotationAlpha = 1.0f;
 	ERetargetTranslationMode TranslationMode;
-	float TranslationMultiplier = 1.0f;
+	float TranslationAlpha = 1.0f;
 
 	bool DriveIKGoal = true;
 	FVector StaticOffset;
@@ -166,10 +172,12 @@ public:
 	void CopySettingsFromAsset(const URetargetChainSettings* AssetChainSettings)
 	{
 		TargetChainName = AssetChainSettings->TargetChain;
-		
+
 		CopyPoseUsingFK = AssetChainSettings->CopyPoseUsingFK;
+		RotationMode = AssetChainSettings->RotationMode;
+		RotationAlpha = AssetChainSettings->RotationAlpha;
 		TranslationMode = AssetChainSettings->TranslationMode;
-		TranslationMultiplier = AssetChainSettings->TranslationMultiplier;
+		TranslationAlpha = AssetChainSettings->TranslationAlpha;
 		
 		DriveIKGoal = AssetChainSettings->DriveIKGoal;
 		Extension = AssetChainSettings->Extension;
@@ -184,26 +192,52 @@ public:
 
 struct FChainFK
 {
-	TArray<FTransform> InitialBoneTransforms;
+	TArray<FTransform> InitialGlobalTransforms;
+
+	TArray<FTransform> InitialLocalTransforms;
 
-	TArray<FTransform> CurrentBoneTransforms;
+	TArray<FTransform> CurrentGlobalTransforms;
 
 	TArray<float> Params;
 
+	int32 ChainParentBoneIndex;
+	FTransform ChainParentInitialGlobalTransform;
+
 	bool Initialize(
+		const FRetargetSkeleton& Skeleton,
 		const TArray<int32>& BoneIndices,
 		const TArray<FTransform> &InitialGlobalPose);
 
 private:
 	
-	bool CalculateBoneParameters();	
+	bool CalculateBoneParameters();
+
+protected:
+
+	static void FillTransformsWithLocalSpaceOfChain(
+		const FRetargetSkeleton& Skeleton,
+		const TArray<FTransform>& InGlobalPose,
+		const TArray<int32>& BoneIndices,
+		TArray<FTransform>& OutLocalTransforms);
+
+	void PutCurrentTransformsInRefPose(
+		const TArray<int32>& BoneIndices,
+		const FRetargetSkeleton& Skeleton,
+		const TArray<FTransform>& InCurrentGlobalPose);
 };
 
 struct FChainEncoderFK : public FChainFK
 {
+	TArray<FTransform> CurrentLocalTransforms;
+
+	FTransform ChainParentCurrentGlobalTransform;
+	
 	void EncodePose(
+		const FRetargetSkeleton& SourceSkeleton,
 		const TArray<int32>& SourceBoneIndices,
-		const TArray<FTransform> &InputGlobalPose);
+		const TArray<FTransform> &InSourceGlobalPose);
+
+	void TransformCurrentChainTransforms(const FTransform& NewParentTransform);
 };
 
 struct FChainDecoderFK : public FChainFK
@@ -217,7 +251,7 @@ struct FChainDecoderFK : public FChainFK
 		const FRootRetargeter& RootRetargeter,
 		const FRetargetChainSettings& Settings,
 		const TArray<int32>& TargetBoneIndices,
-		const FChainEncoderFK& SourceChain,
+		FChainEncoderFK& SourceChain,
 		const FTargetSkeleton& TargetSkeleton,
 		TArray<FTransform> &InOutGlobalPose);
 
diff --git a/Engine/Plugins/Animation/IKRig/Source/IKRig/Public/Retargeter/IKRetargeter.h b/Engine/Plugins/Animation/IKRig/Source/IKRig/Public/Retargeter/IKRetargeter.h
--- a/Engine/Plugins/Animation/IKRig/Source/IKRig/Public/Retargeter/IKRetargeter.h
+++ b/Engine/Plugins/Animation/IKRig/Source/IKRig/Public/Retargeter/IKRetargeter.h
@@ -32,6 +32,16 @@ enum class ERetargetTranslationMode : uint8
 	Absolute		UMETA(DisplayName = "Absolute"),
 };
 
+UENUM(BlueprintType)
+enum class ERetargetRotationMode : uint8
+{
+	
+	Interpolated		UMETA(DisplayName = "Interpolated"),
+	OneToOne			UMETA(DisplayName = "One to One"),
+	OneToOneReversed	UMETA(DisplayName = "One to One Reversed"),
+	None				UMETA(DisplayName = "None"),
+};
+
 UCLASS()
 class IKRIG_API URetargetChainSettings: public UObject
 {
@@ -51,10 +61,24 @@ class IKRIG_API URetargetChainSettings: public UObject
 	UPROPERTY(VisibleAnywhere, Category = "Chain Mapping")
 	FName TargetChain = NAME_None;
 
-	/** Whether to copy the shape of the chain from the source skeleton (using retarget pose offsets). Default is true.
-	 * NOTE: This runs before the IK pass to copy the shape of the FK chain from the source skeleton. */
+	/** Whether to copy the shape of the chain from the source skeleton using the Rotation and Translation modes. Default is true.
+	* NOTE: All FK operations run before the IK pass to copy the shape of the FK chain from the source skeleton. */
 	UPROPERTY(EditAnywhere, Category = "FK Adjustments")
 	bool CopyPoseUsingFK = true;
+	
+	/** Determines how rotation is copied from the source chain to the target chain. Default is Interpolated.
+	* Interpolated: Source and target chains are normalized by their length, then each target bone rotation is generated by finding the rotation at the same normalized distance on the source chain and interpolating between the neighboring bones.
+	* One to One: Each target bone rotation is copied from the equivalent bone in the source chain, based on the order in the chain, starting at the root of the chain. If the target chain has more bones than the source, the extra bones will remain at their reference pose.
+	* One to One Reversed: Same as One-to-One, but starting from the tip of the chain.
+	* None: The rotation of each target bone in the chain is left at the reference pose. */
+	UPROPERTY(EditAnywhere, Category = "FK Adjustments")
+	ERetargetRotationMode RotationMode;
+
+	/** Range +/- infinity. Default 1. Scales the amount of rotation that is applied.
+	*  If Rotation Mode is None this parameter has no effect.
+	*  Otherwise, this parameter blends the rotation of each bone in the chain from the base retarget pose (0) to the retargeted pose (1).*/
+	UPROPERTY(EditAnywhere, Category = "FK Adjustments", meta = (UIMin = "0.0", UIMax = "1.0"))
+	float RotationAlpha = 1.0f;
 
 	/** Determines how translation is copied from the source chain to the target chain. Default is None.
 	* None: Translation of target bones are left unmodified from the retarget pose.
@@ -63,15 +87,15 @@ class IKRIG_API URetargetChainSettings: public UObject
 	UPROPERTY(EditAnywhere, Category = "FK Adjustments")
 	ERetargetTranslationMode TranslationMode;
 
-	/** Range +/- inf. Default 1. Scales the amount of translation that is applied. Exact behavior depends on the Translation Mode.
+	/** Range +/- infinity. Default 1. Scales the amount of translation that is applied. Exact behavior depends on the Translation Mode.
 	*  In None Mode, this parameter has no effect.
 	*  In Globally Scaled and Absolute modes, the translation offset is scaled by this parameter.*/
 	UPROPERTY(EditAnywhere, Category = "FK Adjustments", meta = (UIMin = "0.0", UIMax = "1.0"))
-	float TranslationMultiplier = 1.0f;
+	float TranslationAlpha = 1.0f;
 
-	/** Whether to run IK on this chain (IK runs AFTER the FK pass). Default is true.
+	/** Whether to modify the location of the IK goal on this chain. Default is true.
 	 * NOTE: This only has an effect if the chain has an IK Goal assigned to it in the Target IK Rig asset.
-	 * NOTE: If off, and this chain has an IK Goal, the IK will still be evaluated, but the Goal is set to the input bone location.*/
+	 * NOTE: If off, and this chain has an IK Goal, the IK will still be evaluated, but the Goal is set to the input bone location (from the FK pass).*/
 	UPROPERTY(EditAnywhere, Category = "IK Adjustments")
 	bool DriveIKGoal = true;
 
@@ -273,5 +297,3 @@ private:
 
 	friend class UIKRetargeterController;
 };
-
-
