Fix incorrect local rotation axes when editing retarget pose.

diff --git a/Engine/Plugins/Animation/IKRig/Source/IKRigEditor/Private/RetargetEditor/IKRetargetEditMode.cpp b/Engine/Plugins/Animation/IKRig/Source/IKRigEditor/Private/RetargetEditor/IKRetargetEditMode.cpp
--- a/Engine/Plugins/Animation/IKRig/Source/IKRigEditor/Private/RetargetEditor/IKRetargetEditMode.cpp
+++ b/Engine/Plugins/Animation/IKRig/Source/IKRigEditor/Private/RetargetEditor/IKRetargetEditMode.cpp
@@ -10,6 +10,7 @@
 #include "IPersonaPreviewScene.h"
 #include "Animation/DebugSkelMeshComponent.h"
 #include "Materials/MaterialInstanceDynamic.h"
+#include "RetargetEditor/IKRetargetAnimInstance.h"
 #include "RetargetEditor/IKRetargetHitProxies.h"
 #include "RetargetEditor/IKRetargetEditor.h"
 #include "Retargeter/IKRetargetProcessor.h"
@@ -54,13 +55,13 @@ void FIKRetargetEditMode::Render(const FSceneView* View, FViewport* Viewport, FP
 		return;
 	}
 
-	if (!Controller->GetRetargetProcessor()->IsInitialized())
+	const UIKRetargetProcessor* RetargetProcessor = Controller->GetRetargetProcessor();
+	if (!(RetargetProcessor && RetargetProcessor->IsInitialized()))
 	{
 		return;
 	}
 
-	const UIKRetargeter* Retargeter = Controller->AssetController->GetAsset();
-	const UIKRetargetProcessor* RetargetProcessor = Controller->GetRetargetProcessor();
+	const UIKRetargeter* Asset = Controller->AssetController->GetAsset();
 	const FTargetSkeleton& TargetSkeleton = RetargetProcessor->GetTargetSkeleton();
 	const int32 RootBoneIndex = RetargetProcessor->GetTargetSkeletonRootBone();
 
@@ -86,15 +87,15 @@ void FIKRetargetEditMode::Render(const FSceneView* View, FViewport* Viewport, FP
 		FLinearColor LineColor = bIsAffected ? IKRigDebugRendering::AFFECTED_BONE_COLOR : IKRigDebugRendering::DESELECTED_BONE_COLOR;
 		LineColor = bIsSelected ? IKRigDebugRendering::SELECTED_BONE_COLOR : LineColor;
 
-		const float BoneRadiusSetting = Retargeter->BoneDrawSize;
+		const float BoneRadiusSetting = Asset->BoneDrawSize;
 		const float BoneRadius = FMath::Min(1.0f, MaxDrawRadius) * BoneRadiusSetting;
 		
 		// get the location of the bone on the currently initialized target skeletal mesh
 		// along with array of child positions
-		const FTransform BoneTransform = Controller->GetTargetBoneGlobalTransform(BoneIndex);
+		const FTransform BoneTransform = Controller->GetTargetBoneGlobalTransform(RetargetProcessor, BoneIndex);
 		FVector Start;
 		TArray<FVector> ChildLocations;
-		Controller->GetTargetBoneLineSegments(BoneIndex, Start, ChildLocations);
+		Controller->GetTargetBoneLineSegments(RetargetProcessor, BoneIndex, Start, ChildLocations);
 		
 		// draw the bone proxy
 		PDI->SetHitProxy(new HIKRetargetEditorBoneProxy(TargetSkeleton.BoneNames[BoneIndex]));
@@ -198,8 +199,14 @@ FVector FIKRetargetEditMode::GetWidgetLocation() const
 	{
 		return FVector::ZeroVector;
 	}
+
+	const UIKRetargetProcessor* RetargetProcessor = Controller->GetRetargetProcessor();
+	if (!(RetargetProcessor && RetargetProcessor->IsInitialized()))
+	{
+		return FVector::ZeroVector;
+	}
 	
-	return Controller->GetTargetBoneGlobalTransform(BoneIndex).GetTranslation();
+	return Controller->GetTargetBoneGlobalTransform(RetargetProcessor, BoneIndex).GetTranslation();
 }
 
 bool FIKRetargetEditMode::HandleClick(FEditorViewportClient* InViewportClient, HHitProxy* HitProxy, const FViewportClick& Click)
@@ -401,12 +408,18 @@ void FIKRetargetEditMode::UpdateWidgetTransform()
 		return;
 	}
 
-	UIKRetargeterController* AssetController = Controller->AssetController;
-	const FIKRigSkeleton& Skeleton = AssetController->GetAsset()->GetTargetIKRig()->Skeleton;
+	const UIKRetargetProcessor* RetargetProcessor = Controller->GetRetargetProcessor();
+	if (!(RetargetProcessor && RetargetProcessor->IsInitialized()))
+	{
+		return;
+	}
+	
+	const UIKRetargeterController* AssetController = Controller->AssetController;
+	const FTargetSkeleton& TargetSkeleton = RetargetProcessor->GetTargetSkeleton();
 
 	BoneEdit.Name = BoneEdit.SelectedBones.Last();
-	BoneEdit.Index = Skeleton.GetBoneIndexFromName(BoneEdit.Name);
-	BoneEdit.GlobalTransform = Controller->GetTargetBoneGlobalTransform(BoneEdit.Index);
+	BoneEdit.Index = TargetSkeleton.FindBoneIndexByName(BoneEdit.Name);
+	BoneEdit.GlobalTransform = Controller->GetTargetBoneGlobalTransform(RetargetProcessor, BoneEdit.Index);
 	BoneEdit.AccumulatedGlobalOffset = FQuat::Identity;
 
 	BoneEdit.PrevLocalOffsets.Reset();
@@ -416,10 +429,10 @@ void FIKRetargetEditMode::UpdateWidgetTransform()
 		BoneEdit.PrevLocalOffsets.Add(PrevLocalOffset);
 	}
 	
-	int32 ParentIndex = Skeleton.GetParentIndex(BoneEdit.Index);
+	const int32 ParentIndex = TargetSkeleton.GetParentIndex(BoneEdit.Index);
 	if (ParentIndex != INDEX_NONE)
 	{
-		BoneEdit.ParentGlobalTransform = Controller->GetTargetBoneGlobalTransform(ParentIndex);
+		BoneEdit.ParentGlobalTransform = Controller->GetTargetBoneGlobalTransform(RetargetProcessor, ParentIndex);
 	}else
 	{
 		BoneEdit.ParentGlobalTransform = FTransform::Identity;
diff --git a/Engine/Plugins/Animation/IKRig/Source/IKRigEditor/Private/RetargetEditor/IKRetargetEditorController.cpp b/Engine/Plugins/Animation/IKRig/Source/IKRigEditor/Private/RetargetEditor/IKRetargetEditorController.cpp
--- a/Engine/Plugins/Animation/IKRig/Source/IKRigEditor/Private/RetargetEditor/IKRetargetEditorController.cpp
+++ b/Engine/Plugins/Animation/IKRig/Source/IKRigEditor/Private/RetargetEditor/IKRetargetEditorController.cpp
@@ -106,13 +106,11 @@ USkeletalMesh* FIKRetargetEditorController::GetTargetSkeletalMesh() const
 	return AssetController->GetTargetPreviewMesh();
 }
 
-FTransform FIKRetargetEditorController::GetTargetBoneGlobalTransform(const int32& TargetBoneIndex) const
+FTransform FIKRetargetEditorController::GetTargetBoneGlobalTransform(
+	const UIKRetargetProcessor* RetargetProcessor,
+	const int32& TargetBoneIndex) const
 {
-	UIKRetargetAnimInstance* AnimInstance = TargetAnimInstance.Get();
-	check(AnimInstance);
-
-	const UIKRetargetProcessor* RetargetProcessor = AnimInstance->GetRetargetProcessor();
-	check(RetargetProcessor)
+	check(RetargetProcessor && RetargetProcessor->IsInitialized())
 
 	// get transform of bone
 	FTransform BoneTransform = RetargetProcessor->GetTargetBoneRetargetPoseGlobalTransform(TargetBoneIndex);
@@ -124,28 +122,26 @@ FTransform FIKRetargetEditorController::GetTargetBoneGlobalTransform(const int32
 	return BoneTransform;
 }
 
-FTransform FIKRetargetEditorController::GetTargetBoneLocalTransform(const int32& TargetBoneIndex) const
+FTransform FIKRetargetEditorController::GetTargetBoneLocalTransform(
+	const UIKRetargetProcessor* RetargetProcessor,
+	const int32& TargetBoneIndex) const
 {
-	UIKRetargetAnimInstance* AnimInstance = TargetAnimInstance.Get();
-	check(AnimInstance);
-
-	const UIKRetargetProcessor* RetargetProcessor = AnimInstance->GetRetargetProcessor();
-	check(RetargetProcessor)
+	check(RetargetProcessor && RetargetProcessor->IsInitialized())
 
 	return RetargetProcessor->GetTargetBoneRetargetPoseLocalTransform(TargetBoneIndex);
 }
 
 bool FIKRetargetEditorController::GetTargetBoneLineSegments(
+	const UIKRetargetProcessor* RetargetProcessor,
 	const int32& TargetBoneIndex,
 	FVector& OutStart,
 	TArray<FVector>& OutChildren) const
 {
 	// get the runtime processor
-	const UIKRetargetProcessor* Processor = GetRetargetProcessor();
-	check(Processor && Processor->IsInitialized())
+	check(RetargetProcessor && RetargetProcessor->IsInitialized())
 	
 	// get the target skeleton we want to draw
-	const FTargetSkeleton& TargetSkeleton = Processor->GetTargetSkeleton();
+	const FTargetSkeleton& TargetSkeleton = RetargetProcessor->GetTargetSkeleton();
 	check(TargetSkeleton.BoneNames.IsValidIndex(TargetBoneIndex))
 
 	// get the origin of the bone chain
diff --git a/Engine/Plugins/Animation/IKRig/Source/IKRigEditor/Public/RetargetEditor/IKRetargetEditorController.h b/Engine/Plugins/Animation/IKRig/Source/IKRigEditor/Public/RetargetEditor/IKRetargetEditorController.h
--- a/Engine/Plugins/Animation/IKRig/Source/IKRigEditor/Public/RetargetEditor/IKRetargetEditorController.h
+++ b/Engine/Plugins/Animation/IKRig/Source/IKRigEditor/Public/RetargetEditor/IKRetargetEditorController.h
@@ -65,10 +65,14 @@ public:
 	USkeletalMesh* GetTargetSkeletalMesh() const;
 
 	/** get current chain pose */
-	FTransform GetTargetBoneGlobalTransform(const int32& TargetBoneIndex) const;
-	FTransform GetTargetBoneLocalTransform(const int32& TargetBoneIndex) const;
+	FTransform GetTargetBoneGlobalTransform(const UIKRetargetProcessor* RetargetProcessor, const int32& TargetBoneIndex) const;
+	FTransform GetTargetBoneLocalTransform(const UIKRetargetProcessor* RetargetProcessor, const int32& TargetBoneIndex) const;
 	/** get the line segments to draw from this bone to each child */
-	bool GetTargetBoneLineSegments(const int32& TargetBoneIndex, FVector& OutStart, TArray<FVector>& OutChildren) const;
+	bool GetTargetBoneLineSegments(
+		const UIKRetargetProcessor* RetargetProcessor,
+		const int32& TargetBoneIndex,
+		FVector& OutStart,
+		TArray<FVector>& OutChildren) const;
 
 	/** get the retargeter that is running in the viewport (which is a duplicate of the source asset) */
 	const UIKRetargetProcessor* GetRetargetProcessor() const;
