Fixes crash when exporting retargeted animation due to async compression task.

diff --git a/Engine/Plugins/Animation/IKRig/Source/IKRigEditor/Private/RetargetEditor/IKRetargetBatchOperation.cpp b/Engine/Plugins/Animation/IKRig/Source/IKRigEditor/Private/RetargetEditor/IKRetargetBatchOperation.cpp
--- a/Engine/Plugins/Animation/IKRig/Source/IKRigEditor/Private/RetargetEditor/IKRetargetBatchOperation.cpp
+++ b/Engine/Plugins/Animation/IKRig/Source/IKRigEditor/Private/RetargetEditor/IKRetargetBatchOperation.cpp
@@ -228,9 +228,15 @@ void FIKRetargetBatchOperation::RetargetAssets(
 			UAnimationBlueprintLibrary::CopyAnimationCurveNamesToSkeleton(OldSkeleton, NewSkeleton, AnimSequenceToRetarget, ERawCurveTrackTypes::RCT_Float);	
 			// clear transform curves since those curves won't work in new skeleton
 			IAnimationDataController& Controller = AnimSequenceToRetarget->GetController();
+			const bool ShouldTransactAnimEdits = false;
+			Controller.OpenBracket(FText::FromString("Generating Retargeted Animation Data"), ShouldTransactAnimEdits);
 			Controller.RemoveAllCurvesOfType(ERawCurveTrackTypes::RCT_Transform);
+			// clear bone tracks to prevent recompression
+			Controller.RemoveAllBoneTracks(false);
 			// make sure we update the retarget source
 			UpdateRetargetSource(AnimSequenceToRetarget);
+			// done editing sequence data, close bracket
+			Controller.CloseBracket(ShouldTransactAnimEdits);
 		}
 
 		// replace references to other animation
