Fixing retargeter not reinitializing when source mesh swapped at runtime.

diff --git a/Engine/Plugins/Animation/IKRig/Source/IKRig/Private/AnimNodes/AnimNode_RetargetPoseFromMesh.cpp b/Engine/Plugins/Animation/IKRig/Source/IKRig/Private/AnimNodes/AnimNode_RetargetPoseFromMesh.cpp
--- a/Engine/Plugins/Animation/IKRig/Source/IKRig/Private/AnimNodes/AnimNode_RetargetPoseFromMesh.cpp
+++ b/Engine/Plugins/Animation/IKRig/Source/IKRig/Private/AnimNodes/AnimNode_RetargetPoseFromMesh.cpp
@@ -22,9 +22,8 @@ void FAnimNode_RetargetPoseFromMesh::CacheBones_AnyThread(const FAnimationCacheB
 	{
 		return;
 	}
-
-	const bool bIsRetargeterReady = IKRetargeterAsset && Processor && Processor->IsInitialized();
-	if (!bIsRetargeterReady)
+	
+	if (!Context.AnimInstanceProxy->GetSkelMeshComponent()->SkeletalMesh)
 	{
 		return;
 	}
@@ -33,8 +32,7 @@ void FAnimNode_RetargetPoseFromMesh::CacheBones_AnyThread(const FAnimationCacheB
 	RequiredToTargetBoneMapping.Reset();
 
 	const FReferenceSkeleton& RefSkeleton = RequiredBones.GetReferenceSkeleton();
-	const FTargetSkeleton& TargetSkeleton = Processor->GetTargetSkeleton();
-	
+	const FReferenceSkeleton& TargetSkeleton = Context.AnimInstanceProxy->GetSkelMeshComponent()->SkeletalMesh->GetRefSkeleton();
 	const TArray<FBoneIndexType>& RequiredBonesArray = RequiredBones.GetBoneIndicesArray();
 	for (int32 Index = 0; Index < RequiredBonesArray.Num(); ++Index)
 	{
@@ -42,7 +40,7 @@ void FAnimNode_RetargetPoseFromMesh::CacheBones_AnyThread(const FAnimationCacheB
 		if (ReqBoneIndex != INDEX_NONE)
 		{
 			const FName Name = RefSkeleton.GetBoneName(ReqBoneIndex);
-			const int32 TargetBoneIndex = TargetSkeleton.FindBoneIndexByName(Name);
+			const int32 TargetBoneIndex = TargetSkeleton.FindBoneIndex(Name);
 			if (TargetBoneIndex != INDEX_NONE)
 			{
 				// store require bone to target bone indices
@@ -64,26 +62,32 @@ void FAnimNode_RetargetPoseFromMesh::Evaluate_AnyThread(FPoseContext& Output)
 {
 	DECLARE_SCOPE_HIERARCHICAL_COUNTER_ANIMNODE(Evaluate_AnyThread)
 
-	if (!(IKRetargeterAsset && Processor))
+	if (!(IKRetargeterAsset && Processor && SourceMeshComponent.IsValid()))
 	{
 		return;
 	}
 
-	const bool bIsInitialized = Processor->IsInitialized();
-	const bool bInitializedWithSameMesh = Processor->GetTargetSkeleton().SkeletalMesh == Output.AnimInstanceProxy->GetSkelMeshComponent()->SkeletalMesh;
 	// it's possible in editor to have anim instances initialized before PreUpdate() is called
 	// which results in trying to run the retargeter without an source pose to copy from
-	const bool bSourceMeshBonesCopied = !SourceMeshComponentSpaceBoneTransforms.IsEmpty();
-	if (!(bIsInitialized && bInitializedWithSameMesh && bSourceMeshBonesCopied))
+	const bool bSourcePoseCopied = !SourceMeshComponentSpaceBoneTransforms.IsEmpty();
+
+	// ensure processor was initialized with the currently used assets (source/target meshes and retarget asset)
+	// if processor is not ready this tick, it will be next tick as this state will trigger re-initialization
+	const TObjectPtr<USkeletalMesh> SourceMesh = SourceMeshComponent->SkeletalMesh;
+	const TObjectPtr<USkeletalMesh> TargetMesh = Output.AnimInstanceProxy->GetSkelMeshComponent()->SkeletalMesh;
+	const bool bIsProcessorReady = Processor->WasInitializedWithTheseAssets(SourceMesh, TargetMesh, IKRetargeterAsset);
+
+	// if not ready to run, skip retarget and output the ref pose
+	if (!(bIsProcessorReady && bSourcePoseCopied))
 	{
 		Output.ResetToRefPose();
 		return;
 	}
 
 #if WITH_EDITOR
-	// live preview IK Rig solver settings in the retarget, editor only
+	// live preview source asset settings in the retarget, editor only
 	// NOTE: this copies goal targets as well, but these are overwritten by IK chain goals
-	if (bDriveTargetIKRigWithAsset)
+	if (bDriveWithAsset)
 	{
 		Processor->CopyAllSettingsFromAsset();
 	}
@@ -114,20 +118,20 @@ void FAnimNode_RetargetPoseFromMesh::PreUpdate(const UAnimInstance* InAnimInstan
 {
 	DECLARE_SCOPE_HIERARCHICAL_COUNTER_ANIMNODE(PreUpdate)
 	
-	if (!IsValid(IKRetargeterAsset))
+	if (IKRetargeterAsset.IsNull())
 	{
 		return;
 	}
 	
-	if (!IsValid(Processor))
+	if (Processor.IsNull())
 	{
 		Processor = NewObject<UIKRetargetProcessor>(InAnimInstance->GetOwningComponent());	
 	}
-	
-	EnsureInitialized(InAnimInstance);
-	if (Processor->IsInitialized())
+
+	const TObjectPtr<USkeletalMeshComponent> TargetMeshComponent = InAnimInstance->GetSkelMeshComponent();
+	if (EnsureProcessorIsInitialized(TargetMeshComponent))
 	{
-		CopyBoneTransformsFromSource(InAnimInstance->GetSkelMeshComponent());
+		CopyBoneTransformsFromSource(TargetMeshComponent);
 	}
 }
 
@@ -146,19 +150,18 @@ const UIKRetargetProcessor* FAnimNode_RetargetPoseFromMesh::GetRetargetProcessor
 	return Processor;
 }
 
-void FAnimNode_RetargetPoseFromMesh::EnsureInitialized(const UAnimInstance* InAnimInstance)
+bool FAnimNode_RetargetPoseFromMesh::EnsureProcessorIsInitialized(const TObjectPtr<USkeletalMeshComponent> TargetMeshComponent)
 {
 	// has user supplied a retargeter asset?
 	if (!IKRetargeterAsset)
 	{
-		return;
+		return false;
 	}
-
+	
 	// if user hasn't explicitly connected a source mesh, optionally use the parent mesh component (if there is one) 
 	if (!SourceMeshComponent.IsValid() && bUseAttachedParent)
 	{
-		USkeletalMeshComponent* TargetMesh = InAnimInstance->GetSkelMeshComponent();
-		USkeletalMeshComponent* ParentComponent = Cast<USkeletalMeshComponent>(TargetMesh->GetAttachParent());
+		const TObjectPtr<USkeletalMeshComponent> ParentComponent = Cast<USkeletalMeshComponent>(TargetMeshComponent->GetAttachParent());
 		if (ParentComponent)
 		{
 			SourceMeshComponent = ParentComponent;
@@ -168,67 +171,53 @@ void FAnimNode_RetargetPoseFromMesh::EnsureInitialized(const UAnimInstance* InAn
 	// has a source mesh been plugged in or found?
 	if (!SourceMeshComponent.IsValid())
 	{
-		return; // can't do anything if we don't have a source mesh
+		return false; // can't do anything if we don't have a source mesh component
 	}
 
-	// store all the components that were used to initialize
-	// if in future updates, any of this are mismatched, we have to re-initialize
-	CurrentlyUsedSourceMesh = SourceMeshComponent->SkeletalMesh;
-	CurrentlyUsedTargetMesh = InAnimInstance->GetSkelMeshComponent()->SkeletalMesh;
-	const bool bMeshesAreValid = CurrentlyUsedSourceMesh.IsValid() && CurrentlyUsedTargetMesh.IsValid();
-	if (!bMeshesAreValid)
+	// check that both a source and target mesh exist
+	const TObjectPtr<USkeletalMesh> SourceMesh = SourceMeshComponent->SkeletalMesh;
+	const TObjectPtr<USkeletalMesh> TargetMesh = TargetMeshComponent->SkeletalMesh;
+	if (SourceMesh.IsNull() || TargetMesh.IsNull())
 	{
-		return; // cannot initialize if components are missing skeletal mesh references
+		return false; // cannot initialize if components are missing skeletal mesh references
 	}
-
+	
 	// try initializing the processor
-	if (!Processor->IsInitialized())
+	if (!Processor->WasInitializedWithTheseAssets(SourceMesh, TargetMesh, IKRetargeterAsset))
 	{
 		// initialize retarget processor with source and target skeletal meshes
-		// (anim instance is passed in as outer UObject for new UIKRigProcessor) 
-		Processor->Initialize(
-			CurrentlyUsedSourceMesh.Get(),
-			CurrentlyUsedTargetMesh.Get(),
-			IKRetargeterAsset);
+		// (asset is passed in as outer UObject for new UIKRigProcessor) 
+		Processor->Initialize(SourceMesh,	TargetMesh,IKRetargeterAsset);
 	}
+
+	return Processor->IsInitialized();
 }
 
 void FAnimNode_RetargetPoseFromMesh::CopyBoneTransformsFromSource(USkeletalMeshComponent* TargetMeshComponent)
 {
-	if (!SourceMeshComponent.IsValid())
-	{
-		return; 
-	}
+	// get the mesh component to use as the source
+	const TObjectPtr<USkeletalMeshComponent> ComponentToCopyFrom =  GetComponentToCopyPoseFrom();
 
-	USkeletalMeshComponent* SourceMeshComp =  SourceMeshComponent.Get();
-	
-	// is the source mesh ticking?
-	if (!SourceMeshComp->IsRegistered())
-		
+	// this should not happen as we're guaranteed to be initialized at this stage
+	// but just in case component is lost after initialization, we avoid a crash
+	if (ComponentToCopyFrom.IsNull())
 	{
-		CurrentlyUsedSourceMesh.Reset(); // forces reinitialization when re-registered
-		return; // skip copying pose when component is no longer ticking
+		return; 
 	}
 	
-	// if our source is running under master-pose, then get bone data from there
-	if(USkeletalMeshComponent* MasterPoseComponent = Cast<USkeletalMeshComponent>(SourceMeshComponent->MasterPoseComponent.Get()))
-	{
-		SourceMeshComp = MasterPoseComponent;
-	}
-
-	// re-check mesh component validity as it may have changed to master
-	if (!(SourceMeshComp->SkeletalMesh && SourceMeshComp->IsRegistered()))
+	// skip copying pose when component is no longer ticking
+	if (!ComponentToCopyFrom->IsRegistered())
 	{
-		return; // master pose either missing skeletal mesh reference or not ticking, either way, we aren't copying from it
+		return; 
 	}
 	
 	const bool bUROInSync =
-		SourceMeshComp->ShouldUseUpdateRateOptimizations() &&
-		SourceMeshComp->AnimUpdateRateParams != nullptr &&
+		ComponentToCopyFrom->ShouldUseUpdateRateOptimizations() &&
+		ComponentToCopyFrom->AnimUpdateRateParams != nullptr &&
 		SourceMeshComponent->AnimUpdateRateParams == TargetMeshComponent->AnimUpdateRateParams;
-	const bool bUsingExternalInterpolation = SourceMeshComp->IsUsingExternalInterpolation();
-	const TArray<FTransform>& CachedComponentSpaceTransforms = SourceMeshComp->GetCachedComponentSpaceTransforms();
-	const bool bArraySizesMatch = CachedComponentSpaceTransforms.Num() == SourceMeshComp->GetComponentSpaceTransforms().Num();
+	const bool bUsingExternalInterpolation = ComponentToCopyFrom->IsUsingExternalInterpolation();
+	const TArray<FTransform>& CachedComponentSpaceTransforms = ComponentToCopyFrom->GetCachedComponentSpaceTransforms();
+	const bool bArraySizesMatch = CachedComponentSpaceTransforms.Num() == ComponentToCopyFrom->GetComponentSpaceTransforms().Num();
 
 	// copy source array from the appropriate location
 	SourceMeshComponentSpaceBoneTransforms.Reset();
@@ -238,6 +227,20 @@ void FAnimNode_RetargetPoseFromMesh::CopyBoneTransformsFromSource(USkeletalMeshC
 	}
 	else
 	{
-		SourceMeshComponentSpaceBoneTransforms.Append(SourceMeshComp->GetComponentSpaceTransforms()); // copy directly
+		SourceMeshComponentSpaceBoneTransforms.Append(ComponentToCopyFrom->GetComponentSpaceTransforms()); // copy directly
+	}
+}
+
+TObjectPtr<USkeletalMeshComponent> FAnimNode_RetargetPoseFromMesh::GetComponentToCopyPoseFrom() const
+{
+	// if our source is running under master-pose, then get bone data from there
+	if (SourceMeshComponent.IsValid())
+	{
+		if(USkeletalMeshComponent* MasterPoseComponent = Cast<USkeletalMeshComponent>(SourceMeshComponent->MasterPoseComponent.Get()))
+		{
+			return MasterPoseComponent;
+		}
 	}
+	
+	return SourceMeshComponent.Get();
 }
diff --git a/Engine/Plugins/Animation/IKRig/Source/IKRig/Private/Retargeter/IKRetargetProcessor.cpp b/Engine/Plugins/Animation/IKRig/Source/IKRig/Private/Retargeter/IKRetargetProcessor.cpp
--- a/Engine/Plugins/Animation/IKRig/Source/IKRig/Private/Retargeter/IKRetargetProcessor.cpp
+++ b/Engine/Plugins/Animation/IKRig/Source/IKRig/Private/Retargeter/IKRetargetProcessor.cpp
@@ -1486,7 +1486,25 @@ FTransform UIKRetargetProcessor::GetTargetBoneRetargetPoseLocalTransform(const i
 	return TargetSkeleton.RetargetLocalPose[TargetBoneIndex];
 }
 
+bool UIKRetargetProcessor::WasInitializedWithTheseAssets(
+	const TObjectPtr<USkeletalMesh> InSourceMesh,
+	const TObjectPtr<USkeletalMesh> InTargetMesh,
+	const TObjectPtr<UIKRetargeter> InRetargetAsset)
+{
+	if (!IsInitialized())
+	{
+		return false;
+	}
+
+	const bool bSourceMatches = InSourceMesh == GetSourceSkeleton().SkeletalMesh;
+	const bool bTargetMatches = InTargetMesh == GetTargetSkeleton().SkeletalMesh;
+	const bool bAssetMatches = InRetargetAsset == RetargeterAsset;
+	
+	return bSourceMatches && bTargetMatches && bAssetMatches;
+}
+
 #if WITH_EDITOR
+
 void UIKRetargetProcessor::SetNeedsInitialized()
 {
 	bIsInitialized = false;
diff --git a/Engine/Plugins/Animation/IKRig/Source/IKRig/Public/AnimNodes/AnimNode_RetargetPoseFromMesh.h b/Engine/Plugins/Animation/IKRig/Source/IKRig/Public/AnimNodes/AnimNode_RetargetPoseFromMesh.h
--- a/Engine/Plugins/Animation/IKRig/Source/IKRig/Public/AnimNodes/AnimNode_RetargetPoseFromMesh.h
+++ b/Engine/Plugins/Animation/IKRig/Source/IKRig/Public/AnimNodes/AnimNode_RetargetPoseFromMesh.h
@@ -20,16 +20,16 @@ struct IKRIG_API FAnimNode_RetargetPoseFromMesh : public FAnimNode_Base
 	TWeakObjectPtr<USkeletalMeshComponent> SourceMeshComponent = nullptr;
 
 	/* If SourceMeshComponent is not valid, and if this is true, it will look for attached parent as a source */
-	UPROPERTY(EditAnywhere, BlueprintReadWrite, Category=Settings, meta = (NeverAsPin))
+	UPROPERTY(EditAnywhere, BlueprintReadWrite, Category=Settings, meta=(NeverAsPin))
 	bool bUseAttachedParent = true;
 	
 	/** Retarget asset to use. Must define a Source and Target IK Rig compatible with the SourceMeshComponent and current anim instance.*/
-	UPROPERTY(EditAnywhere, Category = Settings)
+	UPROPERTY(EditAnywhere, BlueprintReadWrite, Category = Settings, meta=(PinHiddenByDefault))
 	TObjectPtr<UIKRetargeter> IKRetargeterAsset = nullptr;
 
 #if WITH_EDITOR
 	/** when true, will copy all setting from target IK Rig asset each tick (for live preview) */
-	bool bDriveTargetIKRigWithAsset = false;
+	bool bDriveWithAsset = false;
 #endif
 	
 	// FAnimNode_Base interface
@@ -50,13 +50,11 @@ struct IKRIG_API FAnimNode_RetargetPoseFromMesh : public FAnimNode_Base
 	const UIKRetargetProcessor* GetRetargetProcessor() const;
 
 private:
-	
-	void EnsureInitialized(const UAnimInstance* InAnimInstance);
+
+	/** returns true if processor is setup and ready to go, false otherwise */
+	bool EnsureProcessorIsInitialized(const TObjectPtr<USkeletalMeshComponent> TargetMeshComponent);
 	void CopyBoneTransformsFromSource(USkeletalMeshComponent* TargetMeshComponent);
-	
-	// source mesh references, cached during init so that we can compare and see if it has changed
-	TWeakObjectPtr<USkeletalMesh> CurrentlyUsedSourceMesh;
-	TWeakObjectPtr<USkeletalMesh> CurrentlyUsedTargetMesh;
+	TObjectPtr<USkeletalMeshComponent> GetComponentToCopyPoseFrom() const;
 
 	/** the runtime processor used to run the retarget and generate new poses */
 	UPROPERTY(Transient)
diff --git a/Engine/Plugins/Animation/IKRig/Source/IKRig/Public/Retargeter/IKRetargetProcessor.h b/Engine/Plugins/Animation/IKRig/Source/IKRig/Public/Retargeter/IKRetargetProcessor.h
--- a/Engine/Plugins/Animation/IKRig/Source/IKRig/Public/Retargeter/IKRetargetProcessor.h
+++ b/Engine/Plugins/Animation/IKRig/Source/IKRig/Public/Retargeter/IKRetargetProcessor.h
@@ -451,6 +451,12 @@ public:
 	/** Get whether this processor is ready to call RunRetargeter() and generate new poses. */
 	bool IsInitialized() const { return bIsInitialized; };
 
+	/** Get whether this processor was initialized with these skeletal meshes and retarget asset*/
+	bool WasInitializedWithTheseAssets(
+		const TObjectPtr<USkeletalMesh> InSourceMesh,
+		const TObjectPtr<USkeletalMesh> InTargetMesh,
+		const TObjectPtr<UIKRetargeter> InRetargetAsset);
+
 	/** Get the currently running IK Rig processor for the target */
 	UIKRigProcessor* GetTargetIKRigProcessor() const { return IKRigProcessor; };
 
diff --git a/Engine/Plugins/Animation/IKRig/Source/IKRigDeveloper/Private/AnimGraphNode_RetargetPoseFromMesh.cpp b/Engine/Plugins/Animation/IKRig/Source/IKRigDeveloper/Private/AnimGraphNode_RetargetPoseFromMesh.cpp
--- a/Engine/Plugins/Animation/IKRig/Source/IKRigDeveloper/Private/AnimGraphNode_RetargetPoseFromMesh.cpp
+++ b/Engine/Plugins/Animation/IKRig/Source/IKRigDeveloper/Private/AnimGraphNode_RetargetPoseFromMesh.cpp
@@ -57,7 +57,21 @@ void UAnimGraphNode_RetargetPoseFromMesh::ValidateAnimNodeDuringCompilation(USke
 	// validate IK Rig asset has been assigned
 	if (!Node.IKRetargeterAsset)
 	{
-		MessageLog.Warning(TEXT("@@ is missing an IKRetargeter asset."), this);
+		UEdGraphPin* Pin = FindPin(GET_MEMBER_NAME_STRING_CHECKED(FAnimNode_RetargetPoseFromMesh, IKRetargeterAsset));
+		if (Pin == nullptr)
+		{
+			// retarget asset unassigned
+			MessageLog.Warning(TEXT("@@ is missing an IKRetargeter asset."), this);
+			return;	
+		}
+		
+		if (!IsPinExposedAndLinked(GET_MEMBER_NAME_STRING_CHECKED(FAnimNode_RetargetPoseFromMesh, IKRetargeterAsset)))
+		{
+			// retarget asset pin not connected to an asset
+			MessageLog.Warning(TEXT("@@ has a IKRetargeter asset pin that is not connected to a valid asset."), this);
+			return;	
+		}
+
 		return;
 	}
 
diff --git a/Engine/Plugins/Animation/IKRig/Source/IKRigEditor/Private/RetargetEditor/IKRetargetAnimInstanceProxy.cpp b/Engine/Plugins/Animation/IKRig/Source/IKRigEditor/Private/RetargetEditor/IKRetargetAnimInstanceProxy.cpp
--- a/Engine/Plugins/Animation/IKRig/Source/IKRigEditor/Private/RetargetEditor/IKRetargetAnimInstanceProxy.cpp
+++ b/Engine/Plugins/Animation/IKRig/Source/IKRigEditor/Private/RetargetEditor/IKRetargetAnimInstanceProxy.cpp
@@ -50,7 +50,7 @@ void FIKRetargetAnimInstanceProxy::SetRetargetAssetAndSourceComponent(
 	IKRetargetNode->IKRetargeterAsset = InIKRetargetAsset;
 	IKRetargetNode->bUseAttachedParent = false;
 	IKRetargetNode->SourceMeshComponent = InSourceMeshComponent;
-	IKRetargetNode->bDriveTargetIKRigWithAsset = true;
+	IKRetargetNode->bDriveWithAsset = true;
 	IKRetargetNode->SetProcessorNeedsInitialized();
 }
 
