Animation: CopyPoseFromMesh and RetargetPoseFromMesh will now walk up the attach parent hierarchy to the parent mesh component when using bUseAttachedParent instead of only checking the immediate parent, which fixes issues when the owning actor is attached via, e.g., a child actor component

diff --git a/Engine/Plugins/Animation/IKRig/Source/IKRig/Private/AnimNodes/AnimNode_RetargetPoseFromMesh.cpp b/Engine/Plugins/Animation/IKRig/Source/IKRig/Private/AnimNodes/AnimNode_RetargetPoseFromMesh.cpp
--- a/Engine/Plugins/Animation/IKRig/Source/IKRig/Private/AnimNodes/AnimNode_RetargetPoseFromMesh.cpp
+++ b/Engine/Plugins/Animation/IKRig/Source/IKRig/Private/AnimNodes/AnimNode_RetargetPoseFromMesh.cpp
@@ -161,7 +161,17 @@ bool FAnimNode_RetargetPoseFromMesh::EnsureProcessorIsInitialized(const TObjectP
 	// if user hasn't explicitly connected a source mesh, optionally use the parent mesh component (if there is one) 
 	if (!SourceMeshComponent.IsValid() && bUseAttachedParent)
 	{
-		const TObjectPtr<USkeletalMeshComponent> ParentComponent = Cast<USkeletalMeshComponent>(TargetMeshComponent->GetAttachParent());
+		// Walk up the attachment chain until we find a skeletal mesh component
+		USkeletalMeshComponent* ParentComponent = nullptr;
+		for (USceneComponent* AttachParentComp = TargetMeshComponent->GetAttachParent(); AttachParentComp != nullptr; AttachParentComp = AttachParentComp->GetAttachParent())
+		{
+			ParentComponent = Cast<USkeletalMeshComponent>(AttachParentComp);
+			if (ParentComponent)
+			{
+				break;
+			}
+		}
+
 		if (ParentComponent)
 		{
 			SourceMeshComponent = ParentComponent;
diff --git a/Engine/Source/Runtime/AnimGraphRuntime/Private/AnimNodes/AnimNode_CopyPoseFromMesh.cpp b/Engine/Source/Runtime/AnimGraphRuntime/Private/AnimNodes/AnimNode_CopyPoseFromMesh.cpp
--- a/Engine/Source/Runtime/AnimGraphRuntime/Private/AnimNodes/AnimNode_CopyPoseFromMesh.cpp
+++ b/Engine/Source/Runtime/AnimGraphRuntime/Private/AnimNodes/AnimNode_CopyPoseFromMesh.cpp
@@ -73,10 +73,20 @@ void FAnimNode_CopyPoseFromMesh::RefreshMeshComponent(USkeletalMeshComponent* Ta
 	{
 		if (TargetMeshComponent)
 		{
-			USkeletalMeshComponent* ParentComponent = Cast<USkeletalMeshComponent>(TargetMeshComponent->GetAttachParent());
-			if (ParentComponent)
+			// Walk up the attachment chain until we find a skeletal mesh component
+			USkeletalMeshComponent* ParentMeshComponent = nullptr;
+			for (USceneComponent* AttachParentComp = TargetMeshComponent->GetAttachParent(); AttachParentComp != nullptr; AttachParentComp = AttachParentComp->GetAttachParent())
 			{
-				ResetMeshComponent(ParentComponent, TargetMeshComponent);
+				ParentMeshComponent = Cast<USkeletalMeshComponent>(AttachParentComp);
+				if (ParentMeshComponent)
+				{
+					break;
+				}
+			}
+
+			if (ParentMeshComponent)
+			{
+				ResetMeshComponent(ParentMeshComponent, TargetMeshComponent);
 			}
 			else
 			{
