Added profile group to the retargeter.

diff --git a/Engine/Plugins/Animation/IKRig/Source/IKRig/Private/AnimNodes/AnimNode_RetargetPoseFromMesh.cpp b/Engine/Plugins/Animation/IKRig/Source/IKRig/Private/AnimNodes/AnimNode_RetargetPoseFromMesh.cpp
--- a/Engine/Plugins/Animation/IKRig/Source/IKRig/Private/AnimNodes/AnimNode_RetargetPoseFromMesh.cpp
+++ b/Engine/Plugins/Animation/IKRig/Source/IKRig/Private/AnimNodes/AnimNode_RetargetPoseFromMesh.cpp
@@ -62,6 +62,8 @@ void FAnimNode_RetargetPoseFromMesh::Evaluate_AnyThread(FPoseContext& Output)
 {
 	DECLARE_SCOPE_HIERARCHICAL_COUNTER_ANIMNODE(Evaluate_AnyThread)
 
+	SCOPE_CYCLE_COUNTER(STAT_IKRetarget);
+
 	if (!(IKRetargeterAsset && Processor && SourceMeshComponent.IsValid()))
 	{
 		return;
diff --git a/Engine/Plugins/Animation/IKRig/Source/IKRig/Public/AnimNodes/AnimNode_RetargetPoseFromMesh.h b/Engine/Plugins/Animation/IKRig/Source/IKRig/Public/AnimNodes/AnimNode_RetargetPoseFromMesh.h
--- a/Engine/Plugins/Animation/IKRig/Source/IKRig/Public/AnimNodes/AnimNode_RetargetPoseFromMesh.h
+++ b/Engine/Plugins/Animation/IKRig/Source/IKRig/Public/AnimNodes/AnimNode_RetargetPoseFromMesh.h
@@ -10,6 +10,8 @@
 
 #include "AnimNode_RetargetPoseFromMesh.generated.h"
 
+DECLARE_CYCLE_STAT(TEXT("IK Retarget"), STAT_IKRetarget, STATGROUP_Anim);
+
 USTRUCT(BlueprintInternalUseOnly)
 struct IKRIG_API FAnimNode_RetargetPoseFromMesh : public FAnimNode_Base
 {
