Fixes crash when chain list tab is closed.

diff --git a/Engine/Plugins/Animation/IKRig/Source/IKRigEditor/Private/RetargetEditor/IKRetargetEditorController.cpp b/Engine/Plugins/Animation/IKRig/Source/IKRigEditor/Private/RetargetEditor/IKRetargetEditorController.cpp
--- a/Engine/Plugins/Animation/IKRig/Source/IKRigEditor/Private/RetargetEditor/IKRetargetEditorController.cpp
+++ b/Engine/Plugins/Animation/IKRig/Source/IKRigEditor/Private/RetargetEditor/IKRetargetEditorController.cpp
@@ -181,7 +181,12 @@ const UIKRetargetProcessor* FIKRetargetEditorController::GetRetargetProcessor()
 void FIKRetargetEditorController::RefreshAllViews() const
 {
 	Editor.Pin()->RegenerateMenusAndToolbars();
-	ChainsView.Get()->RefreshView();
+
+	// cannot assume chains view is always available
+	if (ChainsView.IsValid())
+	{
+		ChainsView.Get()->RefreshView();
+	}
 }
 
 void FIKRetargetEditorController::PlayAnimationAsset(UAnimationAsset* AssetToPlay)
