Fixing off-by-1 frame bug with IK Rig anim node pin inputs.

diff --git a/Engine/Plugins/Animation/IKRig/Source/IKRig/Private/AnimNodes/AnimNode_IKRig.cpp b/Engine/Plugins/Animation/IKRig/Source/IKRig/Private/AnimNodes/AnimNode_IKRig.cpp
--- a/Engine/Plugins/Animation/IKRig/Source/IKRig/Private/AnimNodes/AnimNode_IKRig.cpp
+++ b/Engine/Plugins/Animation/IKRig/Source/IKRig/Private/AnimNodes/AnimNode_IKRig.cpp
@@ -172,30 +172,30 @@ void FAnimNode_IKRig::Initialize_AnyThread(const FAnimationInitializeContext& Co
 void FAnimNode_IKRig::Update_AnyThread(const FAnimationUpdateContext& Context)
 {
 	ActualAlpha = 0.f;
-	// if (IsLODEnabled(Context.AnimInstanceProxy))
-	{
+
 		GetEvaluateGraphExposedInputs().Execute(Context);
 
-		// alpha handlers
-		switch (AlphaInputType)
+	// alpha handlers
+	switch (AlphaInputType)
+	{
+	case EAnimAlphaInputType::Float : 
+		ActualAlpha = AlphaScaleBias.ApplyTo(AlphaScaleBiasClamp.ApplyTo(Alpha, Context.GetDeltaTime()));
+		break;
+	case EAnimAlphaInputType::Bool :
+		ActualAlpha = AlphaBoolBlend.ApplyTo(bAlphaBoolEnabled, Context.GetDeltaTime());
+		break;
+	case EAnimAlphaInputType::Curve :
+		if (UAnimInstance* AnimInstance = Cast<UAnimInstance>(Context.AnimInstanceProxy->GetAnimInstanceObject()))
 		{
-		case EAnimAlphaInputType::Float : 
-			ActualAlpha = AlphaScaleBias.ApplyTo(AlphaScaleBiasClamp.ApplyTo(Alpha, Context.GetDeltaTime()));
-			break;
-		case EAnimAlphaInputType::Bool :
-			ActualAlpha = AlphaBoolBlend.ApplyTo(bAlphaBoolEnabled, Context.GetDeltaTime());
-			break;
-		case EAnimAlphaInputType::Curve :
-			if (UAnimInstance* AnimInstance = Cast<UAnimInstance>(Context.AnimInstanceProxy->GetAnimInstanceObject()))
-			{
-				ActualAlpha = AlphaScaleBiasClamp.ApplyTo(AnimInstance->GetCurveValue(AlphaCurveName), Context.GetDeltaTime());
-			}
-			break;
-		};
+			ActualAlpha = AlphaScaleBiasClamp.ApplyTo(AnimInstance->GetCurveValue(AlphaCurveName), Context.GetDeltaTime());
+		}
+		break;
+	};
 
-		// Make sure Alpha is clamped between 0 and 1.
-		ActualAlpha = FMath::Clamp<float>(ActualAlpha, 0.f, 1.f);
-	}
+	// Make sure Alpha is clamped between 0 and 1.
+	ActualAlpha = FMath::Clamp<float>(ActualAlpha, 0.f, 1.f);
+
+	PropagateInputProperties(Context.AnimInstanceProxy->GetAnimInstanceObject());
 	
 	FAnimNode_Base::Update_AnyThread(Context);
 	Source.Update(Context);
@@ -219,8 +219,6 @@ void FAnimNode_IKRig::PreUpdate(const UAnimInstance* InAnimInstance)
 		const FReferenceSkeleton& RefSkeleton = InAnimInstance->GetSkelMeshComponent()->SkeletalMesh->GetRefSkeleton();
  		IKRigProcessor->Initialize(RigDefinitionAsset, RefSkeleton);
 	}
-
-	PropagateInputProperties(InAnimInstance);
 	
 	// cache list of goal creator components on the actor
 	// TODO tried doing this in Initialize_AnyThread but it would miss some GoalCreator components
