Retarget Pose From Mesh now copies curves
Retarget editor viewport now can select + move source/target meshes

diff --git a/Engine/Plugins/Animation/IKRig/Source/IKRig/Private/AnimNodes/AnimNode_RetargetPoseFromMesh.cpp b/Engine/Plugins/Animation/IKRig/Source/IKRig/Private/AnimNodes/AnimNode_RetargetPoseFromMesh.cpp
--- a/Engine/Plugins/Animation/IKRig/Source/IKRig/Private/AnimNodes/AnimNode_RetargetPoseFromMesh.cpp
+++ b/Engine/Plugins/Animation/IKRig/Source/IKRig/Private/AnimNodes/AnimNode_RetargetPoseFromMesh.cpp
@@ -114,6 +114,19 @@ void FAnimNode_RetargetPoseFromMesh::Evaluate_AnyThread(FPoseContext& Output)
 
 	// convert to local space
 	FCSPose<FCompactPose>::ConvertComponentPosesToLocalPoses(ComponentPose, Output.Pose);
+
+	// copy curves over
+	if (bCopyCurves)
+	{
+		for (const TPair<FName, float>& SourceCurve : SourceCurveList)
+		{
+			if (const SmartName::UID_Type* UID = CurveNameToUIDMap.Find(SourceCurve.Key))
+			{
+				// set source value to output curve
+				Output.Curve.Set(*UID, SourceCurve.Value);
+			}
+		}
+	}
 }
 
 void FAnimNode_RetargetPoseFromMesh::PreUpdate(const UAnimInstance* InAnimInstance)
@@ -134,6 +147,16 @@ void FAnimNode_RetargetPoseFromMesh::PreUpdate(const UAnimInstance* InAnimInstan
 	if (EnsureProcessorIsInitialized(TargetMeshComponent))
 	{
 		CopyBoneTransformsFromSource(TargetMeshComponent);
+
+		if(bCopyCurves)
+		{
+			SourceCurveList.Reset();
+			if (const UAnimInstance* SourceAnimInstance = SourceMeshComponent->GetAnimInstance())
+			{
+				// attribute curve contains all list	
+				SourceCurveList.Append(SourceAnimInstance->GetAnimationCurveList(EAnimCurveType::AttributeCurve));
+			}
+		}
 	}
 }
 
@@ -193,6 +216,14 @@ bool FAnimNode_RetargetPoseFromMesh::EnsureProcessorIsInitialized(const TObjectP
 	{
 		return false; // cannot initialize if components are missing skeletal mesh references
 	}
+
+	// check that both have skeleton assets (shouldn't get this far without a skeleton)
+	const TObjectPtr<USkeleton> SourceSkeleton = SourceMesh->GetSkeleton();
+	const TObjectPtr<USkeleton> TargetSkeleton = TargetMesh->GetSkeleton();
+	if (SourceSkeleton.IsNull() || TargetSkeleton.IsNull())
+	{
+		return false;
+	}
 	
 	// try initializing the processor
 	if (!Processor->WasInitializedWithTheseAssets(SourceMesh, TargetMesh, IKRetargeterAsset))
@@ -200,6 +231,27 @@ bool FAnimNode_RetargetPoseFromMesh::EnsureProcessorIsInitialized(const TObjectP
 		// initialize retarget processor with source and target skeletal meshes
 		// (asset is passed in as outer UObject for new UIKRigProcessor) 
 		Processor->Initialize(SourceMesh,	TargetMesh,IKRetargeterAsset);
+
+		// create a map of curve names to IDs that are present on the target skeleton
+		if (bCopyCurves)
+		{
+			const FSmartNameMapping* SourceContainer = SourceSkeleton->GetSmartNameContainer(USkeleton::AnimCurveMappingName);
+			const FSmartNameMapping* TargetContainer = TargetSkeleton->GetSmartNameContainer(USkeleton::AnimCurveMappingName);
+
+			TArray<FName> SourceCurveNames;
+			SourceContainer->FillNameArray(SourceCurveNames);
+			CurveNameToUIDMap.Reset();
+			for (int32 Index = 0; Index < SourceCurveNames.Num(); ++Index)
+			{
+				SmartName::UID_Type UID = TargetContainer->FindUID(SourceCurveNames[Index]);
+				if (UID != SmartName::MaxUID)
+				{
+					// has a valid UID, add to the list
+					SmartName::UID_Type& Value = CurveNameToUIDMap.Add(SourceCurveNames[Index]);
+					Value = UID;
+				}
+			}
+		}
 	}
 
 	return Processor->IsInitialized();
diff --git a/Engine/Plugins/Animation/IKRig/Source/IKRig/Private/Retargeter/IKRetargeter.cpp b/Engine/Plugins/Animation/IKRig/Source/IKRig/Private/Retargeter/IKRetargeter.cpp
--- a/Engine/Plugins/Animation/IKRig/Source/IKRig/Private/Retargeter/IKRetargeter.cpp
+++ b/Engine/Plugins/Animation/IKRig/Source/IKRig/Private/Retargeter/IKRetargeter.cpp
@@ -43,6 +43,21 @@ void UIKRetargeter::PostLoad()
 			}
 		}
 	}
+
+	#if WITH_EDITORONLY_DATA
+		// load deprecated target actor offset
+		if (!FMath::IsNearlyZero(TargetActorOffset_DEPRECATED))
+		{
+			TargetMeshOffset.X = TargetActorOffset_DEPRECATED;
+		}
+
+		// load deprecated target actor scale
+		if (!FMath::IsNearlyZero(TargetActorScale_DEPRECATED))
+		{
+			TargetMeshScale = TargetActorScale_DEPRECATED;
+		}
+	#endif
+	
 	PRAGMA_ENABLE_DEPRECATION_WARNINGS
 
 	// remove null settings
diff --git a/Engine/Plugins/Animation/IKRig/Source/IKRig/Public/AnimNodes/AnimNode_RetargetPoseFromMesh.h b/Engine/Plugins/Animation/IKRig/Source/IKRig/Public/AnimNodes/AnimNode_RetargetPoseFromMesh.h
--- a/Engine/Plugins/Animation/IKRig/Source/IKRig/Public/AnimNodes/AnimNode_RetargetPoseFromMesh.h
+++ b/Engine/Plugins/Animation/IKRig/Source/IKRig/Public/AnimNodes/AnimNode_RetargetPoseFromMesh.h
@@ -29,6 +29,10 @@ struct IKRIG_API FAnimNode_RetargetPoseFromMesh : public FAnimNode_Base
 	UPROPERTY(EditAnywhere, BlueprintReadWrite, Category = Settings, meta=(PinHiddenByDefault))
 	TObjectPtr<UIKRetargeter> IKRetargeterAsset = nullptr;
 
+	/* Copy curves from SouceMeshComponent. This will copy any curves the source/target Skeleton have in common. */
+	UPROPERTY(EditAnywhere, BlueprintReadWrite, Category = Settings, meta = (NeverAsPin))
+	bool bCopyCurves = true;
+
 #if WITH_EDITOR
 	/** when true, will copy all setting from target IK Rig asset each tick (for live preview) */
 	bool bDriveWithAsset = false;
@@ -67,4 +71,8 @@ private:
 
 	// mapping from required bones to actual bones within the target skeleton
 	TArray< TPair<int32, int32> > RequiredToTargetBoneMapping;
+	// mapping of Skeleton curve names to the UID (for copying curves from the source mesh component)
+	TMap<FName, SmartName::UID_Type> CurveNameToUIDMap;
+	// cached curves, copied on the game thread
+	TMap<FName, float> SourceCurveList;
 };
diff --git a/Engine/Plugins/Animation/IKRig/Source/IKRig/Public/Retargeter/IKRetargeter.h b/Engine/Plugins/Animation/IKRig/Source/IKRig/Public/Retargeter/IKRetargeter.h
--- a/Engine/Plugins/Animation/IKRig/Source/IKRig/Public/Retargeter/IKRetargeter.h
+++ b/Engine/Plugins/Animation/IKRig/Source/IKRig/Public/Retargeter/IKRetargeter.h
@@ -234,7 +234,7 @@ private:
 public:
 
 #if WITH_EDITORONLY_DATA
-	/** The Skeletal Mesh to preview the retarget on.*/
+	/** Optional. Override the Skeletal Mesh to preview the retarget on. Uses the preview mesh from the Target IK Rig asset by default. */
 	UPROPERTY(EditAnywhere, Category = Rigs)
 	TObjectPtr<USkeletalMesh> TargetPreviewMesh = nullptr;
 #endif
@@ -252,18 +252,27 @@ public:
 	bool bRetargetIK = true;
 
 #if WITH_EDITORONLY_DATA
-	/** Move the target actor in the viewport for easier visualization next to the source actor.*/
-	UPROPERTY(EditAnywhere, Category = TargetActorPreview, meta = (UIMin = "-2000.0", UIMax = "2000.0"))
-	float TargetActorOffset = 150.0f;
+	UPROPERTY()
+	float TargetActorOffset_DEPRECATED = 0.0f;
+	UPROPERTY()
+	float TargetActorScale_DEPRECATED = 0.0f;
+
+	/** The offset applied to the target mesh in the editor viewport. */
+	UPROPERTY(EditAnywhere, Category = PreviewSettings)
+	FVector TargetMeshOffset;
 
-	/** Scale the target actor in the viewport for easier visualization next to the source actor.*/
-	UPROPERTY(EditAnywhere, Category = TargetActorPreview, meta = (UIMin = "0.01", UIMax = "10.0"))
-	float TargetActorScale = 1.0f;
+	/** Scale the target mesh in the viewport for easier visualization next to the source.*/
+	UPROPERTY(EditAnywhere, Category = PreviewSettings, meta = (UIMin = "0.01", UIMax = "10.0"))
+	float TargetMeshScale = 1.0f;
+
+	/** The offset applied to the source mesh in the editor viewport. */
+	UPROPERTY(EditAnywhere, Category = PreviewSettings)
+	FVector SourceMeshOffset;
 
 	/** The visual size of the bones in the viewport when editing the retarget pose.*/
-	UPROPERTY(EditAnywhere, Category = PoseEditSettings, meta = (ClampMin = "0.0", UIMin = "0.01", UIMax = "10.0"))
+	UPROPERTY(EditAnywhere, Category = PreviewSettings, meta = (ClampMin = "0.0", UIMin = "0.01", UIMax = "10.0"))
 	float BoneDrawSize = 8.0f;
-
+	
 private:
 	/** A special editor-only mode which forces the retargeter to output the current retarget reference pose,
 	* rather than actually running the retarget and outputting the retargeted pose. Used in Edit-Pose mode.*/
diff --git a/Engine/Plugins/Animation/IKRig/Source/IKRigDeveloper/Public/AnimGraphNode_RetargetPoseFromMesh.h b/Engine/Plugins/Animation/IKRig/Source/IKRigDeveloper/Public/AnimGraphNode_RetargetPoseFromMesh.h
--- a/Engine/Plugins/Animation/IKRig/Source/IKRigDeveloper/Public/AnimGraphNode_RetargetPoseFromMesh.h
+++ b/Engine/Plugins/Animation/IKRig/Source/IKRigDeveloper/Public/AnimGraphNode_RetargetPoseFromMesh.h
@@ -33,6 +33,7 @@ public:
 	virtual FEditorModeID GetEditorMode() const override;
 	virtual void Draw(FPrimitiveDrawInterface* PDI, USkeletalMeshComponent* PreviewSkelMeshComp) const override;
 	virtual void CustomizePinData(UEdGraphPin* Pin, FName SourcePropertyName, int32 ArrayIndex) const override;
+	virtual bool UsingCopyPoseFromMesh() const override { return true; };
 	// End of UAnimGraphNode_Base interface
 
 	static const FName AnimModeName;
diff --git a/Engine/Plugins/Animation/IKRig/Source/IKRigEditor/Private/IKRigEditor.cpp b/Engine/Plugins/Animation/IKRig/Source/IKRigEditor/Private/IKRigEditor.cpp
--- a/Engine/Plugins/Animation/IKRig/Source/IKRigEditor/Private/IKRigEditor.cpp
+++ b/Engine/Plugins/Animation/IKRig/Source/IKRigEditor/Private/IKRigEditor.cpp
@@ -11,7 +11,8 @@
 #include "RigEditor/AssetTypeActions_IKRigDefinition.h"
 #include "RetargetEditor/AssetTypeActions_IKRetargeter.h"
 #include "RetargetEditor/IKRetargetCommands.h"
-#include "RetargetEditor/IKRetargetEditMode.h"
+#include "RetargetEditor/IKRetargetDefaultMode.h"
+#include "RetargetEditor/IKRetargetEditPoseMode.h"
 #include "RigEditor/IKRigCommands.h"
 #include "RigEditor/IKRigEditMode.h"
 #include "RigEditor/IKRigSkeletonCommands.h"
@@ -43,7 +44,8 @@ void FIKRigEditor::StartupModule()
 
 	// register custom editor modes
 	FEditorModeRegistry::Get().RegisterMode<FIKRigEditMode>(FIKRigEditMode::ModeName, LOCTEXT("IKRigEditMode", "IKRig"), FSlateIcon(), false);
-	FEditorModeRegistry::Get().RegisterMode<FIKRetargetEditMode>(FIKRetargetEditMode::ModeName, LOCTEXT("IKRetargetEditMode", "IKRetarget"), FSlateIcon(), false);
+	FEditorModeRegistry::Get().RegisterMode<FIKRetargetDefaultMode>(FIKRetargetDefaultMode::ModeName, LOCTEXT("IKRetargetDefaultMode", "IKRetargetDefault"), FSlateIcon(), false);
+	FEditorModeRegistry::Get().RegisterMode<FIKRetargetEditPoseMode>(FIKRetargetEditPoseMode::ModeName, LOCTEXT("IKRetargetEditMode", "IKRetargetEditPose"), FSlateIcon(), false);
 
 	// register detail customizations
 	FPropertyEditorModule& PropertyEditorModule = FModuleManager::LoadModuleChecked<FPropertyEditorModule>("PropertyEditor");
@@ -60,7 +62,8 @@ void FIKRigEditor::ShutdownModule()
 	FIKRetargetCommands::Unregister();
 	
 	FEditorModeRegistry::Get().UnregisterMode(FIKRigEditMode::ModeName);
-	FEditorModeRegistry::Get().UnregisterMode(FIKRetargetEditMode::ModeName);
+	FEditorModeRegistry::Get().UnregisterMode(FIKRetargetDefaultMode::ModeName);
+	FEditorModeRegistry::Get().UnregisterMode(FIKRetargetEditPoseMode::ModeName);
 
 	// unregister IKRigDefinition asset action
 	if (IKRigDefinitionAssetAction.IsValid())
diff --git a/Engine/Plugins/Animation/IKRig/Source/IKRigEditor/Private/RetargetEditor/IKRetargetApplicationMode.cpp b/Engine/Plugins/Animation/IKRig/Source/IKRigEditor/Private/RetargetEditor/IKRetargetApplicationMode.cpp
new file mode 100644
--- /dev/null
+++ b/Engine/Plugins/Animation/IKRig/Source/IKRigEditor/Private/RetargetEditor/IKRetargetApplicationMode.cpp
@@ -0,0 +1,97 @@
+// Copyright Epic Games, Inc. All Rights Reserved.
+
+#include "RetargetEditor/IKRetargetApplicationMode.h"
+
+#include "IPersonaPreviewScene.h"
+#include "PersonaModule.h"
+#include "Modules/ModuleManager.h"
+#include "PersonaTabs.h"
+#include "RetargetEditor/IKRetargetAssetBrowserTabSummoner.h"
+
+#include "RetargetEditor/IKRetargetChainTabSummoner.h"
+#include "RetargetEditor/IKRetargetEditor.h"
+
+#define LOCTEXT_NAMESPACE "IKRetargetMode"
+
+
+FIKRetargetApplicationMode::FIKRetargetApplicationMode(
+	TSharedRef<FWorkflowCentricApplication> InHostingApp,  
+	TSharedRef<IPersonaPreviewScene> InPreviewScene)
+	: FApplicationMode(IKRetargetApplicationModes::IKRetargetApplicationMode)
+{
+	IKRetargetEditorPtr = StaticCastSharedRef<FIKRetargetEditor>(InHostingApp);
+	TSharedRef<FIKRetargetEditor> IKRetargetEditor = StaticCastSharedRef<FIKRetargetEditor>(InHostingApp);
+
+	FPersonaViewportArgs ViewportArgs(InPreviewScene);
+	ViewportArgs.bAlwaysShowTransformToolbar = true;
+	ViewportArgs.bShowStats = false;
+	ViewportArgs.bShowTurnTable = false;
+	ViewportArgs.ContextName = TEXT("IKRetargetEditor.Viewport");
+
+	// register Persona tabs
+	FPersonaModule& PersonaModule = FModuleManager::LoadModuleChecked<FPersonaModule>("Persona");
+	TabFactories.RegisterFactory(PersonaModule.CreatePersonaViewportTabFactory(InHostingApp, ViewportArgs));
+	TabFactories.RegisterFactory(PersonaModule.CreateDetailsTabFactory(InHostingApp, FOnDetailsCreated::CreateSP(&IKRetargetEditor.Get(), &FIKRetargetEditor::HandleDetailsCreated)));
+
+	// register custom tabs
+	TabFactories.RegisterFactory(MakeShared<FIKRetargetChainTabSummoner>(IKRetargetEditor));
+	TabFactories.RegisterFactory(MakeShared<FIKRetargetAssetBrowserTabSummoner>(IKRetargetEditor));
+
+	// create tab layout
+	TabLayout = FTabManager::NewLayout("Standalone_IKRetargetEditor_Layout_v1.008")
+		->AddArea
+		(
+			FTabManager::NewPrimaryArea()
+			->SetOrientation(Orient_Vertical)
+			->Split
+			(
+				FTabManager::NewSplitter()
+				->SetSizeCoefficient(0.9f)
+				->SetOrientation(Orient_Horizontal)
+				->Split
+				(
+					FTabManager::NewStack()
+					->SetSizeCoefficient(0.6f)
+					->SetHideTabWell(true)
+					->AddTab(FPersonaTabs::PreviewViewportID, ETabState::OpenedTab)
+				)
+				->Split
+				(
+					FTabManager::NewSplitter()
+					->SetSizeCoefficient(0.9f)
+					->SetOrientation(Orient_Vertical)
+					->Split
+					(
+					FTabManager::NewStack()
+						->SetSizeCoefficient(0.6f)
+						->AddTab(FPersonaTabs::DetailsID, ETabState::OpenedTab)
+						->SetForegroundTab(FPersonaTabs::DetailsID)
+					)
+					->Split
+					(
+					FTabManager::NewStack()
+						->SetSizeCoefficient(0.6f)
+						->AddTab(FIKRetargetChainTabSummoner::TabID, ETabState::OpenedTab)
+						->AddTab(FIKRetargetAssetBrowserTabSummoner::TabID, ETabState::OpenedTab)
+						->SetForegroundTab(FIKRetargetChainTabSummoner::TabID)
+					)
+					
+				)
+			)
+		);
+
+	PersonaModule.OnRegisterTabs().Broadcast(TabFactories, InHostingApp);
+	LayoutExtender = MakeShared<FLayoutExtender>();
+	PersonaModule.OnRegisterLayoutExtensions().Broadcast(*LayoutExtender.Get());
+	TabLayout->ProcessExtensions(*LayoutExtender.Get());
+}
+
+void FIKRetargetApplicationMode::RegisterTabFactories(TSharedPtr<FTabManager> InTabManager)
+{
+	TSharedPtr<FIKRetargetEditor> IKRigEditor = IKRetargetEditorPtr.Pin();
+	IKRigEditor->RegisterTabSpawners(InTabManager.ToSharedRef());
+	IKRigEditor->PushTabFactories(TabFactories);
+	FApplicationMode::RegisterTabFactories(InTabManager);
+}
+
+#undef LOCTEXT_NAMESPACE
diff --git a/Engine/Plugins/Animation/IKRig/Source/IKRigEditor/Private/RetargetEditor/IKRetargetDefaultMode.cpp b/Engine/Plugins/Animation/IKRig/Source/IKRigEditor/Private/RetargetEditor/IKRetargetDefaultMode.cpp
new file mode 100644
--- /dev/null
+++ b/Engine/Plugins/Animation/IKRig/Source/IKRigEditor/Private/RetargetEditor/IKRetargetDefaultMode.cpp
@@ -0,0 +1,311 @@
+﻿// Copyright Epic Games, Inc. All Rights Reserved.
+
+#include "RetargetEditor/IKRetargetDefaultMode.h"
+
+#include "EditorViewportClient.h"
+#include "AssetEditorModeManager.h"
+#include "EngineUtils.h"
+#include "IKRigDebugRendering.h"
+#include "IPersonaPreviewScene.h"
+#include "Animation/DebugSkelMeshComponent.h"
+#include "RetargetEditor/IKRetargetEditor.h"
+
+
+#define LOCTEXT_NAMESPACE "IKRetargetDefaultMode"
+
+FName FIKRetargetDefaultMode::ModeName("IKRetargetAssetDefaultMode");
+
+bool FIKRetargetDefaultMode::GetCameraTarget(FSphere& OutTarget) const
+{
+	// target skeletal mesh
+	if (const TSharedPtr<FIKRetargetEditorController> Controller = EditorController.Pin())
+	{
+		if (Controller->SourceSkelMeshComponent)
+		{
+			OutTarget = Controller->SourceSkelMeshComponent->Bounds.GetSphere();
+			return true;
+		}
+	}
+	
+	return false;
+}
+
+IPersonaPreviewScene& FIKRetargetDefaultMode::GetAnimPreviewScene() const
+{
+	return *static_cast<IPersonaPreviewScene*>(static_cast<FAssetEditorModeManager*>(Owner)->GetPreviewScene());
+}
+
+void FIKRetargetDefaultMode::GetOnScreenDebugInfo(TArray<FText>& OutDebugInfo) const
+{
+}
+
+void FIKRetargetDefaultMode::Initialize()
+{
+	const TSharedPtr<FIKRetargetEditorController> Controller = EditorController.Pin();
+	if (!Controller.IsValid())
+	{
+		return; 
+	}
+	
+	// register selection callback overrides
+	TInlineComponentArray<UPrimitiveComponent*> PrimitiveComponents;
+	PrimitiveComponents.Add(Controller->SourceSkelMeshComponent);
+	PrimitiveComponents.Add(Controller->TargetSkelMeshComponent);
+	for (UPrimitiveComponent* PrimitiveComponent : PrimitiveComponents)
+	{
+		PrimitiveComponent->SelectionOverrideDelegate = UPrimitiveComponent::FSelectionOverride::CreateRaw(this, &FIKRetargetDefaultMode::ComponentSelectionOverride);
+	}
+
+	// deselect all
+	SetSelectedComponent(nullptr);
+
+	// update offsets on preview meshes
+	Controller->AddOffsetAndUpdatePreviewMeshPosition(FVector::ZeroVector, Controller->SourceSkelMeshComponent);
+	Controller->AddOffsetAndUpdatePreviewMeshPosition(FVector::ZeroVector, Controller->TargetSkelMeshComponent);
+
+	bIsInitialized = true;
+}
+
+void FIKRetargetDefaultMode::Render(const FSceneView* View, FViewport* Viewport, FPrimitiveDrawInterface* PDI)
+{
+	FEdMode::Render(View, Viewport, PDI);
+}
+
+bool FIKRetargetDefaultMode::AllowWidgetMove()
+{
+	return false;
+}
+
+bool FIKRetargetDefaultMode::ShouldDrawWidget() const
+{
+	return UsesTransformWidget(CurrentWidgetMode);
+}
+
+bool FIKRetargetDefaultMode::UsesTransformWidget() const
+{
+	return UsesTransformWidget(CurrentWidgetMode);
+}
+
+bool FIKRetargetDefaultMode::UsesTransformWidget(UE::Widget::EWidgetMode CheckMode) const
+{
+	const bool bTranslating = CheckMode == UE::Widget::EWidgetMode::WM_Translate;
+	return bTranslating && IsValid(SelectedComponent);
+}
+
+FVector FIKRetargetDefaultMode::GetWidgetLocation() const
+{
+	const TSharedPtr<FIKRetargetEditorController> Controller = EditorController.Pin();
+	if (!Controller.IsValid())
+	{
+		return FVector::ZeroVector; 
+	}
+	
+	if (SelectedComponent.IsNull())
+	{
+		return FVector::ZeroVector; // shouldn't get here
+	}
+
+	return SelectedComponent->GetComponentTransform().GetLocation();
+}
+
+bool FIKRetargetDefaultMode::HandleClick(FEditorViewportClient* InViewportClient, HHitProxy* HitProxy, const FViewportClick& Click)
+{
+	const TSharedPtr<FIKRetargetEditorController> Controller = EditorController.Pin();
+	if (!Controller.IsValid())
+	{
+		return false; 
+	}
+	
+	// did we select an actor in the viewport?
+	const bool bLeftButtonClicked = Click.GetKey() == EKeys::LeftMouseButton;
+	const bool bHitActor = HitProxy && HitProxy->IsA(HActor::StaticGetType());
+	if (bLeftButtonClicked && bHitActor)
+	{
+		const HActor* ActorProxy = static_cast<HActor*>(HitProxy);
+		SetSelectedComponent(const_cast<UPrimitiveComponent*>(ActorProxy->PrimComponent));
+		return true;
+	}
+
+	SetSelectedComponent(nullptr);
+	return false;
+}
+
+bool FIKRetargetDefaultMode::StartTracking(FEditorViewportClient* InViewportClient, FViewport* InViewport)
+{
+	bIsTranslating = false;
+
+	// not manipulating any widget axes, so stop tracking
+	const EAxisList::Type CurrentAxis = InViewportClient->GetCurrentWidgetAxis();
+	if (CurrentAxis == EAxisList::None)
+	{
+		return false; 
+	}
+
+	const TSharedPtr<FIKRetargetEditorController> Controller = EditorController.Pin();
+	if (!Controller.IsValid())
+	{
+		return false; // invalid editor state
+	}
+
+	const bool bTranslating = InViewportClient->GetWidgetMode() == UE::Widget::EWidgetMode::WM_Translate;
+	if (bTranslating && IsValid(SelectedComponent))
+	{
+		bIsTranslating = true;
+		GEditor->BeginTransaction(LOCTEXT("MovePreviewMesh", "Move Preview Mesh"));
+		Controller->AssetController->GetAsset()->Modify();
+		return true;
+	}
+
+	return false;
+}
+
+bool FIKRetargetDefaultMode::EndTracking(FEditorViewportClient* InViewportClient, FViewport* InViewport)
+{
+	GEditor->EndTransaction();
+	bIsTranslating = false;
+	return true;
+}
+
+bool FIKRetargetDefaultMode::InputDelta(
+	FEditorViewportClient* InViewportClient,
+	FViewport* InViewport,
+	FVector& InDrag,
+	FRotator& InRot,
+	FVector& InScale)
+{
+	if (!(bIsTranslating && IsValid(SelectedComponent)))
+	{
+		return false; // not handled
+	}
+
+	const TSharedPtr<FIKRetargetEditorController> Controller = EditorController.Pin();
+	if (!Controller.IsValid())
+	{
+		return false; 
+	}
+
+	if(InViewportClient->GetWidgetMode() != UE::Widget::WM_Translate)
+	{
+		return false;
+	}
+
+	Controller->AddOffsetAndUpdatePreviewMeshPosition(InDrag, SelectedComponent);
+	
+	return true;
+}
+
+bool FIKRetargetDefaultMode::GetCustomDrawingCoordinateSystem(FMatrix& InMatrix, void* InData)
+{
+	const TSharedPtr<FIKRetargetEditorController> Controller = EditorController.Pin();
+	if (!Controller.IsValid())
+	{
+		return false; 
+	}
+
+	if (SelectedComponent.IsNull())
+	{
+		return false;
+	}
+
+	InMatrix = SelectedComponent->GetComponentTransform().ToMatrixNoScale().RemoveTranslation();
+	return true;
+}
+
+bool FIKRetargetDefaultMode::GetCustomInputCoordinateSystem(FMatrix& InMatrix, void* InData)
+{
+	return GetCustomDrawingCoordinateSystem(InMatrix, InData);
+}
+
+void FIKRetargetDefaultMode::Enter()
+{
+	IPersonaEditMode::Enter();
+
+	const TSharedPtr<FIKRetargetEditorController> Controller = EditorController.Pin();
+	if (!Controller.IsValid())
+	{
+		return; 
+	}
+
+	// allow selection of meshes in this mode
+	Controller->Editor.Pin()->GetPersonaToolkit()->GetPreviewScene()->SetAllowMeshHitProxies(true);
+	Controller->SourceSkelMeshComponent->bSelectable = true;
+	Controller->TargetSkelMeshComponent->bSelectable = true;
+
+	// deselect when entering mode
+	SetSelectedComponent(nullptr);
+}
+
+void FIKRetargetDefaultMode::Exit()
+{
+	IPersonaEditMode::Exit();
+
+	const TSharedPtr<FIKRetargetEditorController> Controller = EditorController.Pin();
+	if (!Controller.IsValid())
+	{
+		return; 
+	}
+
+	// disable selection in other modes
+	Controller->Editor.Pin()->GetPersonaToolkit()->GetPreviewScene()->SetAllowMeshHitProxies(false);
+	Controller->SourceSkelMeshComponent->bSelectable = false;
+	Controller->TargetSkelMeshComponent->bSelectable = false;
+	
+	// deselect all
+	SetSelectedComponent(nullptr);
+}
+
+bool FIKRetargetDefaultMode::ComponentSelectionOverride(const UPrimitiveComponent* InComponent) const
+{
+	return InComponent == SelectedComponent;
+}
+
+void FIKRetargetDefaultMode::SetSelectedComponent(UPrimitiveComponent* InComponent)
+{
+	const TSharedPtr<FIKRetargetEditorController> Controller = EditorController.Pin();
+	if (!Controller.IsValid())
+	{
+		return; 
+	}
+	
+	SelectedComponent = InComponent;
+
+	Controller->SourceSkelMeshComponent->PushSelectionToProxy();
+	Controller->TargetSkelMeshComponent->PushSelectionToProxy();
+	
+	Controller->SourceSkelMeshComponent->MarkRenderStateDirty();
+	Controller->TargetSkelMeshComponent->MarkRenderStateDirty();
+
+	// when clicking in empty space, show global details
+	if (SelectedComponent == nullptr)
+	{
+		Controller->DetailsView->SetObject(Controller->AssetController->GetAsset());
+	}
+}
+
+void FIKRetargetDefaultMode::ApplyOffsetToMeshTransform(const FVector& Offset, USceneComponent* Component)
+{
+	constexpr bool bSweep = false;
+	constexpr FHitResult* OutSweepHitResult = nullptr;
+	constexpr ETeleportType Teleport = ETeleportType::ResetPhysics;
+	Component->SetWorldLocation(Offset, bSweep, OutSweepHitResult, Teleport);
+}
+
+void FIKRetargetDefaultMode::Tick(FEditorViewportClient* ViewportClient, float DeltaTime)
+{
+	FEdMode::Tick(ViewportClient, DeltaTime);
+	
+	CurrentWidgetMode = ViewportClient->GetWidgetMode();
+
+	// ensure selection callbacks have been generated
+	if (!bIsInitialized)
+	{
+		Initialize();
+	}
+}
+
+void FIKRetargetDefaultMode::DrawHUD(FEditorViewportClient* ViewportClient, FViewport* Viewport, const FSceneView* View, FCanvas* Canvas)
+{
+	FEdMode::DrawHUD(ViewportClient, Viewport, View, Canvas);
+}
+
+#undef LOCTEXT_NAMESPACE
diff --git a/Engine/Plugins/Animation/IKRig/Source/IKRigEditor/Private/RetargetEditor/IKRetargetEditMode.cpp b/Engine/Plugins/Animation/IKRig/Source/IKRigEditor/Private/RetargetEditor/IKRetargetEditMode.cpp
deleted file mode 100644
diff --git a/Engine/Plugins/Animation/IKRig/Source/IKRigEditor/Private/RetargetEditor/IKRetargetEditPoseMode.cpp b/Engine/Plugins/Animation/IKRig/Source/IKRigEditor/Private/RetargetEditor/IKRetargetEditPoseMode.cpp
new file mode 100644
--- /dev/null
+++ b/Engine/Plugins/Animation/IKRig/Source/IKRigEditor/Private/RetargetEditor/IKRetargetEditPoseMode.cpp
@@ -0,0 +1,528 @@
+// Copyright Epic Games, Inc. All Rights Reserved.
+
+#include "RetargetEditor/IKRetargetEditPoseMode.h"
+
+#include "EditorViewportClient.h"
+#include "AssetEditorModeManager.h"
+#include "EditorScriptingHelpers.h"
+#include "IKRigDebugRendering.h"
+#include "IKRigProcessor.h"
+#include "IPersonaPreviewScene.h"
+#include "Animation/DebugSkelMeshComponent.h"
+#include "Materials/MaterialInstanceDynamic.h"
+#include "RetargetEditor/IKRetargetAnimInstance.h"
+#include "RetargetEditor/IKRetargetHitProxies.h"
+#include "RetargetEditor/IKRetargetEditor.h"
+#include "Retargeter/IKRetargetProcessor.h"
+#include "RigEditor/IKRigEditMode.h"
+
+
+#define LOCTEXT_NAMESPACE "IKRetargeterEditMode"
+
+FName FIKRetargetEditPoseMode::ModeName("IKRetargetAssetEditMode");
+
+bool FIKRetargetEditPoseMode::GetCameraTarget(FSphere& OutTarget) const
+{
+	// target skeletal mesh
+	if (const TSharedPtr<FIKRetargetEditorController> Controller = EditorController.Pin())
+	{
+		if (Controller->TargetSkelMeshComponent)
+		{
+			OutTarget = Controller->TargetSkelMeshComponent->Bounds.GetSphere();
+			return true;
+		}
+	}
+	
+	return false;
+}
+
+IPersonaPreviewScene& FIKRetargetEditPoseMode::GetAnimPreviewScene() const
+{
+	return *static_cast<IPersonaPreviewScene*>(static_cast<FAssetEditorModeManager*>(Owner)->GetPreviewScene());
+}
+
+void FIKRetargetEditPoseMode::GetOnScreenDebugInfo(TArray<FText>& OutDebugInfo) const
+{
+	// todo: provide warnings
+}
+
+void FIKRetargetEditPoseMode::Render(const FSceneView* View, FViewport* Viewport, FPrimitiveDrawInterface* PDI)
+{
+	FEdMode::Render(View, Viewport, PDI);
+
+	const TSharedPtr<FIKRetargetEditorController> Controller = EditorController.Pin();
+	if (!Controller.IsValid())
+	{
+		return;
+	}
+
+	const UIKRetargetProcessor* RetargetProcessor = Controller->GetRetargetProcessor();
+	if (!(RetargetProcessor && RetargetProcessor->IsInitialized()))
+	{
+		return;
+	}
+
+	const UIKRetargeter* Asset = Controller->AssetController->GetAsset();
+	const FTargetSkeleton& TargetSkeleton = RetargetProcessor->GetTargetSkeleton();
+	const int32 RootBoneIndex = RetargetProcessor->GetTargetSkeletonRootBone();
+
+	TSet<int32> OutAffectedBones;
+	TSet<int32> OutSelectedBones;
+	UIKRigProcessor* IKRigProcessor = RetargetProcessor->GetTargetIKRigProcessor();
+	GetAffectedBones(Controller.Get(), IKRigProcessor, OutAffectedBones, OutSelectedBones);
+
+	const float MaxDrawRadius = Controller->TargetSkelMeshComponent->Bounds.SphereRadius * 0.01f;
+	for (int32 BoneIndex = 0; BoneIndex<TargetSkeleton.BoneNames.Num(); ++BoneIndex)
+	{		
+		// filter out bones that cannot be edited as part of the retarget pose
+		const bool bIsRootBone = BoneIndex==RootBoneIndex;
+		const bool bIsInTargetChain = TargetSkeleton.IsBoneInAnyTargetChain[BoneIndex];
+		if (!(bIsRootBone || bIsInTargetChain))
+		{
+			continue;
+		}
+
+		// selected bones are drawn with different color, affected bones are highlighted
+		const bool bIsSelected = OutSelectedBones.Contains(BoneIndex);
+		const bool bIsAffected = OutAffectedBones.Contains(BoneIndex);
+		FLinearColor LineColor = bIsAffected ? IKRigDebugRendering::AFFECTED_BONE_COLOR : IKRigDebugRendering::DESELECTED_BONE_COLOR;
+		LineColor = bIsSelected ? IKRigDebugRendering::SELECTED_BONE_COLOR : LineColor;
+
+		const float BoneRadiusSetting = Asset->BoneDrawSize;
+		const float BoneRadius = FMath::Min(1.0f, MaxDrawRadius) * BoneRadiusSetting;
+		
+		// get the location of the bone on the currently initialized target skeletal mesh
+		// along with array of child positions
+		const FTransform BoneTransform = Controller->GetTargetBoneGlobalTransform(RetargetProcessor, BoneIndex);
+		FVector Start;
+		TArray<FVector> ChildLocations;
+		Controller->GetTargetBoneLineSegments(RetargetProcessor, BoneIndex, Start, ChildLocations);
+		
+		// draw the bone proxy
+		PDI->SetHitProxy(new HIKRetargetEditorBoneProxy(TargetSkeleton.BoneNames[BoneIndex]));
+		IKRigDebugRendering::DrawWireBone(
+			PDI,
+			BoneTransform,
+			ChildLocations,
+			LineColor,
+			SDPG_Foreground,
+			BoneRadius,
+			bIsSelected || bIsAffected);
+		PDI->SetHitProxy(NULL);
+	}
+}
+
+void FIKRetargetEditPoseMode::GetAffectedBones(
+	FIKRetargetEditorController* Controller,
+	UIKRigProcessor* Processor,
+	TSet<int32>& OutAffectedBones,
+	TSet<int32>& OutSelectedBones) const
+{
+	if (!(Controller && Processor))
+	{
+		return;
+	}
+
+	const FIKRigSkeleton& Skeleton = Processor->GetSkeleton();
+
+	// record selected bone indices
+	for (const FName& SelectedBone : BoneEdit.SelectedBones)
+	{
+		OutSelectedBones.Add(Skeleton.GetBoneIndexFromName(SelectedBone));
+	}
+
+	// "affected bones" are the selected bones AND their children, recursively
+	for (int32 SelectedBone : OutSelectedBones)
+	{
+		const int32 EndOfBranch = Skeleton.GetCachedEndOfBranchIndex(SelectedBone);
+		OutAffectedBones.Add(SelectedBone);
+		for (int32 BoneIndex=SelectedBone; BoneIndex<=EndOfBranch; ++BoneIndex)
+		{
+			OutAffectedBones.Add(BoneIndex);
+		}
+	}
+}
+
+bool FIKRetargetEditPoseMode::AllowWidgetMove()
+{
+	return false;
+}
+
+bool FIKRetargetEditPoseMode::ShouldDrawWidget() const
+{
+	return UsesTransformWidget(CurrentWidgetMode);
+}
+
+bool FIKRetargetEditPoseMode::UsesTransformWidget() const
+{
+	return UsesTransformWidget(CurrentWidgetMode);
+}
+
+bool FIKRetargetEditPoseMode::UsesTransformWidget(UE::Widget::EWidgetMode CheckMode) const
+{
+	const bool bIsOnlyRootSelected = IsOnlyRootSelected();
+	const bool bIsAnyBoneSelected = !BoneEdit.SelectedBones.IsEmpty();
+
+	if (!bIsAnyBoneSelected)
+	{
+		return false; // no bones selected, can't transform anything
+	}
+	
+	if (bIsOnlyRootSelected && CheckMode == UE::Widget::EWidgetMode::WM_Translate)
+	{
+		return true; // can translate only the root
+	}
+
+	if (bIsAnyBoneSelected && CheckMode == UE::Widget::EWidgetMode::WM_Rotate)
+	{
+		return true; // can rotate any bone
+	}
+	
+	return false;
+}
+
+FVector FIKRetargetEditPoseMode::GetWidgetLocation() const
+{
+	if (BoneEdit.SelectedBones.IsEmpty())
+	{
+		return FVector::ZeroVector; 
+	}
+
+	const TSharedPtr<FIKRetargetEditorController> Controller = EditorController.Pin();
+	if (!Controller.IsValid())
+	{
+		return FVector::ZeroVector; 
+	}
+
+	const FTargetSkeleton& TargetSkeleton = Controller->GetRetargetProcessor()->GetTargetSkeleton();
+	const int32 BoneIndex = TargetSkeleton.FindBoneIndexByName(BoneEdit.SelectedBones.Last());
+	if (BoneIndex == INDEX_NONE)
+	{
+		return FVector::ZeroVector;
+	}
+
+	const UIKRetargetProcessor* RetargetProcessor = Controller->GetRetargetProcessor();
+	if (!(RetargetProcessor && RetargetProcessor->IsInitialized()))
+	{
+		return FVector::ZeroVector;
+	}
+	
+	return Controller->GetTargetBoneGlobalTransform(RetargetProcessor, BoneIndex).GetTranslation();
+}
+
+bool FIKRetargetEditPoseMode::HandleClick(FEditorViewportClient* InViewportClient, HHitProxy* HitProxy, const FViewportClick& Click)
+{
+	// check for selections
+	if (Click.GetKey() == EKeys::LeftMouseButton)
+	{
+		// selected bone ?
+		if (HitProxy && HitProxy->IsA(HIKRetargetEditorBoneProxy::StaticGetType()))
+		{
+			HIKRetargetEditorBoneProxy* BoneProxy = static_cast<HIKRetargetEditorBoneProxy*>(HitProxy);
+			const bool bReplaceSelection = !(InViewportClient->IsCtrlPressed() || InViewportClient->IsShiftPressed());
+			HandleBoneSelectedInViewport(BoneProxy->BoneName, bReplaceSelection);
+			return true;
+		}
+		else
+		{
+			// clicking in empty space clears selection
+			const bool bReplaceSelection = true;
+			HandleBoneSelectedInViewport(NAME_None, bReplaceSelection);
+
+			// show asset in details view
+			const TSharedPtr<FIKRetargetEditorController> Controller = EditorController.Pin();
+			if (Controller.IsValid())
+			{
+				Controller->DetailsView->SetObject(Controller->AssetController->GetAsset());
+			}
+		}
+	}
+	
+	return false;
+}
+
+bool FIKRetargetEditPoseMode::StartTracking(FEditorViewportClient* InViewportClient, FViewport* InViewport)
+{
+	TrackingState = FIKRetargetTrackingState::None;
+
+	// not manipulating any widget axes, so stop tracking
+	const EAxisList::Type CurrentAxis = InViewportClient->GetCurrentWidgetAxis();
+	if (CurrentAxis == EAxisList::None)
+	{
+		return false; 
+	}
+
+	const TSharedPtr<FIKRetargetEditorController> Controller = EditorController.Pin();
+	if (!Controller.IsValid())
+	{
+		return false; // invalid editor state
+	}
+
+	// get state of viewport
+	const bool bTranslating = InViewportClient->GetWidgetMode() == UE::Widget::EWidgetMode::WM_Translate;
+	const bool bRotating = InViewportClient->GetWidgetMode() == UE::Widget::EWidgetMode::WM_Rotate;
+	const bool bAnyBoneSelected = !BoneEdit.SelectedBones.IsEmpty();
+	const bool bOnlyRootSelected = IsOnlyRootSelected();
+
+	// is any bone being rotated?
+	if (bRotating && bAnyBoneSelected)
+	{
+		// Start a rotation transaction
+		GEditor->BeginTransaction(LOCTEXT("RotateRetargetPoseBone", "Rotate Retarget Pose Bone"));
+		Controller->AssetController->GetAsset()->Modify();
+		TrackingState = FIKRetargetTrackingState::RotatingBone;
+		UpdateWidgetTransform();
+		return true;
+	}
+
+	// is the root being translated?
+	if (bTranslating && bOnlyRootSelected)
+	{
+		// Start a translation transaction
+		GEditor->BeginTransaction(LOCTEXT("TranslateRetargetPoseBone", "Translate Retarget Pose Bone"));
+		Controller->AssetController->GetAsset()->Modify();
+		TrackingState = FIKRetargetTrackingState::TranslatingRoot;
+		UpdateWidgetTransform();
+		return true;
+	}
+
+	return false;
+}
+
+bool FIKRetargetEditPoseMode::EndTracking(FEditorViewportClient* InViewportClient, FViewport* InViewport)
+{	
+	if (TrackingState == FIKRetargetTrackingState::None)
+	{
+		const bool bIsRootSelected = IsRootSelected();
+		const bool bTranslating = InViewportClient->GetWidgetMode() == UE::Widget::EWidgetMode::WM_Translate;
+		// forcibly prevent translation of anything but the root
+		if (!bIsRootSelected && bTranslating)
+		{
+			InViewportClient->SetWidgetMode(UE::Widget::EWidgetMode::WM_Rotate);
+		}
+		
+		return true; // not handled
+	}
+
+	GEditor->EndTransaction();
+	TrackingState = FIKRetargetTrackingState::None;
+	return true;
+}
+
+bool FIKRetargetEditPoseMode::InputDelta(
+	FEditorViewportClient* InViewportClient,
+	FViewport* InViewport,
+	FVector& InDrag,
+	FRotator& InRot,
+	FVector& InScale)
+{
+	if (TrackingState == FIKRetargetTrackingState::None)
+	{
+		return false; // not handled
+	}
+
+	const TSharedPtr<FIKRetargetEditorController> Controller = EditorController.Pin();
+	if (!Controller.IsValid())
+	{
+		return false; 
+	}
+
+	// rotating any bone
+	if (TrackingState == FIKRetargetTrackingState::RotatingBone)
+	{
+		if(InViewportClient->GetWidgetMode() != UE::Widget::WM_Rotate)
+		{
+			return false;
+		}
+
+		// accumulate the rotation from the viewport gizmo (amount of rotation since tracking started)
+		BoneEdit.AccumulatedGlobalOffset = InRot.Quaternion().Inverse() * BoneEdit.AccumulatedGlobalOffset;
+
+		// convert to a delta rotation in the local space of the last selected bone
+		const FQuat BoneGlobalOrig = BoneEdit.GlobalTransform.GetRotation();
+		const FQuat BoneGlobalPlusOffset = BoneEdit.AccumulatedGlobalOffset * BoneGlobalOrig;
+		const FQuat ParentInv = BoneEdit.ParentGlobalTransform.GetRotation().Inverse();
+		const FQuat BoneLocal = ParentInv * BoneGlobalOrig;
+		const FQuat BoneLocalPlusOffset = ParentInv * BoneGlobalPlusOffset;
+		const FQuat BoneLocalOffset = BoneLocal * BoneLocalPlusOffset.Inverse();
+		
+		// apply rotation delta to all selected bones
+		for (int32 SelectionIndex=0; SelectionIndex<BoneEdit.SelectedBones.Num(); ++SelectionIndex)
+		{
+			const FQuat TotalOffsetLocal = BoneLocalOffset * BoneEdit.PrevLocalOffsets[SelectionIndex];
+			const FName& BoneName = BoneEdit.SelectedBones[SelectionIndex];
+			Controller->AssetController->SetRotationOffsetForRetargetPoseBone(BoneName, TotalOffsetLocal);
+		}
+
+		return true;
+	}
+	
+	// translating root
+	if (TrackingState == FIKRetargetTrackingState::TranslatingRoot)
+	{
+		if(InViewportClient->GetWidgetMode() != UE::Widget::WM_Translate)
+		{
+			return false;
+		}
+
+		// apply translation delta to root
+		Controller->AssetController->AddTranslationOffsetToRetargetRootBone(InDrag);
+		return true;
+	}
+	
+	return false;
+}
+
+bool FIKRetargetEditPoseMode::GetCustomDrawingCoordinateSystem(FMatrix& InMatrix, void* InData)
+{
+	if (BoneEdit.SelectedBones.IsEmpty())
+	{
+		return false; // nothing selected to manipulate
+	}
+
+	const TSharedPtr<FIKRetargetEditorController> Controller = EditorController.Pin();
+	if (!Controller.IsValid())
+	{
+		return false; 
+	}
+
+	if (TrackingState == FIKRetargetTrackingState::None)
+	{
+		UpdateWidgetTransform();
+	}
+
+	InMatrix = BoneEdit.GlobalTransform.ToMatrixNoScale().RemoveTranslation();
+	return true;
+}
+
+bool FIKRetargetEditPoseMode::GetCustomInputCoordinateSystem(FMatrix& InMatrix, void* InData)
+{
+	return GetCustomDrawingCoordinateSystem(InMatrix, InData);
+}
+
+void FIKRetargetEditPoseMode::UpdateWidgetTransform()
+{
+	const TSharedPtr<FIKRetargetEditorController> Controller = EditorController.Pin();
+	if (!Controller.IsValid() || BoneEdit.SelectedBones.IsEmpty())
+	{
+		BoneEdit.GlobalTransform = FTransform::Identity;
+		return;
+	}
+
+	const UIKRetargetProcessor* RetargetProcessor = Controller->GetRetargetProcessor();
+	if (!(RetargetProcessor && RetargetProcessor->IsInitialized()))
+	{
+		return;
+	}
+	
+	const UIKRetargeterController* AssetController = Controller->AssetController;
+	const FTargetSkeleton& TargetSkeleton = RetargetProcessor->GetTargetSkeleton();
+
+	BoneEdit.Name = BoneEdit.SelectedBones.Last();
+	BoneEdit.Index = TargetSkeleton.FindBoneIndexByName(BoneEdit.Name);
+	BoneEdit.GlobalTransform = Controller->GetTargetBoneGlobalTransform(RetargetProcessor, BoneEdit.Index);
+	BoneEdit.AccumulatedGlobalOffset = FQuat::Identity;
+
+	BoneEdit.PrevLocalOffsets.Reset();
+	for (int32 SelectionIndex=0; SelectionIndex<BoneEdit.SelectedBones.Num(); ++SelectionIndex)
+	{
+		FQuat PrevLocalOffset = AssetController->GetRotationOffsetForRetargetPoseBone(BoneEdit.SelectedBones[SelectionIndex]);
+		BoneEdit.PrevLocalOffsets.Add(PrevLocalOffset);
+	}
+	
+	const int32 ParentIndex = TargetSkeleton.GetParentIndex(BoneEdit.Index);
+	if (ParentIndex != INDEX_NONE)
+	{
+		BoneEdit.ParentGlobalTransform = Controller->GetTargetBoneGlobalTransform(RetargetProcessor, ParentIndex);
+	}else
+	{
+		BoneEdit.ParentGlobalTransform = FTransform::Identity;
+	}
+}
+
+bool FIKRetargetEditPoseMode::IsRootSelected() const
+{
+	if (BoneEdit.SelectedBones.IsEmpty())
+	{
+		return false;
+	}
+
+	const TSharedPtr<FIKRetargetEditorController> Controller = EditorController.Pin();
+	if (!Controller.IsValid())
+	{
+		return false; 
+	}
+
+	const FName& RootName = Controller->AssetController->GetTargetRootBone();
+	for (const FName& SelectedBone : BoneEdit.SelectedBones)
+	{
+		if (RootName == SelectedBone)
+		{
+			return true;
+		}
+	}
+
+	return false;
+}
+
+bool FIKRetargetEditPoseMode::IsOnlyRootSelected() const
+{
+	if (BoneEdit.SelectedBones.Num() != 1)
+	{
+		return false;
+	}
+
+	const TSharedPtr<FIKRetargetEditorController> Controller = EditorController.Pin();
+	if (!Controller.IsValid())
+	{
+		return false; 
+	}
+
+	return Controller->AssetController->GetTargetRootBone() == BoneEdit.SelectedBones[0];
+}
+
+void FIKRetargetEditPoseMode::Tick(FEditorViewportClient* ViewportClient, float DeltaTime)
+{
+	FEdMode::Tick(ViewportClient, DeltaTime);
+	
+	CurrentWidgetMode = ViewportClient->GetWidgetMode();
+}
+
+void FIKRetargetEditPoseMode::DrawHUD(FEditorViewportClient* ViewportClient, FViewport* Viewport, const FSceneView* View, FCanvas* Canvas)
+{
+	FEdMode::DrawHUD(ViewportClient, Viewport, View, Canvas);
+}
+
+bool FIKRetargetEditPoseMode::IsBoneSelected(const FName& BoneName) const
+{
+	return BoneEdit.SelectedBones.Contains(BoneName);
+}
+
+void FIKRetargetEditPoseMode::HandleBoneSelectedInViewport(const FName& BoneName, bool bReplace)
+{
+	if (bReplace)
+	{
+		BoneEdit.SelectedBones.Reset();
+		if (BoneName != NAME_None)
+		{
+			BoneEdit.SelectedBones.Add(BoneName);
+		}
+	}
+	else
+	{
+		if (BoneName != NAME_None)
+		{
+			const bool bAlreadySelected = BoneEdit.SelectedBones.Contains(BoneName);
+			if (bAlreadySelected)
+			{
+				BoneEdit.SelectedBones.Remove(BoneName);	
+			}
+			else
+			{
+				BoneEdit.SelectedBones.Add(BoneName);
+			}
+		}	
+	}
+}
+
+#undef LOCTEXT_NAMESPACE
diff --git a/Engine/Plugins/Animation/IKRig/Source/IKRigEditor/Private/RetargetEditor/IKRetargetEditor.cpp b/Engine/Plugins/Animation/IKRig/Source/IKRigEditor/Private/RetargetEditor/IKRetargetEditor.cpp
--- a/Engine/Plugins/Animation/IKRig/Source/IKRigEditor/Private/RetargetEditor/IKRetargetEditor.cpp
+++ b/Engine/Plugins/Animation/IKRig/Source/IKRigEditor/Private/RetargetEditor/IKRetargetEditor.cpp
@@ -18,13 +18,14 @@
 #include "Retargeter/IKRetargeter.h"
 #include "RetargetEditor/IKRetargetAnimInstance.h"
 #include "RetargetEditor/IKRetargetCommands.h"
-#include "RetargetEditor/IKRetargetEditMode.h"
-#include "RetargetEditor/IKRetargetMode.h"
+#include "RetargetEditor/IKRetargetEditPoseMode.h"
+#include "RetargetEditor/IKRetargetApplicationMode.h"
+#include "RetargetEditor/IKRetargetDefaultMode.h"
 #include "RetargetEditor/IKRetargetEditorController.h"
 
 #define LOCTEXT_NAMESPACE "IKRetargeterEditor"
 
-const FName IKRetargetEditorModes::IKRetargetEditorMode("IKRetargetEditorMode");
+const FName IKRetargetApplicationModes::IKRetargetApplicationMode("IKRetargetApplicationMode");
 const FName IKRetargetEditorAppName = FName(TEXT("IKRetargetEditorApp"));
 
 FIKRetargetEditor::FIKRetargetEditor()
@@ -65,17 +66,25 @@ void FIKRetargetEditor::InitAssetEditor(
 		bCreateDefaultToolbar, 
 		InAsset);
 
+	// this sets the application mode which defines the tab factory that builds the editor layout
 	AddApplicationMode(
-		IKRetargetEditorModes::IKRetargetEditorMode,
-		MakeShareable(new FIKRetargetMode(SharedThis(this), EditorController->PersonaToolkit->GetPreviewScene())));
+		IKRetargetApplicationModes::IKRetargetApplicationMode,
+		MakeShareable(new FIKRetargetApplicationMode(SharedThis(this),EditorController->PersonaToolkit->GetPreviewScene())));
+	SetCurrentMode(IKRetargetApplicationModes::IKRetargetApplicationMode);
 
-	SetCurrentMode(IKRetargetEditorModes::IKRetargetEditorMode);
+	// set the default editing mode to use in the editor
+	GetEditorModeManager().SetDefaultMode(FIKRetargetDefaultMode::ModeName);
+	
+	// give default editing mode a pointer to the editor controller
+	GetEditorModeManager().ActivateMode(FIKRetargetDefaultMode::ModeName);
+	FIKRetargetDefaultMode* DefaultMode = GetEditorModeManager().GetActiveModeTyped<FIKRetargetDefaultMode>(FIKRetargetDefaultMode::ModeName);
+	DefaultMode->SetEditorController(EditorController);
 
-	GetEditorModeManager().SetDefaultMode(FPersonaEditModes::SkeletonSelection);
-	GetEditorModeManager().ActivateMode(FIKRetargetEditMode::ModeName);
-	FIKRetargetEditMode* EditMode = GetEditorModeManager().GetActiveModeTyped<FIKRetargetEditMode>(FIKRetargetEditMode::ModeName);
-	EditMode->SetEditorController(EditorController);
-	GetEditorModeManager().DeactivateMode(FIKRetargetEditMode::ModeName);
+	// give edit pose mode a pointer to the editor controller
+	GetEditorModeManager().ActivateMode(FIKRetargetEditPoseMode::ModeName);
+	FIKRetargetEditPoseMode* EditPoseMode = GetEditorModeManager().GetActiveModeTyped<FIKRetargetEditPoseMode>(FIKRetargetEditPoseMode::ModeName);
+	EditPoseMode->SetEditorController(EditorController);
+	GetEditorModeManager().DeactivateMode(FIKRetargetEditPoseMode::ModeName);
 
 	ExtendToolbar();
 	RegenerateMenusAndToolbars();
@@ -246,17 +255,9 @@ void FIKRetargetEditor::AddReferencedObjects(FReferenceCollector& Collector)
 
 void FIKRetargetEditor::Tick(float DeltaTime)
 {
-	// apply offset to the target component
-	if (EditorController->TargetSkelMeshComponent)
-	{
-		const UIKRetargeter* Retargeter = EditorController->AssetController->GetAsset();
-		
-		const float TargetOffset = Retargeter->TargetActorOffset;
-		EditorController->TargetSkelMeshComponent->SetRelativeLocation(FVector(TargetOffset,0,0));
-
-		const float TargetScale = Retargeter->TargetActorScale;
-		EditorController->TargetSkelMeshComponent->SetRelativeScale3D(FVector(TargetScale,TargetScale,TargetScale));
-	}
+	// update with latest offsets
+	EditorController->AddOffsetAndUpdatePreviewMeshPosition(FVector::ZeroVector, EditorController->SourceSkelMeshComponent);
+	EditorController->AddOffsetAndUpdatePreviewMeshPosition(FVector::ZeroVector, EditorController->TargetSkelMeshComponent);
 }
 
 TStatId FIKRetargetEditor::GetStatId() const
@@ -318,11 +319,13 @@ void FIKRetargetEditor::HandlePreviewSceneCreated(const TSharedRef<IPersonaPrevi
 
 	// apply mesh to the preview scene
 	InPersonaPreviewScene->SetPreviewMeshComponent(EditorController->SourceSkelMeshComponent);
-	InPersonaPreviewScene->SetAllowMeshHitProxies(false);
-	InPersonaPreviewScene->SetAdditionalMeshesSelectable(false);
-	EditorController->SourceSkelMeshComponent->bSelectable = false;
-	EditorController->TargetSkelMeshComponent->bSelectable = false;
 	InPersonaPreviewScene->SetPreviewMesh(SourceMesh);
+	InPersonaPreviewScene->SetAdditionalMeshesSelectable(false);
+
+	// SetPreviewMesh() sets this flag true, which the render uses to filter out objects for selection highlighting...
+	// but since we want to be able to select the mesh in this viewport, we have to set it back to false
+	EditorController->SourceSkelMeshComponent->bCanHighlightSelectedSections = false;
+	
 	InPersonaPreviewScene->AddComponent(EditorController->SourceSkelMeshComponent, FTransform::Identity);
 	InPersonaPreviewScene->AddComponent(EditorController->TargetSkelMeshComponent, FTransform::Identity);
 }
@@ -373,6 +376,7 @@ void FIKRetargetEditor::OnFinishedChangingDetails(const FPropertyChangedEvent& P
 		{
 			PreviewScene->SetPreviewMeshComponent(EditorController->SourceSkelMeshComponent);
 			PreviewScene->SetPreviewMesh(SourceMesh);
+			EditorController->SourceSkelMeshComponent->bCanHighlightSelectedSections = false;
 		}
 	
 		SetupAnimInstance();
diff --git a/Engine/Plugins/Animation/IKRig/Source/IKRigEditor/Private/RetargetEditor/IKRetargetEditorController.cpp b/Engine/Plugins/Animation/IKRig/Source/IKRigEditor/Private/RetargetEditor/IKRetargetEditorController.cpp
--- a/Engine/Plugins/Animation/IKRig/Source/IKRigEditor/Private/RetargetEditor/IKRetargetEditorController.cpp
+++ b/Engine/Plugins/Animation/IKRig/Source/IKRigEditor/Private/RetargetEditor/IKRetargetEditorController.cpp
@@ -6,7 +6,8 @@
 #include "Widgets/Input/SButton.h"
 #include "Animation/DebugSkelMeshComponent.h"
 #include "RetargetEditor/IKRetargetAnimInstance.h"
-#include "RetargetEditor/IKRetargetEditMode.h"
+#include "RetargetEditor/IKRetargetDefaultMode.h"
+#include "RetargetEditor/IKRetargetEditPoseMode.h"
 #include "RetargetEditor/IKRetargetEditor.h"
 #include "RetargetEditor/SIKRetargetChainMapList.h"
 #include "Retargeter/IKRetargeter.h"
@@ -79,13 +80,40 @@ void FIKRetargetEditorController::OnRetargeterNeedsInitialized(const UIKRetarget
 {
 	// force edit pose mode to be off
 	AssetController->SetEditRetargetPoseMode(false, false); // turn off and don't reinitialize (avoid infinite loop)
-	Editor.Pin()->GetEditorModeManager().DeactivateMode(FIKRetargetEditMode::ModeName);
+	Editor.Pin()->GetEditorModeManager().DeactivateMode(FIKRetargetEditPoseMode::ModeName);
 	// force reinit the runtime retarget processor
 	TargetAnimInstance->SetProcessorNeedsInitialized();
 	// refresh all the UI views
 	RefreshAllViews();
 }
 
+void FIKRetargetEditorController::AddOffsetAndUpdatePreviewMeshPosition(
+	const FVector& Offset,
+	USceneComponent* Component) const
+{
+	UIKRetargeter* Asset = AssetController->GetAsset();
+	FVector Position;
+	float Scale;
+	if (Component == TargetSkelMeshComponent)
+	{
+		Asset->TargetMeshOffset += Offset;
+		Position = Asset->TargetMeshOffset;
+		Scale = Asset->TargetMeshScale;
+	}
+	else
+	{
+		Asset->SourceMeshOffset += Offset;
+		Position = Asset->SourceMeshOffset;
+		Scale = 1.0f;
+	}
+
+	constexpr bool bSweep = false;
+	constexpr FHitResult* OutSweepHitResult = nullptr;
+	constexpr ETeleportType Teleport = ETeleportType::ResetPhysics;
+	Component->SetWorldLocation(Position, bSweep, OutSweepHitResult, Teleport);
+	Component->SetWorldScale3D(FVector(Scale,Scale,Scale));
+}
+
 USkeletalMesh* FIKRetargetEditorController::GetSourceSkeletalMesh() const
 {
 	if (!(AssetController && AssetController->GetAsset()->GetSourceIKRig()))
@@ -116,8 +144,8 @@ FTransform FIKRetargetEditorController::GetTargetBoneGlobalTransform(
 	FTransform BoneTransform = RetargetProcessor->GetTargetBoneRetargetPoseGlobalTransform(TargetBoneIndex);
 
 	// scale and offset
-	BoneTransform.ScaleTranslation(AssetController->GetAsset()->TargetActorScale);
-	BoneTransform.AddToTranslation(FVector(AssetController->GetAsset()->TargetActorOffset, 0.f, 0.f));
+	BoneTransform.ScaleTranslation(AssetController->GetAsset()->TargetMeshScale);
+	BoneTransform.AddToTranslation(AssetController->GetAsset()->TargetMeshOffset);
 
 	return BoneTransform;
 }
@@ -156,13 +184,13 @@ bool FIKRetargetEditorController::GetTargetBoneLineSegments(
 	}
 
 	// add the target translation offset and scale
-	const FVector TargetOffset(AssetController->GetAsset()->TargetActorOffset, 0.f, 0.f);
-	OutStart *= AssetController->GetAsset()->TargetActorScale;
-	OutStart += TargetOffset;
+	const UIKRetargeter* Asset = AssetController->GetAsset();
+	OutStart *= Asset->TargetMeshScale;
+	OutStart += Asset->TargetMeshOffset;
 	for (FVector& ChildPoint : OutChildren)
 	{
-		ChildPoint *= AssetController->GetAsset()->TargetActorScale;
-		ChildPoint += TargetOffset;
+		ChildPoint *= Asset->TargetMeshScale;
+		ChildPoint += Asset->TargetMeshOffset;
 	}
 	
 	return true;
@@ -212,12 +240,16 @@ void FIKRetargetEditorController::HandleEditPose() const
 	AssetController->SetEditRetargetPoseMode(bEditPoseMode);
 	if (bEditPoseMode)
 	{
-		Editor.Pin()->GetEditorModeManager().ActivateMode(FIKRetargetEditMode::ModeName);
+		Editor.Pin()->GetEditorModeManager().DeactivateMode(FIKRetargetDefaultMode::ModeName);
+		Editor.Pin()->GetEditorModeManager().ActivateMode(FIKRetargetEditPoseMode::ModeName);
 		SourceSkelMeshComponent->ShowReferencePose(true);
+		// have to move component back to offset position because ShowReferencePose() sets it back to origin
+		AddOffsetAndUpdatePreviewMeshPosition(FVector::ZeroVector, SourceSkelMeshComponent);
 	}
 	else
 	{
-		Editor.Pin()->GetEditorModeManager().DeactivateMode(FIKRetargetEditMode::ModeName);
+		Editor.Pin()->GetEditorModeManager().DeactivateMode(FIKRetargetEditPoseMode::ModeName);
+		Editor.Pin()->GetEditorModeManager().ActivateMode(FIKRetargetDefaultMode::ModeName);
 		PlayPreviousAnimationAsset();
 	}
 }
diff --git a/Engine/Plugins/Animation/IKRig/Source/IKRigEditor/Private/RetargetEditor/IKRetargetMode.cpp b/Engine/Plugins/Animation/IKRig/Source/IKRigEditor/Private/RetargetEditor/IKRetargetMode.cpp
deleted file mode 100644
diff --git a/Engine/Plugins/Animation/IKRig/Source/IKRigEditor/Public/RetargetEditor/IKRetargetApplicationMode.h b/Engine/Plugins/Animation/IKRig/Source/IKRigEditor/Public/RetargetEditor/IKRetargetApplicationMode.h
new file mode 100644
--- /dev/null
+++ b/Engine/Plugins/Animation/IKRig/Source/IKRigEditor/Public/RetargetEditor/IKRetargetApplicationMode.h
@@ -0,0 +1,29 @@
+// Copyright Epic Games, Inc. All Rights Reserved.
+
+#pragma once
+
+#include "WorkflowOrientedApp/ApplicationMode.h"
+#include "WorkflowOrientedApp/WorkflowCentricApplication.h"
+#include "WorkflowOrientedApp/WorkflowTabManager.h"
+
+class FIKRetargetEditor;
+
+class FIKRetargetApplicationMode : public FApplicationMode
+{
+	
+public:
+	
+	FIKRetargetApplicationMode(TSharedRef<class FWorkflowCentricApplication> InHostingApp, TSharedRef<class IPersonaPreviewScene> InPreviewScene);
+
+	/** FApplicationMode interface */
+	virtual void RegisterTabFactories(TSharedPtr<FTabManager> InTabManager) override;
+	/** END FApplicationMode interface */
+
+protected:
+	
+	/** The hosting app */
+	TWeakPtr<FIKRetargetEditor> IKRetargetEditorPtr;
+
+	/** The tab factories we support */
+	FWorkflowAllowedTabSet TabFactories;
+};
diff --git a/Engine/Plugins/Animation/IKRig/Source/IKRigEditor/Public/RetargetEditor/IKRetargetDefaultMode.h b/Engine/Plugins/Animation/IKRig/Source/IKRigEditor/Public/RetargetEditor/IKRetargetDefaultMode.h
new file mode 100644
--- /dev/null
+++ b/Engine/Plugins/Animation/IKRig/Source/IKRigEditor/Public/RetargetEditor/IKRetargetDefaultMode.h
@@ -0,0 +1,70 @@
+﻿// Copyright Epic Games, Inc. All Rights Reserved.
+
+#pragma once
+
+#include "CoreMinimal.h"
+
+#include "Retargeter/IKRetargeter.h"
+#include "IPersonaEditMode.h"
+
+class UIKRigProcessor;
+class FIKRetargetEditorController;
+class FIKRetargetEditor;
+class FIKRetargetPreviewScene;
+
+
+class FIKRetargetDefaultMode : public IPersonaEditMode
+{
+public:
+	static FName ModeName;
+	
+	FIKRetargetDefaultMode() = default;
+
+	/** glue for all the editor parts to communicate */
+	void SetEditorController(const TSharedPtr<FIKRetargetEditorController> InEditorController) { EditorController = InEditorController; };
+
+	/** IPersonaEditMode interface */
+	virtual bool GetCameraTarget(FSphere& OutTarget) const override;
+	virtual class IPersonaPreviewScene& GetAnimPreviewScene() const override;
+	virtual void GetOnScreenDebugInfo(TArray<FText>& OutDebugInfo) const override;
+	/** END IPersonaEditMode interface */
+
+	/** FEdMode interface */
+	virtual void Initialize() override;
+	virtual void Tick(FEditorViewportClient* ViewportClient, float DeltaTime) override;
+	virtual void Render(const FSceneView* View, FViewport* Viewport, FPrimitiveDrawInterface* PDI) override;
+	virtual void DrawHUD(FEditorViewportClient* ViewportClient, FViewport* Viewport, const FSceneView* View, FCanvas* Canvas) override;
+	virtual bool IsCompatibleWith(FEditorModeID OtherModeID) const override { return true; }
+	virtual bool AllowWidgetMove() override;
+	virtual bool ShouldDrawWidget() const override;
+	virtual bool UsesTransformWidget() const override;
+	virtual bool UsesTransformWidget(UE::Widget::EWidgetMode CheckMode) const override;
+	virtual FVector GetWidgetLocation() const override;
+	virtual bool HandleClick(FEditorViewportClient* InViewportClient, HHitProxy* HitProxy, const FViewportClick& Click) override;
+	virtual bool StartTracking(FEditorViewportClient* InViewportClient, FViewport* InViewport) override;
+	virtual bool EndTracking(FEditorViewportClient* InViewportClient, FViewport* InViewport) override;
+	virtual bool InputDelta(FEditorViewportClient* InViewportClient, FViewport* InViewport, FVector& InDrag, FRotator& InRot, FVector& InScale) override;
+	virtual bool GetCustomDrawingCoordinateSystem(FMatrix& InMatrix, void* InData) override;
+	virtual bool GetCustomInputCoordinateSystem(FMatrix& InMatrix, void* InData) override;
+
+	virtual void Enter() override;
+	virtual void Exit() override;
+	// IS THIS NEEDED
+	virtual bool IsSelectionAllowed( AActor* InActor, bool bInSelection ) const override { return true; }
+	/** END FEdMode interface */
+
+private:
+
+	TObjectPtr<UPrimitiveComponent> SelectedComponent;
+	bool ComponentSelectionOverride(const UPrimitiveComponent* InComponent) const;
+	void SetSelectedComponent(UPrimitiveComponent* InComponent);
+	static void ApplyOffsetToMeshTransform(const FVector& Offset, USceneComponent* Component);
+	
+	/** The hosting app */
+	TWeakPtr<FIKRetargetEditorController> EditorController;
+
+	UE::Widget::EWidgetMode CurrentWidgetMode;
+	bool bIsTranslating = false;
+
+	bool bIsInitialized = false;
+};
diff --git a/Engine/Plugins/Animation/IKRig/Source/IKRigEditor/Public/RetargetEditor/IKRetargetEditMode.h b/Engine/Plugins/Animation/IKRig/Source/IKRigEditor/Public/RetargetEditor/IKRetargetEditMode.h
deleted file mode 100644
diff --git a/Engine/Plugins/Animation/IKRig/Source/IKRigEditor/Public/RetargetEditor/IKRetargetEditPoseMode.h b/Engine/Plugins/Animation/IKRig/Source/IKRigEditor/Public/RetargetEditor/IKRetargetEditPoseMode.h
new file mode 100644
--- /dev/null
+++ b/Engine/Plugins/Animation/IKRig/Source/IKRigEditor/Public/RetargetEditor/IKRetargetEditPoseMode.h
@@ -0,0 +1,92 @@
+// Copyright Epic Games, Inc. All Rights Reserved.
+
+#pragma once
+
+#include "CoreMinimal.h"
+
+#include "Retargeter/IKRetargeter.h"
+#include "IPersonaEditMode.h"
+
+class UIKRigProcessor;
+class FIKRetargetEditorController;
+class FIKRetargetEditor;
+class FIKRetargetPreviewScene;
+
+enum FIKRetargetTrackingState : int8
+{
+	None,
+	RotatingBone,
+	TranslatingRoot,
+};
+
+struct BoneEdit
+{
+	FName Name;							// name of last selected bone
+	int32 Index;						// index of last selected bone
+	FTransform ParentGlobalTransform;	// global transform of parent of last selected bone
+	FTransform GlobalTransform;			// global transform of last selected bone
+	FTransform LocalTransform;			// local transform of last selected bone
+	FQuat AccumulatedGlobalOffset;		// the accumulated offset from rotation gizmo
+	
+	TArray<FQuat> PrevLocalOffsets;		// the prev stored local offsets of all selected bones
+	TArray<FName> SelectedBones;		// the currently selected bones in the viewport
+};
+
+class FIKRetargetEditPoseMode : public IPersonaEditMode
+{
+public:
+	static FName ModeName;
+	
+	FIKRetargetEditPoseMode() = default;
+
+	/** glue for all the editor parts to communicate */
+	void SetEditorController(const TSharedPtr<FIKRetargetEditorController> InEditorController) { EditorController = InEditorController; };
+
+	/** IPersonaEditMode interface */
+	virtual bool GetCameraTarget(FSphere& OutTarget) const override;
+	virtual class IPersonaPreviewScene& GetAnimPreviewScene() const override;
+	virtual void GetOnScreenDebugInfo(TArray<FText>& OutDebugInfo) const override;
+	/** END IPersonaEditMode interface */
+
+	/** FEdMode interface */
+	virtual void Tick(FEditorViewportClient* ViewportClient, float DeltaTime) override;
+	virtual void Render(const FSceneView* View, FViewport* Viewport, FPrimitiveDrawInterface* PDI) override;
+	virtual void DrawHUD(FEditorViewportClient* ViewportClient, FViewport* Viewport, const FSceneView* View, FCanvas* Canvas) override;
+	virtual bool IsCompatibleWith(FEditorModeID OtherModeID) const override { return true; }
+	virtual bool AllowWidgetMove() override;
+	virtual bool ShouldDrawWidget() const override;
+	virtual bool UsesTransformWidget() const override;
+	virtual bool UsesTransformWidget(UE::Widget::EWidgetMode CheckMode) const override;
+	virtual FVector GetWidgetLocation() const override;
+	virtual bool HandleClick(FEditorViewportClient* InViewportClient, HHitProxy* HitProxy, const FViewportClick& Click) override;
+	virtual bool StartTracking(FEditorViewportClient* InViewportClient, FViewport* InViewport) override;
+	virtual bool EndTracking(FEditorViewportClient* InViewportClient, FViewport* InViewport) override;
+	virtual bool InputDelta(FEditorViewportClient* InViewportClient, FViewport* InViewport, FVector& InDrag, FRotator& InRot, FVector& InScale) override;
+	virtual bool GetCustomDrawingCoordinateSystem(FMatrix& InMatrix, void* InData) override;
+	virtual bool GetCustomInputCoordinateSystem(FMatrix& InMatrix, void* InData) override;
+	/** END FEdMode interface */
+
+private:
+
+	void GetAffectedBones(
+		FIKRetargetEditorController* Controller,
+		UIKRigProcessor* Processor,
+		TSet<int32>& OutAffectedBones,
+		TSet<int32>& OutSelectedBones) const;
+
+	UE::Widget::EWidgetMode CurrentWidgetMode;
+
+	bool IsRootSelected() const;
+	bool IsOnlyRootSelected() const;
+	bool IsBoneSelected(const FName& BoneName) const;
+
+	BoneEdit BoneEdit;
+	void UpdateWidgetTransform();
+	void HandleBoneSelectedInViewport(const FName& BoneName, bool bReplace);
+	
+	/** The hosting app */
+	TWeakPtr<FIKRetargetEditorController> EditorController;
+
+	/** viewport selection/editing state */
+	FIKRetargetTrackingState TrackingState;
+};
diff --git a/Engine/Plugins/Animation/IKRig/Source/IKRigEditor/Public/RetargetEditor/IKRetargetEditor.h b/Engine/Plugins/Animation/IKRig/Source/IKRigEditor/Public/RetargetEditor/IKRetargetEditor.h
--- a/Engine/Plugins/Animation/IKRig/Source/IKRigEditor/Public/RetargetEditor/IKRetargetEditor.h
+++ b/Engine/Plugins/Animation/IKRig/Source/IKRigEditor/Public/RetargetEditor/IKRetargetEditor.h
@@ -6,7 +6,7 @@
 
 #include "IHasPersonaToolkit.h"
 #include "IKRetargetEditorController.h"
-#include "IKRetargetMode.h"
+#include "IKRetargetApplicationMode.h"
 #include "IPersonaPreviewScene.h"
 #include "PersonaAssetEditorToolkit.h"
 #include "EditorUndoClient.h"
@@ -22,9 +22,9 @@ class FIKRetargetPreviewScene;
 struct FIKRetargetPose;
 class SEditableTextBox;
 
-namespace IKRetargetEditorModes
+namespace IKRetargetApplicationModes
 {
-	extern const FName IKRetargetEditorMode;
+	extern const FName IKRetargetApplicationMode;
 }
 
 class FIKRetargetEditor :
@@ -97,5 +97,5 @@ private:
 	/** centralized management across all views */
 	TSharedRef<FIKRetargetEditorController> EditorController;
 	
-	friend FIKRetargetMode;
+	friend FIKRetargetApplicationMode;
 };
diff --git a/Engine/Plugins/Animation/IKRig/Source/IKRigEditor/Public/RetargetEditor/IKRetargetEditorController.h b/Engine/Plugins/Animation/IKRig/Source/IKRigEditor/Public/RetargetEditor/IKRetargetEditorController.h
--- a/Engine/Plugins/Animation/IKRig/Source/IKRigEditor/Public/RetargetEditor/IKRetargetEditorController.h
+++ b/Engine/Plugins/Animation/IKRig/Source/IKRigEditor/Public/RetargetEditor/IKRetargetEditorController.h
@@ -53,6 +53,8 @@ public:
 	UPROPERTY(transient, NonTransactional)
 	TWeakObjectPtr<class UIKRetargetAnimInstance> TargetAnimInstance;
 
+	void AddOffsetAndUpdatePreviewMeshPosition(const FVector& Offset, USceneComponent* Component) const;
+
 	/** asset properties tab */
 	TSharedPtr<class IDetailsView> DetailsView;
 
diff --git a/Engine/Plugins/Animation/IKRig/Source/IKRigEditor/Public/RetargetEditor/IKRetargetMode.h b/Engine/Plugins/Animation/IKRig/Source/IKRigEditor/Public/RetargetEditor/IKRetargetMode.h
deleted file mode 100644
diff --git a/Engine/Source/Runtime/Engine/Private/EngineUtils.cpp b/Engine/Source/Runtime/Engine/Private/EngineUtils.cpp
--- a/Engine/Source/Runtime/Engine/Private/EngineUtils.cpp
+++ b/Engine/Source/Runtime/Engine/Private/EngineUtils.cpp
@@ -506,8 +506,8 @@ FStripDataFlags::FStripDataFlags( class FArchive& Ar, uint8 InClassFlags /*= 0*/
 		{
 			// When cooking GlobalStripFlags are automatically generated based on the current target
 			// platform's properties.
-			GlobalStripFlags |= Ar.CookingTarget()->HasEditorOnlyData() ? FStripDataFlags::None : FStripDataFlags::Editor;
-			GlobalStripFlags |= !Ar.CookingTarget()->AllowAudioVisualData() ? FStripDataFlags::Server : FStripDataFlags::None;
+			GlobalStripFlags |= Ar.CookingTarget()->HasEditorOnlyData() ? static_cast<uint8>(FStripDataFlags::EStrippedData::None) : static_cast<uint8>(FStripDataFlags::EStrippedData::Editor);
+			GlobalStripFlags |= !Ar.CookingTarget()->AllowAudioVisualData() ? static_cast<uint8>(FStripDataFlags::EStrippedData::Server) : static_cast<uint8>(FStripDataFlags::EStrippedData::None);
 			ClassStripFlags = InClassFlags;
 		}
 		Ar << GlobalStripFlags;
@@ -550,8 +550,8 @@ FStripDataFlags::FStripDataFlags(FStructuredArchive::FSlot Slot, uint8 InClassFl
 		{
 			// When cooking GlobalStripFlags are automatically generated based on the current target
 			// platform's properties.
-			GlobalStripFlags |= UnderlyingArchive.CookingTarget()->HasEditorOnlyData() ? FStripDataFlags::None : FStripDataFlags::Editor;
-			GlobalStripFlags |= !UnderlyingArchive.CookingTarget()->AllowAudioVisualData() ? FStripDataFlags::Server : FStripDataFlags::None;
+			GlobalStripFlags |= UnderlyingArchive.CookingTarget()->HasEditorOnlyData() ? static_cast<uint8>(FStripDataFlags::EStrippedData::None) : static_cast<uint8>(FStripDataFlags::EStrippedData::Editor);
+			GlobalStripFlags |= !UnderlyingArchive.CookingTarget()->AllowAudioVisualData() ? static_cast<uint8>(FStripDataFlags::EStrippedData::Server) : static_cast<uint8>(FStripDataFlags::EStrippedData::None);
 			ClassStripFlags = InClassFlags;
 		}
 		Record << SA_VALUE(TEXT("GlobalStripFlags"), GlobalStripFlags);
diff --git a/Engine/Source/Runtime/Engine/Private/Model.cpp b/Engine/Source/Runtime/Engine/Private/Model.cpp
--- a/Engine/Source/Runtime/Engine/Private/Model.cpp
+++ b/Engine/Source/Runtime/Engine/Private/Model.cpp
@@ -187,7 +187,7 @@ void UModel::Serialize( FArchive& Ar )
 	Super::Serialize( Ar );
 
 	const int32 StripVertexBufferFlag = 1;
-	FStripDataFlags StripFlags( Ar, GetOuter() && GetOuter()->IsA(ABrush::StaticClass()) ? StripVertexBufferFlag : FStripDataFlags::None );
+	FStripDataFlags StripFlags( Ar, GetOuter() && GetOuter()->IsA(ABrush::StaticClass()) ? StripVertexBufferFlag : static_cast<int32>(FStripDataFlags::EStrippedData::None) );
 
 	Ar << Bounds;
 
diff --git a/Engine/Source/Runtime/Engine/Public/EngineUtils.h b/Engine/Source/Runtime/Engine/Public/EngineUtils.h
--- a/Engine/Source/Runtime/Engine/Public/EngineUtils.h
+++ b/Engine/Source/Runtime/Engine/Public/EngineUtils.h
@@ -824,7 +824,7 @@ class ENGINE_API FStripDataFlags
 public:
 
 	/** Engine strip flags */
-	enum EStrippedData
+	enum class EStrippedData : uint8
 	{
 		None = 0,
 
@@ -916,7 +916,7 @@ public:
 	 */
 	FORCEINLINE bool IsEditorDataStripped() const
 	{
-		return (GlobalStripFlags & FStripDataFlags::Editor) != 0;
+		return (GlobalStripFlags & static_cast<uint8>(FStripDataFlags::EStrippedData::Editor)) != 0;
 	}
 
 	/**
@@ -926,7 +926,7 @@ public:
 	 */
 	bool IsDataStrippedForServer() const
 	{
-		return (GlobalStripFlags & FStripDataFlags::Server) != 0;
+		return (GlobalStripFlags & static_cast<uint8>(FStripDataFlags::EStrippedData::Server)) != 0;
 	}
 
 	/**
