Adding "Go to Ref Pose" button to retargeter.
Moved "Edit Pose" toolbar into details panel.
Added ability to rename retarget poses

diff --git a/Engine/Plugins/Animation/IKRig/Source/IKRigEditor/Private/IKRigEditor.cpp b/Engine/Plugins/Animation/IKRig/Source/IKRigEditor/Private/IKRigEditor.cpp
--- a/Engine/Plugins/Animation/IKRig/Source/IKRigEditor/Private/IKRigEditor.cpp
+++ b/Engine/Plugins/Animation/IKRig/Source/IKRigEditor/Private/IKRigEditor.cpp
@@ -12,6 +12,7 @@
 #include "RetargetEditor/AssetTypeActions_IKRetargeter.h"
 #include "RetargetEditor/IKRetargetCommands.h"
 #include "RetargetEditor/IKRetargetDefaultMode.h"
+#include "RetargetEditor/IKRetargetDetails.h"
 #include "RetargetEditor/IKRetargetEditPoseMode.h"
 #include "RigEditor/IKRigCommands.h"
 #include "RigEditor/IKRigEditMode.h"
@@ -53,6 +54,8 @@ void FIKRigEditor::StartupModule()
 	PropertyEditorModule.RegisterCustomClassLayout(ClassesToUnregisterOnShutdown.Last(), FOnGetDetailCustomizationInstance::CreateStatic(&FIKRigGenericDetailCustomization::MakeInstance));
 	ClassesToUnregisterOnShutdown.Add(UIKRigEffectorGoal::StaticClass()->GetFName());
 	PropertyEditorModule.RegisterCustomClassLayout(ClassesToUnregisterOnShutdown.Last(), FOnGetDetailCustomizationInstance::CreateStatic(&FIKRigGenericDetailCustomization::MakeInstance));
+	ClassesToUnregisterOnShutdown.Add(UIKRetargeter::StaticClass()->GetFName());
+	PropertyEditorModule.RegisterCustomClassLayout(ClassesToUnregisterOnShutdown.Last(), FOnGetDetailCustomizationInstance::CreateStatic(&FIKRetargeterDetails::MakeInstance));
 }
 
 void FIKRigEditor::ShutdownModule()
diff --git a/Engine/Plugins/Animation/IKRig/Source/IKRigEditor/Private/RetargetEditor/IKRetargetCommands.cpp b/Engine/Plugins/Animation/IKRig/Source/IKRigEditor/Private/RetargetEditor/IKRetargetCommands.cpp
--- a/Engine/Plugins/Animation/IKRig/Source/IKRigEditor/Private/RetargetEditor/IKRetargetCommands.cpp
+++ b/Engine/Plugins/Animation/IKRig/Source/IKRigEditor/Private/RetargetEditor/IKRetargetCommands.cpp
@@ -6,10 +6,14 @@
 
 void FIKRetargetCommands::RegisterCommands()
 {
-	UI_COMMAND(EditRetargetPose, "Edit Pose", "Edit retarget reference pose of target mesh. Usually best to match the source mesh.", EUserInterfaceActionType::ToggleButton, FInputChord());
-	UI_COMMAND(NewRetargetPose, "New Pose", "Create a new retarget pose for the target mesh.", EUserInterfaceActionType::Button, FInputChord());
-	UI_COMMAND(DeleteRetargetPose, "Delete Pose", "Delete a retarget pose.", EUserInterfaceActionType::Button, FInputChord());
-	UI_COMMAND(ResetRetargetPose, "Reset Pose", "Resets a retarget pose to the mesh reference pose.", EUserInterfaceActionType::Button, FInputChord());
+	UI_COMMAND(GoToRetargetPose, "Go to Retarget Pose", "Stop playback and set Source and Target meshes to the Retarget Pose.", EUserInterfaceActionType::Button, FInputChord());
+	
+	UI_COMMAND(EditRetargetPose, "Edit Mode", "Enter into mode allowing manual editing of the target skeleton pose in the viewport.", EUserInterfaceActionType::ToggleButton, FInputChord());
+	UI_COMMAND(SetToRefPose, "Set to Ref Pose", "Sets the retarget pose to the mesh reference pose.", EUserInterfaceActionType::Button, FInputChord());
+
+	UI_COMMAND(NewRetargetPose, "New", "Create a new retarget pose for the target mesh.", EUserInterfaceActionType::Button, FInputChord());
+	UI_COMMAND(DeleteRetargetPose, "Delete", "Delete current retarget pose.", EUserInterfaceActionType::Button, FInputChord());
+	UI_COMMAND(RenameRetargetPose, "Rename", "Rename current retarget pose.", EUserInterfaceActionType::Button, FInputChord());
 }
 
 #undef LOCTEXT_NAMESPACE
diff --git a/Engine/Plugins/Animation/IKRig/Source/IKRigEditor/Private/RetargetEditor/IKRetargetDetails.cpp b/Engine/Plugins/Animation/IKRig/Source/IKRigEditor/Private/RetargetEditor/IKRetargetDetails.cpp
new file mode 100644
--- /dev/null
+++ b/Engine/Plugins/Animation/IKRig/Source/IKRigEditor/Private/RetargetEditor/IKRetargetDetails.cpp
@@ -0,0 +1,147 @@
+﻿// Copyright Epic Games, Inc. All Rights Reserved.
+
+#include "RetargetEditor/IKRetargetDetails.h"
+
+#include "DetailCategoryBuilder.h"
+#include "DetailLayoutBuilder.h"
+#include "DetailWidgetRow.h"
+#include "Framework/MultiBox/SToolBarButtonBlock.h"
+#include "RetargetEditor/IKRetargetCommands.h"
+#include "RetargetEditor/IKRetargetEditor.h"
+#include "RetargetEditor/IKRetargetEditorController.h"
+#include "RetargetEditor/IKRetargeterController.h"
+#include "Retargeter/IKRetargeter.h"
+
+#define LOCTEXT_NAMESPACE "IKRetargeterDetails"
+
+TSharedRef<IDetailCustomization> FIKRetargeterDetails::MakeInstance()
+{
+	return MakeShareable(new FIKRetargeterDetails);
+}
+
+void FIKRetargeterDetails::CustomizeDetails(IDetailLayoutBuilder& DetailBuilder)
+{
+	GetAssetController(DetailBuilder);
+	if (!Controller)
+	{
+		return;
+	}
+
+	FIKRetargetEditorController* EditorController = Controller->GetEditorController();
+	if (!EditorController)
+	{
+		return;
+	}
+
+	// add a new category at the top to edit the retarget pose
+	IDetailCategoryBuilder& EditPoseCategoryBuilder = DetailBuilder.EditCategory( "Edit Retarget Pose", LOCTEXT("EditPoseLabel", "Edit Retarget Pose"), ECategoryPriority::Default );
+
+	// fill list of pose names
+	PoseNames.Reset();
+	for (const TTuple<FName, FIKRetargetPose>& Pose : Controller->GetRetargetPoses())
+	{
+		PoseNames.Add(MakeShareable(new FName(Pose.Key)));
+	}
+
+	// add pose selection combo box
+	FDetailWidgetRow& CurrentPoseRow = EditPoseCategoryBuilder.AddCustomRow(LOCTEXT("CurrentPoseLabel", "Current Pose"))
+	.NameContent()
+	[
+		SNew(STextBlock)
+		.Text(LOCTEXT("CurrentPose", "Current Retarget Pose"))
+		.Font(DetailBuilder.GetDetailFont())
+	]
+	.ValueContent()
+	[
+		SNew(SComboBox<TSharedPtr<FName>>)
+		.OptionsSource(&PoseNames)
+		.OnGenerateWidget_Lambda([](TSharedPtr<FName> InItem)
+		{
+			return SNew(STextBlock).Text(FText::FromName(*InItem.Get()));
+		})
+		.OnSelectionChanged(EditorController, &FIKRetargetEditorController::OnPoseSelected)
+		[
+			SNew(STextBlock).Text(EditorController, &FIKRetargetEditorController::GetCurrentPoseName)
+		]
+	];
+
+	// add pose editing toolbar
+	const TSharedRef<FUICommandList>& Commands = Controller->GetEditorController()->Editor.Pin()->GetToolkitCommands();
+	FDetailWidgetRow& ToolbarRow = EditPoseCategoryBuilder.AddCustomRow(LOCTEXT("CurrentPoseLabel", "Edit Pose"))
+	.WholeRowWidget
+	[
+		SNew(SHorizontalBox)
+		+SHorizontalBox::Slot()
+		.FillWidth(1.0f)
+		.HAlign(HAlign_Center)
+		[
+			MakeToolbar(Commands)
+		]
+	];
+}
+
+TSharedRef<SWidget>  FIKRetargeterDetails::MakeToolbar(const TSharedRef<FUICommandList>& Commands)
+{
+	FToolBarBuilder ToolbarBuilder(Commands, FMultiBoxCustomization::None);
+	
+	ToolbarBuilder.BeginSection("Edit Current Pose");
+
+	ToolbarBuilder.AddToolBarButton(
+		FIKRetargetCommands::Get().EditRetargetPose,
+		NAME_None,
+		TAttribute<FText>(),
+		TAttribute<FText>(),
+		FSlateIcon(FAppStyle::Get().GetStyleSetName(),"Icons.Edit"));
+
+	ToolbarBuilder.AddToolBarButton(
+		FIKRetargetCommands::Get().SetToRefPose,
+		NAME_None,
+		TAttribute<FText>(),
+		TAttribute<FText>(),
+		FSlateIcon(FAppStyle::Get().GetStyleSetName(),"Icons.Refresh"));
+
+	ToolbarBuilder.EndSection();
+	ToolbarBuilder.BeginSection("Create Poses");
+
+	ToolbarBuilder.AddToolBarButton(
+		FIKRetargetCommands::Get().NewRetargetPose,
+		NAME_None,
+		TAttribute<FText>(),
+		TAttribute<FText>(),
+		FSlateIcon(FAppStyle::Get().GetStyleSetName(),"Icons.Plus"));
+
+	ToolbarBuilder.AddToolBarButton(
+		FIKRetargetCommands::Get().DeleteRetargetPose,
+		NAME_None,
+		TAttribute<FText>(),
+		TAttribute<FText>(),
+		FSlateIcon(FAppStyle::Get().GetStyleSetName(),"Icons.Delete"));
+
+	ToolbarBuilder.AddToolBarButton(
+		FIKRetargetCommands::Get().RenameRetargetPose,
+		NAME_None,
+		TAttribute<FText>(),
+		TAttribute<FText>(),
+		FSlateIcon(FAppStyle::Get().GetStyleSetName(),"Icons.Settings"));
+
+	ToolbarBuilder.EndSection();
+
+	return ToolbarBuilder.MakeWidget();
+}
+
+void FIKRetargeterDetails::GetAssetController(IDetailLayoutBuilder& DetailBuilder)
+{
+	Controller = nullptr;
+	
+	TArray< TWeakObjectPtr<UObject> > OutObjects;
+	DetailBuilder.GetObjectsBeingCustomized(OutObjects);
+	if (OutObjects.IsEmpty())
+	{
+		return;
+	}
+
+	const TObjectPtr<UIKRetargeter> Asset = CastChecked<UIKRetargeter>(OutObjects[0].Get());
+	Controller = UIKRetargeterController::GetController(Asset);
+}
+
+#undef LOCTEXT_NAMESPACE
diff --git a/Engine/Plugins/Animation/IKRig/Source/IKRigEditor/Private/RetargetEditor/IKRetargetEditPoseMode.cpp b/Engine/Plugins/Animation/IKRig/Source/IKRigEditor/Private/RetargetEditor/IKRetargetEditPoseMode.cpp
--- a/Engine/Plugins/Animation/IKRig/Source/IKRigEditor/Private/RetargetEditor/IKRetargetEditPoseMode.cpp
+++ b/Engine/Plugins/Animation/IKRig/Source/IKRigEditor/Private/RetargetEditor/IKRetargetEditPoseMode.cpp
@@ -400,6 +400,34 @@ bool FIKRetargetEditPoseMode::GetCustomInputCoordinateSystem(FMatrix& InMatrix,
 	return GetCustomDrawingCoordinateSystem(InMatrix, InData);
 }
 
+void FIKRetargetEditPoseMode::Enter()
+{
+	const TSharedPtr<FIKRetargetEditorController> Controller = EditorController.Pin();
+	if (!Controller.IsValid())
+	{
+		return;
+	}
+
+	// put source mesh in reference pose
+	Controller->SourceSkelMeshComponent->ShowReferencePose(true);
+	// have to move component back to offset position because ShowReferencePose() sets it back to origin
+	Controller->AddOffsetAndUpdatePreviewMeshPosition(FVector::ZeroVector, Controller->SourceSkelMeshComponent);
+
+	// put asset in mode where target mesh will output retarget pose (for preview purposes)
+	Controller->AssetController->SetEditRetargetPoseMode(true);
+}
+
+void FIKRetargetEditPoseMode::Exit()
+{
+	const TSharedPtr<FIKRetargetEditorController> Controller = EditorController.Pin();
+	if (!Controller.IsValid())
+	{
+		return;
+	}
+
+	Controller->AssetController->SetEditRetargetPoseMode(false);
+}
+
 void FIKRetargetEditPoseMode::UpdateWidgetTransform()
 {
 	const TSharedPtr<FIKRetargetEditorController> Controller = EditorController.Pin();
diff --git a/Engine/Plugins/Animation/IKRig/Source/IKRigEditor/Private/RetargetEditor/IKRetargetEditor.cpp b/Engine/Plugins/Animation/IKRig/Source/IKRigEditor/Private/RetargetEditor/IKRetargetEditor.cpp
--- a/Engine/Plugins/Animation/IKRig/Source/IKRigEditor/Private/RetargetEditor/IKRetargetEditor.cpp
+++ b/Engine/Plugins/Animation/IKRig/Source/IKRigEditor/Private/RetargetEditor/IKRetargetEditor.cpp
@@ -107,6 +107,13 @@ void FIKRetargetEditor::BindCommands()
 {
 	const FIKRetargetCommands& Commands = FIKRetargetCommands::Get();
 
+	ToolkitCommands->MapAction(
+		Commands.GoToRetargetPose,
+		FExecuteAction::CreateSP(EditorController, &FIKRetargetEditorController::HandleGoToRetargetPose),
+		FCanExecuteAction(),
+		FCanExecuteAction(),
+		EUIActionRepeatMode::RepeatDisabled);
+	
 	ToolkitCommands->MapAction(
         Commands.EditRetargetPose,
         FExecuteAction::CreateSP(EditorController, &FIKRetargetEditorController::HandleEditPose),
@@ -114,10 +121,16 @@ void FIKRetargetEditor::BindCommands()
         FIsActionChecked::CreateSP(EditorController,  &FIKRetargetEditorController::IsEditingPose),
 		EUIActionRepeatMode::RepeatDisabled);
 
+	ToolkitCommands->MapAction(
+		Commands.SetToRefPose,
+		FExecuteAction::CreateSP(EditorController, &FIKRetargetEditorController::HandleResetPose),
+		FCanExecuteAction::CreateSP(EditorController, &FIKRetargetEditorController::CanResetPose),
+		EUIActionRepeatMode::RepeatDisabled);
+
 	ToolkitCommands->MapAction(
 		Commands.NewRetargetPose,
 		FExecuteAction::CreateSP(EditorController, &FIKRetargetEditorController::HandleNewPose),
-		FCanExecuteAction(),
+		FCanExecuteAction::CreateSP(EditorController, &FIKRetargetEditorController::CanNewPose),
 		FCanExecuteAction(),
 		EUIActionRepeatMode::RepeatDisabled);
 
@@ -129,8 +142,9 @@ void FIKRetargetEditor::BindCommands()
 		EUIActionRepeatMode::RepeatDisabled);
 
 	ToolkitCommands->MapAction(
-		Commands.ResetRetargetPose,
-		FExecuteAction::CreateSP(EditorController, &FIKRetargetEditorController::HandleResetPose),
+		Commands.RenameRetargetPose,
+		FExecuteAction::CreateSP(EditorController, &FIKRetargetEditorController::HandleRenamePose),
+		FCanExecuteAction::CreateSP(EditorController, &FIKRetargetEditorController::CanRenamePose),
 		FCanExecuteAction(),
 		EUIActionRepeatMode::RepeatDisabled);
 }
@@ -151,72 +165,14 @@ void FIKRetargetEditor::ExtendToolbar()
 
 void FIKRetargetEditor::FillToolbar(FToolBarBuilder& ToolbarBuilder)
 {
-	ToolbarBuilder.BeginSection("Retarget Pose");
+	ToolbarBuilder.BeginSection("Go To Retarget Pose");
 	{
-		EditorController->PoseNames.Reset();
-		for (const TTuple<FName, FIKRetargetPose>& Pose : EditorController->AssetController->GetRetargetPoses())
-		{
-			EditorController->PoseNames.Add(MakeShareable(new FName(Pose.Key)));
-		}
-
-		TSharedRef<SWidget> PoseListWidget = SNew(SHorizontalBox)
-
-		+ SHorizontalBox::Slot()
-		.AutoWidth()
-		.HAlign(HAlign_Left)
-		.VAlign(VAlign_Center)
-		.Padding(3.0f, 1.0f)
-		[
-			SNew(STextBlock)
-			.Text(LOCTEXT("IKRetargetPoseTitleLabel", "Current Retarget Pose: "))
-		]
-
-		+ SHorizontalBox::Slot()
-		.AutoWidth()
-		.HAlign(HAlign_Left)
-		.VAlign(VAlign_Center)
-		.Padding(3.0f, 1.0f)
-		[
-			SNew(SComboBox<TSharedPtr<FName>>)
-			.OptionsSource(&EditorController->PoseNames)
-			.OnGenerateWidget_Lambda([](TSharedPtr<FName> InItem)
-			{
-				return SNew(STextBlock).Text(FText::FromName(*InItem.Get()));
-			})
-			.OnSelectionChanged(EditorController, &FIKRetargetEditorController::OnPoseSelected)
-			[
-				SNew(STextBlock).Text(EditorController, &FIKRetargetEditorController::GetCurrentPoseName)
-			]
-		];
-		ToolbarBuilder.AddWidget(PoseListWidget);
-		
-		ToolbarBuilder.AddToolBarButton(
-			FIKRetargetCommands::Get().EditRetargetPose,
-			NAME_None,
-			TAttribute<FText>(),
-			TAttribute<FText>(),
-			FSlateIcon(FAppStyle::Get().GetStyleSetName(),"Icons.Edit"));
-
-		ToolbarBuilder.AddToolBarButton(
-			FIKRetargetCommands::Get().NewRetargetPose,
-			NAME_None,
-			TAttribute<FText>(),
-			TAttribute<FText>(),
-			FSlateIcon(FAppStyle::Get().GetStyleSetName(),"Icons.Plus"));
-
-		ToolbarBuilder.AddToolBarButton(
-			FIKRetargetCommands::Get().DeleteRetargetPose,
-			NAME_None,
-			TAttribute<FText>(),
-			TAttribute<FText>(),
-			FSlateIcon(FAppStyle::Get().GetStyleSetName(),"Icons.Delete"));
-
 		ToolbarBuilder.AddToolBarButton(
-			FIKRetargetCommands::Get().ResetRetargetPose,
+			FIKRetargetCommands::Get().GoToRetargetPose,
 			NAME_None,
 			TAttribute<FText>(),
 			TAttribute<FText>(),
-			FSlateIcon(FAppStyle::Get().GetStyleSetName(),"Icons.Refresh"));
+			FSlateIcon(FEditorStyle::GetStyleSetName(),"GenericStop"));
 	}
 	ToolbarBuilder.EndSection();
 }
diff --git a/Engine/Plugins/Animation/IKRig/Source/IKRigEditor/Private/RetargetEditor/IKRetargetEditorController.cpp b/Engine/Plugins/Animation/IKRig/Source/IKRigEditor/Private/RetargetEditor/IKRetargetEditorController.cpp
--- a/Engine/Plugins/Animation/IKRig/Source/IKRigEditor/Private/RetargetEditor/IKRetargetEditorController.cpp
+++ b/Engine/Plugins/Animation/IKRig/Source/IKRigEditor/Private/RetargetEditor/IKRetargetEditorController.cpp
@@ -19,6 +19,7 @@ void FIKRetargetEditorController::Initialize(TSharedPtr<FIKRetargetEditor> InEdi
 {
 	Editor = InEditor;
 	AssetController = UIKRetargeterController::GetController(InAsset);
+	AssetController->SetEditorController(this);
 
 	// bind callbacks when SOURCE or TARGET IK Rigs are modified
 	BindToIKRigAsset(AssetController->GetAsset()->GetSourceIKRigWriteable());
@@ -79,7 +80,6 @@ void FIKRetargetEditorController::OnRetargetChainRemoved(UIKRigDefinition* Modif
 void FIKRetargetEditorController::OnRetargeterNeedsInitialized(const UIKRetargeter* Retargeter) const
 {
 	// force edit pose mode to be off
-	AssetController->SetEditRetargetPoseMode(false, false); // turn off and don't reinitialize (avoid infinite loop)
 	Editor.Pin()->GetEditorModeManager().DeactivateMode(FIKRetargetEditPoseMode::ModeName);
 	// force reinit the runtime retarget processor
 	TargetAnimInstance->SetProcessorNeedsInitialized();
@@ -209,6 +209,7 @@ const UIKRetargetProcessor* FIKRetargetEditorController::GetRetargetProcessor()
 void FIKRetargetEditorController::RefreshAllViews() const
 {
 	Editor.Pin()->RegenerateMenusAndToolbars();
+	DetailsView->ForceRefresh();
 
 	// cannot assume chains view is always available
 	if (ChainsView.IsValid())
@@ -234,24 +235,36 @@ void FIKRetargetEditorController::PlayPreviousAnimationAsset() const
 	}
 }
 
+void FIKRetargetEditorController::HandleGoToRetargetPose() const
+{
+	Editor.Pin()->GetEditorModeManager().DeactivateMode(FIKRetargetEditPoseMode::ModeName);
+	Editor.Pin()->GetEditorModeManager().ActivateMode(FIKRetargetDefaultMode::ModeName);
+
+	// put source back in ref pose
+	SourceSkelMeshComponent->ShowReferencePose(true);
+	// have to move component back to offset position because ShowReferencePose() sets it back to origin
+	AddOffsetAndUpdatePreviewMeshPosition(FVector::ZeroVector, SourceSkelMeshComponent);
+}
+
 void FIKRetargetEditorController::HandleEditPose() const
 {
-	const bool bEditPoseMode = !AssetController->GetEditRetargetPoseMode();
-	AssetController->SetEditRetargetPoseMode(bEditPoseMode);
-	if (bEditPoseMode)
-	{
-		Editor.Pin()->GetEditorModeManager().DeactivateMode(FIKRetargetDefaultMode::ModeName);
-		Editor.Pin()->GetEditorModeManager().ActivateMode(FIKRetargetEditPoseMode::ModeName);
-		SourceSkelMeshComponent->ShowReferencePose(true);
-		// have to move component back to offset position because ShowReferencePose() sets it back to origin
-		AddOffsetAndUpdatePreviewMeshPosition(FVector::ZeroVector, SourceSkelMeshComponent);
-	}
-	else
+	if (IsEditingPose())
 	{
+		// stop pose editing
 		Editor.Pin()->GetEditorModeManager().DeactivateMode(FIKRetargetEditPoseMode::ModeName);
 		Editor.Pin()->GetEditorModeManager().ActivateMode(FIKRetargetDefaultMode::ModeName);
+		
+		// must reinitialize after editing the retarget pose
+		AssetController->BroadcastNeedsReinitialized();
+		// continue playing whatever animation asset was last used
 		PlayPreviousAnimationAsset();
 	}
+	else
+	{
+		// start pose editing
+		Editor.Pin()->GetEditorModeManager().DeactivateMode(FIKRetargetDefaultMode::ModeName);
+		Editor.Pin()->GetEditorModeManager().ActivateMode(FIKRetargetEditPoseMode::ModeName);
+	}
 }
 
 bool FIKRetargetEditorController::CanEditPose() const
@@ -335,12 +348,17 @@ void FIKRetargetEditorController::HandleNewPose()
 	NewPoseWindow.Reset();
 }
 
+bool FIKRetargetEditorController::CanNewPose() const
+{
+	return !IsEditingPose();
+}
+
 FReply FIKRetargetEditorController::CreateNewPose() const
 {
 	const FName NewPoseName = FName(NewPoseEditableText.Get()->GetText().ToString());
 	AssetController->AddRetargetPose(NewPoseName);
 	NewPoseWindow->RequestDestroyWindow();
-	Editor.Pin()->RegenerateMenusAndToolbars();
+	DetailsView->ForceRefresh();
 	return FReply::Handled();
 }
 
@@ -348,13 +366,15 @@ void FIKRetargetEditorController::HandleDeletePose() const
 {
 	const FName CurrentPose = AssetController->GetCurrentRetargetPoseName();
 	AssetController->RemoveRetargetPose(CurrentPose);
-	Editor.Pin()->RegenerateMenusAndToolbars();
+	DetailsView->ForceRefresh();
 }
 
 bool FIKRetargetEditorController::CanDeletePose() const
 {	
 	// cannot delete default pose
-	return AssetController->GetCurrentRetargetPoseName() != UIKRetargeter::GetDefaultPoseName();
+	const bool bNotUsingDefaultPose = AssetController->GetCurrentRetargetPoseName() != UIKRetargeter::GetDefaultPoseName();
+	// cannot delete pose while editing
+	return bNotUsingDefaultPose && !IsEditingPose();
 }
 
 void FIKRetargetEditorController::HandleResetPose() const
@@ -363,6 +383,95 @@ void FIKRetargetEditorController::HandleResetPose() const
 	AssetController->ResetRetargetPose(CurrentPose);
 }
 
+bool FIKRetargetEditorController::CanResetPose() const
+{
+	// only allow reseting pose while editing to avoid confusion
+	return IsEditingPose();
+}
+
+void FIKRetargetEditorController::HandleRenamePose()
+{
+	SAssignNew(RenamePoseWindow, SWindow)
+	.Title(LOCTEXT("RenameRetargetPoseOptions", "Rename Retarget Pose"))
+	.ClientSize(FVector2D(250, 80))
+	.HasCloseButton(true)
+	.SupportsMinimize(false) .SupportsMaximize(false)
+	[
+		SNew(SBorder)
+		.BorderImage( FEditorStyle::GetBrush("Menu.Background") )
+		[
+			SNew(SVerticalBox)
+
+			+ SVerticalBox::Slot()
+			.Padding(4)
+			.AutoHeight()
+			[
+				SAssignNew(NewNameEditableText, SEditableTextBox)
+				.Text(GetCurrentPoseName())
+			]
+
+			+ SVerticalBox::Slot()
+			.Padding(4)
+			.AutoHeight()
+			[
+				SNew(SHorizontalBox)
+				+ SHorizontalBox::Slot()
+				.HAlign(HAlign_Center)
+				[
+					SNew(SButton)
+					.ButtonStyle(FAppStyle::Get(), "Button")
+					.TextStyle( FAppStyle::Get(), "DialogButtonText" )
+					.HAlign(HAlign_Center)
+					.VAlign(VAlign_Center)
+					.Text(LOCTEXT("OkButtonLabel", "Ok"))
+					.IsEnabled_Lambda([this]()
+					{
+						return !GetCurrentPoseName().EqualTo(NewNameEditableText.Get()->GetText());
+					})
+					.OnClicked(this, &FIKRetargetEditorController::RenamePose)
+				]
+				
+				+ SHorizontalBox::Slot()
+				.HAlign(HAlign_Center)
+				[
+					SNew(SButton)
+					.ButtonStyle(FAppStyle::Get(), "Button")
+					.TextStyle( FAppStyle::Get(), "DialogButtonText" )
+					.HAlign(HAlign_Center)
+					.VAlign(VAlign_Center)
+					.Text(LOCTEXT("CancelButtonLabel", "Cancel"))
+					.OnClicked_Lambda( [=]()
+					{
+						RenamePoseWindow->RequestDestroyWindow();
+						return FReply::Handled();
+					})
+				]
+			]	
+		]
+	];
+
+	GEditor->EditorAddModalWindow(RenamePoseWindow.ToSharedRef());
+	RenamePoseWindow.Reset();
+}
+
+FReply FIKRetargetEditorController::RenamePose() const
+{
+	const FName NewPoseName = FName(NewNameEditableText.Get()->GetText().ToString());
+	RenamePoseWindow->RequestDestroyWindow();
+	
+	AssetController->RenameCurrentRetargetPose(NewPoseName);
+	DetailsView->ForceRefresh();
+	return FReply::Handled();
+}
+
+bool FIKRetargetEditorController::CanRenamePose() const
+{
+	// cannot rename default pose
+	const bool bNotUsingDefaultPose = AssetController->GetCurrentRetargetPoseName() != UIKRetargeter::GetDefaultPoseName();
+	// cannot rename pose while editing
+	return bNotUsingDefaultPose && !IsEditingPose();
+}
+
 FText FIKRetargetEditorController::GetCurrentPoseName() const
 {
 	return FText::FromName(AssetController->GetCurrentRetargetPoseName());
diff --git a/Engine/Plugins/Animation/IKRig/Source/IKRigEditor/Private/RetargetEditor/IKRetargeterController.cpp b/Engine/Plugins/Animation/IKRig/Source/IKRigEditor/Private/RetargetEditor/IKRetargeterController.cpp
--- a/Engine/Plugins/Animation/IKRig/Source/IKRigEditor/Private/RetargetEditor/IKRetargeterController.cpp
+++ b/Engine/Plugins/Animation/IKRig/Source/IKRigEditor/Private/RetargetEditor/IKRetargeterController.cpp
@@ -5,6 +5,7 @@
 #include "ScopedTransaction.h"
 #include "Algo/LevenshteinDistance.h"
 #include "Engine/SkeletalMesh.h"
+#include "RetargetEditor/IKRetargetEditorController.h"
 #include "Retargeter/IKRetargeter.h"
 
 #define LOCTEXT_NAMESPACE "IKRetargeterController"
@@ -26,8 +27,9 @@ UIKRetargeterController* UIKRetargeterController::GetController(UIKRetargeter* I
 
 	UIKRetargeterController* Controller = Cast<UIKRetargeterController>(InRetargeterAsset->Controller);
 	// clean the asset before editing
-	Controller->CleanChainMapping();
-	Controller->CleanPoseList();
+	const bool bForceReinitialization = false;
+	Controller->CleanChainMapping(bForceReinitialization);
+	Controller->CleanPoseList(bForceReinitialization);
 	
 	return Controller;
 }
@@ -37,6 +39,11 @@ UIKRetargeter* UIKRetargeterController::GetAsset() const
 	return Asset;
 }
 
+FIKRetargetEditorController* UIKRetargeterController::GetEditorController() const
+{
+	return EditorController;
+}
+
 void UIKRetargeterController::SetSourceIKRig(UIKRigDefinition* SourceIKRig)
 {
 	Asset->SourceIKRigAsset = SourceIKRig;
@@ -100,7 +107,7 @@ void UIKRetargeterController::GetSourceChainNames(TArray<FName>& OutNames) const
 	}
 }
 
-void UIKRetargeterController::CleanChainMapping()
+void UIKRetargeterController::CleanChainMapping(const bool bForceReinitialization)
 {
 	if (!IsValid(Asset->TargetIKRigAsset))
 	{
@@ -160,10 +167,13 @@ void UIKRetargeterController::CleanChainMapping()
 	// enforce the chain order based on the StartBone index
 	SortChainMapping();
 
-	BroadcastNeedsReinitialized();
+	if (bForceReinitialization)
+	{
+		BroadcastNeedsReinitialized();
+	}
 }
 
-void UIKRetargeterController::CleanPoseList()
+void UIKRetargeterController::CleanPoseList(const bool bForceReinitialization)
 {
 	// enforce the existence of a default pose
 	const bool HasDefaultPose = Asset->RetargetPoses.Contains(Asset->DefaultPoseName);
@@ -205,7 +215,10 @@ void UIKRetargeterController::CleanPoseList()
 		}
 	}
 
-	BroadcastNeedsReinitialized();
+	if (bForceReinitialization)
+	{
+		BroadcastNeedsReinitialized();	
+	}
 }
 
 void UIKRetargeterController::AutoMapChains() const
@@ -343,6 +356,30 @@ void UIKRetargeterController::AddRetargetPose(FName NewPoseName) const
 	BroadcastNeedsReinitialized();
 }
 
+void UIKRetargeterController::RenameCurrentRetargetPose(FName NewPoseName) const
+{
+	// do we already have a retarget pose with this name?
+	if (Asset->RetargetPoses.Contains(NewPoseName))
+	{
+		return;
+	}
+	
+	FScopedTransaction Transaction(LOCTEXT("RenameRetargetPose", "Rename Retarget Pose"));
+	Asset->Modify();
+
+	// replace key in the map
+	TMap<FName, FIKRetargetPose>& Poses = Asset->RetargetPoses;
+	const FIKRetargetPose CurrentPose = Poses[Asset->CurrentRetargetPose];
+	Poses.Remove(Asset->CurrentRetargetPose);
+	Poses.Shrink();
+	Poses.Add(NewPoseName, CurrentPose);
+
+	// update current pose name
+	Asset->CurrentRetargetPose = NewPoseName;
+
+	BroadcastNeedsReinitialized();
+}
+
 void UIKRetargeterController::RemoveRetargetPose(FName PoseToRemove) const
 {
 	if (PoseToRemove == Asset->DefaultPoseName)
@@ -381,7 +418,7 @@ void UIKRetargeterController::ResetRetargetPose(FName PoseToReset) const
 	
 	Asset->RetargetPoses[PoseToReset].BoneRotationOffsets.Reset();
 	Asset->RetargetPoses[PoseToReset].RootTranslationOffset = FVector::ZeroVector;
-	
+
 	BroadcastNeedsReinitialized();
 }
 
@@ -428,19 +465,14 @@ void UIKRetargeterController::AddTranslationOffsetToRetargetRootBone(FVector Tra
 	Asset->RetargetPoses[Asset->CurrentRetargetPose].AddTranslationDeltaToRoot(TranslationOffset);
 }
 
-void UIKRetargeterController::SetEditRetargetPoseMode(bool bEditPoseMode, bool bReinitializeAfter) const
+bool UIKRetargeterController::GetEditRetargetPoseMode() const
 {
-	GetAsset()->bEditRetargetPoseMode = bEditPoseMode;
-	if (!bEditPoseMode && bReinitializeAfter)
-	{
-		// must reinitialize after editing the retarget pose
-		BroadcastNeedsReinitialized();
-	}
+	return GetAsset()->bEditRetargetPoseMode;
 }
 
-bool UIKRetargeterController::GetEditRetargetPoseMode() const
+void UIKRetargeterController::SetEditRetargetPoseMode(const bool bMode) const
 {
-	return GetAsset()->bEditRetargetPoseMode;
+	GetAsset()->bEditRetargetPoseMode = bMode;
 }
 
 FName UIKRetargeterController::MakePoseNameUnique(FName PoseName) const
diff --git a/Engine/Plugins/Animation/IKRig/Source/IKRigEditor/Public/RetargetEditor/IKRetargetCommands.h b/Engine/Plugins/Animation/IKRig/Source/IKRigEditor/Public/RetargetEditor/IKRetargetCommands.h
--- a/Engine/Plugins/Animation/IKRig/Source/IKRigEditor/Public/RetargetEditor/IKRetargetCommands.h
+++ b/Engine/Plugins/Animation/IKRig/Source/IKRigEditor/Public/RetargetEditor/IKRetargetCommands.h
@@ -17,18 +17,24 @@ public:
 		FIKRetargetEditorStyle::Get().GetStyleSetName() // Icon Style Set
 	)
 	{}
+
+	/** reset reference pose */
+	TSharedPtr< FUICommandInfo > GoToRetargetPose;
 	
 	/** edit reference pose */
 	TSharedPtr< FUICommandInfo > EditRetargetPose;
 
+	/** reset reference pose */
+	TSharedPtr< FUICommandInfo > SetToRefPose;
+
 	/** new reference pose */
 	TSharedPtr< FUICommandInfo > NewRetargetPose;
 
 	/** delete reference pose */
 	TSharedPtr< FUICommandInfo > DeleteRetargetPose;
 
-	/** reset reference pose */
-	TSharedPtr< FUICommandInfo > ResetRetargetPose;
+	/** rename reference pose */
+	TSharedPtr< FUICommandInfo > RenameRetargetPose;
 	
 	/** export animation */
 	TSharedPtr< FUICommandInfo > ExportAnimation;
diff --git a/Engine/Plugins/Animation/IKRig/Source/IKRigEditor/Public/RetargetEditor/IKRetargetDetails.h b/Engine/Plugins/Animation/IKRig/Source/IKRigEditor/Public/RetargetEditor/IKRetargetDetails.h
new file mode 100644
--- /dev/null
+++ b/Engine/Plugins/Animation/IKRig/Source/IKRigEditor/Public/RetargetEditor/IKRetargetDetails.h
@@ -0,0 +1,30 @@
+﻿// Copyright Epic Games, Inc. All Rights Reserved.
+
+#pragma once
+
+#include "CoreMinimal.h"
+#include "IDetailCustomization.h"
+#include "IKRetargeterController.h"
+
+// details customization for IKRetargeter asset
+class FIKRetargeterDetails : public IDetailCustomization
+{
+	
+public:
+	
+	// makes a new instance of this detail layout class for a specific detail view requesting it
+	static TSharedRef<IDetailCustomization> MakeInstance();
+
+	// Begin IDetailCustomization interface
+	virtual void CustomizeDetails(IDetailLayoutBuilder& DetailBuilder) override;
+	// End IDetailCustomization interface
+
+private:
+
+	TSharedRef<SWidget> MakeToolbar(const TSharedRef<FUICommandList>& UICommandList);
+
+	TObjectPtr<UIKRetargeterController> Controller;
+	void GetAssetController(IDetailLayoutBuilder& DetailBuilder);
+
+	TArray<TSharedPtr<FName>> PoseNames;
+};
diff --git a/Engine/Plugins/Animation/IKRig/Source/IKRigEditor/Public/RetargetEditor/IKRetargetEditPoseMode.h b/Engine/Plugins/Animation/IKRig/Source/IKRigEditor/Public/RetargetEditor/IKRetargetEditPoseMode.h
--- a/Engine/Plugins/Animation/IKRig/Source/IKRigEditor/Public/RetargetEditor/IKRetargetEditPoseMode.h
+++ b/Engine/Plugins/Animation/IKRig/Source/IKRigEditor/Public/RetargetEditor/IKRetargetEditPoseMode.h
@@ -64,6 +64,8 @@ public:
 	virtual bool InputDelta(FEditorViewportClient* InViewportClient, FViewport* InViewport, FVector& InDrag, FRotator& InRot, FVector& InScale) override;
 	virtual bool GetCustomDrawingCoordinateSystem(FMatrix& InMatrix, void* InData) override;
 	virtual bool GetCustomInputCoordinateSystem(FMatrix& InMatrix, void* InData) override;
+	virtual void Enter() override;
+	virtual void Exit() override;
 	/** END FEdMode interface */
 
 private:
diff --git a/Engine/Plugins/Animation/IKRig/Source/IKRigEditor/Public/RetargetEditor/IKRetargetEditorController.h b/Engine/Plugins/Animation/IKRig/Source/IKRigEditor/Public/RetargetEditor/IKRetargetEditorController.h
--- a/Engine/Plugins/Animation/IKRig/Source/IKRigEditor/Public/RetargetEditor/IKRetargetEditorController.h
+++ b/Engine/Plugins/Animation/IKRig/Source/IKRigEditor/Public/RetargetEditor/IKRetargetEditorController.h
@@ -85,17 +85,32 @@ public:
 	UAnimationAsset* PreviousAsset = nullptr;
 	/** END Sequence Browser */
 
+	// go to retarget pose
+	void HandleGoToRetargetPose() const;
+	
 	/** edit retarget poses */
 	void HandleEditPose() const;
 	bool CanEditPose() const;
 	bool IsEditingPose() const;
+
+	void HandleResetPose() const;
+	bool CanResetPose() const;
+
 	void HandleNewPose();
+	bool CanNewPose() const;
 	FReply CreateNewPose() const;
-	void HandleDeletePose() const;
-	bool CanDeletePose() const;
-	void HandleResetPose() const;
 	TSharedPtr<SWindow> NewPoseWindow;
 	TSharedPtr<SEditableTextBox> NewPoseEditableText;
+	
+	void HandleDeletePose() const;
+	bool CanDeletePose() const;
+	
+	void HandleRenamePose();
+	FReply RenamePose() const;
+	bool CanRenamePose() const;
+	TSharedPtr<SWindow> RenamePoseWindow;
+	TSharedPtr<SEditableTextBox> NewNameEditableText;
+	
 	TArray<TSharedPtr<FName>> PoseNames;
 	FText GetCurrentPoseName() const;
 	void OnPoseSelected(TSharedPtr<FName> InPoseName, ESelectInfo::Type SelectInfo) const;
diff --git a/Engine/Plugins/Animation/IKRig/Source/IKRigEditor/Public/RetargetEditor/IKRetargeterController.h b/Engine/Plugins/Animation/IKRig/Source/IKRigEditor/Public/RetargetEditor/IKRetargeterController.h
--- a/Engine/Plugins/Animation/IKRig/Source/IKRigEditor/Public/RetargetEditor/IKRetargeterController.h
+++ b/Engine/Plugins/Animation/IKRig/Source/IKRigEditor/Public/RetargetEditor/IKRetargeterController.h
@@ -9,6 +9,7 @@
 
 #include "IKRetargeterController.generated.h"
 
+class FIKRetargetEditorController;
 class URetargetChainSettings;
 class UIKRigDefinition;
 class UIKRetargeter;
@@ -27,6 +28,10 @@ public:
 	/** Get access to the retargeter asset.
 	 *@warning Do not make modifications to the asset directly. Use this API instead. */
 	UIKRetargeter* GetAsset() const;
+	/** Get access to the editor controller.*/
+	FIKRetargetEditorController* GetEditorController() const;
+	/** Set the currently used editor controller.*/
+	void SetEditorController(FIKRetargetEditorController* InEditorController) { EditorController = InEditorController; };
 
 	/** SOURCE / TARGET
 	* 
@@ -54,7 +59,7 @@ public:
 	/** Get names of all the source bone chains. */
 	void GetSourceChainNames(TArray<FName>& OutNames) const;
 	/** Remove invalid chain mappings (no longer existing in currently referenced source/target IK Rig assets) */
-	void CleanChainMapping();
+	void CleanChainMapping(const bool bForceReinitialization=true);
 	/** Use fuzzy string search to find "best" Source chain to map to each Target chain */
 	void AutoMapChains() const;
 	/** Callback when IK Rig chain is renamed. Retains existing mappings using the new name */
@@ -71,9 +76,11 @@ public:
 	 * 
 	 */
 	/** Remove bones from retarget poses that are no longer in skeleton */
-	void CleanPoseList();
+	void CleanPoseList(const bool bForceReinitialization=true);
 	/** Add new retarget pose. */
 	void AddRetargetPose(FName NewPoseName) const;
+	/** Rename current retarget pose. */
+	void RenameCurrentRetargetPose(FName NewPoseName) const;
 	/** Remove a retarget pose. */
 	void RemoveRetargetPose(FName PoseToRemove) const;
 	/** Reset a retarget pose (removes all stored deltas, returning pose to reference pose */
@@ -90,10 +97,10 @@ public:
 	FQuat GetRotationOffsetForRetargetPoseBone(FName BoneName) const;
 	/** Add a delta translation to the root bone (used in Edit Mode in the retarget editor) */
 	void AddTranslationOffsetToRetargetRootBone(FVector TranslationOffset) const;
-	/** Set whether to output retarget pose. Will output current retarget pose if true, or run retarget otherwise. */
-	void SetEditRetargetPoseMode(bool bOutputRetargetPose, bool bReinitializeAfter=true) const;
 	/** Get whether in mode to output retarget pose (true) or run retarget (false). */
 	bool GetEditRetargetPoseMode() const;
+	/** Set whether retarget asset should be in "Edit Pose" mode. */
+	void SetEditRetargetPoseMode(const bool bMode) const;
 	/** Add a numbered suffix to the given pose name to make it unique. */
 	FName MakePoseNameUnique(FName PoseName) const;
 	/** END RETARGET POSE EDITING */
@@ -121,5 +128,8 @@ private:
 	void SortChainMapping() const;
 
 	/** The actual asset that this Controller modifies. */
-	UIKRetargeter* Asset = nullptr;
+	TObjectPtr<UIKRetargeter> Asset = nullptr;
+
+	/** The editor controller for this asset. */
+	TObjectPtr<FIKRetargetEditorController> EditorController = nullptr;
 };
