Addressed missed comment in code review

diff --git a/Engine/Plugins/Animation/IKRig/Source/IKRigEditor/Private/RetargetEditor/IKRetargetDetails.cpp b/Engine/Plugins/Animation/IKRig/Source/IKRigEditor/Private/RetargetEditor/IKRetargetDetails.cpp
--- a/Engine/Plugins/Animation/IKRig/Source/IKRigEditor/Private/RetargetEditor/IKRetargetDetails.cpp
+++ b/Engine/Plugins/Animation/IKRig/Source/IKRigEditor/Private/RetargetEditor/IKRetargetDetails.cpp
@@ -21,7 +21,7 @@ TSharedRef<IDetailCustomization> FIKRetargeterDetails::MakeInstance()
 
 void FIKRetargeterDetails::CustomizeDetails(IDetailLayoutBuilder& DetailBuilder)
 {
-	GetAssetController(DetailBuilder);
+	const TObjectPtr<UIKRetargeterController> Controller = GetAssetControllerFromSelectedObjects(DetailBuilder);
 	if (!Controller)
 	{
 		return;
@@ -129,19 +129,17 @@ TSharedRef<SWidget>  FIKRetargeterDetails::MakeToolbar(const TSharedRef<FUIComma
 	return ToolbarBuilder.MakeWidget();
 }
 
-void FIKRetargeterDetails::GetAssetController(IDetailLayoutBuilder& DetailBuilder)
+TObjectPtr<UIKRetargeterController> FIKRetargeterDetails::GetAssetControllerFromSelectedObjects(IDetailLayoutBuilder& DetailBuilder) const
 {
-	Controller = nullptr;
-	
 	TArray< TWeakObjectPtr<UObject> > OutObjects;
 	DetailBuilder.GetObjectsBeingCustomized(OutObjects);
 	if (OutObjects.IsEmpty())
 	{
-		return;
+		return nullptr;
 	}
 
 	const TObjectPtr<UIKRetargeter> Asset = CastChecked<UIKRetargeter>(OutObjects[0].Get());
-	Controller = UIKRetargeterController::GetController(Asset);
+	return UIKRetargeterController::GetController(Asset);
 }
 
 #undef LOCTEXT_NAMESPACE
diff --git a/Engine/Plugins/Animation/IKRig/Source/IKRigEditor/Public/RetargetEditor/IKRetargetDetails.h b/Engine/Plugins/Animation/IKRig/Source/IKRigEditor/Public/RetargetEditor/IKRetargetDetails.h
--- a/Engine/Plugins/Animation/IKRig/Source/IKRigEditor/Public/RetargetEditor/IKRetargetDetails.h
+++ b/Engine/Plugins/Animation/IKRig/Source/IKRigEditor/Public/RetargetEditor/IKRetargetDetails.h
@@ -22,9 +22,8 @@ public:
 private:
 
 	TSharedRef<SWidget> MakeToolbar(const TSharedRef<FUICommandList>& UICommandList);
-
-	TObjectPtr<UIKRetargeterController> Controller;
-	void GetAssetController(IDetailLayoutBuilder& DetailBuilder);
+	
+	TObjectPtr<UIKRetargeterController> GetAssetControllerFromSelectedObjects(IDetailLayoutBuilder& DetailBuilder) const;
 
 	TArray<TSharedPtr<FName>> PoseNames;
 };
