bulk edition within the details panel

diff --git a/Engine/Plugins/Animation/IKRig/Source/IKRig/Private/IKRigDefinition.cpp b/Engine/Plugins/Animation/IKRig/Source/IKRig/Private/IKRigDefinition.cpp
--- a/Engine/Plugins/Animation/IKRig/Source/IKRig/Private/IKRigDefinition.cpp
+++ b/Engine/Plugins/Animation/IKRig/Source/IKRig/Private/IKRigDefinition.cpp
@@ -47,75 +47,94 @@ TOptional<FTransform::FReal> UIKRigEffectorGoal::GetNumericValue(
 	return TOptional<FTransform::FReal>();
 }
 
-void UIKRigEffectorGoal::OnNumericValueChanged(ESlateTransformComponent::Type Component,
-	ESlateRotationRepresentation::Type Representation, ESlateTransformSubComponent::Type SubComponent,
-	FTransform::FReal Value, ETextCommit::Type CommitType, EIKRigTransformType::Type TransformType)
+TTuple<FTransform, FTransform> UIKRigEffectorGoal::PrepareNumericValueChanged( ESlateTransformComponent::Type Component,
+																			ESlateRotationRepresentation::Type Representation,
+																			ESlateTransformSubComponent::Type SubComponent,
+																			FTransform::FReal Value,
+																			EIKRigTransformType::Type TransformType) const
 {
+	const FTransform& InTransform = TransformType == EIKRigTransformType::Current ? CurrentTransform : InitialTransform;
+	FTransform OutTransform = InTransform;
+	SAdvancedTransformInputBox<FTransform>::ApplyNumericValueChange(OutTransform, Value, Component, Representation, SubComponent);
+	return MakeTuple(InTransform, OutTransform);
+}
+
+void UIKRigEffectorGoal::SetTransform(const FTransform& InTransform, EIKRigTransformType::Type InTransformType)
+{
+	// we assume that InTransform is not equal to the one it's being assigned to
 	Modify();
 
-	switch(TransformType)
+	FTransform& TransformChanged = InTransformType == EIKRigTransformType::Current ? CurrentTransform : InitialTransform;
+	TransformChanged = InTransform;
+}
+
+namespace
+{
+	
+template<typename DataType>
+void GetContentFromData(const DataType& InData, FString& Content)
+{
+	TBaseStructure<DataType>::Get()->ExportText(Content, &InData, &InData, nullptr, PPF_None, nullptr);
+}
+
+class FIKRigEffectorGoalErrorPipe : public FOutputDevice
+{
+public:
+	int32 NumErrors;
+
+	FIKRigEffectorGoalErrorPipe()
+		: FOutputDevice()
+		, NumErrors(0)
+	{}
+
+	virtual void Serialize(const TCHAR* V, ELogVerbosity::Type Verbosity, const class FName& Category) override
 	{
-		case EIKRigTransformType::Current:
-		{
-			SAdvancedTransformInputBox<FTransform>::ApplyNumericValueChange(
-				CurrentTransform,
-				Value,
-				Component,
-				Representation,
-				SubComponent);
-			break;
-		}
-		case EIKRigTransformType::Reference:
-		{
-			SAdvancedTransformInputBox<FTransform>::ApplyNumericValueChange(
-				InitialTransform,
-				Value,
-				Component,
-				Representation,
-				SubComponent);
-			break;
-		}
-		default:
-		{
-			break;
-		}
+		NumErrors++;
 	}
-}
+};
 
-void UIKRigEffectorGoal::OnCopyToClipboard(ESlateTransformComponent::Type Component, EIKRigTransformType::Type TransformType)
+template<typename DataType>
+bool GetDataFromContent(const FString& Content, DataType& OutData)
 {
-	const FTransform Xfo = TransformType == EIKRigTransformType::Current ? CurrentTransform : InitialTransform;
+	FIKRigEffectorGoalErrorPipe ErrorPipe;
+	static UScriptStruct* DataStruct = TBaseStructure<DataType>::Get();
+	DataStruct->ImportText(*Content, &OutData, nullptr, PPF_None, &ErrorPipe, DataStruct->GetName(), true);
+	return (ErrorPipe.NumErrors == 0);
+}
 	
+}
+
+void UIKRigEffectorGoal::OnCopyToClipboard(ESlateTransformComponent::Type Component, EIKRigTransformType::Type TransformType) const
+{
+	const FTransform& Xfo = TransformType == EIKRigTransformType::Current ? CurrentTransform : InitialTransform;
+
 	FString Content;
 	switch(Component)
 	{
 	case ESlateTransformComponent::Location:
 		{
-			const FVector Data = Xfo.GetLocation();
-			TBaseStructure<FVector>::Get()->ExportText(Content, &Data, &Data, nullptr, PPF_None, nullptr);
+			GetContentFromData(Xfo.GetLocation(), Content);
 			break;
 		}
 	case ESlateTransformComponent::Rotation:
 		{
-			const FRotator Data = Xfo.Rotator();
-			TBaseStructure<FRotator>::Get()->ExportText(Content, &Data, &Data, nullptr, PPF_None, nullptr);
+			GetContentFromData(Xfo.Rotator(), Content);
 			break;
 		}
 	case ESlateTransformComponent::Scale:
 		{
-			const FVector Data = Xfo.GetScale3D();
-			TBaseStructure<FVector>::Get()->ExportText(Content, &Data, &Data, nullptr, PPF_None, nullptr);
+			GetContentFromData(Xfo.GetScale3D(), Content);
 			break;
 		}
 	case ESlateTransformComponent::Max:
 	default:
 		{
-			TBaseStructure<FTransform>::Get()->ExportText(Content, &Xfo, &Xfo, nullptr, PPF_None, nullptr);
+			GetContentFromData(Xfo, Content);
 			break;
 		}
 	}
 
-	if(!Content.IsEmpty())
+	if (!Content.IsEmpty())
 	{
 		FPlatformApplicationMisc::ClipboardCopy(*Content);
 	}
@@ -126,42 +145,20 @@ void UIKRigEffectorGoal::OnPasteFromClipboard(ESlateTransformComponent::Type Com
 	FString Content;
 	FPlatformApplicationMisc::ClipboardPaste(Content);
 
-	if(Content.IsEmpty())
+	if (Content.IsEmpty())
 	{
 		return;
 	}
 
 	FTransform& Xfo = TransformType == EIKRigTransformType::Current ? CurrentTransform : InitialTransform;
-	
 	Modify();
-
-	class FIKRigEffectorGoalErrorPipe : public FOutputDevice
-	{
-	public:
-
-		int32 NumErrors;
-
-		FIKRigEffectorGoalErrorPipe()
-			: FOutputDevice()
-			, NumErrors(0)
-		{
-		}
-
-		virtual void Serialize(const TCHAR* V, ELogVerbosity::Type Verbosity, const class FName& Category) override
-		{
-			NumErrors++;
-		}
-	};
-
-	FIKRigEffectorGoalErrorPipe ErrorPipe;
-				
+	
 	switch(Component)
 	{
 		case ESlateTransformComponent::Location:
 		{
 			FVector Data = Xfo.GetLocation();
-			TBaseStructure<FVector>::Get()->ImportText(*Content, &Data, nullptr, PPF_None, &ErrorPipe, TBaseStructure<FVector>::Get()->GetName(), true);
-			if(ErrorPipe.NumErrors == 0)
+			if (GetDataFromContent(Content, Data))
 			{
 				Xfo.SetLocation(Data);
 			}
@@ -170,8 +167,7 @@ void UIKRigEffectorGoal::OnPasteFromClipboard(ESlateTransformComponent::Type Com
 		case ESlateTransformComponent::Rotation:
 		{
 			FRotator Data = Xfo.Rotator();
-			TBaseStructure<FRotator>::Get()->ImportText(*Content, &Data, nullptr, PPF_None, &ErrorPipe, TBaseStructure<FRotator>::Get()->GetName(), true);
-			if(ErrorPipe.NumErrors == 0)
+			if (GetDataFromContent(Content, Data))
 			{
 				Xfo.SetRotation(FQuat(Data));
 			}
@@ -180,8 +176,7 @@ void UIKRigEffectorGoal::OnPasteFromClipboard(ESlateTransformComponent::Type Com
 		case ESlateTransformComponent::Scale:
 		{
 			FVector Data = Xfo.GetScale3D();
-			TBaseStructure<FVector>::Get()->ImportText(*Content, &Data, nullptr, PPF_None, &ErrorPipe, TBaseStructure<FVector>::Get()->GetName(), true);
-			if(ErrorPipe.NumErrors == 0)
+			if (GetDataFromContent(Content, Data))
 			{
 				Xfo.SetScale3D(Data);
 			}
@@ -191,8 +186,7 @@ void UIKRigEffectorGoal::OnPasteFromClipboard(ESlateTransformComponent::Type Com
 		default:
 		{
 			FTransform Data = Xfo;
-			TBaseStructure<FTransform>::Get()->ImportText(*Content, &Data, nullptr, PPF_None, &ErrorPipe, TBaseStructure<FTransform>::Get()->GetName(), true);
-			if(ErrorPipe.NumErrors == 0)
+			if (GetDataFromContent(Content, Data))
 			{
 				Xfo = Data;
 			}
@@ -221,6 +215,10 @@ bool UIKRigEffectorGoal::TransformDiffersFromDefault(
 			{
 				return !(CurrentTransform.GetScale3D() - InitialTransform.GetScale3D()).IsNearlyZero();
 			}
+			default:
+			{
+				return false;
+			}
 		}
 	}
 	return false;
@@ -247,6 +245,10 @@ void UIKRigEffectorGoal::ResetTransformToDefault(
 			CurrentTransform.SetScale3D(InitialTransform.GetScale3D());
 			break;
 		}
+		default:
+		{
+			break;
+		}
 	}
 }
 
diff --git a/Engine/Plugins/Animation/IKRig/Source/IKRig/Public/IKRigDefinition.h b/Engine/Plugins/Animation/IKRig/Source/IKRig/Public/IKRigDefinition.h
--- a/Engine/Plugins/Animation/IKRig/Source/IKRig/Public/IKRigDefinition.h
+++ b/Engine/Plugins/Animation/IKRig/Source/IKRig/Public/IKRigDefinition.h
@@ -107,17 +107,18 @@ public:
 		ESlateTransformSubComponent::Type SubComponent,
 		EIKRigTransformType::Type TransformType
 	) const;
-	
-	void OnNumericValueChanged(
+
+	TTuple<FTransform, FTransform> PrepareNumericValueChanged(
 		ESlateTransformComponent::Type Component,
 		ESlateRotationRepresentation::Type Representation,
 		ESlateTransformSubComponent::Type SubComponent,
 		FTransform::FReal Value,
-		ETextCommit::Type CommitType,
 		EIKRigTransformType::Type TransformType
-	);
+	) const;
+
+	void SetTransform( const FTransform& InTransform, EIKRigTransformType::Type InTransformType);
 
-	void OnCopyToClipboard(ESlateTransformComponent::Type Component, EIKRigTransformType::Type TransformType);
+	void OnCopyToClipboard(ESlateTransformComponent::Type Component, EIKRigTransformType::Type TransformType) const;
 	void OnPasteFromClipboard(ESlateTransformComponent::Type Component, EIKRigTransformType::Type TransformType);
 
 	bool TransformDiffersFromDefault(ESlateTransformComponent::Type Component, TSharedPtr<IPropertyHandle> PropertyHandle) const;
diff --git a/Engine/Plugins/Animation/IKRig/Source/IKRigEditor/Private/RigEditor/IKRigDetailCustomizations.cpp b/Engine/Plugins/Animation/IKRig/Source/IKRigEditor/Private/RigEditor/IKRigDetailCustomizations.cpp
--- a/Engine/Plugins/Animation/IKRig/Source/IKRigEditor/Private/RigEditor/IKRigDetailCustomizations.cpp
+++ b/Engine/Plugins/Animation/IKRig/Source/IKRigEditor/Private/RigEditor/IKRigDetailCustomizations.cpp
@@ -18,92 +18,104 @@ namespace IKRigDetailCustomizationsConstants
 	static const float ItemWidth = 125.0f;
 }
 
+const TArray<FText>& FIKRigGenericDetailCustomization::GetButtonLabels()
+{
+	static const TArray<FText> ButtonLabels
+	{
+		LOCTEXT("CurrentTransform", "Current"),
+		LOCTEXT("ReferenceTransform", "Reference")
+	};
+	return ButtonLabels;
+}
+
+const TArray<EIKRigTransformType::Type>& FIKRigGenericDetailCustomization::GetTransformTypes()
+{
+	static const TArray<EIKRigTransformType::Type> TransformTypes
+	{
+		EIKRigTransformType::Current,
+		EIKRigTransformType::Reference
+	};
+	return TransformTypes;
+}
+
+template<typename ClassToCustomize>
+TArray<ClassToCustomize*> GetCustomizedObjects(const TArray<TWeakObjectPtr<UObject>>& InObjectsBeingCustomized)
+{
+	TArray<ClassToCustomize*> CustomizedObjects;
+	CustomizedObjects.Reserve(InObjectsBeingCustomized.Num());
+
+	auto IsOfCustomType = [](TWeakObjectPtr<UObject> InObject) { return InObject.Get() && InObject->IsA<ClassToCustomize>(); };
+	auto CastAsCustomType = [](TWeakObjectPtr<UObject> InObject) { return Cast<ClassToCustomize>(InObject); };
+	Algo::TransformIf(InObjectsBeingCustomized, CustomizedObjects, IsOfCustomType, CastAsCustomType);
+	
+	return CustomizedObjects;
+}
+
 void FIKRigGenericDetailCustomization::CustomizeDetails(IDetailLayoutBuilder& DetailBuilder)
 {
 	TArray<TWeakObjectPtr<UObject>> ObjectsBeingCustomized = DetailBuilder.GetSelectedObjects();
+	ObjectsBeingCustomized.RemoveAll( [](const TWeakObjectPtr<UObject>& InObject) { return !InObject.IsValid(); } );
 
+	// if no selected objects
+	if (ObjectsBeingCustomized.IsEmpty())
+	{
+		return;
+	}
+	
 	// make sure all types are of the same class
-	UClass* DetailsClass = nullptr;
-	for(const TWeakObjectPtr<UObject>& ObjectBeingCustomized : ObjectsBeingCustomized)
+	const UClass* DetailsClass = ObjectsBeingCustomized[0]->GetClass();
+	const int32 NumObjects = ObjectsBeingCustomized.Num();
+	for (int32 Index = 1; Index < NumObjects; Index++)
 	{
-		if(!ObjectBeingCustomized.IsValid())
+		if (ObjectsBeingCustomized[Index]->GetClass() != DetailsClass)
 		{
-			continue;
-		}
-
-		if(DetailsClass)
-		{
-			if(ObjectBeingCustomized->GetClass() != DetailsClass)
-			{
-				// multiple different things - fall back to default
-				// details panel behavior
-				return;
-			}
-		}
-		else
-		{
-			DetailsClass = ObjectBeingCustomized->GetClass();
+			// multiple different things - fallback to default details panel behavior
+			return;
 		}
 	}
 
-	// assuming the classes are all the same 
-	for(const TWeakObjectPtr<UObject>& ObjectBeingCustomized : ObjectsBeingCustomized)
+	// assuming the classes are all the same
+	if (ObjectsBeingCustomized[0]->IsA<UIKRigBoneDetails>())
 	{
-		if(!ObjectBeingCustomized.IsValid())
-		{
-			continue;
-		}
-
-		if(ObjectBeingCustomized->IsA<UIKRigBoneDetails>())
-		{
-			CustomizeDetailsForClass<UIKRigBoneDetails>(DetailBuilder, ObjectsBeingCustomized);	
-		}
-		else if(ObjectBeingCustomized->IsA<UIKRigEffectorGoal>())
-		{
-			CustomizeDetailsForClass<UIKRigEffectorGoal>(DetailBuilder, ObjectsBeingCustomized);	
-		}
+		return CustomizeDetailsForClass<UIKRigBoneDetails>(DetailBuilder, ObjectsBeingCustomized);
+	}
 
-		break;
+	if (ObjectsBeingCustomized[0]->IsA<UIKRigEffectorGoal>())
+	{
+		return CustomizeDetailsForClass<UIKRigEffectorGoal>(DetailBuilder, ObjectsBeingCustomized);
 	}
 }
 
 template <>
 void FIKRigGenericDetailCustomization::CustomizeDetailsForClass<UIKRigBoneDetails>(
 	IDetailLayoutBuilder& DetailBuilder,
-	TArray<TWeakObjectPtr<UObject>> ObjectsBeingCustomized)
+	const TArray<TWeakObjectPtr<UObject>>& InObjectsBeingCustomized)
 {
-	UIKRigBoneDetails* BoneDetails = Cast<UIKRigBoneDetails>(ObjectsBeingCustomized[0].Get());
-	if(BoneDetails == nullptr)
+	const TArray<UIKRigBoneDetails*> Bones = GetCustomizedObjects<UIKRigBoneDetails>(InObjectsBeingCustomized);
+	if (Bones.IsEmpty())
 	{
 		return;
-		
 	}
-	const TArray<FText> ButtonLabels =
-	{
-		LOCTEXT("CurrentTransform", "Current"),
-		LOCTEXT("ReferenceTransform", "Reference")
-	};
+
+	const TArray<FText>& ButtonLabels = GetButtonLabels();
+	const TArray<EIKRigTransformType::Type>& TransformTypes = GetTransformTypes();
 	
-	const TArray<FText> ButtonTooltips =
+	static const TArray<FText> ButtonTooltips
 	{
 		LOCTEXT("CurrentBoneTransformTooltip", "The current transform of the bone"),
 		LOCTEXT("ReferenceBoneTransformTooltip", "The reference transform of the bone")
 	};
+	
+	static const TAttribute<TArray<EIKRigTransformType::Type>> VisibleTransforms =
+		TArray<EIKRigTransformType::Type>({EIKRigTransformType::Current});
 
-	const TArray<EIKRigTransformType::Type> TransformTypes =
+	const TArray<TSharedRef<IPropertyHandle>> Properties
 	{
-		EIKRigTransformType::Current,
-		EIKRigTransformType::Reference
+		DetailBuilder.GetProperty(GET_MEMBER_NAME_CHECKED(UIKRigBoneDetails, CurrentTransform)),
+		DetailBuilder.GetProperty(GET_MEMBER_NAME_CHECKED(UIKRigBoneDetails, ReferenceTransform))
 	};
 
-	static TAttribute<TArray<EIKRigTransformType::Type>> VisibleTransforms =
-		TArray<EIKRigTransformType::Type>({EIKRigTransformType::Current});
-
-	TArray<TSharedRef<IPropertyHandle>> Properties;
-	Properties.Add(DetailBuilder.GetProperty(GET_MEMBER_NAME_CHECKED(UIKRigBoneDetails, CurrentTransform)));
-	Properties.Add(DetailBuilder.GetProperty(GET_MEMBER_NAME_CHECKED(UIKRigBoneDetails, ReferenceTransform)));
-
-	for(TSharedRef<IPropertyHandle> Property : Properties)
+	for (TSharedRef<IPropertyHandle> Property : Properties)
 	{
 		DetailBuilder.HideProperty(Property);
 	}
@@ -148,11 +160,48 @@ void FIKRigGenericDetailCustomization::CustomizeDetailsForClass<UIKRigBoneDetail
 	{
 		const EIKRigTransformType::Type TransformType = (EIKRigTransformType::Type)PropertyIndex; 
 
-		TransformWidgetArgs.OnGetIsComponentRelative_UObject(BoneDetails, &UIKRigBoneDetails::IsComponentRelative, TransformType);
-		TransformWidgetArgs.OnIsComponentRelativeChanged_UObject(BoneDetails, &UIKRigBoneDetails::OnComponentRelativeChanged, TransformType);
-		TransformWidgetArgs.Transform_UObject(BoneDetails, &UIKRigBoneDetails::GetTransform, TransformType);
-		TransformWidgetArgs.OnCopyToClipboard_UObject(BoneDetails, &UIKRigBoneDetails::OnCopyToClipboard, TransformType);
-		TransformWidgetArgs.OnPasteFromClipboard_UObject(BoneDetails, &UIKRigBoneDetails::OnPasteFromClipboard, TransformType);
+		// get/set relative
+		TransformWidgetArgs.OnGetIsComponentRelative_Lambda( [Bones, TransformType](ESlateTransformComponent::Type InComponent)
+		{
+			return Bones.ContainsByPredicate( [&](const UIKRigBoneDetails* Bone)
+			{
+				return Bone->IsComponentRelative(InComponent, TransformType);
+			} );
+		})
+		.OnIsComponentRelativeChanged_Lambda( [Bones, TransformType](ESlateTransformComponent::Type InComponent, bool bIsRelative)
+		{
+			for (UIKRigBoneDetails* Bone: Bones)
+			{
+				Bone->OnComponentRelativeChanged(InComponent, bIsRelative, TransformType);
+			}
+		} );
+
+		// get bones transforms
+		TransformWidgetArgs.Transform_Lambda([Bones, TransformType]()
+		{
+			TOptional<FTransform> Value = Bones[0]->GetTransform(TransformType);
+			if (Value)
+			{
+				for (int32 Index = 1; Index < Bones.Num(); Index++)
+				{
+					const TOptional<FTransform> CurrentValue = Bones[Index]->GetTransform(TransformType);
+					if (CurrentValue)
+					{
+						if (!Value->Equals( *CurrentValue))
+						{
+							return TOptional<FTransform>();
+						}
+					}
+				}
+			}
+			return Value;
+		});
+		
+		// copy/paste bones transforms
+		TransformWidgetArgs.OnCopyToClipboard_UObject(Bones[0], &UIKRigBoneDetails::OnCopyToClipboard, TransformType);
+		// FIXME we must connect paste capabilities here otherwise copy is not enabled...
+		// see DetailSingleItemRow::OnContextMenuOpening for more explanations
+		TransformWidgetArgs.OnPasteFromClipboard_UObject(Bones[0], &UIKRigBoneDetails::OnPasteFromClipboard, TransformType);
 
 		TransformWidgetArgs.Visibility_Lambda([TransformChoiceWidget, TransformType]() -> EVisibility
 		{
@@ -170,40 +219,33 @@ void FIKRigGenericDetailCustomization::CustomizeDetailsForClass<UIKRigBoneDetail
 template <>
 void FIKRigGenericDetailCustomization::CustomizeDetailsForClass<UIKRigEffectorGoal>(
 	IDetailLayoutBuilder& DetailBuilder,
-	TArray<TWeakObjectPtr<UObject>> ObjectsBeingCustomized)
+	const TArray<TWeakObjectPtr<UObject>>& InObjectsBeingCustomized)
 {
-	UIKRigEffectorGoal* EffectorGoal = Cast<UIKRigEffectorGoal>(ObjectsBeingCustomized[0].Get());
-	if(EffectorGoal == nullptr)
+	const TArray<UIKRigEffectorGoal*> Goals = GetCustomizedObjects<UIKRigEffectorGoal>(InObjectsBeingCustomized);
+	if (Goals.IsEmpty())
 	{
 		return;
-		
 	}
-	const TArray<FText> ButtonLabels =
-	{
-		LOCTEXT("CurrentTransform", "Current"),
-		LOCTEXT("ReferenceTransform", "Reference")
-	};
-	
-	const TArray<FText> ButtonTooltips =
+
+	const TArray<FText>& ButtonLabels = GetButtonLabels();
+	const TArray<EIKRigTransformType::Type>& TransformTypes = GetTransformTypes();	
+
+	static const TArray<FText> ButtonTooltips
 	{
 		LOCTEXT("CurrentGoalTransformTooltip", "The current transform of the goal"),
 		LOCTEXT("ReferenceGoalTransformTooltip", "The reference transform of the goal")
 	};
 
-	const TArray<EIKRigTransformType::Type> TransformTypes =
-	{
-		EIKRigTransformType::Current,
-		EIKRigTransformType::Reference
-	};
-
 	static TAttribute<TArray<EIKRigTransformType::Type>> VisibleTransforms =
 		TArray<EIKRigTransformType::Type>({EIKRigTransformType::Current});
 
-	TArray<TSharedRef<IPropertyHandle>> Properties;
-	Properties.Add(DetailBuilder.GetProperty(GET_MEMBER_NAME_CHECKED(UIKRigEffectorGoal, CurrentTransform)));
-	Properties.Add(DetailBuilder.GetProperty(GET_MEMBER_NAME_CHECKED(UIKRigEffectorGoal, InitialTransform)));
+	const TArray<TSharedRef<IPropertyHandle>> Properties
+	{
+		DetailBuilder.GetProperty(GET_MEMBER_NAME_CHECKED(UIKRigEffectorGoal, CurrentTransform)),
+		DetailBuilder.GetProperty(GET_MEMBER_NAME_CHECKED(UIKRigEffectorGoal, InitialTransform))
+	};
 
-	for(TSharedRef<IPropertyHandle> Property : Properties)
+	for (TSharedRef<IPropertyHandle> Property : Properties)
 	{
 		DetailBuilder.HideProperty(Property);
 	}
@@ -255,18 +297,105 @@ void FIKRigGenericDetailCustomization::CustomizeDetailsForClass<UIKRigEffectorGo
 			.IsEnabled(false)
 			.DisplayScaleLock(false);
 		}
-		
-		TransformWidgetArgs.OnGetNumericValue_UObject(EffectorGoal, &UIKRigEffectorGoal::GetNumericValue, TransformType);
-		TransformWidgetArgs.OnNumericValueChanged_UObject(EffectorGoal, &UIKRigEffectorGoal::OnNumericValueChanged, ETextCommit::Default, TransformType);
-		TransformWidgetArgs.OnNumericValueCommitted_UObject(EffectorGoal, &UIKRigEffectorGoal::OnNumericValueChanged, TransformType);
 
-		TransformWidgetArgs.OnCopyToClipboard_UObject(EffectorGoal, &UIKRigEffectorGoal::OnCopyToClipboard, TransformType);
-		TransformWidgetArgs.OnPasteFromClipboard_UObject(EffectorGoal, &UIKRigEffectorGoal::OnPasteFromClipboard, TransformType);
+		// get goals transforms
+		TransformWidgetArgs.OnGetNumericValue_Lambda( [Goals, TransformType](ESlateTransformComponent::Type Component,
+														ESlateRotationRepresentation::Type Representation,
+														ESlateTransformSubComponent::Type SubComponent)
+		{
+			TOptional<FVector::FReal> Value = Goals[0]->GetNumericValue(Component, Representation, SubComponent, TransformType);
+			if (Value)
+			{
+				for (int32 Index = 1; Index < Goals.Num(); Index++)
+				{
+					const TOptional<FVector::FReal> CurrentValue = Goals[Index]->GetNumericValue(Component, Representation, SubComponent, TransformType);
+					if (CurrentValue)
+					{
+						if (!FMath::IsNearlyEqual(*Value, *CurrentValue))
+						{
+							return TOptional<FVector::FReal>();
+						}
+					}
+				}
+			}
+			return Value;
+		} );
+
+		// set goals transforms
+		TransformWidgetArgs.OnNumericValueChanged_Lambda([this, Goals, TransformType]( ESlateTransformComponent::Type InComponent,
+															  ESlateRotationRepresentation::Type InRepresentation,
+															  ESlateTransformSubComponent::Type InSubComponent,
+															  FVector::FReal InValue)
+		{
+			FTransform CurrentTransform, UpdatedTransform;
+			for (UIKRigEffectorGoal* Goal: Goals)
+			{
+				Tie(CurrentTransform, UpdatedTransform) =
+					Goal->PrepareNumericValueChanged(InComponent, InRepresentation, InSubComponent, InValue, TransformType);
+
+				// don't do anything if the transform has not been changed
+				if (!UpdatedTransform.Equals(CurrentTransform))
+				{
+					// prepare transaction if needed
+					if (!ValueChangedTransaction.IsValid())
+					{
+						ValueChangedTransaction = MakeShareable(new FScopedTransaction(LOCTEXT("ChangeNumericValue", "Change Numeric Value")));
+					}
+					
+					Goal->SetTransform(UpdatedTransform, TransformType);
+				}
+			}
+		})
+		.OnNumericValueCommitted_Lambda([this, Goals, TransformType](ESlateTransformComponent::Type InComponent,
+											ESlateRotationRepresentation::Type InRepresentation,
+											ESlateTransformSubComponent::Type InSubComponent,
+											FVector::FReal InValue,
+											ETextCommit::Type InCommitType)
+		{
+			if (!ValueChangedTransaction.IsValid())
+			{
+				ValueChangedTransaction = MakeShareable(new FScopedTransaction(LOCTEXT("ChangeNumericValue", "Change Numeric Value")));
+			}
+			
+			FTransform CurrentTransform, UpdatedTransform;
+			for (UIKRigEffectorGoal* Goal: Goals)
+			{
+				Tie(CurrentTransform, UpdatedTransform) = Goal->PrepareNumericValueChanged( InComponent, InRepresentation, InSubComponent, InValue, TransformType);					
+				Goal->SetTransform(UpdatedTransform, TransformType);
+			}
+			
+			ValueChangedTransaction.Reset();
+		});
+
+		// copy/paste values
+		TransformWidgetArgs.OnCopyToClipboard_UObject(Goals[0], &UIKRigEffectorGoal::OnCopyToClipboard, TransformType)
+		.OnPasteFromClipboard_Lambda([Goals, TransformType](ESlateTransformComponent::Type InComponent)
+		{
+			FScopedTransaction Transaction(LOCTEXT("PasteTransform", "Paste Transform"));
+			for (UIKRigEffectorGoal* Goal: Goals)
+			{
+				Goal->OnPasteFromClipboard(InComponent, TransformType);
+			};
+		});
 
+		// reset to default
 		const TSharedPtr<IPropertyHandle> PropertyHandle = Properties[PropertyIndex];
-		TransformWidgetArgs.DiffersFromDefault_UObject(EffectorGoal, &UIKRigEffectorGoal::TransformDiffersFromDefault, PropertyHandle);
-		TransformWidgetArgs.OnResetToDefault_UObject(EffectorGoal, &UIKRigEffectorGoal::ResetTransformToDefault, PropertyHandle);
+		TransformWidgetArgs.DiffersFromDefault_Lambda([Goals,PropertyHandle](ESlateTransformComponent::Type InComponent)
+		{
+			return Goals.ContainsByPredicate( [&](const UIKRigEffectorGoal* Goal)
+			{
+				return Goal->TransformDiffersFromDefault(InComponent, PropertyHandle);
+			} );
+		} )
+		.OnResetToDefault_Lambda([Goals,PropertyHandle](ESlateTransformComponent::Type InComponent)
+		{
+			for (UIKRigEffectorGoal* Goal: Goals)
+			{
+				Goal->ResetTransformToDefault(InComponent, PropertyHandle);
+			}
+		} );
 
+		// visibility
 		TransformWidgetArgs.Visibility_Lambda([TransformChoiceWidget, TransformType]() -> EVisibility
 		{
 			return TransformChoiceWidget->HasValue(TransformType) ? EVisibility::Visible : EVisibility::Collapsed;
diff --git a/Engine/Plugins/Animation/IKRig/Source/IKRigEditor/Private/RigEditor/IKRigDetailCustomizations.h b/Engine/Plugins/Animation/IKRig/Source/IKRigEditor/Private/RigEditor/IKRigDetailCustomizations.h
--- a/Engine/Plugins/Animation/IKRig/Source/IKRigEditor/Private/RigEditor/IKRigDetailCustomizations.h
+++ b/Engine/Plugins/Animation/IKRig/Source/IKRigEditor/Private/RigEditor/IKRigDetailCustomizations.h
@@ -4,6 +4,7 @@
 
 #include "CoreMinimal.h"
 #include "IDetailCustomization.h"
+#include "ScopedTransaction.h"
 #include "RigEditor/IKRigEditorController.h"
 
 class FIKRigGenericDetailCustomization : public IDetailCustomization
@@ -18,13 +19,22 @@ public:
 	/** IDetailCustomization interface */
 	virtual void CustomizeDetails(IDetailLayoutBuilder& DetailBuilder) override;
 
+private:
+
+	/** Labels and transform types used by the transform widget */
+	static const TArray<FText>& GetButtonLabels();
+	static const TArray<EIKRigTransformType::Type>& GetTransformTypes();
+
 	/** Per class customization */
 	template<typename ClassToCustomize>
-	void CustomizeDetailsForClass(IDetailLayoutBuilder& DetailBuilder, TArray<TWeakObjectPtr<UObject>> ObjectsBeingCustomized) {}
+	void CustomizeDetailsForClass(IDetailLayoutBuilder& DetailBuilder, const TArray<TWeakObjectPtr<UObject>>& ObjectsBeingCustomized) {}
 
 	/** Template specializations */
 	template<>
-	void CustomizeDetailsForClass<UIKRigBoneDetails>(IDetailLayoutBuilder& DetailBuilder, TArray<TWeakObjectPtr<UObject>> ObjectsBeingCustomized);
+	void CustomizeDetailsForClass<UIKRigBoneDetails>(IDetailLayoutBuilder& DetailBuilder, const TArray<TWeakObjectPtr<UObject>>& ObjectsBeingCustomized);
 	template<>
-	void CustomizeDetailsForClass<UIKRigEffectorGoal>(IDetailLayoutBuilder& DetailBuilder, TArray<TWeakObjectPtr<UObject>> ObjectsBeingCustomized);
+	void CustomizeDetailsForClass<UIKRigEffectorGoal>(IDetailLayoutBuilder& DetailBuilder, const TArray<TWeakObjectPtr<UObject>>& ObjectsBeingCustomized);
+
+	/** Transaction to handle value changed. It's being effective once the value committed */
+	TSharedPtr<FScopedTransaction> ValueChangedTransaction;
 };
\ No newline at end of file
diff --git a/Engine/Plugins/Animation/IKRig/Source/IKRigEditor/Private/RigEditor/IKRigEditorController.cpp b/Engine/Plugins/Animation/IKRig/Source/IKRigEditor/Private/RigEditor/IKRigEditorController.cpp
--- a/Engine/Plugins/Animation/IKRig/Source/IKRigEditor/Private/RigEditor/IKRigEditorController.cpp
+++ b/Engine/Plugins/Animation/IKRig/Source/IKRigEditor/Private/RigEditor/IKRigEditorController.cpp
@@ -125,7 +125,18 @@ void UIKRigBoneDetails::OnComponentRelativeChanged(
 
 #if WITH_EDITOR
 
-void UIKRigBoneDetails::OnCopyToClipboard(ESlateTransformComponent::Type Component, EIKRigTransformType::Type TransformType)
+namespace
+{
+
+template<typename DataType>
+void GetContentFromData(const DataType& InData, FString& Content)
+{
+	TBaseStructure<DataType>::Get()->ExportText(Content, &InData, &InData, nullptr, PPF_None, nullptr);
+}
+	
+}
+
+void UIKRigBoneDetails::OnCopyToClipboard(ESlateTransformComponent::Type Component, EIKRigTransformType::Type TransformType) const
 {
 	TOptional<FTransform> Optional = GetTransform(TransformType);
 	if(!Optional.IsSet())
@@ -140,25 +151,23 @@ void UIKRigBoneDetails::OnCopyToClipboard(ESlateTransformComponent::Type Compone
 	{
 	case ESlateTransformComponent::Location:
 		{
-			const FVector Data = Xfo.GetLocation();
-			TBaseStructure<FVector>::Get()->ExportText(Content, &Data, &Data, nullptr, PPF_None, nullptr);
+			GetContentFromData(Xfo.GetLocation(), Content);
 			break;
 		}
 	case ESlateTransformComponent::Rotation:
 		{
-			const FRotator Data = Xfo.Rotator();
-			TBaseStructure<FRotator>::Get()->ExportText(Content, &Data, &Data, nullptr, PPF_None, nullptr);
+			GetContentFromData(Xfo.Rotator(), Content);
 			break;
 		}
 	case ESlateTransformComponent::Scale:
 		{
-			const FVector Data = Xfo.GetScale3D();
-			TBaseStructure<FVector>::Get()->ExportText(Content, &Data, &Data, nullptr, PPF_None, nullptr);
+			GetContentFromData(Xfo.GetScale3D(), Content);
 			break;
 		}
 	case ESlateTransformComponent::Max:
 	default:
 		{
+			GetContentFromData(Xfo, Content);
 			TBaseStructure<FTransform>::Get()->ExportText(Content, &Xfo, &Xfo, nullptr, PPF_None, nullptr);
 			break;
 		}
@@ -306,21 +315,25 @@ void FIKRigEditorController::ClearSelection()
 	ShowEmptyDetails();
 }
 
-void FIKRigEditorController::HandleGoalSelectedInViewport(const FName& GoalName, bool bReplace)
+void FIKRigEditorController::HandleGoalSelectedInViewport(const FName& GoalName, bool bReplace) const
 {
 	if (SkeletonView.IsValid())
 	{
 		SkeletonView->AddSelectedItemFromViewport(GoalName, IKRigTreeElementType::GOAL, bReplace);
+		ShowDetailsForElements(SkeletonView->GetSelectedItems());
+		return;
 	}
-	
+
 	ShowDetailsForGoal(GoalName);
 }
 
-void FIKRigEditorController::HandleBoneSelectedInViewport(const FName& BoneName, bool bReplace)
+void FIKRigEditorController::HandleBoneSelectedInViewport(const FName& BoneName, bool bReplace) const
 {
 	if (SkeletonView.IsValid())
 	{
 		SkeletonView->AddSelectedItemFromViewport(BoneName, IKRigTreeElementType::BONE, bReplace);
+		ShowDetailsForElements(SkeletonView->GetSelectedItems());
+		return;
 	}
 	
 	ShowDetailsForBone(BoneName);
@@ -557,13 +570,13 @@ bool FIKRigEditorController::IsElementExcludedBone(TSharedRef<FIKRigTreeElement>
 	return AssetController->GetBoneExcluded(TreeElement->BoneName);
 }
 
-void FIKRigEditorController::ShowDetailsForBone(const FName BoneName)
+void FIKRigEditorController::ShowDetailsForBone(const FName BoneName) const
 {
 	BoneDetails->SetBone(BoneName);
 	DetailsView->SetObject(BoneDetails);
 }
 
-void FIKRigEditorController::ShowDetailsForBoneSettings(const FName BoneName, int32 SolverIndex)
+void FIKRigEditorController::ShowDetailsForBoneSettings(const FName& BoneName, int32 SolverIndex) const
 {
 	if (UObject* BoneSettings = AssetController->GetSettingsForBone(BoneName, SolverIndex))
 	{
@@ -571,12 +584,12 @@ void FIKRigEditorController::ShowDetailsForBoneSettings(const FName BoneName, in
 	}
 }
 
-void FIKRigEditorController::ShowDetailsForGoal(const FName GoalName)
+void FIKRigEditorController::ShowDetailsForGoal(const FName& GoalName) const
 {
 	DetailsView->SetObject(AssetController->GetGoal(GoalName));
 }
 
-void FIKRigEditorController::ShowDetailsForGoalSettings(const FName GoalName, const int32 SolverIndex)
+void FIKRigEditorController::ShowDetailsForGoalSettings(const FName GoalName, const int32 SolverIndex) const
 {
 	// get solver that owns this effector
 	if (const UIKRigSolver* SolverWithEffector = AssetController->GetSolver(SolverIndex))
@@ -588,49 +601,70 @@ void FIKRigEditorController::ShowDetailsForGoalSettings(const FName GoalName, co
 	}
 }
 
-void FIKRigEditorController::ShowDetailsForSolver(const int32 SolverIndex)
+void FIKRigEditorController::ShowDetailsForSolver(const int32 SolverIndex) const
 {
 	DetailsView->SetObject(AssetController->GetSolver(SolverIndex));
 }
 
-void FIKRigEditorController::ShowEmptyDetails()
+void FIKRigEditorController::ShowEmptyDetails() const
 {
 	DetailsView->SetObject(AssetController->GetAsset());
 }
 
-void FIKRigEditorController::ShowDetailsForElements(const TArray<TSharedPtr<FIKRigTreeElement>>& InItems)
+void FIKRigEditorController::ShowDetailsForElements(const TArray<TSharedPtr<FIKRigTreeElement>>& InItems) const
 {
-	// TODO
-	// this function is currently only dealing with displaying a single element in the details panel (the last one)
-	// as this is the current state of the details panel in the IKRig editor. However, it should be able to handle
-	// multiple items. This is a new feature that has to be done that's why this function's interface takes an array
-	// of items as parameter for future implementation.
-	
-	if (InItems.Num() > 0)
+	if (!InItems.Num())
+	{
+		ShowEmptyDetails();
+		return;
+	}
+
+	const TSharedPtr<FIKRigTreeElement>& LastItem = InItems.Last();
+
+	// check is the items are all of the same type
+	const bool bContainsSeveralTypes = InItems.ContainsByPredicate( [LastItem](const TSharedPtr<FIKRigTreeElement>& Item)
 	{
-		const TSharedPtr<FIKRigTreeElement>& LastItem = InItems.Last();
-		switch (LastItem->ElementType)
+		return Item->ElementType != LastItem->ElementType;
+	});
+
+	// if all elements are similar then treat them once
+	if (!bContainsSeveralTypes)
+	{
+		TArray<TWeakObjectPtr<>> Objects;
+		for (const TSharedPtr<FIKRigTreeElement>& Item: InItems)
 		{
-			case IKRigTreeElementType::BONE:
-				ShowDetailsForBone(LastItem->BoneName);
-				break;
-			
-			case IKRigTreeElementType::GOAL:
-				ShowDetailsForGoal(LastItem->GoalName);
-				break;
-			
-			case IKRigTreeElementType::SOLVERGOAL:
-				ShowDetailsForGoalSettings(LastItem->SolverGoalName, LastItem->SolverGoalIndex);
-				break;
-			
-			case IKRigTreeElementType::BONE_SETTINGS:
-				ShowDetailsForBoneSettings(LastItem->BoneSettingBoneName, LastItem->BoneSettingsSolverIndex);
-				break;
-			
-			default:
-				ensure(false);
-				break;
+			TWeakObjectPtr<> Object = Item->GetObject();
+			if (Object.IsValid())
+			{
+				Objects.Add(Object);
+			}
 		}
+		DetailsView->SetObjects(Objects);
+		return;
+	}
+
+	// fallback to the last selected element
+	switch (LastItem->ElementType)
+	{
+	case IKRigTreeElementType::BONE:
+		ShowDetailsForBone(LastItem->BoneName);
+		break;
+		
+	case IKRigTreeElementType::GOAL:
+		ShowDetailsForGoal(LastItem->GoalName);
+		break;
+		
+	case IKRigTreeElementType::SOLVERGOAL:
+		ShowDetailsForGoalSettings(LastItem->SolverGoalName, LastItem->SolverGoalIndex);
+		break;
+		
+	case IKRigTreeElementType::BONE_SETTINGS:
+		ShowDetailsForBoneSettings(LastItem->BoneSettingBoneName, LastItem->BoneSettingsSolverIndex);
+		break;
+		
+	default:
+		ensure(false);
+		break;
 	}
 }
 
@@ -687,4 +721,21 @@ void FIKRigEditorController::InitializeSolvers() const
 	}
 }
 
+TObjectPtr<UIKRigBoneDetails> FIKRigEditorController::CreateBoneDetails(const TSharedPtr<FIKRigTreeElement const>& InBoneItem) const
+{
+	// ensure that the element is related to a bone
+	if (InBoneItem->ElementType != IKRigTreeElementType::BONE)
+	{
+		return nullptr;
+	}
+	
+	// create and store a new one
+	UIKRigBoneDetails* NewBoneDetails = NewObject<UIKRigBoneDetails>();
+	NewBoneDetails->SelectedBone = InBoneItem->BoneName;
+	NewBoneDetails->AnimInstancePtr = AnimInstance;
+	NewBoneDetails->AssetPtr = AssetController->GetAsset();
+	
+	return NewBoneDetails;
+}
+
 #undef LOCTEXT_NAMESPACE
diff --git a/Engine/Plugins/Animation/IKRig/Source/IKRigEditor/Private/RigEditor/SIKRigSkeleton.cpp b/Engine/Plugins/Animation/IKRig/Source/IKRigEditor/Private/RigEditor/SIKRigSkeleton.cpp
--- a/Engine/Plugins/Animation/IKRig/Source/IKRigEditor/Private/RigEditor/SIKRigSkeleton.cpp
+++ b/Engine/Plugins/Animation/IKRig/Source/IKRigEditor/Private/RigEditor/SIKRigSkeleton.cpp
@@ -20,11 +20,13 @@
 
 #define LOCTEXT_NAMESPACE "SIKRigSkeleton"
 
-FIKRigTreeElement::FIKRigTreeElement(const FText& InKey, IKRigTreeElementType InType)
-{
-	Key = InKey;
-	ElementType = InType;
-}
+FIKRigTreeElement::FIKRigTreeElement(const FText& InKey, IKRigTreeElementType InType,
+									 const TSharedRef<FIKRigEditorController>& InEditorController)
+	: Key(InKey)
+	, ElementType(InType)
+	, EditorController(InEditorController)
+	, OptionalBoneDetails(nullptr)
+{}
 
 TSharedRef<ITableRow> FIKRigTreeElement::MakeTreeRowWidget(
 	TSharedRef<FIKRigEditorController> InEditorController,
@@ -41,6 +43,48 @@ void FIKRigTreeElement::RequestRename()
 	OnRenameRequested.ExecuteIfBound();
 }
 
+TWeakObjectPtr< UObject > FIKRigTreeElement::GetObject() const
+{
+	const TSharedPtr<FIKRigEditorController>& Controller = EditorController.Pin();
+	if (!Controller.IsValid())
+	{
+		return nullptr; 
+	}
+
+	if (Controller->AssetController == nullptr)
+	{
+		return nullptr;
+	}
+
+	switch(ElementType)
+	{
+	case IKRigTreeElementType::BONE:
+		if (OptionalBoneDetails == nullptr)
+		{
+			OptionalBoneDetails = Controller->CreateBoneDetails(AsShared());
+		}
+		return OptionalBoneDetails;
+		break;
+	case IKRigTreeElementType::BONE_SETTINGS:
+		return Controller->AssetController->GetSettingsForBone(BoneSettingBoneName, BoneSettingsSolverIndex);
+		break;
+	case IKRigTreeElementType::GOAL:
+		return Controller->AssetController->GetGoal(GoalName);
+		break;
+	case IKRigTreeElementType::SOLVERGOAL:
+		if (const UIKRigSolver* SolverWithEffector = Controller->AssetController->GetSolver(SolverGoalIndex))
+		{
+			return SolverWithEffector->GetGoalSettings(SolverGoalName);
+		}
+		break;
+	default:
+		return nullptr;
+		break;
+	}
+	
+	return nullptr;
+}
+
 void SIKRigSkeletonItem::Construct(
 	const FArguments& InArgs,
 	TSharedRef<FIKRigEditorController> InEditorController,
@@ -166,6 +210,7 @@ void SIKRigSkeletonItem::Construct(
 		    .Text(this, &SIKRigSkeletonItem::GetName)
 		    .Font(TextFont)
 			.ColorAndOpacity(TextColor)
+			.OnVerifyTextChanged(this, &SIKRigSkeletonItem::OnVerifyNameChanged)
 		    .OnTextCommitted(this, &SIKRigSkeletonItem::OnNameCommitted)
 		    .MultiLine(false)
         ];
@@ -202,6 +247,12 @@ void SIKRigSkeletonItem::OnNameCommitted(const FText& InText, ETextCommit::Type
 	SkeletonView.Pin()->ReplaceItemInSelection(OldText, WeakRigTreeElement.Pin()->Key);
 }
 
+bool SIKRigSkeletonItem::OnVerifyNameChanged(const FText& InText, FText& OutErrorMessage) const
+{
+	// TODO let the user know when/why the goal can't be renamed 
+	return true;
+}
+
 FText SIKRigSkeletonItem::GetName() const
 {
 	return WeakRigTreeElement.Pin()->Key;
@@ -393,6 +444,11 @@ void SIKRigSkeleton::GetSelectedBoneChains(TArray<FIKRigSkeletonChain>& OutChain
 	return Skeleton.GetChainsInList(SelectedBones, OutChains);
 }
 
+TArray<TSharedPtr<FIKRigTreeElement>> SIKRigSkeleton::GetSelectedItems() const
+{
+	return TreeView->GetSelectedItems();
+}
+
 bool SIKRigSkeleton::HasSelectedItems() const
 {
 	return TreeView->GetNumItemsSelected() > 0;
@@ -1224,17 +1280,19 @@ bool SIKRigSkeleton::CanRenameGoal() const
 	return SelectedGoals.Num() == 1;
 }
 
-void SIKRigSkeleton::RefreshTreeView(bool IsInitialSetup)
+void SIKRigSkeleton::RefreshTreeView(bool IsInitialSetup /*=false*/)
 {
 	const TSharedPtr<FIKRigEditorController> Controller = EditorController.Pin();
 	if (!Controller.IsValid())
 	{
 		return;
 	}
+
+	const TSharedRef<FIKRigEditorController> ControllerRef = Controller.ToSharedRef();
 	
 	// save expansion and selection state
 	TreeView->SaveAndClearState();
-	
+
 	// reset all tree items
 	RootElements.Reset();
 	AllElements.Reset();
@@ -1261,7 +1319,7 @@ void SIKRigSkeleton::RefreshTreeView(bool IsInitialSetup)
 	{
 		// create "Bone" tree element for this bone
 		const FText BoneDisplayName = FText::FromName(BoneName);
-		TSharedPtr<FIKRigTreeElement> BoneElement = MakeShared<FIKRigTreeElement>(BoneDisplayName, IKRigTreeElementType::BONE);
+		TSharedPtr<FIKRigTreeElement> BoneElement = MakeShared<FIKRigTreeElement>(BoneDisplayName, IKRigTreeElementType::BONE, ControllerRef);
 		BoneElement.Get()->BoneName = BoneName;
 		const int32 BoneElementIndex = AllElements.Add(BoneElement);
 		BoneTreeElementIndices.Add(BoneName, BoneElementIndex);
@@ -1273,7 +1331,7 @@ void SIKRigSkeleton::RefreshTreeView(bool IsInitialSetup)
 			{
 				const FText SolverDisplayName = FText::FromString(AssetController->GetSolverUniqueName(SolverIndex));
 				const FText BoneSettingDisplayName = FText::Format(LOCTEXT("BoneSettings", "{0} settings for {1}"), BoneDisplayName, SolverDisplayName);
-				TSharedPtr<FIKRigTreeElement> SettingsItem = MakeShared<FIKRigTreeElement>(BoneSettingDisplayName, IKRigTreeElementType::BONE_SETTINGS);
+				TSharedPtr<FIKRigTreeElement> SettingsItem = MakeShared<FIKRigTreeElement>(BoneSettingDisplayName, IKRigTreeElementType::BONE_SETTINGS, ControllerRef);
 				SettingsItem->BoneSettingBoneName = BoneName;
 				SettingsItem->BoneSettingsSolverIndex = SolverIndex;
 				AllElements.Add(SettingsItem);
@@ -1293,7 +1351,7 @@ void SIKRigSkeleton::RefreshTreeView(bool IsInitialSetup)
 			
 			// make new element for goal
 			const FText GoalDisplayName = FText::FromName(Goal->GoalName);
-			TSharedPtr<FIKRigTreeElement> GoalItem = MakeShared<FIKRigTreeElement>(GoalDisplayName, IKRigTreeElementType::GOAL);
+			TSharedPtr<FIKRigTreeElement> GoalItem = MakeShared<FIKRigTreeElement>(GoalDisplayName, IKRigTreeElementType::GOAL, ControllerRef);
 			GoalItem->GoalName = Goal->GoalName;
 			AllElements.Add(GoalItem);
 
@@ -1309,7 +1367,7 @@ void SIKRigSkeleton::RefreshTreeView(bool IsInitialSetup)
 					// make new element for solver goal
 					const FText SolverDisplayName = FText::FromString(AssetController->GetSolverUniqueName(SolverIndex));
 					const FText SolverGoalDisplayName = FText::Format(LOCTEXT("GoalSettingsForSolver", "{0} settings for solver {1}"), FText::FromName(Goal->GoalName), SolverDisplayName);
-					TSharedPtr<FIKRigTreeElement> SolverGoalItem = MakeShared<FIKRigTreeElement>(SolverGoalDisplayName, IKRigTreeElementType::SOLVERGOAL);
+					TSharedPtr<FIKRigTreeElement> SolverGoalItem = MakeShared<FIKRigTreeElement>(SolverGoalDisplayName, IKRigTreeElementType::SOLVERGOAL, ControllerRef);
 					SolverGoalItem->SolverGoalIndex = SolverIndex;
 					SolverGoalItem->SolverGoalName = Goal->GoalName;
 					AllElements.Add(SolverGoalItem);
diff --git a/Engine/Plugins/Animation/IKRig/Source/IKRigEditor/Public/RigEditor/IKRigEditorController.h b/Engine/Plugins/Animation/IKRig/Source/IKRigEditor/Public/RigEditor/IKRigEditorController.h
--- a/Engine/Plugins/Animation/IKRig/Source/IKRigEditor/Public/RigEditor/IKRigEditorController.h
+++ b/Engine/Plugins/Animation/IKRig/Source/IKRigEditor/Public/RigEditor/IKRigEditorController.h
@@ -54,13 +54,12 @@ public:
 	UPROPERTY()
 	TWeakObjectPtr<UIKRigDefinition> AssetPtr;
 
+#if WITH_EDITOR
+
 	TOptional<FTransform> GetTransform(EIKRigTransformType::Type TransformType) const;
 	bool IsComponentRelative(ESlateTransformComponent::Type Component, EIKRigTransformType::Type TransformType) const;
 	void OnComponentRelativeChanged(ESlateTransformComponent::Type Component, bool bIsRelative, EIKRigTransformType::Type TransformType);
-
-#if WITH_EDITOR
-
-	void OnCopyToClipboard(ESlateTransformComponent::Type Component, EIKRigTransformType::Type TransformType);
+	void OnCopyToClipboard(ESlateTransformComponent::Type Component, EIKRigTransformType::Type TransformType) const;
 	void OnPasteFromClipboard(ESlateTransformComponent::Type Component, EIKRigTransformType::Type TransformType);
 
 #endif
@@ -87,9 +86,9 @@ public:
 	/** clear all selected objects */
 	void ClearSelection();
 	/** callback when goal is selected in the viewport */
-	void HandleGoalSelectedInViewport(const FName& GoalName, bool bReplace);
+	void HandleGoalSelectedInViewport(const FName& GoalName, bool bReplace) const;
 	/** callback when bone is selected in the viewport */
-	void HandleBoneSelectedInViewport(const FName& BoneName, bool bReplace);
+	void HandleBoneSelectedInViewport(const FName& BoneName, bool bReplace) const;
 	/** reset all goals to initial transforms */
 	void Reset() const;
 	/** refresh all views */
@@ -123,20 +122,20 @@ public:
 	/** determine if the element is an excluded bone*/
 	bool IsElementExcludedBone(TSharedRef<FIKRigTreeElement> TreeElement);
 	
-	/** show transform of bone in details view */
-	void ShowDetailsForBone(const FName BoneName);
-	/** show BONE settings in details view */
-	void ShowDetailsForBoneSettings(const FName BoneName, int32 SolverIndex);
-	/** show GOAL settings in details view */
-	void ShowDetailsForGoal(const FName GoalName);
-	/** show EFFECTOR settings in details view */
-	void ShowDetailsForGoalSettings(const FName GoalName, const int32 SolverIndex);
-	/** show SOLVER settings in details view */
-	void ShowDetailsForSolver(const int32 SolverIndex);
+	/** show single transform of bone in details view */
+	void ShowDetailsForBone(const FName BoneName) const;
+	/** show single BONE settings in details view */
+	void ShowDetailsForBoneSettings(const FName& BoneName, int32 SolverIndex) const;
+	/** show single GOAL settings in details view */
+	void ShowDetailsForGoal(const FName& GoalName) const;
+	/** show single EFFECTOR settings in details view */
+	void ShowDetailsForGoalSettings(const FName GoalName, const int32 SolverIndex) const;
+	/** show single SOLVER settings in details view */
+	void ShowDetailsForSolver(const int32 SolverIndex) const;
 	/** show nothing in details view */
-	void ShowEmptyDetails();
+	void ShowEmptyDetails() const;
 	/** show selected items in details view */
-	void ShowDetailsForElements(const TArray<TSharedPtr<FIKRigTreeElement>>& InItems);
+	void ShowDetailsForElements(const TArray<TSharedPtr<FIKRigTreeElement>>& InItems) const;
 	
 	/** set details tab view */
 	void SetDetailsView(const TSharedPtr<class IDetailsView>& InDetailsView){ DetailsView = InDetailsView; };
@@ -185,6 +184,9 @@ public:
 	};
 	/** END FGCObject interface */
 
+	/** UIKRigBoneDetails factory **/
+	TObjectPtr<UIKRigBoneDetails> CreateBoneDetails(const TSharedPtr<FIKRigTreeElement const>& InItem) const;
+	
 private:
 	/** right after importing a skeleton, we ask user what solver they want to use */
 	bool PromptToAddSolver() const;
diff --git a/Engine/Plugins/Animation/IKRig/Source/IKRigEditor/Public/RigEditor/SIKRigSkeleton.h b/Engine/Plugins/Animation/IKRig/Source/IKRigEditor/Public/RigEditor/SIKRigSkeleton.h
--- a/Engine/Plugins/Animation/IKRig/Source/IKRigEditor/Public/RigEditor/SIKRigSkeleton.h
+++ b/Engine/Plugins/Animation/IKRig/Source/IKRigEditor/Public/RigEditor/SIKRigSkeleton.h
@@ -16,6 +16,7 @@ class FIKRigEditorController;
 class SIKRigSkeleton;
 class FIKRigEditorToolkit;
 class USkeletalMesh;
+class UIKRigBoneDetails;
 
 enum class IKRigTreeElementType { BONE, GOAL, SOLVERGOAL, BONE_SETTINGS };
 
@@ -23,7 +24,7 @@ class FIKRigTreeElement : public TSharedFromThis<FIKRigTreeElement>
 {
 public:
 	
-	FIKRigTreeElement(const FText& InKey, IKRigTreeElementType InType);
+	FIKRigTreeElement(const FText& InKey, IKRigTreeElementType InType, const TSharedRef<FIKRigEditorController>& InEditorController);
 
 	TSharedRef<ITableRow> MakeTreeRowWidget(
 		TSharedRef<FIKRigEditorController> InEditorController,
@@ -54,7 +55,17 @@ public:
 	/** delegate for when the context menu requests a rename */
 	void RequestRename();
 	DECLARE_DELEGATE(FOnRenameRequested);
-	FOnRenameRequested OnRenameRequested;	
+	FOnRenameRequested OnRenameRequested;
+
+	/** get the underlying object */
+	TWeakObjectPtr< UObject > GetObject() const;
+
+private:
+	/** centralized editor controls */
+	TWeakPtr<FIKRigEditorController> EditorController;
+
+	/** on demand bone details object */	
+	mutable TObjectPtr<UIKRigBoneDetails> OptionalBoneDetails = nullptr;
 };
 
 class SIKRigSkeletonItem : public STableRow<TSharedPtr<FIKRigTreeElement>>
@@ -70,7 +81,8 @@ public:
 		TSharedPtr<SIKRigSkeleton> InHierarchy);
 
 private:
-	
+
+	bool OnVerifyNameChanged(const FText& InText, FText& OutErrorMessage) const;
 	void OnNameCommitted(const FText& InText, ETextCommit::Type InCommitType) const;
 	FText GetName() const;
 
@@ -196,6 +208,7 @@ public:
 		IKRigTreeElementType ItemType,
 		const bool bReplace);
 	void GetSelectedBoneChains(TArray<FIKRigSkeletonChain>& OutChains);
+	TArray<TSharedPtr<FIKRigTreeElement>> GetSelectedItems() const;
 	bool HasSelectedItems() const;
 	/** END selection state queries */
 
