Makes meshes fully swappable in IK Rig and Retarget editors.
Various small UX improvements

diff --git a/Engine/Plugins/Animation/IKRig/Source/IKRig/Private/AnimNodes/AnimNode_IKRig.cpp b/Engine/Plugins/Animation/IKRig/Source/IKRig/Private/AnimNodes/AnimNode_IKRig.cpp
--- a/Engine/Plugins/Animation/IKRig/Source/IKRig/Private/AnimNodes/AnimNode_IKRig.cpp
+++ b/Engine/Plugins/Animation/IKRig/Source/IKRig/Private/AnimNodes/AnimNode_IKRig.cpp
@@ -173,7 +173,7 @@ void FAnimNode_IKRig::Update_AnyThread(const FAnimationUpdateContext& Context)
 {
 	ActualAlpha = 0.f;
 
-		GetEvaluateGraphExposedInputs().Execute(Context);
+	GetEvaluateGraphExposedInputs().Execute(Context);
 
 	// alpha handlers
 	switch (AlphaInputType)
@@ -216,8 +216,7 @@ void FAnimNode_IKRig::PreUpdate(const UAnimInstance* InAnimInstance)
 	// initialize the IK Rig (will only try once on the current version of the rig asset)
 	if (!IKRigProcessor->IsInitialized())
 	{
-		const FReferenceSkeleton& RefSkeleton = InAnimInstance->GetSkelMeshComponent()->SkeletalMesh->GetRefSkeleton();
- 		IKRigProcessor->Initialize(RigDefinitionAsset, RefSkeleton);
+ 		IKRigProcessor->Initialize(RigDefinitionAsset, InAnimInstance->GetSkelMeshComponent()->SkeletalMesh);
 	}
 	
 	// cache list of goal creator components on the actor
diff --git a/Engine/Plugins/Animation/IKRig/Source/IKRig/Private/IKRig.cpp b/Engine/Plugins/Animation/IKRig/Source/IKRig/Private/IKRig.cpp
--- a/Engine/Plugins/Animation/IKRig/Source/IKRig/Private/IKRig.cpp
+++ b/Engine/Plugins/Animation/IKRig/Source/IKRig/Private/IKRig.cpp
@@ -7,7 +7,6 @@
 
 #define LOCTEXT_NAMESPACE "FIKRigModule"
 
-DEFINE_LOG_CATEGORY(LogIKRig);
 
 // Unique IK Rig Object version id
 const FGuid FIKRigObjectVersion::GUID(0xF6DFBB78, 0xBB50A0E4, 0x4018B84D, 0x60CBAF23);
diff --git a/Engine/Plugins/Animation/IKRig/Source/IKRig/Private/IKRigDefinition.cpp b/Engine/Plugins/Animation/IKRig/Source/IKRig/Private/IKRigDefinition.cpp
--- a/Engine/Plugins/Animation/IKRig/Source/IKRig/Private/IKRigDefinition.cpp
+++ b/Engine/Plugins/Animation/IKRig/Source/IKRig/Private/IKRigDefinition.cpp
@@ -6,14 +6,13 @@
 #include "Engine/SkeletalMesh.h"
 
 #if WITH_EDITOR
-
 #include "HAL/PlatformApplicationMisc.h"
-
 #endif
 
-#define LOCTEXT_NAMESPACE	"IKRigDefinition"
+#define LOCTEXT_NAMESPACE "IKRigDefinition"
 
 #if WITH_EDITOR
+const FName UIKRigDefinition::GetPreviewMeshPropertyName() { return GET_MEMBER_NAME_STRING_CHECKED(UIKRigDefinition, PreviewSkeletalMesh); };
 
 TOptional<FTransform::FReal> UIKRigEffectorGoal::GetNumericValue(
 	ESlateTransformComponent::Type Component,
@@ -273,6 +272,13 @@ void UIKRigDefinition::Serialize(FArchive& Ar)
 	Ar.UsingCustomVersion(FIKRigObjectVersion::GUID);
 }
 
+void UIKRigDefinition::PostLoad()
+{
+	Super::PostLoad();
+
+	Log.SetLogTarget(GetUniqueIDAsName(), false);
+}
+
 const FBoneChain* UIKRigDefinition::GetRetargetChainByName(FName ChainName) const
 {
 	for (const FBoneChain& Chain : RetargetDefinition.BoneChains)
@@ -293,7 +299,13 @@ void UIKRigDefinition::SetPreviewMesh(USkeletalMesh* PreviewMesh, bool bMarkAsDi
 
 USkeletalMesh* UIKRigDefinition::GetPreviewMesh() const
 {
-	return PreviewSkeletalMesh.Get();
+	return PreviewSkeletalMesh.LoadSynchronous();
+}
+
+FName UIKRigDefinition::GetUniqueIDAsName() const
+{
+	static const FString IKRigIDPrefix("IKRig_");
+	return FName(IKRigIDPrefix + FString::FromInt(GetUniqueID()));
 }
 
 #undef LOCTEXT_NAMESPACE
diff --git a/Engine/Plugins/Animation/IKRig/Source/IKRig/Private/IKRigLogger.cpp b/Engine/Plugins/Animation/IKRig/Source/IKRig/Private/IKRigLogger.cpp
new file mode 100644
--- /dev/null
+++ b/Engine/Plugins/Animation/IKRig/Source/IKRig/Private/IKRigLogger.cpp
@@ -0,0 +1,55 @@
+﻿// Copyright Epic Games, Inc. All Rights Reserved.
+
+#include "IKRigLogger.h"
+
+#include "Logging/MessageLog.h"
+
+DEFINE_LOG_CATEGORY(LogIKRig);
+
+void FIKRigLogger::SetLogTarget(const FName InLogName, bool bInSuppressWarnings)
+{
+	LogName = InLogName;
+	bWarningsSuppressed = bInSuppressWarnings;
+}
+
+FName FIKRigLogger::GetLogTarget() const
+{
+	return LogName;
+}
+	
+void FIKRigLogger::LogError(const FText& Message) const
+{
+	// print to the global output log
+	UE_LOG(LogIKRig, Error, TEXT("%s"), *Message.ToString());
+
+	// print to the output log in the asset editor
+	FMessageLog MessageLog(LogName);
+	MessageLog.Error(Message);
+}
+
+void FIKRigLogger::LogWarning(const FText& Message) const
+{
+	if (bWarningsSuppressed)
+	{
+		return;
+	}
+
+	// print to the global output log
+	UE_LOG(LogIKRig, Warning, TEXT("%s"), *Message.ToString());
+
+	// print to the output log in the asset editor
+	FMessageLog MessageLog(LogName);
+	MessageLog.Warning(Message);
+}
+
+void FIKRigLogger::LogEditorMessage(const FText& Message) const
+{
+	if (bWarningsSuppressed)
+	{
+		return;
+	}
+
+	// print to the output log in the asset editor
+	FMessageLog MessageLog(LogName);
+	MessageLog.Info(Message);
+}
diff --git a/Engine/Plugins/Animation/IKRig/Source/IKRig/Private/IKRigProcessor.cpp b/Engine/Plugins/Animation/IKRig/Source/IKRig/Private/IKRigProcessor.cpp
--- a/Engine/Plugins/Animation/IKRig/Source/IKRig/Private/IKRigProcessor.cpp
+++ b/Engine/Plugins/Animation/IKRig/Source/IKRig/Private/IKRigProcessor.cpp
@@ -4,7 +4,11 @@
 #include "IKRigDefinition.h"
 #include "IKRigSolver.h"
 
-void UIKRigProcessor::Initialize(const UIKRigDefinition* InRigAsset, const FIKRigInputSkeleton& InputSkeleton)
+#include "Engine/SkeletalMesh.h"
+
+#define LOCTEXT_NAMESPACE "IKRigProcessor"
+
+void UIKRigProcessor::Initialize(const UIKRigDefinition* InRigAsset, const USkeletalMesh* SkeletalMesh)
 {
 	// we instantiate UObjects here which MUST be done on game thread...
 	check(IsInGameThread());
@@ -22,17 +26,21 @@ void UIKRigProcessor::Initialize(const UIKRigDefinition* InRigAsset, const FIKRi
 	bTriedToInitialize = true;
 
 	// copy skeleton data from the actual skeleton we want to run on
-	Skeleton.SetInputSkeleton(InputSkeleton, InRigAsset->Skeleton.ExcludedBones);
+	Skeleton.SetInputSkeleton(SkeletalMesh, InRigAsset->Skeleton.ExcludedBones);
 	
 	if (InRigAsset->Skeleton.BoneNames.IsEmpty())
 	{
-		UE_LOG(LogTemp, Error, TEXT("Trying to initialize IKRig that has no skeleton: %s"), *InRigAsset->GetName());
+		InRigAsset->Log.LogError(FText::Format(
+			LOCTEXT("NoSkeleton", "Trying to initialize IK Rig, '{0}' that has no skeleton."),
+			FText::FromString(InRigAsset->GetName())));
 		return;
 	}
 
-	if (!UIKRigProcessor::IsIKRigCompatibleWithSkeleton(InRigAsset, InputSkeleton))
+	if (!UIKRigProcessor::IsIKRigCompatibleWithSkeleton(InRigAsset, SkeletalMesh))
 	{
-		UE_LOG(LogTemp, Error, TEXT("Trying to initialize IKRig with a Skeleton that is missing required bones. See output log. %s"), *InRigAsset->GetName());
+		InRigAsset->Log.LogError(FText::Format(
+			LOCTEXT("NoSkeleton", "Trying to initialize IKRig, '{0}' with a Skeleton that is missing required bones. See prior warnings."),
+			FText::FromString(InRigAsset->GetName())));
 		return;
 	}
 	
@@ -56,8 +64,9 @@ void UIKRigProcessor::Initialize(const UIKRigDefinition* InRigAsset, const FIKRi
 		// validate that the skeleton we are trying to solve this goal on contains the bone the goal expects
 		if (NewGoalBone.BoneIndex == INDEX_NONE)
 		{
-			UE_LOG(LogTemp, Warning, TEXT("IK Rig, %s has a Goal, '%s' that references an unknown bone, '%s'. Cannot evaluate."),
-				*InRigAsset->GetName(), *EffectorGoal->GoalName.ToString(), *EffectorGoal->BoneName.ToString());
+			InRigAsset->Log.LogError(FText::Format(
+				LOCTEXT("MissingGoalBone", "IK Rig, '{0}' has a Goal, '{1}' that references an unknown bone, '{2}'. Cannot evaluate."),
+				FText::FromString(InRigAsset->GetName()), FText::FromName(EffectorGoal->GoalName), FText::FromName(EffectorGoal->BoneName) ));
 			return;
 		}
 
@@ -67,8 +76,12 @@ void UIKRigProcessor::Initialize(const UIKRigDefinition* InRigAsset, const FIKRi
 		{
 			if (Bone->BoneName != NewGoalBone.BoneName)
 			{
-				UE_LOG(LogTemp, Warning, TEXT("IK Rig, %s has a Goal, '%s' that references different bones in different solvers, '%s' and '%s'. Cannot evaluate."),
-                *InRigAsset->GetName(), *EffectorGoal->GoalName.ToString(), *Bone->BoneName.ToString(), *NewGoalBone.BoneName.ToString());
+				InRigAsset->Log.LogError(FText::Format(
+				LOCTEXT("DuplicateGoal", "IK Rig, '{0}' has a Goal, '{1}' that references different bones in different solvers, '{2}' and '{3}'. Cannot evaluate."),
+                FText::FromString(InRigAsset->GetName()),
+                FText::FromName(EffectorGoal->GoalName),
+                FText::FromName(Bone->BoneName), 
+                FText::FromName(NewGoalBone.BoneName)));
 				return;
 			}
 		}
@@ -84,9 +97,10 @@ void UIKRigProcessor::Initialize(const UIKRigDefinition* InRigAsset, const FIKRi
 	{
 		if (!IKRigSolver)
 		{
-			// this can happen if asset references deleted IK Solver type
-			// which should only happen during development (if at all)
-			UE_LOG(LogTemp, Warning, TEXT("IK Rig, %s has null/unknown solver in it. Please remove it."), *InRigAsset->GetName());
+			// this can happen if asset references deleted IK Solver type which should only happen during development (if at all)
+			InRigAsset->Log.LogWarning(FText::Format(
+				LOCTEXT("UnknownSolver", "IK Rig, '{0}' has null/unknown solver in it. Please remove it."),
+				FText::FromString(InRigAsset->GetName())));
 			continue;
 		}
 
@@ -98,16 +112,14 @@ void UIKRigProcessor::Initialize(const UIKRigDefinition* InRigAsset, const FIKRi
 		Solvers.Add(Solver);
 	}
 
+	InRigAsset->Log.LogEditorMessage(FText::Format(
+				LOCTEXT("SuccessfulInit", "IK Rig, '{0}' ready to run on {1}."),
+				FText::FromString(InRigAsset->GetName()), FText::FromString(SkeletalMesh->GetName())));
+	
 	bInitialized = true;
 }
 
-void UIKRigProcessor::Initialize(const UIKRigDefinition* InRigAsset, const FReferenceSkeleton& RefSkeleton)
-{
-	const FIKRigInputSkeleton InputSkeleton = FIKRigInputSkeleton(RefSkeleton);
-	Initialize(InRigAsset, InputSkeleton);
-}
-
-bool UIKRigProcessor::IsIKRigCompatibleWithSkeleton(const UIKRigDefinition* InRigAsset, const FIKRigInputSkeleton& InputSkeleton)
+bool UIKRigProcessor::IsIKRigCompatibleWithSkeleton(const UIKRigDefinition* InRigAsset,	const FIKRigInputSkeleton& InputSkeleton)
 {
 	// first we validate that all the required bones are in the input skeleton...
 	
@@ -135,7 +147,11 @@ bool UIKRigProcessor::IsIKRigCompatibleWithSkeleton(const UIKRigDefinition* InRi
 	{
 		if (!InputSkeleton.BoneNames.Contains(RequiredBone))
 		{
-			UE_LOG(LogTemp, Warning, TEXT("IK Rig, '%s' is missing a required bone in Skeletal Mesh: '%s'."), *InRigAsset->GetName(), *RequiredBone.ToString());
+			InRigAsset->Log.LogError(FText::Format(
+			LOCTEXT("MissingBone", "IK Rig, '{0}' is missing a required bone, '{1}' in the Skeletal Mesh."),
+			FText::FromString(InRigAsset->GetName()),
+			FText::FromName(RequiredBone)));
+			
 			bAllRequiredBonesFound = false;
 		}
 	}
@@ -165,21 +181,24 @@ bool UIKRigProcessor::IsIKRigCompatibleWithSkeleton(const UIKRigDefinition* InRi
 			if (!InputSkeleton.BoneNames.IsValidIndex(InputParentIndex))
 			{
 				bAllParentsValid = false;
-				UE_LOG(LogTemp, Error,
-					TEXT("IK Rig is running on a skeleton with a required bone, '%s', that expected to have a valid parent. The expected parent was, '%s'."),
-					*RequiredBone.ToString(),
-					*AssetParentName.ToString());
+
+				InRigAsset->Log.LogError(FText::Format(
+			LOCTEXT("MissingParent", "IK Rig is running on a skeleton with a required bone, '{0}', that expected to have a valid parent. The expected parent was, '{1}'."),
+				FText::FromName(RequiredBone),
+				FText::FromName(AssetParentName)));
+				
 				continue;
 			}
+			
 			const FName& InputParentName = InputSkeleton.BoneNames[InputParentIndex];
 			if (AssetParentName != InputParentName)
 			{
 				// we only warn about this, because it may be nice not to have the exact same hierarchy
-				UE_LOG(LogTemp, Warning,
-					TEXT("IK Rig is running on a skeleton with a required bone, '%s', that has a different parent '%s'. The expected parent was, '%s'."),
-					*InputParentName.ToString(),
-					*InputParentName.ToString(),
-					*AssetParentName.ToString());
+				InRigAsset->Log.LogWarning(FText::Format(
+				LOCTEXT("MissingParent", "IK Rig is running on a skeleton with a required bone, '{0}', that has a different parent '{1}'. The expected parent was, '{2}'."),
+				FText::FromName(RequiredBone),
+				FText::FromName(InputParentName),
+				FText::FromName(AssetParentName)));
 				continue;
 			}
 		}
@@ -391,3 +410,5 @@ void UIKRigProcessor::ResolveFinalGoalTransforms(const FTransform& WorldToCompon
             Goal.RotationAlpha);
 	}
 }
+
+#undef LOCTEXT_NAMESPACE
\ No newline at end of file
diff --git a/Engine/Plugins/Animation/IKRig/Source/IKRig/Private/IKRigSkeleton.cpp b/Engine/Plugins/Animation/IKRig/Source/IKRig/Private/IKRigSkeleton.cpp
--- a/Engine/Plugins/Animation/IKRig/Source/IKRig/Private/IKRigSkeleton.cpp
+++ b/Engine/Plugins/Animation/IKRig/Source/IKRig/Private/IKRigSkeleton.cpp
@@ -2,12 +2,39 @@
 
 #include "IKRigSkeleton.h"
 
+#include "Engine/SkeletalMesh.h"
+
 // This is the default end of branch index value, meaning we haven't cached it yet
 #define FIKRIGSKELETON_INVALID_EO_BRANCH_INDEX -2
 
-void FIKRigSkeleton::SetInputSkeleton(const FReferenceSkeleton& RefSkeleton, const TArray<FName>& InExcludedBones)
+FIKRigInputSkeleton::FIKRigInputSkeleton(const USkeletalMesh* SkeletalMesh)
+{
+	Initialize(SkeletalMesh);
+}
+
+void FIKRigInputSkeleton::Initialize(const USkeletalMesh* SkeletalMesh)
+{
+	Reset();
+
+	SkeletalMeshName = FName(SkeletalMesh->GetName());
+	const FReferenceSkeleton& RefSkeleton = SkeletalMesh->GetRefSkeleton();
+	const TArray<FMeshBoneInfo>& BoneInfo = RefSkeleton.GetRefBoneInfo();
+	for (int32 BoneIndex=0; BoneIndex<BoneInfo.Num(); ++BoneIndex)
+	{
+		BoneNames.Add(BoneInfo[BoneIndex].Name);
+		ParentIndices.Add(BoneInfo[BoneIndex].ParentIndex);
+		LocalRefPose.Add(RefSkeleton.GetRefBonePose()[BoneIndex]);
+	}
+}
+
+void FIKRigInputSkeleton::Reset()
+{
+	*this = FIKRigInputSkeleton();
+}
+
+void FIKRigSkeleton::SetInputSkeleton(const USkeletalMesh* SkeletalMesh, const TArray<FName>& InExcludedBones)
 {
-	const FIKRigInputSkeleton InputSkeleton = FIKRigInputSkeleton(RefSkeleton);
+	const FIKRigInputSkeleton InputSkeleton = FIKRigInputSkeleton(SkeletalMesh);
 	SetInputSkeleton(InputSkeleton, InExcludedBones);
 }
 
@@ -50,13 +77,7 @@ void FIKRigSkeleton::SetInputSkeleton(const FIKRigInputSkeleton& InputSkeleton,
 
 void FIKRigSkeleton::Reset()
 {
-	BoneNames.Reset();
-	ParentIndices.Reset();
-	ExcludedBones.Reset();
-	CurrentPoseGlobal.Reset();
-	CurrentPoseLocal.Reset();
-	RefPoseGlobal.Reset();
-	CachedEndOfBranchIndices.Reset();
+	*this = FIKRigSkeleton();
 }
 
 int32 FIKRigSkeleton::GetBoneIndexFromName(const FName InName) const
diff --git a/Engine/Plugins/Animation/IKRig/Source/IKRig/Private/Retargeter/IKRetargetProcessor.cpp b/Engine/Plugins/Animation/IKRig/Source/IKRig/Private/Retargeter/IKRetargetProcessor.cpp
--- a/Engine/Plugins/Animation/IKRig/Source/IKRig/Private/Retargeter/IKRetargetProcessor.cpp
+++ b/Engine/Plugins/Animation/IKRig/Source/IKRig/Private/Retargeter/IKRetargetProcessor.cpp
@@ -3,10 +3,12 @@
 #include "Retargeter/IKRetargetProcessor.h"
 
 #include "IKRigDefinition.h"
+#include "IKRigLogger.h"
 #include "IKRigProcessor.h"
 #include "Engine/SkeletalMesh.h"
 #include "Retargeter/IKRetargeter.h"
 
+#define LOCTEXT_NAMESPACE "IKRetargetProcessor"
 
 void FRetargetSkeleton::Initialize(USkeletalMesh* InSkeletalMesh)
 {	
@@ -307,7 +309,8 @@ void FTargetSkeleton::SetBoneIsRetargeted(const int32 BoneIndex, const bool IsRe
 bool FChainFK::Initialize(
 	const FRetargetSkeleton& Skeleton,
 	const TArray<int32>& BoneIndices,
-	const TArray<FTransform>& InitialGlobalPose)
+	const TArray<FTransform>& InitialGlobalPose,
+	FIKRigLogger& Log)
 {
 	check(!BoneIndices.IsEmpty());
 
@@ -338,10 +341,10 @@ bool FChainFK::Initialize(
 	}
 
 	// calculate parameter of each bone, normalized by the length of the bone chain
-	return CalculateBoneParameters();
+	return CalculateBoneParameters(Log);
 }
 
-bool FChainFK::CalculateBoneParameters()
+bool FChainFK::CalculateBoneParameters(FIKRigLogger& Log)
 {
 	Params.Reset();
 	
@@ -365,7 +368,7 @@ bool FChainFK::CalculateBoneParameters()
 	// cannot retarget chain if all the bones are sitting directly on each other
 	if (TotalChainLength <= KINDA_SMALL_NUMBER)
 	{
-		UE_LOG(LogTemp, Warning, TEXT("IK Retargeter bone chain length is too small to reliably retarget."));
+		Log.LogWarning(LOCTEXT("TinyBoneChain", "IK Retargeter bone chain length is too small to reliably retarget."));
 		return false;
 	}
 
@@ -736,11 +739,12 @@ FTransform FChainDecoderFK::GetTransformAtParam(
 
 bool FChainRetargeterIK::InitializeSource(
 	const TArray<int32>& BoneIndices,
-	const TArray<FTransform>& SourceInitialGlobalPose)
+	const TArray<FTransform>& SourceInitialGlobalPose,
+	FIKRigLogger& Log)
 {
 	if (BoneIndices.Num() < 3)
 	{
-		UE_LOG(LogTemp, Warning, TEXT("IK Retargeter trying to retarget source bone chain with IK but it has less than 3 joints."));
+		Log.LogWarning(LOCTEXT("SourceChainLessThanThree", "IK Retargeter trying to retarget source bone chain with IK but it has less than 3 joints."));
 		return false;
 	}
 	
@@ -758,7 +762,7 @@ bool FChainRetargeterIK::InitializeSource(
 
 	if (Length <= KINDA_SMALL_NUMBER)
 	{
-		UE_LOG(LogTemp, Warning, TEXT("IK Retargeter trying to retarget source bone chain with IK, but it is zero length!"));
+		Log.LogWarning(LOCTEXT("SourceZeroLengthIK", "IK Retargeter trying to retarget source bone chain with IK, but it is zero length!"));
     	return false;
 	}
 
@@ -790,11 +794,12 @@ void FChainRetargeterIK::EncodePose(const TArray<FTransform>& InSourceGlobalPose
 
 bool FChainRetargeterIK::InitializeTarget(
 	const TArray<int32>& BoneIndices,
-	const TArray<FTransform> &TargetInitialGlobalPose)
+	const TArray<FTransform> &TargetInitialGlobalPose,
+	FIKRigLogger& Log)
 {
 	if (BoneIndices.Num() < 3)
 	{
-		UE_LOG(LogTemp, Warning, TEXT("IK Retargeter trying to retarget target bone chain with IK but it has less than 3 joints."));
+		Log.LogWarning(LOCTEXT("TargetChainLessThanThree", "IK Retargeter trying to retarget target bone chain with IK but it has less than 3 joints."));
 		return false;
 	}
 	
@@ -808,7 +813,7 @@ bool FChainRetargeterIK::InitializeTarget(
 
 	if (Target.InitialLength <= KINDA_SMALL_NUMBER)
 	{
-		UE_LOG(LogTemp, Warning, TEXT("IK Retargeter trying to retarget bone chain with IK, but it is zero length!"));
+		Log.LogWarning(LOCTEXT("TargetZeroLengthIK", "IK Retargeter trying to retarget target bone chain with IK, but it is zero length!"));
 		return false;
 	}
 	
@@ -897,30 +902,35 @@ bool FRetargetChainPair::Initialize(
     const FBoneChain& SourceBoneChain,
     const FBoneChain& TargetBoneChain,
     const FRetargetSkeleton& SourceSkeleton,
-    const FTargetSkeleton& TargetSkeleton)
+    const FTargetSkeleton& TargetSkeleton,
+    FIKRigLogger& Log)
 {
 	// validate source bone chain is compatible with source skeletal mesh
-	const bool bIsSourceValid = ValidateBoneChainWithSkeletalMesh(true, SourceBoneChain, SourceSkeleton);
+	const bool bIsSourceValid = ValidateBoneChainWithSkeletalMesh(true, SourceBoneChain, SourceSkeleton, Log);
 	if (!bIsSourceValid)
 	{
-		UE_LOG(LogTemp, Warning, TEXT("IK Retargeter source bone chain, '%s', is not compatible with Skeletal Mesh: '%s'"),
-			*SourceBoneChain.ChainName.ToString(), *SourceSkeleton.SkeletalMesh->GetName());
+		Log.LogWarning( FText::Format(
+			LOCTEXT("IncompatibleSourceChain", "IK Retargeter source bone chain, '{0}', is not compatible with Skeletal Mesh: '{1}'"),
+			FText::FromName(SourceBoneChain.ChainName), FText::FromString(SourceSkeleton.SkeletalMesh->GetName())));
 		return false;
 	}
 
 	// validate target bone chain is compatible with target skeletal mesh
-	const bool bIsTargetValid = ValidateBoneChainWithSkeletalMesh(false, TargetBoneChain, TargetSkeleton);
+	const bool bIsTargetValid = ValidateBoneChainWithSkeletalMesh(false, TargetBoneChain, TargetSkeleton, Log);
 	if (!bIsTargetValid)
     {
-		UE_LOG(LogTemp, Warning, TEXT("IK Retargeter target bone chain, '%s', is not compatible with Skeletal Mesh: '%s'"),
-            *TargetBoneChain.ChainName.ToString(), *TargetSkeleton.SkeletalMesh->GetName());
+		Log.LogWarning( FText::Format(
+			LOCTEXT("IncompatibleTargetChain", "IK Retargeter target bone chain, '{0}', is not compatible with Skeletal Mesh: '{1}'"),
+			FText::FromName(TargetBoneChain.ChainName), FText::FromString(TargetSkeleton.SkeletalMesh->GetName())));
 		return false;
     }
 
 	// ensure valid settings object
 	if (InSettings == nullptr)
 	{
-		UE_LOG(LogTemp, Warning, TEXT("IK Retargeter target bone chain, '%s', has null settings."), *TargetBoneChain.ChainName.ToString());
+		Log.LogWarning( FText::Format(
+			LOCTEXT("MissingChainSettings", "IK Retargeter target bone chain, '{0}', has null settings."),
+			FText::FromName(TargetBoneChain.ChainName)));
 		return false;
 	}
 
@@ -935,7 +945,8 @@ bool FRetargetChainPair::Initialize(
 bool FRetargetChainPair::ValidateBoneChainWithSkeletalMesh(
     const bool IsSource,
     const FBoneChain& BoneChain,
-    const FRetargetSkeleton& RetargetSkeleton)
+    const FRetargetSkeleton& RetargetSkeleton,
+    FIKRigLogger& Log)
 {
 	// record the chain indices
 	TArray<int32>& BoneIndices = IsSource ? SourceBoneIndices : TargetBoneIndices;
@@ -946,22 +957,27 @@ bool FRetargetChainPair::ValidateBoneChainWithSkeletalMesh(
 	// warn if START bone not found
 	if (!ResolvedChain.bFoundStartBone)
 	{
-		UE_LOG(LogTemp, Warning, TEXT("IK Retargeter bone chain, %s, could not find start bone, %s in mesh %s"),
-            *BoneChain.ChainName.ToString(), *BoneChain.StartBone.BoneName.ToString(), *RetargetSkeleton.SkeletalMesh->GetName());
+		Log.LogWarning( FText::Format(
+			LOCTEXT("MissingStartBone", "IK Retargeter bone chain, {0}, could not find start bone, {1} in mesh {2}"),
+			FText::FromName(BoneChain.ChainName),
+			FText::FromName(BoneChain.StartBone.BoneName),
+			FText::FromString(RetargetSkeleton.SkeletalMesh->GetName())));
 	}
 	
 	// warn if END bone not found
 	if (!ResolvedChain.bFoundEndBone)
 	{
-		UE_LOG(LogTemp, Warning, TEXT("IK Retargeter bone chain, %s, could not find end bone, %s in mesh %s"),
-            *BoneChain.ChainName.ToString(), *BoneChain.EndBone.BoneName.ToString(), *RetargetSkeleton.SkeletalMesh->GetName());
+		Log.LogWarning( FText::Format(
+			LOCTEXT("MissingEndBone", "IK Retargeter bone chain, {0}, could not find end bone, {1} in mesh {2}"),
+			FText::FromName(BoneChain.ChainName), FText::FromName(BoneChain.EndBone.BoneName), FText::FromString(RetargetSkeleton.SkeletalMesh->GetName())));
 	}
 
 	// warn if END bone was not a child of START bone
-	if (!ResolvedChain.bEndIsStartOrChildOfStart)
+	if (ResolvedChain.bFoundEndBone && !ResolvedChain.bEndIsStartOrChildOfStart)
 	{
-		UE_LOG(LogTemp, Warning, TEXT("IK Retargeter bone chain, %s, end bone, '%s' was not a child of the start bone '%s'."),
-			*BoneChain.ChainName.ToString(), *BoneChain.EndBone.BoneName.ToString(), *BoneChain.StartBone.BoneName.ToString());
+		Log.LogWarning( FText::Format(
+			LOCTEXT("EndNotChildtOfStart", "IK Retargeter bone chain, {0}, end bone, '{1}' was not a child of the start bone '{2}'."),
+			FText::FromName(BoneChain.ChainName), FText::FromName(BoneChain.EndBone.BoneName), FText::FromName(BoneChain.StartBone.BoneName)));
 	}
 	
 	return ResolvedChain.IsValid();
@@ -972,29 +988,32 @@ bool FRetargetChainPairFK::Initialize(
 	const FBoneChain& SourceBoneChain,
 	const FBoneChain& TargetBoneChain,
 	const FRetargetSkeleton& SourceSkeleton,
-	const FTargetSkeleton& TargetSkeleton)
+	const FTargetSkeleton& TargetSkeleton,
+	FIKRigLogger& Log)
 {
-	const bool bChainInitialized = FRetargetChainPair::Initialize(InSettings, SourceBoneChain, TargetBoneChain, SourceSkeleton, TargetSkeleton);
+	const bool bChainInitialized = FRetargetChainPair::Initialize(InSettings, SourceBoneChain, TargetBoneChain, SourceSkeleton, TargetSkeleton, Log);
 	if (!bChainInitialized)
 	{
 		return false;
 	}
 
 	// initialize SOURCE FK chain encoder with retarget pose
-	const bool bFKEncoderInitialized = FKEncoder.Initialize(SourceSkeleton, SourceBoneIndices, SourceSkeleton.RetargetGlobalPose);
+	const bool bFKEncoderInitialized = FKEncoder.Initialize(SourceSkeleton, SourceBoneIndices, SourceSkeleton.RetargetGlobalPose, Log);
 	if (!bFKEncoderInitialized)
 	{
-		UE_LOG(LogTemp, Warning, TEXT("IK Retargeter failed to initialize FK encoder, '%s', on Skeletal Mesh: '%s'"),
-            *SourceBoneChainName.ToString(), *SourceSkeleton.SkeletalMesh->GetName());
+		Log.LogWarning( FText::Format(
+			LOCTEXT("BadFKEncoder", "IK Retargeter failed to initialize FK encoder, '{0}', on Skeletal Mesh: '{1}'"),
+			FText::FromName(SourceBoneChainName), FText::FromString(SourceSkeleton.SkeletalMesh->GetName())));
 		return false;
 	}
 
 	// initialize TARGET FK chain decoder with retarget pose
-	const bool bFKDecoderInitialized = FKDecoder.Initialize(TargetSkeleton, TargetBoneIndices, TargetSkeleton.RetargetGlobalPose);
+	const bool bFKDecoderInitialized = FKDecoder.Initialize(TargetSkeleton, TargetBoneIndices, TargetSkeleton.RetargetGlobalPose, Log);
 	if (!bFKDecoderInitialized)
 	{
-		UE_LOG(LogTemp, Warning, TEXT("IK Retargeter failed to initialize FK decoder, '%s', on Skeletal Mesh: '%s'"),
-            *TargetBoneChainName.ToString(), *TargetSkeleton.SkeletalMesh->GetName());
+		Log.LogWarning( FText::Format(
+			LOCTEXT("BadFKDecoder", "IK Retargeter failed to initialize FK decoder, '{0}', on Skeletal Mesh: '{1}'"),
+			FText::FromName(TargetBoneChainName), FText::FromString(TargetSkeleton.SkeletalMesh->GetName())));
 		return false;
 	}
 
@@ -1006,7 +1025,8 @@ bool FRetargetChainPairIK::Initialize(
 	const FBoneChain& SourceBoneChain,
 	const FBoneChain& TargetBoneChain,
 	const FRetargetSkeleton& SourceSkeleton,
-	const FTargetSkeleton& TargetSkeleton)
+	const FTargetSkeleton& TargetSkeleton,
+	FIKRigLogger& Log)
 {
 	// validate if this chain even uses an IK Goal
 	const bool bUsingIKGoal = TargetBoneChain.IKGoalName != NAME_None && TargetBoneChain.IKGoalName != "- None -";
@@ -1019,41 +1039,47 @@ bool FRetargetChainPairIK::Initialize(
 	IKGoalName = TargetBoneChain.IKGoalName;
 
 	// initialize bone chains
-	const bool bChainInitialized = FRetargetChainPair::Initialize(InSettings, SourceBoneChain, TargetBoneChain, SourceSkeleton, TargetSkeleton);
+	const bool bChainInitialized = FRetargetChainPair::Initialize(InSettings, SourceBoneChain, TargetBoneChain, SourceSkeleton, TargetSkeleton, Log);
 	if (!bChainInitialized)
 	{
 		return false;
 	}
 
 	// initialize SOURCE IK chain encoder with retarget pose
-	const bool bIKEncoderInitialized = IKChainRetargeter.InitializeSource(SourceBoneIndices, SourceSkeleton.RetargetGlobalPose);
+	const bool bIKEncoderInitialized = IKChainRetargeter.InitializeSource(SourceBoneIndices, SourceSkeleton.RetargetGlobalPose, Log);
 	if (!bIKEncoderInitialized)
 	{
-		UE_LOG(LogTemp, Warning, TEXT("IK Retargeter failed to initialize IK encoder, '%s', on Skeletal Mesh: '%s'"),
-        *SourceBoneChainName.ToString(), *SourceSkeleton.SkeletalMesh->GetName());
+		Log.LogWarning( FText::Format(
+			LOCTEXT("BadIKEncoder", "IK Retargeter failed to initialize IK encoder, '{0}', on Skeletal Mesh: '{1}'"),
+			FText::FromName(SourceBoneChainName), FText::FromString(SourceSkeleton.SkeletalMesh->GetName())));
 		return false;
 	}
 
 	// initialize TARGET IK chain decoder with retarget pose
-	const bool bIKDecoderInitialized = IKChainRetargeter.InitializeTarget(TargetBoneIndices, TargetSkeleton.RetargetGlobalPose);
+	const bool bIKDecoderInitialized = IKChainRetargeter.InitializeTarget(TargetBoneIndices, TargetSkeleton.RetargetGlobalPose, Log);
 	if (!bIKDecoderInitialized)
 	{
-		UE_LOG(LogTemp, Warning, TEXT("IK Retargeter failed to initialize IK decoder, '%s', on Skeletal Mesh: '%s'"),
-        *TargetBoneChainName.ToString(), *TargetSkeleton.SkeletalMesh->GetName());
+		Log.LogWarning( FText::Format(
+			LOCTEXT("BadIKDecoder", "IK Retargeter failed to initialize IK decoder, '{0}', on Skeletal Mesh: '{1}'"),
+			FText::FromName(TargetBoneChainName), FText::FromString(TargetSkeleton.SkeletalMesh->GetName())));
 		return false;
 	}
 
 	return true;
 }
 
-bool FRootRetargeter::InitializeSource(const FName SourceRootBoneName, const FRetargetSkeleton& SourceSkeleton)
+bool FRootRetargeter::InitializeSource(
+	const FName SourceRootBoneName,
+	const FRetargetSkeleton& SourceSkeleton,
+	FIKRigLogger& Log)
 {
 	// validate target root bone exists
 	Source.BoneIndex = SourceSkeleton.FindBoneIndexByName(SourceRootBoneName);
 	if (Source.BoneIndex == INDEX_NONE)
 	{
-		UE_LOG(LogTemp, Warning, TEXT("IK Retargeter could not find source root bone, %s in mesh %s"),
-			*SourceRootBoneName.ToString(), *SourceSkeleton.SkeletalMesh->GetName());
+		Log.LogWarning( FText::Format(
+			LOCTEXT("MissingSourceRoot", "IK Retargeter could not find source root bone, {0} in mesh {1}"),
+			FText::FromName(SourceRootBoneName), FText::FromString(SourceSkeleton.SkeletalMesh->GetName())));
 		return false;
 	}
 	
@@ -1066,7 +1092,7 @@ bool FRootRetargeter::InitializeSource(const FName SourceRootBoneName, const FRe
 	if (InitialHeight < KINDA_SMALL_NUMBER)
 	{
 		// warn user and push it up slightly to avoid divide by zero
-		UE_LOG(LogTemp, Warning, TEXT("IK Retargeter root bone is very near the ground plane. This is probably not correct."));
+		Log.LogWarning(LOCTEXT("BadRootHeight", "IK Retargeter root bone is very near the ground plane. This is probably not intentional."));
 		InitialHeight = 1.0f;
 	}
 
@@ -1076,14 +1102,18 @@ bool FRootRetargeter::InitializeSource(const FName SourceRootBoneName, const FRe
 	return true;
 }
 
-bool FRootRetargeter::InitializeTarget(const FName TargetRootBoneName, const FTargetSkeleton& TargetSkeleton)
+bool FRootRetargeter::InitializeTarget(
+	const FName TargetRootBoneName,
+	const FTargetSkeleton& TargetSkeleton,
+	FIKRigLogger& Log)
 {
 	// validate target root bone exists
 	Target.BoneIndex = TargetSkeleton.FindBoneIndexByName(TargetRootBoneName);
 	if (Target.BoneIndex == INDEX_NONE)
 	{
-		UE_LOG(LogTemp, Warning, TEXT("IK Retargeter could not find target root bone, %s in mesh %s"),
-            *TargetRootBoneName.ToString(), *TargetSkeleton.SkeletalMesh->GetName());
+		Log.LogWarning( FText::Format(
+			LOCTEXT("CountNotFindRootBone", "IK Retargeter could not find target root bone, {0} in mesh {1}"),
+			FText::FromName(TargetRootBoneName), FText::FromString(TargetSkeleton.SkeletalMesh->GetName())));
 		return false;
 	}
 
@@ -1130,10 +1160,11 @@ void FRootRetargeter::DecodePose(
 void UIKRetargetProcessor::Initialize(
 		USkeletalMesh* SourceSkeletalMesh,
 		USkeletalMesh* TargetSkeletalMesh,
-		UIKRetargeter* InRetargeterAsset)
+		UIKRetargeter* InRetargeterAsset,
+		const bool bSuppressWarnings)
 {
 	bIsInitialized = false;
-
+	
 	// record source asset
 	RetargeterAsset = InRetargeterAsset;
 
@@ -1148,27 +1179,27 @@ void UIKRetargetProcessor::Initialize(
 	// check prerequisite assets
 	if (!SourceSkeletalMesh)
 	{
-		UE_LOG(LogTemp, Error, TEXT("IK Retargeter unable to initialize. Missing source Skeletal Mesh asset."));
+		RetargeterAsset->Log.LogError(LOCTEXT("MissingSourceMesh", "IK Retargeter unable to initialize. Missing source Skeletal Mesh asset."));
 		return;
 	}
 	if (!TargetSkeletalMesh)
 	{
-		UE_LOG(LogTemp, Error, TEXT("IK Retargeter unable to initialize. Missing target Skeletal Mesh asset."));
+		RetargeterAsset->Log.LogError(LOCTEXT("MissingTargetMesh", "IK Retargeter unable to initialize. Missing target Skeletal Mesh asset."));
 		return;
 	}
 	if (!RetargeterAsset->GetSourceIKRig())
 	{
-		UE_LOG(LogTemp, Error, TEXT("IK Retargeter unable to initialize. Missing source IK Rig asset."));
+		RetargeterAsset->Log.LogError(LOCTEXT("MissingSourceIKRig", "IK Retargeter unable to initialize. Missing source IK Rig asset."));
 		return;
 	}
 	if (!RetargeterAsset->GetTargetIKRig())
 	{
-		UE_LOG(LogTemp, Error, TEXT("IK Retargeter unable to initialize. Missing target IK Rig asset."));
+		RetargeterAsset->Log.LogError(LOCTEXT("MissingTargetIKRig", "IK Retargeter unable to initialize. Missing target IK Rig asset."));
 		return;
 	}
 	if (!RetargeterAsset->GetCurrentRetargetPose())
 	{
-		UE_LOG(LogTemp, Error, TEXT("IK Retargeter unable to initialize. Missing retarget pose."));
+		RetargeterAsset->Log.LogError(LOCTEXT("MissingRetargetPose", "IK Retargeter unable to initialize. Missing retarget pose."));
 		return;
 	}
 	
@@ -1182,31 +1213,36 @@ void UIKRetargetProcessor::Initialize(
 
 	// initialize roots
 	bRootsInitialized = InitializeRoots();
-	if (!bRootsInitialized)
-	{
-		// couldn't match up any BoneChain pairs, no retargeting possible
-		UE_LOG(LogTemp, Error, TEXT("IK Retargeter unable to initialize one or more root bones on source, %s and target, %s"),
-            *SourceSkeleton.SkeletalMesh->GetName(), *TargetSkeleton.SkeletalMesh->GetName());
-	}
 
 	// initialize pairs of bone chains
 	bAtLeastOneValidBoneChainPair = InitializeBoneChainPairs();
 	if (!bAtLeastOneValidBoneChainPair)
 	{
 		// couldn't match up any BoneChain pairs, no limb retargeting possible
-		UE_LOG(LogTemp, Warning, TEXT("IK Retargeter unable to create any Bone Chain pairs between source, %s and target, %s"),
-            *SourceSkeleton.SkeletalMesh->GetName(), *TargetSkeleton.SkeletalMesh->GetName());
+		RetargeterAsset->Log.LogWarning( FText::Format(
+			LOCTEXT("NoMappedChains", "IK Retargeter unable to map any bone chains between source, {0} and target, {1}"),
+			FText::FromString(SourceSkeleton.SkeletalMesh->GetName()), FText::FromString(TargetSkeleton.SkeletalMesh->GetName())));
 	}
 
 	// initialize the IKRigProcessor for doing IK decoding
-	bIKRigInitialized = InitializeIKRig(this, TargetSkeletalMesh->GetRefSkeleton());
+	bIKRigInitialized = InitializeIKRig(this, TargetSkeletalMesh);
 	if (!bIKRigInitialized)
 	{
 		// couldn't initialize the IK Rig, we don't disable the retargeter in this case, just warn the user
-		UE_LOG(LogTemp, Warning, TEXT("IK Retargeter unable to initialize IK Rig for %s. See output for details."),
-			*TargetSkeleton.SkeletalMesh->GetName());
+		RetargeterAsset->Log.LogWarning( FText::Format(
+			LOCTEXT("CouldNotInitializeIKRig", "IK Retargeter was unable to initialize the IK Rig, {0} for the Skeletal Mesh {1}. See previous warnings."),
+			FText::FromString(RetargeterAsset->GetTargetIKRig()->GetName()), FText::FromString(TargetSkeleton.SkeletalMesh->GetName())));
 	}
 
+	// must have a mapped root bone OR at least a single mapped chain to be able to do any retargeting at all
+	if (bRootsInitialized && bAtLeastOneValidBoneChainPair)
+	{
+		// confirm for the user that the IK Rig was initialized successfully
+		RetargeterAsset->Log.LogEditorMessage(FText::Format(
+				LOCTEXT("SuccessfulInit", "Success! The IK Retargeter is ready to transfer animation from the source, {0} to the target, {1}"),
+				FText::FromString(SourceSkeleton.SkeletalMesh->GetName()), FText::FromString(TargetSkeleton.SkeletalMesh->GetName())));
+	}
+	
 	bIsInitialized = true;
 }
 
@@ -1214,20 +1250,22 @@ bool UIKRetargetProcessor::InitializeRoots()
 {
 	// initialize root encoder
 	const FName SourceRootBoneName = RetargeterAsset->GetSourceIKRig()->GetRetargetRoot();
-	const bool bRootEncoderInit = RootRetargeter.InitializeSource(SourceRootBoneName, SourceSkeleton);
+	const bool bRootEncoderInit = RootRetargeter.InitializeSource(SourceRootBoneName, SourceSkeleton, RetargeterAsset->Log);
 	if (!bRootEncoderInit)
 	{
-		UE_LOG(LogTemp, Warning, TEXT("IK Retargeter unable to initialize source root, '%s' on skeletal mesh: '%s'"),
-            *SourceRootBoneName.ToString(), *SourceSkeleton.SkeletalMesh->GetName());
+		RetargeterAsset->Log.LogWarning( FText::Format(
+			LOCTEXT("NoSourceRoot", "IK Retargeter unable to initialize source root, '{0}' on skeletal mesh: '{1}'"),
+			FText::FromName(SourceRootBoneName), FText::FromString(SourceSkeleton.SkeletalMesh->GetName())));
 	}
 
 	// initialize root decoder
 	const FName TargetRootBoneName = RetargeterAsset->GetTargetIKRig()->GetRetargetRoot();
-	const bool bRootDecoderInit = RootRetargeter.InitializeTarget(TargetRootBoneName, TargetSkeleton);
+	const bool bRootDecoderInit = RootRetargeter.InitializeTarget(TargetRootBoneName, TargetSkeleton, RetargeterAsset->Log);
 	if (!bRootDecoderInit)
 	{
-		UE_LOG(LogTemp, Warning, TEXT("IK Retargeter unable to initialize target root, '%s' on skeletal mesh: '%s'"),
-            *TargetRootBoneName.ToString(), *TargetSkeleton.SkeletalMesh->GetName());
+		RetargeterAsset->Log.LogWarning( FText::Format(
+			LOCTEXT("NoTargetRoot", "IK Retargeter unable to initialize target root, '{0}' on skeletal mesh: '{1}'"),
+			FText::FromName(TargetRootBoneName), FText::FromString(TargetSkeleton.SkeletalMesh->GetName())));
 	}
 
 	return bRootEncoderInit && bRootDecoderInit;
@@ -1248,8 +1286,9 @@ bool UIKRetargetProcessor::InitializeBoneChainPairs()
 		const FBoneChain* TargetBoneChain = RetargeterAsset->GetTargetIKRig()->GetRetargetChainByName(ChainMap->TargetChain);
 		if (!TargetBoneChain)
 		{
-			UE_LOG(LogTemp, Error, TEXT("IK Retargeter missing target bone chain: %s. Please update the mapping."), *ChainMap->TargetChain.ToString());
-			continue;
+			RetargeterAsset->Log.LogWarning( FText::Format(
+			LOCTEXT("MissingTargetChain", "IK Retargeter missing target bone chain: {0}. Please update the mapping."),
+			FText::FromString(ChainMap->TargetChain.ToString())));
 		}
 		
 		// user opted to not map this to anything, we don't need to spam a warning about it
@@ -1262,20 +1301,22 @@ bool UIKRetargetProcessor::InitializeBoneChainPairs()
 		const FBoneChain* SourceBoneChain = RetargeterAsset->GetSourceIKRig()->GetRetargetChainByName(ChainMap->SourceChain);
 		if (!SourceBoneChain)
 		{
-			UE_LOG(LogTemp, Error, TEXT("IK Retargeter missing source bone chain: %s"), *ChainMap->SourceChain.ToString());
+			RetargeterAsset->Log.LogWarning( FText::Format(
+			LOCTEXT("MissingSourceChain", "IK Retargeter missing source bone chain: {0}"),
+			FText::FromString(ChainMap->SourceChain.ToString())));
 			continue;
 		}
 
 		// all chains are loaded as FK (giving IK better starting pose)
 		FRetargetChainPairFK ChainPair;
-		if (ChainPair.Initialize(ChainMap, *SourceBoneChain, *TargetBoneChain, SourceSkeleton, TargetSkeleton))
+		if (ChainPair.Initialize(ChainMap, *SourceBoneChain, *TargetBoneChain, SourceSkeleton, TargetSkeleton, RetargeterAsset->Log))
 		{
 			ChainPairsFK.Add(ChainPair);
 		}
 		
 		// load IK chain
 		FRetargetChainPairIK ChainPairIK;
-		if (ChainPairIK.Initialize(ChainMap, *SourceBoneChain, *TargetBoneChain, SourceSkeleton, TargetSkeleton))
+		if (ChainPairIK.Initialize(ChainMap, *SourceBoneChain, *TargetBoneChain, SourceSkeleton, TargetSkeleton, RetargeterAsset->Log))
 		{
 			ChainPairsIK.Add(ChainPairIK);
 		}
@@ -1325,14 +1366,14 @@ bool UIKRetargetProcessor::InitializeBoneChainPairs()
 	return !(ChainPairsIK.IsEmpty() && ChainPairsFK.IsEmpty());
 }
 
-bool UIKRetargetProcessor::InitializeIKRig(UObject* Outer, const FReferenceSkeleton& InRefSkeleton)
+bool UIKRetargetProcessor::InitializeIKRig(UObject* Outer, const USkeletalMesh* InSkeletalMesh)
 {	
 	// initialize IK Rig runtime processor
 	if (!IKRigProcessor)
 	{
 		IKRigProcessor = NewObject<UIKRigProcessor>(Outer);	
 	}
-	IKRigProcessor->Initialize(RetargeterAsset->GetTargetIKRig(), InRefSkeleton);
+	IKRigProcessor->Initialize(RetargeterAsset->GetTargetIKRig(), InSkeletalMesh);
 	if (!IKRigProcessor->IsInitialized())
 	{
 		return false;
@@ -1344,8 +1385,9 @@ bool UIKRetargetProcessor::InitializeIKRig(UObject* Outer, const FReferenceSkele
 		// does the IK rig have the IK goal this bone chain requires?
 		if (!IKRigProcessor->GetGoalContainer().FindGoalByName(ChainPair.IKGoalName))
 		{
-			UE_LOG(LogTemp, Error, TEXT("IK Retargeter has target bone chain, %s that references an IK Goal, %s that is not present in any of the solvers in the IK Rig asset."),
-            *ChainPair.TargetBoneChainName.ToString(), *ChainPair.IKGoalName.ToString());
+			RetargeterAsset->Log.LogError( FText::Format(
+			LOCTEXT("CountNotFindRootBone", "IK Retargeter has target bone chain, {0} that references an IK Goal, {1} that is not present in any of the solvers in the IK Rig asset."),
+			FText::FromName(ChainPair.TargetBoneChainName), FText::FromName(ChainPair.IKGoalName)));
 			return false;
 		}
 	}
@@ -1541,3 +1583,5 @@ void UIKRetargetProcessor::CopyAllSettingsFromAsset()
 	}
 }
 #endif
+
+#undef LOCTEXT_NAMESPACE
\ No newline at end of file
diff --git a/Engine/Plugins/Animation/IKRig/Source/IKRig/Private/Retargeter/IKRetargeter.cpp b/Engine/Plugins/Animation/IKRig/Source/IKRig/Private/Retargeter/IKRetargeter.cpp
--- a/Engine/Plugins/Animation/IKRig/Source/IKRig/Private/Retargeter/IKRetargeter.cpp
+++ b/Engine/Plugins/Animation/IKRig/Source/IKRig/Private/Retargeter/IKRetargeter.cpp
@@ -2,13 +2,12 @@
 
 #include "Retargeter/IKRetargeter.h"
 
-const FName UIKRetargeter::DefaultPoseName = "Default Pose";
+#if WITH_EDITOR
 const FName UIKRetargeter::GetSourceIKRigPropertyName() { return GET_MEMBER_NAME_STRING_CHECKED(UIKRetargeter, SourceIKRigAsset); };
 const FName UIKRetargeter::GetTargetIKRigPropertyName() { return GET_MEMBER_NAME_STRING_CHECKED(UIKRetargeter, TargetIKRigAsset); };
-#if WITH_EDITOR
-const FName UIKRetargeter::GetTargetPreviewMeshPropertyName() { return GET_MEMBER_NAME_STRING_CHECKED(UIKRetargeter, TargetPreviewMesh); };
+const FName UIKRetargeter::GetSourcePreviewMeshPropertyName() { return GET_MEMBER_NAME_STRING_CHECKED(UIKRetargeter, SourcePreviewMesh); };
+const FName UIKRetargeter::GetTargetPreviewMeshPropertyName() { return GET_MEMBER_NAME_STRING_CHECKED(UIKRetargeter, TargetPreviewMesh); }
 #endif
-const FName UIKRetargeter::GetDefaultPoseName() { return DefaultPoseName; }
 
 void UIKRetargeter::PostLoad()
 {
@@ -62,6 +61,15 @@ void UIKRetargeter::PostLoad()
 
 	// remove null settings
 	ChainSettings.Remove(nullptr);
+
+	// ensure current pose is valid
+	if (!RetargetPoses.Contains(CurrentRetargetPose))
+	{
+		CurrentRetargetPose = UIKRetargeter::GetDefaultPoseName();
+	}
+
+	// set the log target
+	Log.SetLogTarget(GetUniqueIDAsName(), false);
 }
 
 void FIKRetargetPose::SetBoneRotationOffset(FName BoneName, FQuat RotationDelta, const FIKRigSkeleton& Skeleton)
@@ -92,3 +100,15 @@ void FIKRetargetPose::SortHierarchically(const FIKRigSkeleton& Skeleton)
 		return Skeleton.GetBoneIndexFromName(A) > Skeleton.GetBoneIndexFromName(B);
 	});
 }
+
+FName UIKRetargeter::GetUniqueIDAsName() const
+{
+	static const FString IKRetargetIDPrefix("IKRetargeter_");
+	return FName(IKRetargetIDPrefix + FString::FromInt(GetUniqueID()));
+}
+
+const FName UIKRetargeter::GetDefaultPoseName()
+{
+	static const FName DefaultPoseName = "Default Pose";
+	return DefaultPoseName;
+};
diff --git a/Engine/Plugins/Animation/IKRig/Source/IKRig/Public/AnimNodes/AnimNode_RetargetPoseFromMesh.h b/Engine/Plugins/Animation/IKRig/Source/IKRig/Public/AnimNodes/AnimNode_RetargetPoseFromMesh.h
--- a/Engine/Plugins/Animation/IKRig/Source/IKRig/Public/AnimNodes/AnimNode_RetargetPoseFromMesh.h
+++ b/Engine/Plugins/Animation/IKRig/Source/IKRig/Public/AnimNodes/AnimNode_RetargetPoseFromMesh.h
@@ -29,6 +29,10 @@ struct IKRIG_API FAnimNode_RetargetPoseFromMesh : public FAnimNode_Base
 	UPROPERTY(EditAnywhere, BlueprintReadWrite, Category = Settings, meta=(PinHiddenByDefault))
 	TObjectPtr<UIKRetargeter> IKRetargeterAsset = nullptr;
 
+	/* Toggle whether to print warnings about missing or incorrectly configured retarget configurations. */
+	UPROPERTY(EditAnywhere, BlueprintReadWrite, Category=Debug, meta = (NeverAsPin))
+	bool bSuppressWarnings = false;
+	
 	/* Copy curves from SouceMeshComponent. This will copy any curves the source/target Skeleton have in common. */
 	UPROPERTY(EditAnywhere, BlueprintReadWrite, Category = Settings, meta = (NeverAsPin))
 	bool bCopyCurves = true;
diff --git a/Engine/Plugins/Animation/IKRig/Source/IKRig/Public/IKRig.h b/Engine/Plugins/Animation/IKRig/Source/IKRig/Public/IKRig.h
--- a/Engine/Plugins/Animation/IKRig/Source/IKRig/Public/IKRig.h
+++ b/Engine/Plugins/Animation/IKRig/Source/IKRig/Public/IKRig.h
@@ -13,5 +13,3 @@ public:
 	virtual void StartupModule() override;
 	virtual void ShutdownModule() override;
 };
-
-DECLARE_LOG_CATEGORY_EXTERN(LogIKRig, Warning, All);
diff --git a/Engine/Plugins/Animation/IKRig/Source/IKRig/Public/IKRigDefinition.h b/Engine/Plugins/Animation/IKRig/Source/IKRig/Public/IKRigDefinition.h
--- a/Engine/Plugins/Animation/IKRig/Source/IKRig/Public/IKRigDefinition.h
+++ b/Engine/Plugins/Animation/IKRig/Source/IKRig/Public/IKRigDefinition.h
@@ -3,6 +3,7 @@
 #pragma once
 
 #include "BoneContainer.h"
+#include "IKRigLogger.h"
 #include "UObject/Object.h"
 #include "Interfaces/Interface_PreviewMeshProvider.h"
 #include "IKRigSkeleton.h"
@@ -181,6 +182,12 @@ class IKRIG_API UIKRigDefinition : public UObject, public IInterface_PreviewMesh
 
 public:
 
+	/** The skeletal mesh to run the IK solve on (loaded into viewport).
+	* NOTE: you can assign ANY Skeletal Mesh to apply the IK Rig to. Compatibility is determined when a new mesh is assigned
+	* by comparing it's hierarchy with the goals, solvers and bone settings required by the rig. See output log for details. */
+	UPROPERTY(AssetRegistrySearchable, EditAnywhere, Category = PreviewMesh)
+	TSoftObjectPtr<USkeletalMesh> PreviewSkeletalMesh;
+
 #if WITH_EDITORONLY_DATA
 
 	/**Draw bones in the viewport.*/
@@ -211,10 +218,7 @@ public:
 
 	virtual void Serialize(FArchive& Ar) override;
 
-	/** The skeletal mesh that was used as the source of the skeleton data. Also used for preview.
-	* NOTE: the IK rig may be played back on ANY skeleton that is compatible with it's hierarchy. */
-	UPROPERTY(AssetRegistrySearchable, VisibleAnywhere, Category = "Imported Skeleton")
-	TObjectPtr<class USkeletalMesh> PreviewSkeletalMesh;
+	virtual void PostLoad() override;
 	
 	/** hierarchy and bone-pose transforms */
 	UPROPERTY()
@@ -239,6 +243,17 @@ public:
 	virtual void SetPreviewMesh(USkeletalMesh* PreviewMesh, bool bMarkAsDirty = true) override;
 	virtual USkeletalMesh* GetPreviewMesh() const override;
 	/** END IInterface_PreviewMeshProvider interface */
+
+	/** Get the unique ID of this UObject as a name. */
+	FName GetUniqueIDAsName() const;
+
+	/** logging system */
+	FIKRigLogger Log;
+
+#if WITH_EDITOR
+	/* Get name of Preview Mesh property */
+	static const FName GetPreviewMeshPropertyName();
+#endif
 	
 private:
 
diff --git a/Engine/Plugins/Animation/IKRig/Source/IKRig/Public/IKRigLogger.h b/Engine/Plugins/Animation/IKRig/Source/IKRig/Public/IKRigLogger.h
new file mode 100644
--- /dev/null
+++ b/Engine/Plugins/Animation/IKRig/Source/IKRig/Public/IKRigLogger.h
@@ -0,0 +1,33 @@
+﻿// Copyright Epic Games, Inc. All Rights Reserved.
+
+#pragma once
+
+#include "Internationalization/Text.h"
+#include "UObject/NameTypes.h"
+
+DECLARE_LOG_CATEGORY_EXTERN(LogIKRig, Warning, All);
+
+struct IKRIG_API FIKRigLogger
+{
+	/** Set the name of the log to output messages to, and whether to suppress warnings */
+	void SetLogTarget(const FName InLogName, bool bInSuppressWarnings);
+	/** Get the name this log is currently outputting to */
+	FName GetLogTarget() const;
+	/** Log a warning message to display to user. */
+	void LogError(const FText& Message) const;
+	/** Log a warning message to display to user. */
+	void LogWarning(const FText& Message) const;
+	/** Log a message to display to editor output log. */
+	void LogEditorMessage(const FText& Message) const;
+
+private:
+	/** the name of the output log this logger will send messages to
+	 *
+	 * For the IK Rig and Retargeting editors, we desire to filter the messages that originate only from the asset
+	 * that is being edited. Therefore we name the log using the unique ID of the UObject itself (valid for lifetime of UObject between loads)
+	 */
+	FName LogName;
+	
+	/** when true, warnings will not be printed out. */
+	bool bWarningsSuppressed = false;
+};
\ No newline at end of file
diff --git a/Engine/Plugins/Animation/IKRig/Source/IKRig/Public/IKRigProcessor.h b/Engine/Plugins/Animation/IKRig/Source/IKRig/Public/IKRigProcessor.h
--- a/Engine/Plugins/Animation/IKRig/Source/IKRig/Public/IKRigProcessor.h
+++ b/Engine/Plugins/Animation/IKRig/Source/IKRig/Public/IKRigProcessor.h
@@ -4,6 +4,7 @@
 
 #include "CoreMinimal.h"
 #include "IKRigDataTypes.h"
+#include "IKRigLogger.h"
 #include "IKRigSkeleton.h"
 
 #include "IKRigProcessor.generated.h"
@@ -42,12 +43,9 @@ public:
 	/** setup a new processor to run the given IKRig asset
 	 *  NOTE!! this function creates new UObjects and consequently MUST be called from the main thread!!
 	 *  @param InRigAsset - the IK Rig defining the collection of solvers to execute and all the rig settings
-	 *  @param InputSkeleton - the skeleton in reference pose that you want to solve the IK on
+	 *  @param SkeletalMesh - the skeletal mesh you want to solve the IK on
 	 */
-	void Initialize(const UIKRigDefinition* InRigAsset, const FIKRigInputSkeleton& InputSkeleton);
-
-	/** Anim Graph: convenience to initialize directly from an FReferenceSkeleton. */
-	void Initialize(const UIKRigDefinition* InRigAsset, const FReferenceSkeleton& RefSkeleton);
+	void Initialize(const UIKRigDefinition* InRigAsset, const USkeletalMesh* SkeletalMesh);
 
 	//
 	// BEGIN UPDATE SEQUENCE FUNCTIONS
diff --git a/Engine/Plugins/Animation/IKRig/Source/IKRig/Public/IKRigSkeleton.h b/Engine/Plugins/Animation/IKRig/Source/IKRig/Public/IKRigSkeleton.h
--- a/Engine/Plugins/Animation/IKRig/Source/IKRig/Public/IKRigSkeleton.h
+++ b/Engine/Plugins/Animation/IKRig/Source/IKRig/Public/IKRigSkeleton.h
@@ -4,10 +4,9 @@
 
 #include "CoreMinimal.h"
 
-#include "ReferenceSkeleton.h"
-
 #include "IKRigSkeleton.generated.h"
 
+class USkeletalMesh;
 class UIKRigDefinition;
 class UIKRigSolver;
 
@@ -42,33 +41,15 @@ struct IKRIG_API FIKRigInputSkeleton
 	TArray<FName> BoneNames;
 	TArray<int32> ParentIndices;
 	TArray<FTransform> LocalRefPose;
+	FName SkeletalMeshName;
 
 	FIKRigInputSkeleton() = default;
 	
-	FIKRigInputSkeleton(const FReferenceSkeleton& RefSkeleton)
-	{
-		Initialize(RefSkeleton);
-	}
-
-	void Initialize(const FReferenceSkeleton& RefSkeleton)
-	{
-		Reset();
-		
-		const TArray<FMeshBoneInfo>& BoneInfo = RefSkeleton.GetRefBoneInfo();
-		for (int32 BoneIndex=0; BoneIndex<BoneInfo.Num(); ++BoneIndex)
-		{
-			BoneNames.Add(BoneInfo[BoneIndex].Name);
-			ParentIndices.Add(BoneInfo[BoneIndex].ParentIndex);
-			LocalRefPose.Add(RefSkeleton.GetRefBonePose()[BoneIndex]);
-		}
-	}
-
-	void Reset()
-	{
-		BoneNames.Reset();
-		ParentIndices.Reset();
-		LocalRefPose.Reset();
-	}
+	FIKRigInputSkeleton(const USkeletalMesh* SkeletalMesh);
+
+	void Initialize(const USkeletalMesh* SkeletalMesh);
+
+	void Reset();
 };
 
 USTRUCT()
@@ -100,7 +81,7 @@ struct IKRIG_API FIKRigSkeleton
 	UPROPERTY(VisibleAnywhere, Category = Skeleton)
 	TArray<FTransform> RefPoseGlobal;
 
-	void SetInputSkeleton(const FReferenceSkeleton& RefSkeleton, const TArray<FName>& InExcludedBones);
+	void SetInputSkeleton(const USkeletalMesh* SkeletalMesh, const TArray<FName>& InExcludedBones);
 
 	void SetInputSkeleton(const FIKRigInputSkeleton& InputSkeleton, const TArray<FName>& InExcludedBones);
 
diff --git a/Engine/Plugins/Animation/IKRig/Source/IKRig/Public/Retargeter/IKRetargetProcessor.h b/Engine/Plugins/Animation/IKRig/Source/IKRig/Public/Retargeter/IKRetargetProcessor.h
--- a/Engine/Plugins/Animation/IKRig/Source/IKRig/Public/Retargeter/IKRetargetProcessor.h
+++ b/Engine/Plugins/Animation/IKRig/Source/IKRig/Public/Retargeter/IKRetargetProcessor.h
@@ -4,6 +4,7 @@
 
 #include "CoreMinimal.h"
 #include "IKRetargeter.h"
+#include "IKRigLogger.h"
 #include "IKRetargetProcessor.generated.h"
 
 class URetargetChainSettings;
@@ -140,9 +141,15 @@ struct FRootRetargeter
 
 	void Reset();
 	
-	bool InitializeSource(const FName SourceRootBoneName, const FRetargetSkeleton& SourceSkeleton);
+	bool InitializeSource(
+		const FName SourceRootBoneName,
+		const FRetargetSkeleton& SourceSkeleton,
+		FIKRigLogger& Log);
 	
-	bool InitializeTarget(const FName TargetRootBoneName, const FTargetSkeleton& TargetSkeleton);
+	bool InitializeTarget(
+		const FName TargetRootBoneName,
+		const FTargetSkeleton& TargetSkeleton,
+		FIKRigLogger& Log);
 
 	void EncodePose(const TArray<FTransform> &SourceGlobalPose);
 	
@@ -206,11 +213,12 @@ struct FChainFK
 	bool Initialize(
 		const FRetargetSkeleton& Skeleton,
 		const TArray<int32>& BoneIndices,
-		const TArray<FTransform> &InitialGlobalPose);
+		const TArray<FTransform> &InitialGlobalPose,
+		FIKRigLogger& Log);
 
 private:
 	
-	bool CalculateBoneParameters();
+	bool CalculateBoneParameters(FIKRigLogger& Log);
 
 protected:
 
@@ -326,9 +334,15 @@ struct FChainRetargeterIK
 	
 	FTargetChainIK Target;
 
-	bool InitializeSource(const TArray<int32>& BoneIndices, const TArray<FTransform> &SourceInitialGlobalPose);
+	bool InitializeSource(
+		const TArray<int32>& BoneIndices,
+		const TArray<FTransform> &SourceInitialGlobalPose,
+		FIKRigLogger& Log);
 	
-	bool InitializeTarget(const TArray<int32>& BoneIndices, const TArray<FTransform> &TargetInitialGlobalPose);
+	bool InitializeTarget(
+		const TArray<int32>& BoneIndices,
+		const TArray<FTransform> &TargetInitialGlobalPose,
+		FIKRigLogger& Log);
 
 	void EncodePose(const TArray<FTransform> &SourceInputGlobalPose);
 	
@@ -357,14 +371,16 @@ struct FRetargetChainPair
 		const FBoneChain& SourceBoneChain,
 		const FBoneChain& TargetBoneChain,
 		const FRetargetSkeleton& SourceSkeleton,
-		const FTargetSkeleton& TargetSkeleton);
+		const FTargetSkeleton& TargetSkeleton,
+		FIKRigLogger& Log);
 
 private:
 
 	bool ValidateBoneChainWithSkeletalMesh(
 		const bool IsSource,
 		const FBoneChain& BoneChain,
-		const FRetargetSkeleton& RetargetSkeleton);
+		const FRetargetSkeleton& RetargetSkeleton,
+		FIKRigLogger& Log);
 };
 
 struct FRetargetChainPairFK : FRetargetChainPair
@@ -378,7 +394,8 @@ struct FRetargetChainPairFK : FRetargetChainPair
         const FBoneChain& SourceBoneChain,
         const FBoneChain& TargetBoneChain,
         const FRetargetSkeleton& SourceSkeleton,
-        const FTargetSkeleton& TargetSkeleton) override;
+        const FTargetSkeleton& TargetSkeleton,
+        FIKRigLogger& Log) override;
 };
 
 struct FRetargetChainPairIK : FRetargetChainPair
@@ -394,7 +411,8 @@ struct FRetargetChainPairIK : FRetargetChainPair
         const FBoneChain& SourceBoneChain,
         const FBoneChain& TargetBoneChain,
         const FRetargetSkeleton& SourceSkeleton,
-        const FTargetSkeleton& TargetSkeleton) override;
+        const FTargetSkeleton& TargetSkeleton,
+        FIKRigLogger& Log) override;
 };
 
 /** The runtime processor that converts an input pose from a source skeleton into an output pose on a target skeleton.
@@ -415,12 +433,14 @@ public:
 	* @param SourceSkeleton - the skeletal mesh to poses FROM
 	* @param TargetSkeleton - the skeletal mesh to poses TO
 	* @param InRetargeterAsset - the source asset to use for retargeting settings
+	* @param bSuppressWarnings - if true, will not output warnings during initialization
 	* @warning - Initialization does a lot of validation and can fail for many reasons. Check bIsLoadedAndValid afterwards.
 	*/
 	void Initialize(
 		USkeletalMesh *SourceSkeleton,
 		USkeletalMesh *TargetSkeleton,
-		UIKRetargeter* InRetargeterAsset);
+		UIKRetargeter* InRetargeterAsset,
+		const bool bSuppressWarnings=false);
 
 	/**
 	* Run the retarget to generate a new pose.
@@ -510,7 +530,7 @@ private:
 	bool InitializeBoneChainPairs();
 
 	/** Initializes the IK Rig that evaluates the IK solve for the target IK chains */
-	bool InitializeIKRig(UObject* Outer, const FReferenceSkeleton& InRefSkeleton);
+	bool InitializeIKRig(UObject* Outer, const USkeletalMesh* InSkeletalMesh);
 	
 	/** Internal retarget phase for the root. */
 	void RunRootRetarget(const TArray<FTransform>& InGlobalTransforms, TArray<FTransform>& OutGlobalTransforms);
diff --git a/Engine/Plugins/Animation/IKRig/Source/IKRig/Public/Retargeter/IKRetargeter.h b/Engine/Plugins/Animation/IKRig/Source/IKRig/Public/Retargeter/IKRetargeter.h
--- a/Engine/Plugins/Animation/IKRig/Source/IKRig/Public/Retargeter/IKRetargeter.h
+++ b/Engine/Plugins/Animation/IKRig/Source/IKRig/Public/Retargeter/IKRetargeter.h
@@ -203,17 +203,21 @@ public:
 	const TArray<TObjectPtr<URetargetChainSettings>> GetAllChainSettings() const { return ChainSettings; };
 	/** Get read-only access to a retarget pose */
 	const FIKRetargetPose* GetCurrentRetargetPose() const { return &RetargetPoses[CurrentRetargetPose]; };
+	/** Get the unique ID of this UObject as a name. */
+	FName GetUniqueIDAsName() const;
+	/* Get name of default pose */
+	static const FName GetDefaultPoseName();
 
+#if WITH_EDITOR
 	/* Get name of Source IK Rig property */
 	static const FName GetSourceIKRigPropertyName();
 	/* Get name of Target IK Rig property */
 	static const FName GetTargetIKRigPropertyName();
-#if WITH_EDITOR
+	/* Get name of Source Preview Mesh property */
+	static const FName GetSourcePreviewMeshPropertyName();
 	/* Get name of Target Preview Mesh property */
 	static const FName GetTargetPreviewMeshPropertyName();
 #endif
-	/* Get name of default pose */
-	static const FName GetDefaultPoseName();
 
 #if WITH_EDITOR
 	bool IsInEditRetargetPoseMode() const { return bEditRetargetPoseMode; }
@@ -224,20 +228,26 @@ public:
 private:
 
 	/** The rig to copy animation FROM.*/
-	UPROPERTY(VisibleAnywhere, Category = Rigs)
-	TObjectPtr<UIKRigDefinition> SourceIKRigAsset = nullptr;
+	UPROPERTY(EditAnywhere, Category = Source)
+	TSoftObjectPtr<UIKRigDefinition> SourceIKRigAsset = nullptr;
+
+#if WITH_EDITORONLY_DATA
+	/** Optional. Override the Skeletal Mesh to copy animation from. Uses the preview mesh from the Source IK Rig asset by default. */
+	UPROPERTY(EditAnywhere, Category = Source)
+	TSoftObjectPtr<USkeletalMesh> SourcePreviewMesh = nullptr;
+#endif
 	
 	/** The rig to copy animation TO.*/
-	UPROPERTY(EditAnywhere, Category = Rigs)
-	TObjectPtr<UIKRigDefinition> TargetIKRigAsset = nullptr;
-
-public:
+	UPROPERTY(EditAnywhere, Category = Target)
+	TSoftObjectPtr<UIKRigDefinition> TargetIKRigAsset = nullptr;
 
 #if WITH_EDITORONLY_DATA
 	/** Optional. Override the Skeletal Mesh to preview the retarget on. Uses the preview mesh from the Target IK Rig asset by default. */
-	UPROPERTY(EditAnywhere, Category = Rigs)
-	TObjectPtr<USkeletalMesh> TargetPreviewMesh = nullptr;
+	UPROPERTY(EditAnywhere, Category = Target)
+	TSoftObjectPtr<USkeletalMesh> TargetPreviewMesh = nullptr;
 #endif
+
+public:
 	
 	/** When false, translational motion of skeleton root is not copied. Useful for debugging.*/
 	UPROPERTY(EditAnywhere, Category = RetargetPhases)
@@ -251,6 +261,9 @@ public:
 	UPROPERTY(EditAnywhere, Category = RetargetPhases)
 	bool bRetargetIK = true;
 
+	/** log warnings and errors */
+	FIKRigLogger Log;
+
 #if WITH_EDITORONLY_DATA
 	UPROPERTY()
 	float TargetActorOffset_DEPRECATED = 0.0f;
@@ -301,8 +314,7 @@ private:
 	
 	/** The set of retarget poses available as options for retargeting.*/
 	UPROPERTY()
-	FName CurrentRetargetPose = DefaultPoseName;
-	static const FName DefaultPoseName;
+	FName CurrentRetargetPose;
 
 	friend class UIKRetargeterController;
 };
diff --git a/Engine/Plugins/Animation/IKRig/Source/IKRigDeveloper/Private/AnimGraphNode_RetargetPoseFromMesh.cpp b/Engine/Plugins/Animation/IKRig/Source/IKRigDeveloper/Private/AnimGraphNode_RetargetPoseFromMesh.cpp
--- a/Engine/Plugins/Animation/IKRig/Source/IKRigDeveloper/Private/AnimGraphNode_RetargetPoseFromMesh.cpp
+++ b/Engine/Plugins/Animation/IKRig/Source/IKRigDeveloper/Private/AnimGraphNode_RetargetPoseFromMesh.cpp
@@ -49,7 +49,7 @@ void UAnimGraphNode_RetargetPoseFromMesh::ValidateAnimNodeDuringCompilation(USke
 	{
 		if (!IsPinExposedAndLinked(GET_MEMBER_NAME_STRING_CHECKED(FAnimNode_RetargetPoseFromMesh, SourceMeshComponent)))
 		{
-			MessageLog.Warning(TEXT("@@ is missing a Source Skeletal Mesh Component reference."), this);
+			MessageLog.Error(TEXT("@@ is missing a Source Skeletal Mesh Component reference."), this);
 			return;
 		}
 	}
@@ -92,7 +92,7 @@ void UAnimGraphNode_RetargetPoseFromMesh::ValidateAnimNodeDuringCompilation(USke
 		return;
 	}
 
-	if (ForSkeleton)
+	if (ForSkeleton && !Node.bSuppressWarnings)
 	{
 		// validate that target bone chains exist on this skeleton
 		const FReferenceSkeleton &RefSkel = ForSkeleton->GetReferenceSkeleton();
diff --git a/Engine/Plugins/Animation/IKRig/Source/IKRigEditor/IKRigEditor.Build.cs b/Engine/Plugins/Animation/IKRig/Source/IKRigEditor/IKRigEditor.Build.cs
--- a/Engine/Plugins/Animation/IKRig/Source/IKRigEditor/IKRigEditor.Build.cs
+++ b/Engine/Plugins/Animation/IKRig/Source/IKRigEditor/IKRigEditor.Build.cs
@@ -36,6 +36,7 @@ namespace UnrealBuildTool.Rules
                     "SkeletonEditor",
                     "ContentBrowser",
                     "AnimationBlueprintLibrary",
+                    "MessageLog",
                     
 					"PropertyEditor",
 					"BlueprintGraph",
@@ -48,7 +49,7 @@ namespace UnrealBuildTool.Rules
 					"AnimationCore",
                     "AnimationWidgets",
                     "ApplicationCore",
-				}
+                }
             );
 
             PrivateIncludePathModuleNames.AddRange(
diff --git a/Engine/Plugins/Animation/IKRig/Source/IKRigEditor/Private/IKRigEditor.cpp b/Engine/Plugins/Animation/IKRig/Source/IKRigEditor/Private/IKRigEditor.cpp
--- a/Engine/Plugins/Animation/IKRig/Source/IKRigEditor/Private/IKRigEditor.cpp
+++ b/Engine/Plugins/Animation/IKRig/Source/IKRigEditor/Private/IKRigEditor.cpp
@@ -20,9 +20,9 @@
 #include "RigEditor/IKRigDetailCustomizations.h"
 #include "RigEditor/IKRigEditorController.h"
 
-IMPLEMENT_MODULE(FIKRigEditor, IKRigEditor)
+DEFINE_LOG_CATEGORY(LogIKRigEditor);
 
-DEFINE_LOG_CATEGORY_STATIC(LogIKRigEditor, Log, All);
+IMPLEMENT_MODULE(FIKRigEditor, IKRigEditor)
 
 #define LOCTEXT_NAMESPACE "IKRigEditor"
 
diff --git a/Engine/Plugins/Animation/IKRig/Source/IKRigEditor/Private/RetargetEditor/IKRetargetApplicationMode.cpp b/Engine/Plugins/Animation/IKRig/Source/IKRigEditor/Private/RetargetEditor/IKRetargetApplicationMode.cpp
--- a/Engine/Plugins/Animation/IKRig/Source/IKRigEditor/Private/RetargetEditor/IKRetargetApplicationMode.cpp
+++ b/Engine/Plugins/Animation/IKRig/Source/IKRigEditor/Private/RetargetEditor/IKRetargetApplicationMode.cpp
@@ -10,6 +10,7 @@
 
 #include "RetargetEditor/IKRetargetChainTabSummoner.h"
 #include "RetargetEditor/IKRetargetEditor.h"
+#include "RetargetEditor/IKRetargetOutputLogTabSummoner.h"
 
 #define LOCTEXT_NAMESPACE "IKRetargetMode"
 
@@ -36,9 +37,10 @@ FIKRetargetApplicationMode::FIKRetargetApplicationMode(
 	// register custom tabs
 	TabFactories.RegisterFactory(MakeShared<FIKRetargetChainTabSummoner>(IKRetargetEditor));
 	TabFactories.RegisterFactory(MakeShared<FIKRetargetAssetBrowserTabSummoner>(IKRetargetEditor));
+	TabFactories.RegisterFactory(MakeShared<FIKRetargetOutputLogTabSummoner>(IKRetargetEditor));
 
 	// create tab layout
-	TabLayout = FTabManager::NewLayout("Standalone_IKRetargetEditor_Layout_v1.008")
+	TabLayout = FTabManager::NewLayout("Standalone_IKRetargetEditor_Layout_v1.012")
 		->AddArea
 		(
 			FTabManager::NewPrimaryArea()
@@ -46,31 +48,43 @@ FIKRetargetApplicationMode::FIKRetargetApplicationMode(
 			->Split
 			(
 				FTabManager::NewSplitter()
-				->SetSizeCoefficient(0.9f)
+				->SetSizeCoefficient(0.8f)
 				->SetOrientation(Orient_Horizontal)
 				->Split
 				(
+					FTabManager::NewSplitter()
+					->SetSizeCoefficient(0.9f)
+					->SetOrientation(Orient_Vertical)
+					->Split
+					(
+						FTabManager::NewStack()
+						->SetSizeCoefficient(0.8f)
+						->SetHideTabWell(true)
+						->AddTab(FPersonaTabs::PreviewViewportID, ETabState::OpenedTab)
+					)
+					->Split
+					(
 					FTabManager::NewStack()
-					->SetSizeCoefficient(0.6f)
-					->SetHideTabWell(true)
-					->AddTab(FPersonaTabs::PreviewViewportID, ETabState::OpenedTab)
+						->SetSizeCoefficient(0.2f)
+						->AddTab(FIKRetargetOutputLogTabSummoner::TabID, ETabState::OpenedTab)
+					)
 				)
 				->Split
 				(
 					FTabManager::NewSplitter()
-					->SetSizeCoefficient(0.9f)
+					->SetSizeCoefficient(0.2f)
 					->SetOrientation(Orient_Vertical)
 					->Split
 					(
 					FTabManager::NewStack()
-						->SetSizeCoefficient(0.6f)
+						->SetSizeCoefficient(0.7f)
 						->AddTab(FPersonaTabs::DetailsID, ETabState::OpenedTab)
 						->SetForegroundTab(FPersonaTabs::DetailsID)
 					)
 					->Split
 					(
 					FTabManager::NewStack()
-						->SetSizeCoefficient(0.6f)
+						->SetSizeCoefficient(0.3f)
 						->AddTab(FIKRetargetChainTabSummoner::TabID, ETabState::OpenedTab)
 						->AddTab(FIKRetargetAssetBrowserTabSummoner::TabID, ETabState::OpenedTab)
 						->SetForegroundTab(FIKRetargetChainTabSummoner::TabID)
diff --git a/Engine/Plugins/Animation/IKRig/Source/IKRigEditor/Private/RetargetEditor/IKRetargetDetails.cpp b/Engine/Plugins/Animation/IKRig/Source/IKRigEditor/Private/RetargetEditor/IKRetargetDetails.cpp
--- a/Engine/Plugins/Animation/IKRig/Source/IKRigEditor/Private/RetargetEditor/IKRetargetDetails.cpp
+++ b/Engine/Plugins/Animation/IKRig/Source/IKRigEditor/Private/RetargetEditor/IKRetargetDetails.cpp
@@ -65,6 +65,10 @@ void FIKRetargeterDetails::CustomizeDetails(IDetailLayoutBuilder& DetailBuilder)
 		]
 	];
 
+	// add the bone size slider
+	TSharedRef<IPropertyHandle> BoneSizeHandle = DetailBuilder.GetProperty(GET_MEMBER_NAME_CHECKED(UIKRetargeter, BoneDrawSize));
+	EditPoseCategoryBuilder.AddProperty(BoneSizeHandle);
+
 	// add pose editing toolbar
 	const TSharedRef<FUICommandList>& Commands = Controller->GetEditorController()->Editor.Pin()->GetToolkitCommands();
 	FDetailWidgetRow& ToolbarRow = EditPoseCategoryBuilder.AddCustomRow(LOCTEXT("CurrentPoseLabel", "Edit Pose"))
diff --git a/Engine/Plugins/Animation/IKRig/Source/IKRigEditor/Private/RetargetEditor/IKRetargetEditor.cpp b/Engine/Plugins/Animation/IKRig/Source/IKRigEditor/Private/RetargetEditor/IKRetargetEditor.cpp
--- a/Engine/Plugins/Animation/IKRig/Source/IKRigEditor/Private/RetargetEditor/IKRetargetEditor.cpp
+++ b/Engine/Plugins/Animation/IKRig/Source/IKRigEditor/Private/RetargetEditor/IKRetargetEditor.cpp
@@ -223,6 +223,8 @@ TStatId FIKRetargetEditor::GetStatId() const
 
 void FIKRetargetEditor::PostUndo(bool bSuccess)
 {
+	EditorController->ClearOutputLog();
+	
 	const bool WasEditing = EditorController->IsEditingPose();
 	
 	EditorController->AssetController->BroadcastNeedsReinitialized();
@@ -238,6 +240,8 @@ void FIKRetargetEditor::PostUndo(bool bSuccess)
 
 void FIKRetargetEditor::PostRedo(bool bSuccess)
 {
+	EditorController->ClearOutputLog();
+	
 	const bool WasEditing = EditorController->IsEditingPose();
 	
 	EditorController->AssetController->BroadcastNeedsReinitialized();
@@ -307,18 +311,24 @@ void FIKRetargetEditor::HandleDetailsCreated(const TSharedRef<class IDetailsView
 
 void FIKRetargetEditor::OnFinishedChangingDetails(const FPropertyChangedEvent& PropertyChangedEvent)
 {
-	const bool bTargetChanged = PropertyChangedEvent.GetPropertyName() == UIKRetargeter::GetTargetIKRigPropertyName();
-	const bool bPreviewChanged = PropertyChangedEvent.GetPropertyName() == UIKRetargeter::GetTargetPreviewMeshPropertyName();
+	const bool bSourceIKRigChanged = PropertyChangedEvent.GetPropertyName() == UIKRetargeter::GetSourceIKRigPropertyName();
+	const bool bTargetIKRigChanged = PropertyChangedEvent.GetPropertyName() == UIKRetargeter::GetTargetIKRigPropertyName();
+	const bool bSourcePreviewChanged = PropertyChangedEvent.GetPropertyName() == UIKRetargeter::GetSourcePreviewMeshPropertyName();
+	const bool bTargetPreviewChanged = PropertyChangedEvent.GetPropertyName() == UIKRetargeter::GetTargetPreviewMeshPropertyName();
 
-	if (bTargetChanged)
+	if (bTargetIKRigChanged || bSourceIKRigChanged)
 	{
+		EditorController->ClearOutputLog();
 		EditorController->BindToIKRigAsset(EditorController->AssetController->GetAsset()->GetTargetIKRigWriteable());
+		EditorController->BindToIKRigAsset(EditorController->AssetController->GetAsset()->GetSourceIKRigWriteable());
 		EditorController->AssetController->CleanChainMapping();
 		EditorController->AssetController->AutoMapChains();
 	}
 	
-	if (bTargetChanged || bPreviewChanged)
+	if (bTargetIKRigChanged || bSourceIKRigChanged || bTargetPreviewChanged || bSourcePreviewChanged)
 	{
+		EditorController->ClearOutputLog();
+		
 		// set the source and target skeletal meshes on the component
 		// NOTE: this must be done AFTER setting the AnimInstance so that the correct root anim node is loaded
 		USkeletalMesh* SourceMesh = EditorController->GetSourceSkeletalMesh();
diff --git a/Engine/Plugins/Animation/IKRig/Source/IKRigEditor/Private/RetargetEditor/IKRetargetEditorController.cpp b/Engine/Plugins/Animation/IKRig/Source/IKRigEditor/Private/RetargetEditor/IKRetargetEditorController.cpp
--- a/Engine/Plugins/Animation/IKRig/Source/IKRigEditor/Private/RetargetEditor/IKRetargetEditorController.cpp
+++ b/Engine/Plugins/Animation/IKRig/Source/IKRigEditor/Private/RetargetEditor/IKRetargetEditorController.cpp
@@ -11,6 +11,7 @@
 #include "RetargetEditor/IKRetargetEditor.h"
 #include "RetargetEditor/SIKRetargetChainMapList.h"
 #include "Retargeter/IKRetargeter.h"
+#include "RigEditor/SIKRigOutputLog.h"
 #include "RigEditor/IKRigController.h"
 
 #define LOCTEXT_NAMESPACE "IKRetargetEditorController"
@@ -59,6 +60,7 @@ void FIKRetargetEditorController::OnIKRigNeedsInitialized(UIKRigDefinition* Modi
 	}
 
 	// the target anim instance has the IK RetargetPoseFromMesh node which needs reinitialized with new asset version
+	ClearOutputLog();
 	TargetAnimInstance->SetProcessorNeedsInitialized();
 	RefreshAllViews();
 }
@@ -79,6 +81,8 @@ void FIKRetargetEditorController::OnRetargetChainRemoved(UIKRigDefinition* Modif
 
 void FIKRetargetEditorController::OnRetargeterNeedsInitialized(const UIKRetargeter* Retargeter) const
 {
+	// clear the output log
+	ClearOutputLog();
 	// force edit pose mode to be off
 	Editor.Pin()->GetEditorModeManager().DeactivateMode(FIKRetargetEditPoseMode::ModeName);
 	// force reinit the runtime retarget processor
@@ -116,22 +120,12 @@ void FIKRetargetEditorController::AddOffsetAndUpdatePreviewMeshPosition(
 
 USkeletalMesh* FIKRetargetEditorController::GetSourceSkeletalMesh() const
 {
-	if (!(AssetController && AssetController->GetAsset()->GetSourceIKRig()))
-	{
-		return nullptr;
-	}
-
-	return AssetController->GetAsset()->GetSourceIKRig()->PreviewSkeletalMesh.Get();
+	return AssetController ? AssetController->GetSourcePreviewMesh() : nullptr;
 }
 
 USkeletalMesh* FIKRetargetEditorController::GetTargetSkeletalMesh() const
 {
-	if (!(AssetController && AssetController->GetAsset()->GetTargetIKRig()))
-	{
-		return nullptr;
-	}
-
-	return AssetController->GetTargetPreviewMesh();
+	return AssetController ? AssetController->GetTargetPreviewMesh() : nullptr;
 }
 
 FTransform FIKRetargetEditorController::GetTargetBoneGlobalTransform(
@@ -206,6 +200,14 @@ const UIKRetargetProcessor* FIKRetargetEditorController::GetRetargetProcessor()
 	return nullptr;	
 }
 
+void FIKRetargetEditorController::ClearOutputLog() const
+{
+	if (OutputLogView.IsValid())
+	{
+		OutputLogView.Get()->ClearLog();
+	}
+}
+
 void FIKRetargetEditorController::RefreshAllViews() const
 {
 	Editor.Pin()->RegenerateMenusAndToolbars();
@@ -216,6 +218,12 @@ void FIKRetargetEditorController::RefreshAllViews() const
 	{
 		ChainsView.Get()->RefreshView();
 	}
+
+	// refresh the asset browser to ensure it shows compatible sequences
+	if (AssetBrowserView.IsValid())
+	{
+		AssetBrowserView.Get()->RefreshView();
+	}
 }
 
 void FIKRetargetEditorController::PlayAnimationAsset(UAnimationAsset* AssetToPlay)
diff --git a/Engine/Plugins/Animation/IKRig/Source/IKRigEditor/Private/RetargetEditor/IKRetargetFactory.cpp b/Engine/Plugins/Animation/IKRig/Source/IKRigEditor/Private/RetargetEditor/IKRetargetFactory.cpp
--- a/Engine/Plugins/Animation/IKRig/Source/IKRigEditor/Private/RetargetEditor/IKRetargetFactory.cpp
+++ b/Engine/Plugins/Animation/IKRig/Source/IKRigEditor/Private/RetargetEditor/IKRetargetFactory.cpp
@@ -2,6 +2,7 @@
 
 #include "RetargetEditor/IKRetargetFactory.h"
 #include "IKRig/Public/Retargeter/IKRetargeter.h"
+#include "IKRigEditor.h"
 #include "AssetTypeCategories.h"
 #include "ContentBrowserModule.h"
 #include "IContentBrowserSingleton.h"
@@ -10,6 +11,7 @@
 #include "RetargetEditor/IKRetargeterController.h"
 #include "Widgets/SWindow.h"
 #include "Widgets/Layout/SBorder.h"
+#include "Engine/SkeletalMesh.h"
 
 #define LOCTEXT_NAMESPACE "IKRetargeterFactory"
 
@@ -31,13 +33,13 @@ UObject* UIKRetargetFactory::FactoryCreateNew(
 {
 	if (!SourceIKRig.IsValid())
 	{
-		UE_LOG(LogTemp, Warning, TEXT("Unable to create IK Retargter. No source IK Rig asset supplied."));
+		UE_LOG(LogIKRigEditor, Warning, TEXT("Unable to create IK Retargter. No source IK Rig asset supplied."));
 		return nullptr;
 	}
-
-	if (!SourceIKRig->GetPreviewMesh())
+	
+	if (IsValid(SourceIKRig->GetPreviewMesh()))
 	{
-		UE_LOG(LogTemp, Warning, TEXT("Unable to create IK Retargter. Source IK Rig does not have an associated skeletal mesh."));
+		UE_LOG(LogIKRigEditor, Warning, TEXT("Unable to create IK Retargter. Source IK Rig does not have an associated skeletal mesh."));
 		return nullptr;
 	}
 	
diff --git a/Engine/Plugins/Animation/IKRig/Source/IKRigEditor/Private/RetargetEditor/IKRetargetOutputLogTabSummoner.cpp b/Engine/Plugins/Animation/IKRig/Source/IKRigEditor/Private/RetargetEditor/IKRetargetOutputLogTabSummoner.cpp
new file mode 100644
--- /dev/null
+++ b/Engine/Plugins/Animation/IKRig/Source/IKRigEditor/Private/RetargetEditor/IKRetargetOutputLogTabSummoner.cpp
@@ -0,0 +1,41 @@
+﻿// Copyright Epic Games, Inc. All Rights Reserved.
+#include "RetargetEditor/IKRetargetOutputLogTabSummoner.h"
+
+#include "IDocumentation.h"
+#include "RigEditor/IKRigEditorStyle.h"
+#include "RetargetEditor/IKRetargetEditor.h"
+#include "RigEditor/SIKRigOutputLog.h"
+#include "Retargeter/IKRetargetProcessor.h"
+
+#define LOCTEXT_NAMESPACE "IKRetargetOutputLogTabSummoner"
+
+const FName FIKRetargetOutputLogTabSummoner::TabID(TEXT("RetargetOutputLog"));
+
+FIKRetargetOutputLogTabSummoner::FIKRetargetOutputLogTabSummoner(
+	const TSharedRef<FIKRetargetEditor>& InRetargetEditor)
+	: FWorkflowTabFactory(TabID, InRetargetEditor),
+	IKRetargetEditor(InRetargetEditor)
+{
+	bIsSingleton = true; // only allow a single instance of this tab
+	
+	TabLabel = LOCTEXT("IKRetargetOutputLogTabLabel", "Retarget Output Log");
+	TabIcon = FSlateIcon(FIKRigEditorStyle::Get().GetStyleSetName(), "IKRig.TabIcon");
+
+	ViewMenuDescription = LOCTEXT("IKRetargetOutputLog_ViewMenu_Desc", "Retarget Output Log");
+	ViewMenuTooltip = LOCTEXT("IKRetargetOutputLog_ViewMenu_ToolTip", "Show the Retargeting Output Log Tab");
+}
+
+TSharedPtr<SToolTip> FIKRetargetOutputLogTabSummoner::CreateTabToolTipWidget(const FWorkflowTabSpawnInfo& Info) const
+{
+	return  IDocumentation::Get()->CreateToolTip(LOCTEXT("IKRetargetOutputLogTooltip", "View warnings and errors while retargeting."), NULL, TEXT("Shared/Editors/Persona"), TEXT("IKRetargetOutputLog_Window"));
+}
+
+TSharedRef<SWidget> FIKRetargetOutputLogTabSummoner::CreateTabBody(const FWorkflowTabSpawnInfo& Info) const
+{
+	TSharedPtr<SIKRigOutputLog>& OutputLogView = IKRetargetEditor.Pin()->GetController()->OutputLogView;
+	FName LogName = IKRetargetEditor.Pin()->GetController()->AssetController->GetAsset()->Log.GetLogTarget();
+	
+	return SNew(SIKRigOutputLog, LogName, OutputLogView);
+}
+
+#undef LOCTEXT_NAMESPACE 
diff --git a/Engine/Plugins/Animation/IKRig/Source/IKRigEditor/Private/RetargetEditor/IKRetargeterController.cpp b/Engine/Plugins/Animation/IKRig/Source/IKRigEditor/Private/RetargetEditor/IKRetargeterController.cpp
--- a/Engine/Plugins/Animation/IKRig/Source/IKRigEditor/Private/RetargetEditor/IKRetargeterController.cpp
+++ b/Engine/Plugins/Animation/IKRig/Source/IKRigEditor/Private/RetargetEditor/IKRetargeterController.cpp
@@ -4,6 +4,7 @@
 
 #include "ScopedTransaction.h"
 #include "Algo/LevenshteinDistance.h"
+#include "Engine/AssetManager.h"
 #include "Engine/SkeletalMesh.h"
 #include "RetargetEditor/IKRetargetEditorController.h"
 #include "Retargeter/IKRetargeter.h"
@@ -39,9 +40,14 @@ UIKRetargeter* UIKRetargeterController::GetAsset() const
 	return Asset;
 }
 
-FIKRetargetEditorController* UIKRetargeterController::GetEditorController() const
+FName UIKRetargeterController::GetAssetIDAsName() const 
 {
-	return EditorController;
+	if (!Asset)
+	{
+		return NAME_None;
+	}
+
+	return Asset->GetUniqueIDAsName();
 }
 
 void UIKRetargeterController::SetSourceIKRig(UIKRigDefinition* SourceIKRig)
@@ -55,7 +61,25 @@ void UIKRetargeterController::SetTargetIKRig(UIKRigDefinition* TargetIKRig)
 	AutoMapChains();
 }
 
-USkeletalMesh* UIKRetargeterController::GetTargetPreviewMesh()
+USkeletalMesh* UIKRetargeterController::GetSourcePreviewMesh() const
+{
+	// can't preview anything if target IK Rig is null
+	if (!Asset->GetSourceIKRig())
+	{
+		return nullptr;
+	}
+
+	// optionally prefer override if one is provided
+	if (!Asset->SourcePreviewMesh.IsNull())
+	{
+		return Asset->SourcePreviewMesh.LoadSynchronous();
+	}
+
+	// fallback to preview mesh from IK Rig asset
+	return Asset->GetSourceIKRig()->GetPreviewMesh();
+}
+
+USkeletalMesh* UIKRetargeterController::GetTargetPreviewMesh() const
 {
 	// can't preview anything if target IK Rig is null
 	if (!Asset->GetTargetIKRig())
@@ -64,28 +88,28 @@ USkeletalMesh* UIKRetargeterController::GetTargetPreviewMesh()
 	}
 
 	// optionally prefer override if one is provided
-	if (IsValid(Asset->TargetPreviewMesh))
+	if (!Asset->TargetPreviewMesh.IsNull())
 	{
-		return Asset->TargetPreviewMesh.Get();
+		return Asset->TargetPreviewMesh.LoadSynchronous();
 	}
 
 	// fallback to preview mesh from IK Rig asset
-	return Asset->GetTargetIKRig()->PreviewSkeletalMesh.Get();
+	return Asset->GetTargetIKRig()->GetPreviewMesh();
 }
 
 FName UIKRetargeterController::GetSourceRootBone() const
 {
-	return IsValid(Asset->SourceIKRigAsset) ? Asset->SourceIKRigAsset->GetRetargetRoot() : FName("None");
+	return Asset->SourceIKRigAsset.IsValid() ? Asset->SourceIKRigAsset->GetRetargetRoot() : FName("None");
 }
 
 FName UIKRetargeterController::GetTargetRootBone() const
 {
-	return IsValid(Asset->TargetIKRigAsset) ? Asset->TargetIKRigAsset->GetRetargetRoot() : FName("None");
+	return Asset->TargetIKRigAsset.IsValid() ? Asset->TargetIKRigAsset->GetRetargetRoot() : FName("None");
 }
 
 void UIKRetargeterController::GetTargetChainNames(TArray<FName>& OutNames) const
 {
-	if (IsValid(Asset->TargetIKRigAsset))
+	if (Asset->TargetIKRigAsset.IsValid())
 	{
 		const TArray<FBoneChain>& Chains = Asset->TargetIKRigAsset->GetRetargetChains();
 		for (const FBoneChain& Chain : Chains)
@@ -97,7 +121,7 @@ void UIKRetargeterController::GetTargetChainNames(TArray<FName>& OutNames) const
 
 void UIKRetargeterController::GetSourceChainNames(TArray<FName>& OutNames) const
 {
-	if (IsValid(Asset->SourceIKRigAsset))
+	if (Asset->SourceIKRigAsset.IsValid())
 	{
 		const TArray<FBoneChain>& Chains = Asset->SourceIKRigAsset->GetRetargetChains();
 		for (const FBoneChain& Chain : Chains)
@@ -109,7 +133,7 @@ void UIKRetargeterController::GetSourceChainNames(TArray<FName>& OutNames) const
 
 void UIKRetargeterController::CleanChainMapping(const bool bForceReinitialization)
 {
-	if (!IsValid(Asset->TargetIKRigAsset))
+	if (!Asset->TargetIKRigAsset.IsValid())
 	{
 		// empty the chain mapping
 		Asset->ChainSettings.Empty();
@@ -176,20 +200,20 @@ void UIKRetargeterController::CleanChainMapping(const bool bForceReinitializatio
 void UIKRetargeterController::CleanPoseList(const bool bForceReinitialization)
 {
 	// enforce the existence of a default pose
-	const bool HasDefaultPose = Asset->RetargetPoses.Contains(Asset->DefaultPoseName);
+	const bool HasDefaultPose = Asset->RetargetPoses.Contains(UIKRetargeter::GetDefaultPoseName());
 	if (!HasDefaultPose)
 	{
-		Asset->RetargetPoses.Emplace(Asset->DefaultPoseName);
+		Asset->RetargetPoses.Emplace(UIKRetargeter::GetDefaultPoseName());
 	}
 	
 	// use default pose unless set to something else
 	if (Asset->CurrentRetargetPose == NAME_None)
 	{
-		Asset->CurrentRetargetPose = Asset->DefaultPoseName;
+		Asset->CurrentRetargetPose = UIKRetargeter::GetDefaultPoseName();
 	}
 
 	// remove all bone offsets that are no longer part of the target skeleton
-	if (IsValid(Asset->TargetIKRigAsset))
+	if (Asset->TargetIKRigAsset.IsValid())
 	{
 		const TArray<FName> AllowedBoneNames = Asset->TargetIKRigAsset->Skeleton.BoneNames;
 		for (TTuple<FName, FIKRetargetPose>& Pose : Asset->RetargetPoses)
@@ -329,21 +353,6 @@ const TArray<TObjectPtr<URetargetChainSettings>>& UIKRetargeterController::GetCh
 	return Asset->ChainSettings;
 }
 
-USkeleton* UIKRetargeterController::GetSourceSkeletonAsset() const
-{
-	if (!IsValid(Asset->SourceIKRigAsset))
-	{
-		return nullptr;
-	}
-
-	if (!Asset->SourceIKRigAsset->PreviewSkeletalMesh)
-	{
-		return nullptr;
-	}
-
-	return Asset->SourceIKRigAsset->PreviewSkeletalMesh->GetSkeleton();
-}
-
 void UIKRetargeterController::AddRetargetPose(FName NewPoseName) const
 {
 	FScopedTransaction Transaction(LOCTEXT("AddRetargetPose", "Add Retarget Pose"));
@@ -382,7 +391,7 @@ void UIKRetargeterController::RenameCurrentRetargetPose(FName NewPoseName) const
 
 void UIKRetargeterController::RemoveRetargetPose(FName PoseToRemove) const
 {
-	if (PoseToRemove == Asset->DefaultPoseName)
+	if (PoseToRemove == Asset->GetDefaultPoseName())
 	{
 		return; // cannot remove default pose
 	}
@@ -400,7 +409,7 @@ void UIKRetargeterController::RemoveRetargetPose(FName PoseToRemove) const
 	// did we remove the currently used pose?
 	if (Asset->CurrentRetargetPose == PoseToRemove)
 	{
-		Asset->CurrentRetargetPose = UIKRetargeter::DefaultPoseName;
+		Asset->CurrentRetargetPose = UIKRetargeter::GetDefaultPoseName();
 	}
 
 	BroadcastNeedsReinitialized();
diff --git a/Engine/Plugins/Animation/IKRig/Source/IKRigEditor/Private/RetargetEditor/SIKRetargetAssetBrowser.cpp b/Engine/Plugins/Animation/IKRig/Source/IKRigEditor/Private/RetargetEditor/SIKRetargetAssetBrowser.cpp
--- a/Engine/Plugins/Animation/IKRig/Source/IKRigEditor/Private/RetargetEditor/SIKRetargetAssetBrowser.cpp
+++ b/Engine/Plugins/Animation/IKRig/Source/IKRigEditor/Private/RetargetEditor/SIKRetargetAssetBrowser.cpp
@@ -9,6 +9,8 @@
 #include "Animation/AnimMontage.h"
 #include "Animation/AnimSequence.h"
 #include "Animation/PoseAsset.h"
+
+#include "IKRigEditor.h"
 #include "RetargetEditor/IKRetargetBatchOperation.h"
 #include "RetargetEditor/IKRetargetEditorController.h"
 #include "RetargetEditor/SRetargetAnimAssetsWindow.h"
@@ -22,6 +24,7 @@ void SIKRetargetAssetBrowser::Construct(
 	TSharedRef<FIKRetargetEditorController> InEditorController)
 {
 	EditorController = InEditorController;
+	EditorController.Pin()->AssetBrowserView = SharedThis(this);
 	
 	ChildSlot
     [
@@ -45,10 +48,10 @@ void SIKRetargetAssetBrowser::Construct(
 		]
     ];
 
-	AddAssetBrowser();
+	RefreshView();
 }
 
-void SIKRetargetAssetBrowser::AddAssetBrowser()
+void SIKRetargetAssetBrowser::RefreshView()
 {
 	FAssetPickerConfig AssetPickerConfig;
 
@@ -115,7 +118,7 @@ FReply SIKRetargetAssetBrowser::OnExportButtonClicked() const
 	TArray<FAssetData> SelectedAssets = GetCurrentSelectionDelegate.Execute();
 	for (const FAssetData& Asset : SelectedAssets)
 	{
-		UE_LOG(LogTemp, Display, TEXT("Duplicating and Retargeting: %s"), *Asset.GetFullName());
+		UE_LOG(LogIKRigEditor, Display, TEXT("Duplicating and Retargeting: %s"), *Asset.GetFullName());
 
 		BatchContext.AssetsToRetarget.Add(Asset.GetAsset());
 	}
diff --git a/Engine/Plugins/Animation/IKRig/Source/IKRigEditor/Private/RetargetEditor/SIKRetargetChainMapList.cpp b/Engine/Plugins/Animation/IKRig/Source/IKRigEditor/Private/RetargetEditor/SIKRetargetChainMapList.cpp
--- a/Engine/Plugins/Animation/IKRig/Source/IKRigEditor/Private/RetargetEditor/SIKRetargetChainMapList.cpp
+++ b/Engine/Plugins/Animation/IKRig/Source/IKRigEditor/Private/RetargetEditor/SIKRetargetChainMapList.cpp
@@ -342,12 +342,19 @@ EVisibility SIKRetargetChainMapList::IsAutoMapButtonVisible() const
 
 FReply SIKRetargetChainMapList::OnAutoMapButtonClicked() const
 {
+	const TSharedPtr<FIKRetargetEditorController> Controller = EditorController.Pin();
+	if (!Controller.IsValid())
+	{
+		return FReply::Unhandled();
+	}
+	
 	UIKRetargeterController* RetargeterController = GetRetargetController();
 	if (!RetargeterController)
 	{
 		return FReply::Unhandled();
 	}
-	
+
+	Controller->ClearOutputLog();
 	RetargeterController->CleanChainMapping();
 	RetargeterController->AutoMapChains();
 	return FReply::Handled();
diff --git a/Engine/Plugins/Animation/IKRig/Source/IKRigEditor/Private/RetargetEditor/SRetargetAnimAssetsWindow.cpp b/Engine/Plugins/Animation/IKRig/Source/IKRigEditor/Private/RetargetEditor/SRetargetAnimAssetsWindow.cpp
--- a/Engine/Plugins/Animation/IKRig/Source/IKRigEditor/Private/RetargetEditor/SRetargetAnimAssetsWindow.cpp
+++ b/Engine/Plugins/Animation/IKRig/Source/IKRigEditor/Private/RetargetEditor/SRetargetAnimAssetsWindow.cpp
@@ -13,6 +13,7 @@
 #include "PropertyCustomizationHelpers.h"
 #include "Animation/DebugSkelMeshComponent.h"
 #include "Misc/ScopedSlowTask.h"
+#include "RetargetEditor/IKRetargeterController.h"
 #include "Retargeter/IKRetargeter.h"
 #include "Widgets/Input/SEditableTextBox.h"
 #include "Widgets/Notifications/SNotificationList.h"
@@ -260,27 +261,6 @@ void SRetargetAnimAssetsWindow::Construct(const FArguments& InArgs)
 						})
 						.ObjectPath(this, &SRetargetAnimAssetsWindow::GetCurrentSourceMeshPath)
 						.OnObjectChanged(this, &SRetargetAnimAssetsWindow::SourceMeshAssigned)
-						.OnShouldFilterAsset_Lambda([this](const FAssetData& AssetData)
-						{
-							if (!BatchContext.IKRetargetAsset)
-							{
-								return true;
-							}
-							
-							USkeletalMesh* Mesh = Cast<USkeletalMesh>(AssetData.GetAsset());
-							if (!Mesh)
-							{
-								return true;
-							}
-							
-							USkeletalMesh* PreviewMesh = BatchContext.IKRetargetAsset->GetSourceIKRig()->GetPreviewMesh();
-							if (!PreviewMesh)
-							{
-								return true;
-							}
-							
-							return Mesh->GetSkeleton() != PreviewMesh->GetSkeleton();
-						})
 					]
 				]
 
@@ -335,30 +315,6 @@ void SRetargetAnimAssetsWindow::Construct(const FArguments& InArgs)
 						})
 						.ObjectPath(this, &SRetargetAnimAssetsWindow::GetCurrentTargetMeshPath)
 						.OnObjectChanged(this, &SRetargetAnimAssetsWindow::TargetMeshAssigned)
-						.OnShouldFilterAsset_Lambda([this](const FAssetData& AssetData)
-						{
-							/*
-							if (!BatchContext.IKRetargetAsset)
-							{
-								return true;
-							}
-							
-							USkeletalMesh* Mesh = Cast<USkeletalMesh>(AssetData.GetAsset());
-							if (!Mesh)
-							{
-								return true;
-							}
-
-							
-							USkeletalMesh* PreviewMesh = BatchContext.IKRetargetAsset->TargetIKRigAsset->GetPreviewMesh();
-							if (!PreviewMesh)
-							{
-								return true;
-							}
-							
-							return Mesh->GetSkeleton() != PreviewMesh->GetSkeleton();*/
-							return false;
-						})
 					]
 				]
 			]
@@ -707,12 +663,9 @@ void SRetargetAnimAssetsWindow::RetargeterAssigned(const FAssetData& InAssetData
 {
 	UIKRetargeter* InRetargeter = Cast<UIKRetargeter>(InAssetData.GetAsset());
 	BatchContext.IKRetargetAsset = InRetargeter;
-	const UIKRigDefinition* SourceIKRig = InRetargeter ? InRetargeter->GetSourceIKRig() : nullptr;
-	const UIKRigDefinition* TargetIKRig = InRetargeter ? InRetargeter->GetTargetIKRig() : nullptr;
-	USkeletalMesh* SourceMesh =  SourceIKRig ? SourceIKRig->GetPreviewMesh() : nullptr;
-	USkeletalMesh* TargetMesh =  TargetIKRig ? TargetIKRig->GetPreviewMesh() : nullptr;
-	SourceMeshAssigned(FAssetData(SourceMesh));
-	TargetMeshAssigned(FAssetData(TargetMesh));
+	const UIKRetargeterController* Controller = UIKRetargeterController::GetController(InRetargeter);
+	SourceMeshAssigned(FAssetData(Controller->GetSourcePreviewMesh()));
+	TargetMeshAssigned(FAssetData(Controller->GetTargetPreviewMesh()));
 }
 
 ECheckBoxState SRetargetAnimAssetsWindow::IsRemappingReferencedAssets() const
diff --git a/Engine/Plugins/Animation/IKRig/Source/IKRigEditor/Private/RigEditor/IKRigController.cpp b/Engine/Plugins/Animation/IKRig/Source/IKRigEditor/Private/RigEditor/IKRigController.cpp
--- a/Engine/Plugins/Animation/IKRig/Source/IKRigEditor/Private/RigEditor/IKRigController.cpp
+++ b/Engine/Plugins/Animation/IKRig/Source/IKRigEditor/Private/RigEditor/IKRigController.cpp
@@ -11,7 +11,6 @@
 
 #define LOCTEXT_NAMESPACE "IKRigController"
 
-
 UIKRigController* UIKRigController::GetIKRigController(UIKRigDefinition* InIKRigDefinition)
 {
 	if (!InIKRigDefinition)
@@ -34,6 +33,16 @@ UIKRigDefinition* UIKRigController::GetAsset() const
 	return Asset;
 }
 
+FName UIKRigController::GetAssetIDAsName() const
+{
+	if (!Asset)
+	{
+		return NAME_None;
+	}
+
+	return FName(FString::FromInt(Asset->GetUniqueID()));
+}
+
 void UIKRigController::AddBoneSetting(const FName& BoneName, int32 SolverIndex) const
 {
 	check(Asset->Solvers.IsValidIndex(SolverIndex))
@@ -410,11 +419,12 @@ bool UIKRigController::SetSkeletalMesh(USkeletalMesh* SkeletalMesh, bool bTransa
 	}
 	
 	// first determine runtime compatibility between the IK Rig asset and the skeleton we're trying to run it on
-	const FReferenceSkeleton& RefSkeleton = SkeletalMesh->GetRefSkeleton();
-	const FIKRigInputSkeleton InputSkeleton = FIKRigInputSkeleton(RefSkeleton);
-	if (!UIKRigProcessor::IsIKRigCompatibleWithSkeleton(GetAsset(), InputSkeleton))
+	const FIKRigInputSkeleton InputSkeleton = FIKRigInputSkeleton(SkeletalMesh);
+	if (!UIKRigProcessor::IsIKRigCompatibleWithSkeleton(Asset, InputSkeleton))
 	{
-		UE_LOG(LogTemp, Warning, TEXT("Trying to initialize IKRig with a Skeleton that is missing required bones. See output log. %s"), *GetAsset()->GetName());
+		Asset->Log.LogError( FText::Format(
+			LOCTEXT("IncompatibleSkeleton", "Trying to initialize IKRig with a Skeleton that is missing required bones. See output log. {0}"),
+			FText::FromString(Asset->GetName())));
 		return false;
 	}
 
@@ -496,11 +506,6 @@ FTransform UIKRigController::GetRefPoseTransformOfBone(const FName& BoneName) co
 	return Asset->Skeleton.RefPoseGlobal[BoneIndex];
 }
 
-USkeletalMesh* UIKRigController::GetSkeletalMesh() const
-{
-	return Asset->PreviewSkeletalMesh.Get();
-}
-
 // -------------------------------------------------------
 // SOLVERS
 //
diff --git a/Engine/Plugins/Animation/IKRig/Source/IKRigEditor/Private/RigEditor/IKRigDefinitionFactory.cpp b/Engine/Plugins/Animation/IKRig/Source/IKRigEditor/Private/RigEditor/IKRigDefinitionFactory.cpp
--- a/Engine/Plugins/Animation/IKRig/Source/IKRigEditor/Private/RigEditor/IKRigDefinitionFactory.cpp
+++ b/Engine/Plugins/Animation/IKRig/Source/IKRigEditor/Private/RigEditor/IKRigDefinitionFactory.cpp
@@ -1,6 +1,7 @@
 // Copyright Epic Games, Inc. All Rights Reserved.
 
 #include "RigEditor/IKRigDefinitionFactory.h"
+#include "IKRigEditor.h"
 #include "IKRigDefinition.h"
 #include "AssetTypeCategories.h"
 #include "ContentBrowserModule.h"
@@ -30,7 +31,7 @@ UObject* UIKRigDefinitionFactory::FactoryCreateNew(
 {
 	if (!SkeletalMesh.IsValid())
 	{
-		UE_LOG(LogTemp, Warning, TEXT("Unable to create IK Rig. No Skeletal Mesh asset supplied."));
+		UE_LOG(LogIKRigEditor, Warning, TEXT("Unable to create IK Rig. No Skeletal Mesh asset supplied."));
 		return nullptr;
 	}
 	
diff --git a/Engine/Plugins/Animation/IKRig/Source/IKRigEditor/Private/RigEditor/IKRigEditorController.cpp b/Engine/Plugins/Animation/IKRig/Source/IKRigEditor/Private/RigEditor/IKRigEditorController.cpp
--- a/Engine/Plugins/Animation/IKRig/Source/IKRigEditor/Private/RigEditor/IKRigEditorController.cpp
+++ b/Engine/Plugins/Animation/IKRig/Source/IKRigEditor/Private/RigEditor/IKRigEditorController.cpp
@@ -4,6 +4,7 @@
 
 #include "IPersonaToolkit.h"
 #include "SKismetInspector.h"
+#include "Animation/DebugSkelMeshComponent.h"
 #include "Dialog/SCustomDialog.h"
 #include "Dialogs/Dialogs.h"
 
@@ -11,6 +12,9 @@
 #include "RigEditor/SIKRigSkeleton.h"
 #include "RigEditor/SIKRigSolverStack.h"
 #include "RigEditor/IKRigAnimInstance.h"
+#include "RigEditor/IKRigToolkit.h"
+#include "RigEditor/SIKRigAssetBrowser.h"
+#include "RigEditor/SIKRigOutputLog.h"
 
 #if WITH_EDITOR
 
@@ -221,6 +225,7 @@ void FIKRigEditorController::OnIKRigNeedsInitialized(UIKRigDefinition* ModifiedI
 
 void FIKRigEditorController::Reset() const
 {
+	SkelMeshComponent->ShowReferencePose(true);
 	AssetController->ResetGoalTransforms();
 }
 
@@ -245,6 +250,12 @@ void FIKRigEditorController::RefreshAllViews() const
 	{
 		RetargetingView->RefreshView();
 	}
+
+	// refresh the asset browser to ensure it shows compatible sequences
+	if (AssetBrowserView.IsValid())
+	{
+		AssetBrowserView.Get()->RefreshView();
+	}
 }
 
 void FIKRigEditorController::RefreshTreeView() const
@@ -255,6 +266,14 @@ void FIKRigEditorController::RefreshTreeView() const
 	}
 }
 
+void FIKRigEditorController::ClearOutputLog() const
+{
+	if (OutputLogView.IsValid())
+	{
+		OutputLogView->ClearLog();
+	}
+}
+
 void FIKRigEditorController::AddNewGoals(const TArray<FName>& GoalNames, const TArray<FName>& BoneNames)
 {
 	check(GoalNames.Num() == BoneNames.Num());
@@ -668,6 +687,43 @@ void FIKRigEditorController::ShowDetailsForElements(const TArray<TSharedPtr<FIKR
 	}
 }
 
+void FIKRigEditorController::OnFinishedChangingDetails(const FPropertyChangedEvent& PropertyChangedEvent)
+{
+	const bool bPreviewChanged = PropertyChangedEvent.GetPropertyName() == UIKRigDefinition::GetPreviewMeshPropertyName();
+	
+	if (bPreviewChanged)
+	{
+		// set the source and target skeletal meshes on the component
+		// NOTE: this must be done AFTER setting the AnimInstance so that the correct root anim node is loaded
+		USkeletalMesh* NewMesh = AssetController->GetAsset()->GetPreviewMesh();
+		if (NewMesh)
+		{
+			// apply the mesh to the preview scene
+			TSharedRef<IPersonaPreviewScene> PreviewScene = EditorToolkit.Pin()->GetPersonaToolkit()->GetPreviewScene();
+			if (PreviewScene->GetPreviewMesh() != NewMesh)
+			{
+				PreviewScene->SetPreviewMeshComponent(SkelMeshComponent);
+				PreviewScene->SetPreviewMesh(NewMesh);
+			}
+
+			ClearOutputLog();
+			AssetController->SetSkeletalMesh(NewMesh);
+			AnimInstance->SetProcessorNeedsInitialized();
+			AnimInstance->InitializeAnimation();
+			AssetController->BroadcastNeedsReinitialized();
+			AssetController->ResetGoalTransforms();
+			RefreshAllViews();
+		}
+	}
+}
+
+void FIKRigEditorController::SetDetailsView(const TSharedPtr<IDetailsView>& InDetailsView)
+{
+	DetailsView = InDetailsView;
+	DetailsView->OnFinishedChangingProperties().AddSP(this, &FIKRigEditorController::OnFinishedChangingDetails);
+	ShowEmptyDetails();
+}
+
 void FIKRigEditorController::AddNewRetargetChain(const FName ChainName, const FName StartBone, const FName EndBone)
 {
 	FIKRigRetargetChainSettings Settings(ChainName, StartBone, EndBone);
diff --git a/Engine/Plugins/Animation/IKRig/Source/IKRigEditor/Private/RigEditor/IKRigMode.cpp b/Engine/Plugins/Animation/IKRig/Source/IKRigEditor/Private/RigEditor/IKRigMode.cpp
--- a/Engine/Plugins/Animation/IKRig/Source/IKRigEditor/Private/RigEditor/IKRigMode.cpp
+++ b/Engine/Plugins/Animation/IKRig/Source/IKRigEditor/Private/RigEditor/IKRigMode.cpp
@@ -8,6 +8,7 @@
 #include "Modules/ModuleManager.h"
 #include "PersonaTabs.h"
 #include "RigEditor/IKRigAssetBrowserTabSummoner.h"
+#include "RigEditor/IKRigOutputLogTabSummoner.h"
 #include "RigEditor/IKRigSkeletonTabSummoner.h"
 #include "RigEditor/IKRigSolverStackTabSummoner.h"
 #include "RigEditor/IKRigRetargetChainTabSummoner.h"
@@ -31,7 +32,6 @@ FIKRigMode::FIKRigMode(
 	// register Persona tabs
 	FPersonaModule& PersonaModule = FModuleManager::LoadModuleChecked<FPersonaModule>("Persona");
 	TabFactories.RegisterFactory(PersonaModule.CreatePersonaViewportTabFactory(InHostingApp, ViewportArgs));
-	TabFactories.RegisterFactory(PersonaModule.CreateAdvancedPreviewSceneTabFactory(InHostingApp, InPreviewScene));
 	TabFactories.RegisterFactory(PersonaModule.CreateDetailsTabFactory(InHostingApp, FOnDetailsCreated::CreateSP(&IKRigEditor.Get(), &FIKRigEditorToolkit::HandleDetailsCreated)));
 
 	// register custom tabs
@@ -39,9 +39,10 @@ FIKRigMode::FIKRigMode(
 	TabFactories.RegisterFactory(MakeShared<FIKRigSkeletonTabSummoner>(IKRigEditor));
 	TabFactories.RegisterFactory(MakeShared<FIKRigSolverStackTabSummoner>(IKRigEditor));
 	TabFactories.RegisterFactory(MakeShared<FIKRigRetargetChainTabSummoner>(IKRigEditor));
+	TabFactories.RegisterFactory(MakeShared<FIKRigOutputLogTabSummoner>(IKRigEditor));
 
 	// create tab layout
-	TabLayout = FTabManager::NewLayout("Standalone_IKRigEditor_Layout_v1.123")
+	TabLayout = FTabManager::NewLayout("Standalone_IKRigEditor_Layout_v1.125")
 		->AddArea
 		(
 			FTabManager::NewPrimaryArea()
@@ -71,10 +72,22 @@ FIKRigMode::FIKRigMode(
 				)
 				->Split
 				(
+					FTabManager::NewSplitter()
+					->SetSizeCoefficient(0.8f)
+					->SetOrientation(Orient_Vertical)
+					->Split
+					(
 					FTabManager::NewStack()
-					->SetSizeCoefficient(0.6f)
-					->SetHideTabWell(true)
-					->AddTab(FPersonaTabs::PreviewViewportID, ETabState::OpenedTab)
+						->SetSizeCoefficient(0.6f)
+						->SetHideTabWell(true)
+						->AddTab(FPersonaTabs::PreviewViewportID, ETabState::OpenedTab)
+					)
+					->Split
+					(
+						FTabManager::NewStack()
+						->SetSizeCoefficient(0.2f)
+						->AddTab(FIKRigOutputLogTabSummoner::TabID, ETabState::OpenedTab)
+					)
 				)
 				->Split
 				(
@@ -86,7 +99,6 @@ FIKRigMode::FIKRigMode(
 						FTabManager::NewStack()
 						->SetSizeCoefficient(0.6f)
 						->AddTab(FPersonaTabs::DetailsID, ETabState::OpenedTab)
-						->AddTab(FPersonaTabs::AdvancedPreviewSceneSettingsID, ETabState::OpenedTab)
 						->SetForegroundTab(FPersonaTabs::DetailsID)
                     )
                     ->Split
diff --git a/Engine/Plugins/Animation/IKRig/Source/IKRigEditor/Private/RigEditor/IKRigOutputLogTabSummoner.cpp b/Engine/Plugins/Animation/IKRig/Source/IKRigEditor/Private/RigEditor/IKRigOutputLogTabSummoner.cpp
new file mode 100644
--- /dev/null
+++ b/Engine/Plugins/Animation/IKRig/Source/IKRigEditor/Private/RigEditor/IKRigOutputLogTabSummoner.cpp
@@ -0,0 +1,41 @@
+﻿// Copyright Epic Games, Inc. All Rights Reserved.
+#include "RigEditor/IKRigOutputLogTabSummoner.h"
+
+#include "IDocumentation.h"
+#include "RigEditor/IKRigToolkit.h"
+#include "RigEditor/IKRigEditorStyle.h"
+#include "RigEditor/SIKRigOutputLog.h"
+
+#define LOCTEXT_NAMESPACE "IKRigOutputLogTabSummoner"
+
+const FName FIKRigOutputLogTabSummoner::TabID(TEXT("IKRigOutputLog"));
+
+FIKRigOutputLogTabSummoner::FIKRigOutputLogTabSummoner(
+	const TSharedRef<FIKRigEditorToolkit>& InRigEditor)
+	: FWorkflowTabFactory(TabID, InRigEditor),
+	IKRigEditor(InRigEditor)
+{
+	bIsSingleton = true; // only allow a single instance of this tab
+	
+	TabLabel = LOCTEXT("IKRigOutputLogTabLabel", "IK Rig Output");
+	TabIcon = FSlateIcon(FIKRigEditorStyle::Get().GetStyleSetName(), "IKRig.TabIcon");
+
+	ViewMenuDescription = LOCTEXT("IKRigOutputLog_ViewMenu_Desc", "IK Rig Output");
+	ViewMenuTooltip = LOCTEXT("IKRigOutputLog_ViewMenu_ToolTip", "Show the IK Rig Output Log Tab");
+}
+
+TSharedPtr<SToolTip> FIKRigOutputLogTabSummoner::CreateTabToolTipWidget(const FWorkflowTabSpawnInfo& Info) const
+{
+	return  IDocumentation::Get()->CreateToolTip(LOCTEXT("IKRigOutputLogTooltip", "View warnings and errors from this rig."), NULL, TEXT("Shared/Editors/Persona"), TEXT("IKRigOutputLog_Window"));
+}
+
+TSharedRef<SWidget> FIKRigOutputLogTabSummoner::CreateTabBody(const FWorkflowTabSpawnInfo& Info) const
+{
+	const TSharedRef<FIKRigEditorController> Controller = IKRigEditor.Pin()->GetController();
+	TSharedPtr<SIKRigOutputLog>& OutputLogView = Controller->OutputLogView;
+	FName LogName = Controller->AssetController->GetAsset()->Log.GetLogTarget();
+	
+	return SNew(SIKRigOutputLog, LogName, OutputLogView);
+}
+
+#undef LOCTEXT_NAMESPACE 
diff --git a/Engine/Plugins/Animation/IKRig/Source/IKRigEditor/Private/RigEditor/IKRigToolkit.cpp b/Engine/Plugins/Animation/IKRig/Source/IKRigEditor/Private/RigEditor/IKRigToolkit.cpp
--- a/Engine/Plugins/Animation/IKRig/Source/IKRigEditor/Private/RigEditor/IKRigToolkit.cpp
+++ b/Engine/Plugins/Animation/IKRig/Source/IKRigEditor/Private/RigEditor/IKRigToolkit.cpp
@@ -55,10 +55,6 @@ void FIKRigEditorToolkit::InitAssetEditor(
 	
 	const FPersonaModule& PersonaModule = FModuleManager::LoadModuleChecked<FPersonaModule>("Persona");
 	PersonaToolkit = PersonaModule.CreatePersonaToolkit(IKRigAsset, PersonaToolkitArgs);
-	
-	
-	// when/if preview mesh is changed, we need to reinitialize the anim instance
-    PersonaToolkit->GetPreviewScene()->RegisterOnPreviewMeshChanged(FOnPreviewMeshChanged::CreateSP(this, &FIKRigEditorToolkit::HandlePreviewMeshChanged));
 
 	const TSharedRef<IAssetFamily> AssetFamily = PersonaModule.CreatePersonaAssetFamily(IKRigAsset);
 	AssetFamily->RecordAssetOpened(FAssetData(IKRigAsset));
@@ -107,7 +103,7 @@ void FIKRigEditorToolkit::BindCommands()
 
 	ToolkitCommands->MapAction(
         Commands.Reset,
-        FExecuteAction::CreateSP(this, &FIKRigEditorToolkit::HandleReset),
+        FExecuteAction::CreateSP(EditorController, &FIKRigEditorController::Reset),
 		EUIActionRepeatMode::RepeatDisabled);
 }
 
@@ -198,12 +194,16 @@ void FIKRigEditorToolkit::HandlePreviewSceneCreated(const TSharedRef<IPersonaPre
 	EditorController->SkelMeshComponent = NewObject<UDebugSkelMeshComponent>(Actor);
 
 	// setup an apply an anim instance to the skeletal mesh component
-	EditorController->AnimInstance = NewObject<UIKRigAnimInstance>(EditorController->SkelMeshComponent, TEXT("IKRigAnimScriptInstance"));
-	SetupAnimInstance();
+	UIKRigAnimInstance* AnimInstance = NewObject<UIKRigAnimInstance>(EditorController->SkelMeshComponent, TEXT("IKRigAnimScriptInstance"));
+	EditorController->AnimInstance = AnimInstance;
+	AnimInstance->SetIKRigAsset(EditorController->AssetController->GetAsset());
+	EditorController->SkelMeshComponent->PreviewInstance = AnimInstance;
+	EditorController->OnIKRigNeedsInitialized(EditorController->AssetController->GetAsset());
+	AnimInstance->InitializeAnimation();
 
 	// set the skeletal mesh on the component
 	// NOTE: this must be done AFTER setting the AnimInstance so that the correct root anim node is loaded
-	USkeletalMesh* Mesh = EditorController->AssetController->GetSkeletalMesh();
+	USkeletalMesh* Mesh = EditorController->AssetController->GetAsset()->GetPreviewMesh();
 	EditorController->SkelMeshComponent->SetSkeletalMesh(Mesh);
 
 	// apply mesh to the preview scene
@@ -215,60 +215,9 @@ void FIKRigEditorToolkit::HandlePreviewSceneCreated(const TSharedRef<IPersonaPre
 	InPersonaPreviewScene->AddComponent(EditorController->SkelMeshComponent, FTransform::Identity);
 }
 
-void FIKRigEditorToolkit::HandlePreviewMeshChanged(USkeletalMesh* InOldSkeletalMesh, USkeletalMesh* InNewSkeletalMesh)
-{
-	if (InOldSkeletalMesh == InNewSkeletalMesh)
-	{
-		return; // already set to this skeletal mesh
-	}
-
-	// we do not reset the current skeletal mesh to keep track of the last valid one but we still need to reinit 
-	if (!InNewSkeletalMesh)
-	{
-		EditorController->AssetController->BroadcastNeedsReinitialized();
-		return;
-	}
-	
-	// update asset with new skeletal mesh (will copy new skeleton data)
-	if (!EditorController->AssetController->SetSkeletalMesh(InNewSkeletalMesh))
-	{
-		return; // mesh was not set (incompatible for some reason) todo, show reason in UI someplace
-	}
-
-	// update anim instance to use new skeletal mesh
-	// this is required so that the bone containers passed around during update/eval are correctly sized
-	UDebugSkelMeshComponent* EditorSkelComp = Cast<UDebugSkelMeshComponent>(GetPersonaToolkit()->GetPreviewScene()->GetPreviewMeshComponent());
-	if (EditorSkelComp)
-	{
-		bool bWasCreated = false;
-		EditorController->AnimInstance = FAnimCustomInstanceHelper::BindToSkeletalMeshComponent<UIKRigAnimInstance>(EditorSkelComp , bWasCreated);
-		SetupAnimInstance();
-	}
-
-	EditorController->SkelMeshComponent->SetSkeletalMesh(InNewSkeletalMesh);
-	
-	EditorController->RefreshAllViews();
-}
-
 void FIKRigEditorToolkit::HandleDetailsCreated(const TSharedRef<class IDetailsView>& InDetailsView) const
 {
 	EditorController->SetDetailsView(InDetailsView);
-	EditorController->ShowEmptyDetails();
-}
-
-void FIKRigEditorToolkit::HandleReset()
-{
-	EditorController.Get().Reset();
 }
 
-void FIKRigEditorToolkit::SetupAnimInstance()
-{
-	UIKRigAnimInstance* InAnimInstance = EditorController->AnimInstance.Get();
-	InAnimInstance->SetIKRigAsset(EditorController->AssetController->GetAsset());
-	EditorController->SkelMeshComponent->PreviewInstance = InAnimInstance;
-	InAnimInstance->InitializeAnimation();
-	EditorController->OnIKRigNeedsInitialized(EditorController->AssetController->GetAsset());
-}
-
-
 #undef LOCTEXT_NAMESPACE
diff --git a/Engine/Plugins/Animation/IKRig/Source/IKRigEditor/Private/RigEditor/SIKRigAssetBrowser.cpp b/Engine/Plugins/Animation/IKRig/Source/IKRigEditor/Private/RigEditor/SIKRigAssetBrowser.cpp
--- a/Engine/Plugins/Animation/IKRig/Source/IKRigEditor/Private/RigEditor/SIKRigAssetBrowser.cpp
+++ b/Engine/Plugins/Animation/IKRig/Source/IKRigEditor/Private/RigEditor/SIKRigAssetBrowser.cpp
@@ -18,6 +18,7 @@ void SIKRigAssetBrowser::Construct(
 	TSharedRef<FIKRigEditorController> InEditorController)
 {
 	EditorController = InEditorController;
+	EditorController.Pin()->AssetBrowserView = SharedThis(this);
 	
 	ChildSlot
     [
@@ -28,10 +29,10 @@ void SIKRigAssetBrowser::Construct(
 		]
     ];
 
-	AddAssetBrowser();
+	RefreshView();
 }
 
-void SIKRigAssetBrowser::AddAssetBrowser()
+void SIKRigAssetBrowser::RefreshView()
 {
 	FAssetPickerConfig AssetPickerConfig;
 	
diff --git a/Engine/Plugins/Animation/IKRig/Source/IKRigEditor/Private/RigEditor/SIKRigOutputLog.cpp b/Engine/Plugins/Animation/IKRig/Source/IKRigEditor/Private/RigEditor/SIKRigOutputLog.cpp
new file mode 100644
--- /dev/null
+++ b/Engine/Plugins/Animation/IKRig/Source/IKRigEditor/Private/RigEditor/SIKRigOutputLog.cpp
@@ -0,0 +1,43 @@
+﻿// Copyright Epic Games, Inc. All Rights Reserved.
+
+#include "RigEditor/SIKRigOutputLog.h"
+
+#include "IMessageLogListing.h"
+#include "MessageLogModule.h"
+#include "RetargetEditor/IKRetargetEditorController.h"
+
+void SIKRigOutputLog::Construct(const FArguments& InArgs, const FName InLogName, TSharedPtr<SIKRigOutputLog>& OutLog)
+{
+	OutLog = SharedThis(this);
+
+	FMessageLogModule& MessageLogModule = FModuleManager::LoadModuleChecked<FMessageLogModule>("MessageLog");
+
+	if (MessageLogModule.IsRegisteredLogListing(InLogName))
+	{
+		MessageLogListing = MessageLogModule.GetLogListing(InLogName);
+	}
+	else
+	{
+		FMessageLogInitializationOptions InitOptions;
+		InitOptions.bShowFilters = false;
+		InitOptions.bShowPages = false;
+		InitOptions.bAllowClear = false;
+		MessageLogModule.RegisterLogListing(InLogName, FText::FromString("IK Rig Log"), InitOptions);
+		MessageLogListing = MessageLogModule.GetLogListing(InLogName);
+	}
+
+	OutputLogWidget = MessageLogModule.CreateLogListingWidget( MessageLogListing.ToSharedRef() );
+	
+	ChildSlot
+	[
+		SNew(SBox)
+		[
+			OutputLogWidget.ToSharedRef()
+		]
+	];
+}
+
+void SIKRigOutputLog::ClearLog() const
+{
+	MessageLogListing.Get()->ClearMessages();
+}
diff --git a/Engine/Plugins/Animation/IKRig/Source/IKRigEditor/Private/RigEditor/SIKRigRetargetChainList.cpp b/Engine/Plugins/Animation/IKRig/Source/IKRigEditor/Private/RigEditor/SIKRigRetargetChainList.cpp
--- a/Engine/Plugins/Animation/IKRig/Source/IKRigEditor/Private/RigEditor/SIKRigRetargetChainList.cpp
+++ b/Engine/Plugins/Animation/IKRig/Source/IKRigEditor/Private/RigEditor/SIKRigRetargetChainList.cpp
@@ -265,7 +265,7 @@ const FReferenceSkeleton& SIKRigRetargetChainRow::GetReferenceSkeleton() const
 		return DummySkeleton; 
 	}
 
-	USkeletalMesh* SkeletalMesh = Controller->AssetController->GetSkeletalMesh();
+	USkeletalMesh* SkeletalMesh = Controller->AssetController->GetAsset()->GetPreviewMesh();
 	if (SkeletalMesh == nullptr)
 	{
 		return DummySkeleton;
diff --git a/Engine/Plugins/Animation/IKRig/Source/IKRigEditor/Public/IKRigEditor.h b/Engine/Plugins/Animation/IKRig/Source/IKRigEditor/Public/IKRigEditor.h
--- a/Engine/Plugins/Animation/IKRig/Source/IKRigEditor/Public/IKRigEditor.h
+++ b/Engine/Plugins/Animation/IKRig/Source/IKRigEditor/Public/IKRigEditor.h
@@ -18,3 +18,5 @@ private:
 
 	TArray<FName> ClassesToUnregisterOnShutdown;
 };
+
+DECLARE_LOG_CATEGORY_EXTERN(LogIKRigEditor, Warning, All);
diff --git a/Engine/Plugins/Animation/IKRig/Source/IKRigEditor/Public/RetargetEditor/IKRetargetEditorController.h b/Engine/Plugins/Animation/IKRig/Source/IKRigEditor/Public/RetargetEditor/IKRetargetEditorController.h
--- a/Engine/Plugins/Animation/IKRig/Source/IKRigEditor/Public/RetargetEditor/IKRetargetEditorController.h
+++ b/Engine/Plugins/Animation/IKRig/Source/IKRigEditor/Public/RetargetEditor/IKRetargetEditorController.h
@@ -3,12 +3,14 @@
 
 #include "IKRetargeterController.h"
 #include "IPersonaToolkit.h"
+#include "SIKRetargetAssetBrowser.h"
 #include "UObject/ObjectMacros.h"
 #include "UObject/ObjectPtr.h"
 #include "Input/Reply.h"
 #include "Widgets/Input/SEditableTextBox.h"
 #include "Widgets/SWindow.h"
 
+class SIKRigOutputLog;
 class UIKRetargetProcessor;
 class SIKRetargetChainMapList;
 class UIKRetargetAnimInstance;
@@ -56,10 +58,17 @@ public:
 	void AddOffsetAndUpdatePreviewMeshPosition(const FVector& Offset, USceneComponent* Component) const;
 
 	/** asset properties tab */
-	TSharedPtr<class IDetailsView> DetailsView;
-
+	TSharedPtr<IDetailsView> DetailsView;
 	/** chain list view */
 	TSharedPtr<SIKRetargetChainMapList> ChainsView;
+	/** asset browser view */
+	TSharedPtr<SIKRetargetAssetBrowser> AssetBrowserView;
+	/** output log view */
+	TSharedPtr<SIKRigOutputLog> OutputLogView;
+	/** clear the output log */
+	void ClearOutputLog() const;
+	/** force refresh all views in the editor */
+	void RefreshAllViews() const;
 
 	/** get the source skeletal mesh we are copying FROM */
 	USkeletalMesh* GetSourceSkeletalMesh() const;
@@ -115,6 +124,4 @@ public:
 	FText GetCurrentPoseName() const;
 	void OnPoseSelected(TSharedPtr<FName> InPoseName, ESelectInfo::Type SelectInfo) const;
 	/* END edit retarget poses */
-
-	void RefreshAllViews() const;
 };
diff --git a/Engine/Plugins/Animation/IKRig/Source/IKRigEditor/Public/RetargetEditor/IKRetargetOutputLogTabSummoner.h b/Engine/Plugins/Animation/IKRig/Source/IKRigEditor/Public/RetargetEditor/IKRetargetOutputLogTabSummoner.h
new file mode 100644
--- /dev/null
+++ b/Engine/Plugins/Animation/IKRig/Source/IKRigEditor/Public/RetargetEditor/IKRetargetOutputLogTabSummoner.h
@@ -0,0 +1,21 @@
+﻿// Copyright Epic Games, Inc. All Rights Reserved.
+
+#pragma once
+
+#include "CoreMinimal.h"
+#include "WorkflowOrientedApp/WorkflowTabFactory.h"
+
+class FIKRetargetEditor;
+
+struct FIKRetargetOutputLogTabSummoner : public FWorkflowTabFactory
+{
+public:
+	static const FName TabID;
+	
+	FIKRetargetOutputLogTabSummoner(const TSharedRef<FIKRetargetEditor>& InIKRetargetEditor);
+	virtual TSharedRef<SWidget> CreateTabBody(const FWorkflowTabSpawnInfo& Info) const override;
+	virtual TSharedPtr<SToolTip> CreateTabToolTipWidget(const FWorkflowTabSpawnInfo& Info) const override;
+
+protected:
+	TWeakPtr<FIKRetargetEditor> IKRetargetEditor;
+};
diff --git a/Engine/Plugins/Animation/IKRig/Source/IKRigEditor/Public/RetargetEditor/IKRetargeterController.h b/Engine/Plugins/Animation/IKRig/Source/IKRigEditor/Public/RetargetEditor/IKRetargeterController.h
--- a/Engine/Plugins/Animation/IKRig/Source/IKRigEditor/Public/RetargetEditor/IKRetargeterController.h
+++ b/Engine/Plugins/Animation/IKRig/Source/IKRigEditor/Public/RetargetEditor/IKRetargeterController.h
@@ -28,8 +28,10 @@ public:
 	/** Get access to the retargeter asset.
 	 *@warning Do not make modifications to the asset directly. Use this API instead. */
 	UIKRetargeter* GetAsset() const;
+	/** Get unique asset integer ID as a name */
+	FName GetAssetIDAsName() const;
 	/** Get access to the editor controller.*/
-	FIKRetargetEditorController* GetEditorController() const;
+	FIKRetargetEditorController* GetEditorController() const { return EditorController; };
 	/** Set the currently used editor controller.*/
 	void SetEditorController(FIKRetargetEditorController* InEditorController) { EditorController = InEditorController; };
 
@@ -40,11 +42,10 @@ public:
 	void SetSourceIKRig(UIKRigDefinition* SourceIKRig);
 	/** Set the IK Rig to use as the target (to copy animation TO) */
 	void SetTargetIKRig(UIKRigDefinition* TargetIKRig);
+	/** Get source skeletal mesh */
+	USkeletalMesh* GetSourcePreviewMesh() const;
 	/** Get target skeletal mesh */
-	USkeletalMesh* GetTargetPreviewMesh();
-	
-	/** Get the USkeleton used on the Source asset */
-	USkeleton* GetSourceSkeletonAsset() const;
+	USkeletalMesh* GetTargetPreviewMesh() const;
 
 	/** Get name of the Root bone used for retargeting the Source skeleton. */
 	FName GetSourceRootBone() const;
diff --git a/Engine/Plugins/Animation/IKRig/Source/IKRigEditor/Public/RetargetEditor/SIKRetargetAssetBrowser.h b/Engine/Plugins/Animation/IKRig/Source/IKRigEditor/Public/RetargetEditor/SIKRetargetAssetBrowser.h
--- a/Engine/Plugins/Animation/IKRig/Source/IKRigEditor/Public/RetargetEditor/SIKRetargetAssetBrowser.h
+++ b/Engine/Plugins/Animation/IKRig/Source/IKRigEditor/Public/RetargetEditor/SIKRetargetAssetBrowser.h
@@ -18,10 +18,10 @@ public:
 
 	void Construct(const FArguments& InArgs, TSharedRef<FIKRetargetEditorController> InEditorController);
 
+	void RefreshView();
+	
 private:
 
-	void AddAssetBrowser();
-
 	void OnPathChange(const FString& NewPath);
 
 	FReply OnExportButtonClicked() const;
diff --git a/Engine/Plugins/Animation/IKRig/Source/IKRigEditor/Public/RigEditor/IKRigController.h b/Engine/Plugins/Animation/IKRig/Source/IKRigEditor/Public/RigEditor/IKRigController.h
--- a/Engine/Plugins/Animation/IKRig/Source/IKRigEditor/Public/RigEditor/IKRigController.h
+++ b/Engine/Plugins/Animation/IKRig/Source/IKRigEditor/Public/RigEditor/IKRigController.h
@@ -31,6 +31,8 @@ public:
 	
 	/** Get the asset this controller controls. */
 	UIKRigDefinition* GetAsset() const;
+	/** Get unique asset integer ID as a name */
+	FName GetAssetIDAsName() const;
 
 	/** SKELETON
 	 * 
@@ -38,8 +40,6 @@ public:
 	/** Sets the preview mesh to use, can optionally reinitialize the skeleton with bReImportBone=true.
 	 * Returns true if the mesh was able to be set. False if it was incompatible for any reason. */
 	bool SetSkeletalMesh(USkeletalMesh* SkeletalMesh, bool bTransact=false) const;
-	/** Get the skeletal mesh asset this IK Rig was initialized with */
-	USkeletalMesh* GetSkeletalMesh() const;
 	/** Get read-access to the IKRig skeleton representation */
 	const FIKRigSkeleton& GetIKRigSkeleton() const;
 	/** Get the USkeleton asset this rig was initialized with */
diff --git a/Engine/Plugins/Animation/IKRig/Source/IKRigEditor/Public/RigEditor/IKRigEditorController.h b/Engine/Plugins/Animation/IKRig/Source/IKRigEditor/Public/RigEditor/IKRigEditorController.h
--- a/Engine/Plugins/Animation/IKRig/Source/IKRigEditor/Public/RigEditor/IKRigEditorController.h
+++ b/Engine/Plugins/Animation/IKRig/Source/IKRigEditor/Public/RigEditor/IKRigEditorController.h
@@ -10,6 +10,8 @@
 
 #include "IKRigEditorController.generated.h"
 
+class SIKRigAssetBrowser;
+class SIKRigOutputLog;
 class UIKRigAnimInstance;
 class FIKRigEditorToolkit;
 class SIKRigSolverStack;
@@ -95,6 +97,8 @@ public:
 	void RefreshAllViews() const;
 	/** refresh just the skeleton tree view */
 	void RefreshTreeView() const;
+	/** clear the output log */
+	void ClearOutputLog() const;
 
 	/** return list of those solvers in the stack that are selected by user */
 	void GetSelectedSolvers(TArray<TSharedPtr<FSolverStackElement> >& OutSelectedSolvers);
@@ -136,9 +140,11 @@ public:
 	void ShowEmptyDetails() const;
 	/** show selected items in details view */
 	void ShowDetailsForElements(const TArray<TSharedPtr<FIKRigTreeElement>>& InItems) const;
+	/** callback when detail is edited */
+	void OnFinishedChangingDetails(const FPropertyChangedEvent& PropertyChangedEvent);
 	
 	/** set details tab view */
-	void SetDetailsView(const TSharedPtr<class IDetailsView>& InDetailsView){ DetailsView = InDetailsView; };
+	void SetDetailsView(const TSharedPtr<class IDetailsView>& InDetailsView);
 	/** set skeleton tab view */
 	void SetSkeletonsView(const TSharedPtr<SIKRigSkeleton>& InSkeletonView){ SkeletonView = InSkeletonView; };
 	/** set solver stack tab view */
@@ -160,7 +166,7 @@ public:
 
 	/** viewport anim instance */
 	UPROPERTY(transient, NonTransactional)
-	TWeakObjectPtr<class UIKRigAnimInstance> AnimInstance;
+	TWeakObjectPtr<UIKRigAnimInstance> AnimInstance;
 
 	/** the persona toolkit */
 	TWeakPtr<FIKRigEditorToolkit> EditorToolkit;
@@ -178,16 +184,14 @@ public:
 	{
 		Collector.AddReferencedObject(BoneDetails);
 	};
-	virtual FString GetReferencerName() const override
-	{
-		return "IKRigEditorController";
-	};
+	virtual FString GetReferencerName() const override { return "IKRigEditorController"; };
 	/** END FGCObject interface */
 
 	/** UIKRigBoneDetails factory **/
 	TObjectPtr<UIKRigBoneDetails> CreateBoneDetails(const TSharedPtr<FIKRigTreeElement const>& InItem) const;
 	
 private:
+	
 	/** right after importing a skeleton, we ask user what solver they want to use */
 	bool PromptToAddSolver() const;
 
@@ -195,7 +199,7 @@ private:
 	void InitializeSolvers() const;
 
 	/** asset properties tab */
-	TSharedPtr<class IDetailsView> DetailsView;
+	TSharedPtr<IDetailsView> DetailsView;
 
 	/** the skeleton tree view */
 	TSharedPtr<SIKRigSkeleton> SkeletonView;
@@ -206,10 +210,19 @@ private:
 	/** the solver stack view */
 	TSharedPtr<SIKRigRetargetChainList> RetargetingView;
 
+	/** asset browser view */
+	TSharedPtr<SIKRigAssetBrowser> AssetBrowserView;
+
+	/** output log view */
+	TSharedPtr<SIKRigOutputLog> OutputLogView;
+
 	EIKRigSelectionType LastSelectedType;
 
 	UPROPERTY()
 	TObjectPtr<UIKRigBoneDetails> BoneDetails;
+
+	friend struct FIKRigOutputLogTabSummoner;
+	friend class SIKRigAssetBrowser;
 };
 
 struct FIKRigSolverTypeAndName
diff --git a/Engine/Plugins/Animation/IKRig/Source/IKRigEditor/Public/RigEditor/IKRigOutputLogTabSummoner.h b/Engine/Plugins/Animation/IKRig/Source/IKRigEditor/Public/RigEditor/IKRigOutputLogTabSummoner.h
new file mode 100644
--- /dev/null
+++ b/Engine/Plugins/Animation/IKRig/Source/IKRigEditor/Public/RigEditor/IKRigOutputLogTabSummoner.h
@@ -0,0 +1,21 @@
+﻿// Copyright Epic Games, Inc. All Rights Reserved.
+
+#pragma once
+
+#include "CoreMinimal.h"
+#include "WorkflowOrientedApp/WorkflowTabFactory.h"
+
+class FIKRigEditorToolkit;
+
+struct FIKRigOutputLogTabSummoner : public FWorkflowTabFactory
+{
+public:
+	static const FName TabID;
+	
+	FIKRigOutputLogTabSummoner(const TSharedRef<FIKRigEditorToolkit>& InIKRigEditor);
+	virtual TSharedRef<SWidget> CreateTabBody(const FWorkflowTabSpawnInfo& Info) const override;
+	virtual TSharedPtr<SToolTip> CreateTabToolTipWidget(const FWorkflowTabSpawnInfo& Info) const override;
+
+protected:
+	TWeakPtr<FIKRigEditorToolkit> IKRigEditor;
+};
diff --git a/Engine/Plugins/Animation/IKRig/Source/IKRigEditor/Public/RigEditor/IKRigToolkit.h b/Engine/Plugins/Animation/IKRig/Source/IKRigEditor/Public/RigEditor/IKRigToolkit.h
--- a/Engine/Plugins/Animation/IKRig/Source/IKRigEditor/Public/RigEditor/IKRigToolkit.h
+++ b/Engine/Plugins/Animation/IKRig/Source/IKRigEditor/Public/RigEditor/IKRigToolkit.h
@@ -12,21 +12,6 @@
 #include "PersonaAssetEditorToolkit.h"
 #include "SIKRigSkeleton.h"
 
-class IAnimationSequenceBrowser;
-class UIKRigSkeletalMeshComponent;
-class FSolverStackElement;
-class SIKRigSolverStack;
-class SIKRigSkeleton;
-class FIKRigSkeletonTreeBuilder;
-class IDetailsView;
-class ISkeletonTree;
-class FGGAssetEditorToolbar;
-class UIKRigDefinition;
-class UIKRigSettings;
-class FIKRigToolbar;
-class SIKRigViewportTabBody;
-class FIKRigPreviewScene;
-
 namespace IKRigEditorModes
 {
 	extern const FName IKRigEditorMode;
@@ -61,10 +46,7 @@ public:
 	
 	/** FGCObject interface */
 	virtual void AddReferencedObjects(FReferenceCollector& Collector) override;
-	virtual FString GetReferencerName() const override
-	{
-		return TEXT("FIKRigEditorToolkit");
-	}
+	virtual FString GetReferencerName() const override { return TEXT("FIKRigEditorToolkit"); }
 	/** END FGCObject interface */
 
 	//** FTickableEditorObject Interface
@@ -84,29 +66,22 @@ public:
 
 	TSharedRef<FIKRigEditorController> GetController() const {return EditorController;};
 
+	/** preview scene setup */
+	void HandlePreviewSceneCreated(const TSharedRef<IPersonaPreviewScene>& InPersonaPreviewScene);
+	void HandleDetailsCreated(const TSharedRef<class IDetailsView>& InDetailsView) const;
+	/** END preview scene setup */
+
 private:
 
 	/** toolbar */
 	void BindCommands();
 	void ExtendToolbar();
 	void FillToolbar(FToolBarBuilder& ToolbarBuilder);
-	void HandleReset();
 	/** END toolbar */
 	
-	/** preview scene setup */
-	void HandlePreviewSceneCreated(const TSharedRef<IPersonaPreviewScene>& InPersonaPreviewScene);
-	void HandlePreviewMeshChanged(USkeletalMesh* InOldSkeletalMesh, USkeletalMesh* InNewSkeletalMesh);
-	void HandleDetailsCreated(const TSharedRef<class IDetailsView>& InDetailsView) const;
-	void SetupAnimInstance();
-	/** END preview scene setup */
-	
 	/** centralized management of selection across skeleton view and viewport */
 	TSharedRef<FIKRigEditorController> EditorController;
 
 	/** Preview scene to be supplied by IHasPersonaToolkit::GetPersonaToolkit */
 	TSharedPtr<IPersonaToolkit> PersonaToolkit;
-	
-	friend FIKRigMode;
-	friend SIKRigSkeleton;
-	friend SIKRigSolverStack;
 };
diff --git a/Engine/Plugins/Animation/IKRig/Source/IKRigEditor/Public/RigEditor/SIKRigAssetBrowser.h b/Engine/Plugins/Animation/IKRig/Source/IKRigEditor/Public/RigEditor/SIKRigAssetBrowser.h
--- a/Engine/Plugins/Animation/IKRig/Source/IKRigEditor/Public/RigEditor/SIKRigAssetBrowser.h
+++ b/Engine/Plugins/Animation/IKRig/Source/IKRigEditor/Public/RigEditor/SIKRigAssetBrowser.h
@@ -20,7 +20,7 @@ public:
 
 private:
 
-	void AddAssetBrowser();
+	void RefreshView();
 
 	void OnPathChange(const FString& NewPath);
 	
diff --git a/Engine/Plugins/Animation/IKRig/Source/IKRigEditor/Public/RigEditor/SIKRigOutputLog.h b/Engine/Plugins/Animation/IKRig/Source/IKRigEditor/Public/RigEditor/SIKRigOutputLog.h
new file mode 100644
--- /dev/null
+++ b/Engine/Plugins/Animation/IKRig/Source/IKRigEditor/Public/RigEditor/SIKRigOutputLog.h
@@ -0,0 +1,27 @@
+﻿// Copyright Epic Games, Inc. All Rights Reserved.
+
+#pragma once
+
+#include "Widgets/Layout/SBox.h"
+
+class IMessageLogListing;
+class FIKRetargetEditorController;
+
+class SIKRigOutputLog : public SBox
+{
+	
+public:
+	
+	SLATE_BEGIN_ARGS(SIKRigOutputLog) {}
+	SLATE_END_ARGS()
+
+	void Construct(const FArguments& InArgs, const FName InLogName, TSharedPtr<SIKRigOutputLog>& OutLog);
+	
+	void ClearLog() const;
+	
+private:
+	
+	/** the output log */
+	TSharedPtr<SWidget> OutputLogWidget;
+	TSharedPtr<IMessageLogListing> MessageLogListing;
+};
\ No newline at end of file
