IKRig assets crash on creation.

diff --git a/Engine/Plugins/Animation/IKRig/Source/IKRig/Private/IKRigDefinition.cpp b/Engine/Plugins/Animation/IKRig/Source/IKRig/Private/IKRigDefinition.cpp
--- a/Engine/Plugins/Animation/IKRig/Source/IKRig/Private/IKRigDefinition.cpp
+++ b/Engine/Plugins/Animation/IKRig/Source/IKRig/Private/IKRigDefinition.cpp
@@ -279,6 +279,13 @@ void UIKRigDefinition::PostLoad()
 	Log.SetLogTarget(GetUniqueIDAsName(), false);
 }
 
+void UIKRigDefinition::PostInitProperties()
+{
+	UObject::PostInitProperties();
+
+	Log.SetLogTarget(GetUniqueIDAsName(), false);
+}
+
 const FBoneChain* UIKRigDefinition::GetRetargetChainByName(FName ChainName) const
 {
 	for (const FBoneChain& Chain : RetargetDefinition.BoneChains)
diff --git a/Engine/Plugins/Animation/IKRig/Source/IKRig/Private/Retargeter/IKRetargeter.cpp b/Engine/Plugins/Animation/IKRig/Source/IKRig/Private/Retargeter/IKRetargeter.cpp
--- a/Engine/Plugins/Animation/IKRig/Source/IKRig/Private/Retargeter/IKRetargeter.cpp
+++ b/Engine/Plugins/Animation/IKRig/Source/IKRig/Private/Retargeter/IKRetargeter.cpp
@@ -6,7 +6,7 @@
 const FName UIKRetargeter::GetSourceIKRigPropertyName() { return GET_MEMBER_NAME_STRING_CHECKED(UIKRetargeter, SourceIKRigAsset); };
 const FName UIKRetargeter::GetTargetIKRigPropertyName() { return GET_MEMBER_NAME_STRING_CHECKED(UIKRetargeter, TargetIKRigAsset); };
 const FName UIKRetargeter::GetSourcePreviewMeshPropertyName() { return GET_MEMBER_NAME_STRING_CHECKED(UIKRetargeter, SourcePreviewMesh); };
-const FName UIKRetargeter::GetTargetPreviewMeshPropertyName() { return GET_MEMBER_NAME_STRING_CHECKED(UIKRetargeter, TargetPreviewMesh); }
+const FName UIKRetargeter::GetTargetPreviewMeshPropertyName() { return GET_MEMBER_NAME_STRING_CHECKED(UIKRetargeter, TargetPreviewMesh); };
 #endif
 
 void UIKRetargeter::PostLoad()
@@ -72,6 +72,14 @@ void UIKRetargeter::PostLoad()
 	Log.SetLogTarget(GetUniqueIDAsName(), false);
 }
 
+void UIKRetargeter::PostInitProperties()
+{
+	UObject::PostInitProperties();
+
+	// set the log target
+	Log.SetLogTarget(GetUniqueIDAsName(), false);
+}
+
 void FIKRetargetPose::SetBoneRotationOffset(FName BoneName, FQuat RotationDelta, const FIKRigSkeleton& Skeleton)
 {
 	FQuat* RotOffset = BoneRotationOffsets.Find(BoneName);
diff --git a/Engine/Plugins/Animation/IKRig/Source/IKRig/Public/IKRigDefinition.h b/Engine/Plugins/Animation/IKRig/Source/IKRig/Public/IKRigDefinition.h
--- a/Engine/Plugins/Animation/IKRig/Source/IKRig/Public/IKRigDefinition.h
+++ b/Engine/Plugins/Animation/IKRig/Source/IKRig/Public/IKRigDefinition.h
@@ -216,9 +216,11 @@ public:
 
 #endif
 
+	/** UObject */
 	virtual void Serialize(FArchive& Ar) override;
-
 	virtual void PostLoad() override;
+	virtual void PostInitProperties() override;
+	/** END UObject */
 	
 	/** hierarchy and bone-pose transforms */
 	UPROPERTY()
diff --git a/Engine/Plugins/Animation/IKRig/Source/IKRig/Public/Retargeter/IKRetargeter.h b/Engine/Plugins/Animation/IKRig/Source/IKRig/Public/Retargeter/IKRetargeter.h
--- a/Engine/Plugins/Animation/IKRig/Source/IKRig/Public/Retargeter/IKRetargeter.h
+++ b/Engine/Plugins/Animation/IKRig/Source/IKRig/Public/Retargeter/IKRetargeter.h
@@ -208,6 +208,11 @@ public:
 	/* Get name of default pose */
 	static const FName GetDefaultPoseName();
 
+	/** UObject */
+	virtual void PostLoad() override;
+	virtual void PostInitProperties() override;
+	/** END UObject */
+
 #if WITH_EDITOR
 	/* Get name of Source IK Rig property */
 	static const FName GetSourceIKRigPropertyName();
@@ -222,8 +227,6 @@ public:
 #if WITH_EDITOR
 	bool IsInEditRetargetPoseMode() const { return bEditRetargetPoseMode; }
 #endif
-	
-	virtual void PostLoad() override;
 
 private:
 
diff --git a/Engine/Plugins/Animation/IKRig/Source/IKRigEditor/Private/RetargetEditor/IKRetargetFactory.cpp b/Engine/Plugins/Animation/IKRig/Source/IKRigEditor/Private/RetargetEditor/IKRetargetFactory.cpp
--- a/Engine/Plugins/Animation/IKRig/Source/IKRigEditor/Private/RetargetEditor/IKRetargetFactory.cpp
+++ b/Engine/Plugins/Animation/IKRig/Source/IKRigEditor/Private/RetargetEditor/IKRetargetFactory.cpp
@@ -37,7 +37,7 @@ UObject* UIKRetargetFactory::FactoryCreateNew(
 		return nullptr;
 	}
 	
-	if (IsValid(SourceIKRig->GetPreviewMesh()))
+	if (!IsValid(SourceIKRig->GetPreviewMesh()))
 	{
 		UE_LOG(LogIKRigEditor, Warning, TEXT("Unable to create IK Retargter. Source IK Rig does not have an associated skeletal mesh."));
 		return nullptr;
