Fix crash when deleting goal and effector at same time.

diff --git a/Engine/Plugins/Animation/IKRig/Source/IKRigEditor/Private/RigEditor/IKRigController.cpp b/Engine/Plugins/Animation/IKRig/Source/IKRigEditor/Private/RigEditor/IKRigController.cpp
--- a/Engine/Plugins/Animation/IKRig/Source/IKRigEditor/Private/RigEditor/IKRigController.cpp
+++ b/Engine/Plugins/Animation/IKRig/Source/IKRigEditor/Private/RigEditor/IKRigController.cpp
@@ -873,9 +873,16 @@ bool UIKRigController::ConnectGoalToSolver(const UIKRigEffectorGoal& Goal, int32
 bool UIKRigController::DisconnectGoalFromSolver(const FName& GoalToRemove, int32 SolverIndex) const
 {
 	// can't remove goal that is not present in the core
-	check(GetGoalIndex(GoalToRemove) != INDEX_NONE);
+	if (GetGoalIndex(GoalToRemove) == INDEX_NONE)
+	{
+		return false;
+	}
+	
 	// can't remove goal from a solver with an invalid index
-	check(Asset->Solvers.IsValidIndex(SolverIndex))
+	if (!Asset->Solvers.IsValidIndex(SolverIndex))
+	{
+		return false;
+	}
 
     FScopedTransaction Transaction(LOCTEXT("DisconnectGoalSolver_Label", "Disconnect Goal from Solver"));
 	UIKRigSolver* Solver = Asset->Solvers[SolverIndex];
