Fix issue, not force loading meshes in retargeter when using override preview meshes. Regresson from 19598746

diff --git a/Engine/Plugins/Animation/IKRig/Source/IKRig/Public/Retargeter/IKRetargeter.h b/Engine/Plugins/Animation/IKRig/Source/IKRig/Public/Retargeter/IKRetargeter.h
--- a/Engine/Plugins/Animation/IKRig/Source/IKRig/Public/Retargeter/IKRetargeter.h
+++ b/Engine/Plugins/Animation/IKRig/Source/IKRig/Public/Retargeter/IKRetargeter.h
@@ -190,15 +190,15 @@ class IKRIG_API UIKRetargeter : public UObject
 public:
 
 	/** Get read-only access to the source IK Rig asset */
-	const UIKRigDefinition* GetSourceIKRig() const { return SourceIKRigAsset.Get(); };
+	const UIKRigDefinition* GetSourceIKRig() const { return SourceIKRigAsset.LoadSynchronous(); };
 	/** Get read-only access to the target IK Rig asset */
-	const UIKRigDefinition* GetTargetIKRig() const { return TargetIKRigAsset.Get(); };
+	const UIKRigDefinition* GetTargetIKRig() const { return TargetIKRigAsset.LoadSynchronous(); };
 	/** Get read-write access to the source IK Rig asset.
 	 * WARNING: do not use for editing the data model. Use Controller class instead. */
-	 UIKRigDefinition* GetSourceIKRigWriteable() const { return SourceIKRigAsset.Get(); };
+	 UIKRigDefinition* GetSourceIKRigWriteable() const { return SourceIKRigAsset.LoadSynchronous(); };
 	/** Get read-write access to the target IK Rig asset.
 	 * WARNING: do not use for editing the data model. Use Controller class instead. */
-	UIKRigDefinition* GetTargetIKRigWriteable() const { return TargetIKRigAsset.Get(); };
+	UIKRigDefinition* GetTargetIKRigWriteable() const { return TargetIKRigAsset.LoadSynchronous(); };
 	/** Get read-only access to the chain mapping */
 	const TArray<TObjectPtr<URetargetChainSettings>> GetAllChainSettings() const { return ChainSettings; };
 	/** Get read-only access to a retarget pose */
