Minor UI change in IK Rig editor: moved delete icon into it's own column.

diff --git a/Engine/Plugins/Animation/IKRig/Source/IKRigEditor/Private/RigEditor/SIKRigRetargetChainList.cpp b/Engine/Plugins/Animation/IKRig/Source/IKRigEditor/Private/RigEditor/SIKRigRetargetChainList.cpp
--- a/Engine/Plugins/Animation/IKRig/Source/IKRigEditor/Private/RigEditor/SIKRigRetargetChainList.cpp
+++ b/Engine/Plugins/Animation/IKRig/Source/IKRigEditor/Private/RigEditor/SIKRigRetargetChainList.cpp
@@ -24,6 +24,7 @@ static const FName ColumnId_ChainNameLabel( "Chain Name" );
 static const FName ColumnId_ChainStartLabel( "Start Bone" );
 static const FName ColumnId_ChainEndLabel( "End Bone" );
 static const FName ColumnId_IKGoalLabel( "IK Goal" );
+static const FName ColumnId_DeleteChainLabel( "Delete Chain" );
 
 TSharedRef<ITableRow> FRetargetChainElement::MakeListRowWidget(
 	const TSharedRef<STableViewBase>& InOwnerTable,
@@ -110,10 +111,9 @@ TSharedRef<SWidget> SIKRigRetargetChainRow::GenerateWidgetForColumn(const FName&
 		];
 		return EndWidget;
 	}
-	else
+
+	if (ColumnName == ColumnId_IKGoalLabel)
 	{
-		// ColumnId_IKGoalLabel
-		
 		TSharedRef<SWidget> GoalWidget =
 		SNew(SHorizontalBox)
 		+ SHorizontalBox::Slot()
@@ -130,7 +130,14 @@ TSharedRef<SWidget> SIKRigRetargetChainRow::GenerateWidgetForColumn(const FName&
 				SNew(STextBlock)
 				.Text(this, &SIKRigRetargetChainRow::GetGoalName)
 			]
-		]
+		];
+		return GoalWidget;
+	}
+
+	// ColumnName == ColumnId_DeleteChainLabel
+	{
+		TSharedRef<SWidget> DeleteWidget =
+		SNew(SHorizontalBox)
 		+ SHorizontalBox::Slot()
 		.AutoWidth()
 		.HAlign(HAlign_Right)
@@ -160,7 +167,7 @@ TSharedRef<SWidget> SIKRigRetargetChainRow::GenerateWidgetForColumn(const FName&
 				.ColorAndOpacity(FSlateColor::UseForeground())
 			]
 		];
-		return GoalWidget;
+		return DeleteWidget;
 	}
 }
 
@@ -357,6 +364,9 @@ void SIKRigRetargetChainList::Construct(const FArguments& InArgs, TSharedRef<FIK
 
 				+ SHeaderRow::Column( ColumnId_IKGoalLabel )
 				.DefaultLabel( LOCTEXT( "IKGoalColumnLabel", "IK Goal" ) )
+				
+				+ SHeaderRow::Column( ColumnId_DeleteChainLabel )
+				.DefaultLabel( LOCTEXT( "DeleteChainColumnLabel", "Delete Chain" ) )
 			)
 		]
     ];
