Fixing IK Retargeter editor bug, GoToRetargetPose not showing retarget pose on target character.

diff --git a/Engine/Plugins/Animation/IKRig/Source/IKRig/Private/Retargeter/IKRetargetProcessor.cpp b/Engine/Plugins/Animation/IKRig/Source/IKRig/Private/Retargeter/IKRetargetProcessor.cpp
--- a/Engine/Plugins/Animation/IKRig/Source/IKRig/Private/Retargeter/IKRetargetProcessor.cpp
+++ b/Engine/Plugins/Animation/IKRig/Source/IKRig/Private/Retargeter/IKRetargetProcessor.cpp
@@ -1403,7 +1403,9 @@ TArray<FTransform>&  UIKRetargetProcessor::RunRetargeter(const TArray<FTransform
 	// in edit mode we just want to see the edited reference pose, not actually run the retargeting
 	// as long as the retargeter is reinitialized after every modification to the limb rotation offsets,
 	// then the TargetSkeleton.RetargetGlobalPose will contain the updated retarget pose.
-	if (RetargeterAsset->IsInEditRetargetPoseMode() && RetargeterAsset->GetTargetIKRig())
+	const ERetargeterOutputMode CurrentMode = RetargeterAsset->GetOutputMode();
+	const bool bOutputRetargetPose = CurrentMode == ERetargeterOutputMode::EditRetargetPose || CurrentMode == ERetargeterOutputMode::ShowRetargetPose;
+	if (bOutputRetargetPose && RetargeterAsset->GetTargetIKRig())
 	{
 		const FName RootBoneName = RetargeterAsset->GetTargetIKRig()->GetRetargetRoot();
 		TargetSkeleton.GenerateRetargetPose(RetargeterAsset->GetCurrentRetargetPose(), RootBoneName);
diff --git a/Engine/Plugins/Animation/IKRig/Source/IKRig/Public/Retargeter/IKRetargeter.h b/Engine/Plugins/Animation/IKRig/Source/IKRig/Public/Retargeter/IKRetargeter.h
--- a/Engine/Plugins/Animation/IKRig/Source/IKRig/Public/Retargeter/IKRetargeter.h
+++ b/Engine/Plugins/Animation/IKRig/Source/IKRig/Public/Retargeter/IKRetargeter.h
@@ -182,6 +182,14 @@ public:
 	void SortHierarchically(const FIKRigSkeleton& Skeleton);
 };
 
+#if WITH_EDITOR
+enum class ERetargeterOutputMode : uint8
+{
+	RunRetarget,
+	ShowRetargetPose,
+	EditRetargetPose
+};
+#endif
 
 UCLASS(Blueprintable)
 class IKRIG_API UIKRetargeter : public UObject
@@ -225,7 +233,8 @@ public:
 #endif
 
 #if WITH_EDITOR
-	bool IsInEditRetargetPoseMode() const { return bEditRetargetPoseMode; }
+	ERetargeterOutputMode GetOutputMode() const { return bOutputMode; }
+	void SetOutputMode(const ERetargeterOutputMode InMode) { bOutputMode = InMode; }
 #endif
 
 private:
@@ -292,8 +301,7 @@ public:
 private:
 	/** A special editor-only mode which forces the retargeter to output the current retarget reference pose,
 	* rather than actually running the retarget and outputting the retargeted pose. Used in Edit-Pose mode.*/
-	UPROPERTY()
-	bool bEditRetargetPoseMode = false;
+	ERetargeterOutputMode bOutputMode;
 
 	/** The controller responsible for managing this asset's data (all editor mutation goes through this) */
 	UPROPERTY(Transient)
diff --git a/Engine/Plugins/Animation/IKRig/Source/IKRigEditor/Private/RetargetEditor/IKRetargetDefaultMode.cpp b/Engine/Plugins/Animation/IKRig/Source/IKRigEditor/Private/RetargetEditor/IKRetargetDefaultMode.cpp
--- a/Engine/Plugins/Animation/IKRig/Source/IKRigEditor/Private/RetargetEditor/IKRetargetDefaultMode.cpp
+++ b/Engine/Plugins/Animation/IKRig/Source/IKRigEditor/Private/RetargetEditor/IKRetargetDefaultMode.cpp
@@ -230,6 +230,7 @@ void FIKRetargetDefaultMode::Enter()
 	Controller->Editor.Pin()->GetPersonaToolkit()->GetPreviewScene()->SetAllowMeshHitProxies(true);
 	Controller->SourceSkelMeshComponent->bSelectable = true;
 	Controller->TargetSkelMeshComponent->bSelectable = true;
+	Controller->AssetController->GetAsset()->SetOutputMode(ERetargeterOutputMode::RunRetarget);
 
 	// deselect when entering mode
 	SetSelectedComponent(nullptr);
diff --git a/Engine/Plugins/Animation/IKRig/Source/IKRigEditor/Private/RetargetEditor/IKRetargetEditPoseMode.cpp b/Engine/Plugins/Animation/IKRig/Source/IKRigEditor/Private/RetargetEditor/IKRetargetEditPoseMode.cpp
--- a/Engine/Plugins/Animation/IKRig/Source/IKRigEditor/Private/RetargetEditor/IKRetargetEditPoseMode.cpp
+++ b/Engine/Plugins/Animation/IKRig/Source/IKRigEditor/Private/RetargetEditor/IKRetargetEditPoseMode.cpp
@@ -414,7 +414,7 @@ void FIKRetargetEditPoseMode::Enter()
 	Controller->AddOffsetAndUpdatePreviewMeshPosition(FVector::ZeroVector, Controller->SourceSkelMeshComponent);
 
 	// put asset in mode where target mesh will output retarget pose (for preview purposes)
-	Controller->AssetController->SetEditRetargetPoseMode(true);
+	Controller->AssetController->GetAsset()->SetOutputMode(ERetargeterOutputMode::EditRetargetPose);
 }
 
 void FIKRetargetEditPoseMode::Exit()
@@ -425,7 +425,7 @@ void FIKRetargetEditPoseMode::Exit()
 		return;
 	}
 
-	Controller->AssetController->SetEditRetargetPoseMode(false);
+	Controller->AssetController->GetAsset()->SetOutputMode(ERetargeterOutputMode::RunRetarget);
 }
 
 void FIKRetargetEditPoseMode::UpdateWidgetTransform()
diff --git a/Engine/Plugins/Animation/IKRig/Source/IKRigEditor/Private/RetargetEditor/IKRetargetEditorController.cpp b/Engine/Plugins/Animation/IKRig/Source/IKRigEditor/Private/RetargetEditor/IKRetargetEditorController.cpp
--- a/Engine/Plugins/Animation/IKRig/Source/IKRigEditor/Private/RetargetEditor/IKRetargetEditorController.cpp
+++ b/Engine/Plugins/Animation/IKRig/Source/IKRigEditor/Private/RetargetEditor/IKRetargetEditorController.cpp
@@ -232,6 +232,8 @@ void FIKRetargetEditorController::PlayAnimationAsset(UAnimationAsset* AssetToPla
 	{
 		SourceAnimInstance->SetAnimationAsset(AssetToPlay);
 		PreviousAsset = AssetToPlay;
+		// tell asset to output the retargeted pose
+		AssetController->GetAsset()->SetOutputMode(ERetargeterOutputMode::RunRetarget);
 	}
 }
 
@@ -240,6 +242,8 @@ void FIKRetargetEditorController::PlayPreviousAnimationAsset() const
 	if (PreviousAsset)
 	{
 		SourceAnimInstance->SetAnimationAsset(PreviousAsset);
+		// tell asset to output the retarget
+		AssetController->GetAsset()->SetOutputMode(ERetargeterOutputMode::RunRetarget);
 	}
 }
 
@@ -252,6 +256,8 @@ void FIKRetargetEditorController::HandleGoToRetargetPose() const
 	SourceSkelMeshComponent->ShowReferencePose(true);
 	// have to move component back to offset position because ShowReferencePose() sets it back to origin
 	AddOffsetAndUpdatePreviewMeshPosition(FVector::ZeroVector, SourceSkelMeshComponent);
+	// tell asset to output the retarget pose
+	AssetController->GetAsset()->SetOutputMode(ERetargeterOutputMode::ShowRetargetPose);
 }
 
 void FIKRetargetEditorController::HandleEditPose() const
@@ -288,7 +294,7 @@ bool FIKRetargetEditorController::CanEditPose() const
 
 bool FIKRetargetEditorController::IsEditingPose() const
 {
-	return AssetController->GetEditRetargetPoseMode();
+	return AssetController->GetAsset()->GetOutputMode() == ERetargeterOutputMode::EditRetargetPose;
 }
 
 void FIKRetargetEditorController::HandleNewPose()
diff --git a/Engine/Plugins/Animation/IKRig/Source/IKRigEditor/Private/RetargetEditor/IKRetargeterController.cpp b/Engine/Plugins/Animation/IKRig/Source/IKRigEditor/Private/RetargetEditor/IKRetargeterController.cpp
--- a/Engine/Plugins/Animation/IKRig/Source/IKRigEditor/Private/RetargetEditor/IKRetargeterController.cpp
+++ b/Engine/Plugins/Animation/IKRig/Source/IKRigEditor/Private/RetargetEditor/IKRetargeterController.cpp
@@ -474,16 +474,6 @@ void UIKRetargeterController::AddTranslationOffsetToRetargetRootBone(FVector Tra
 	Asset->RetargetPoses[Asset->CurrentRetargetPose].AddTranslationDeltaToRoot(TranslationOffset);
 }
 
-bool UIKRetargeterController::GetEditRetargetPoseMode() const
-{
-	return GetAsset()->bEditRetargetPoseMode;
-}
-
-void UIKRetargeterController::SetEditRetargetPoseMode(const bool bMode) const
-{
-	GetAsset()->bEditRetargetPoseMode = bMode;
-}
-
 FName UIKRetargeterController::MakePoseNameUnique(FName PoseName) const
 {
 	FName UniqueName = PoseName;
diff --git a/Engine/Plugins/Animation/IKRig/Source/IKRigEditor/Public/RetargetEditor/IKRetargeterController.h b/Engine/Plugins/Animation/IKRig/Source/IKRigEditor/Public/RetargetEditor/IKRetargeterController.h
--- a/Engine/Plugins/Animation/IKRig/Source/IKRigEditor/Public/RetargetEditor/IKRetargeterController.h
+++ b/Engine/Plugins/Animation/IKRig/Source/IKRigEditor/Public/RetargetEditor/IKRetargeterController.h
@@ -98,10 +98,6 @@ public:
 	FQuat GetRotationOffsetForRetargetPoseBone(FName BoneName) const;
 	/** Add a delta translation to the root bone (used in Edit Mode in the retarget editor) */
 	void AddTranslationOffsetToRetargetRootBone(FVector TranslationOffset) const;
-	/** Get whether in mode to output retarget pose (true) or run retarget (false). */
-	bool GetEditRetargetPoseMode() const;
-	/** Set whether retarget asset should be in "Edit Pose" mode. */
-	void SetEditRetargetPoseMode(const bool bMode) const;
 	/** Add a numbered suffix to the given pose name to make it unique. */
 	FName MakePoseNameUnique(FName PoseName) const;
 	/** END RETARGET POSE EDITING */
