Fixes regression causing chain mappings to load incorrectly.
Fixes bug, chain mappings get blown away when clearing target IK Rig asset.

diff --git a/Engine/Plugins/Animation/IKRig/Source/IKRigEditor/Private/RetargetEditor/IKRetargeterController.cpp b/Engine/Plugins/Animation/IKRig/Source/IKRigEditor/Private/RetargetEditor/IKRetargeterController.cpp
--- a/Engine/Plugins/Animation/IKRig/Source/IKRigEditor/Private/RetargetEditor/IKRetargeterController.cpp
+++ b/Engine/Plugins/Animation/IKRig/Source/IKRigEditor/Private/RetargetEditor/IKRetargeterController.cpp
@@ -99,19 +99,21 @@ USkeletalMesh* UIKRetargeterController::GetTargetPreviewMesh() const
 
 FName UIKRetargeterController::GetSourceRootBone() const
 {
-	return Asset->SourceIKRigAsset.IsValid() ? Asset->SourceIKRigAsset->GetRetargetRoot() : FName("None");
+	const UIKRigDefinition* SourceIKRig = Asset->GetSourceIKRig();
+	return SourceIKRig ? SourceIKRig->GetRetargetRoot() : FName("None");
 }
 
 FName UIKRetargeterController::GetTargetRootBone() const
 {
-	return Asset->TargetIKRigAsset.IsValid() ? Asset->TargetIKRigAsset->GetRetargetRoot() : FName("None");
+	const UIKRigDefinition* TargetIKRig = Asset->GetTargetIKRig();
+	return TargetIKRig ? TargetIKRig->GetRetargetRoot() : FName("None");
 }
 
 void UIKRetargeterController::GetTargetChainNames(TArray<FName>& OutNames) const
 {
-	if (Asset->TargetIKRigAsset.IsValid())
+	if (const UIKRigDefinition* TargetIKRig = Asset->GetTargetIKRig())
 	{
-		const TArray<FBoneChain>& Chains = Asset->TargetIKRigAsset->GetRetargetChains();
+		const TArray<FBoneChain>& Chains = TargetIKRig->GetRetargetChains();
 		for (const FBoneChain& Chain : Chains)
 		{
 			OutNames.Add(Chain.ChainName);
@@ -121,9 +123,9 @@ void UIKRetargeterController::GetTargetChainNames(TArray<FName>& OutNames) const
 
 void UIKRetargeterController::GetSourceChainNames(TArray<FName>& OutNames) const
 {
-	if (Asset->SourceIKRigAsset.IsValid())
+	if (const UIKRigDefinition* SourceIKRig = Asset->GetSourceIKRig())
 	{
-		const TArray<FBoneChain>& Chains = Asset->SourceIKRigAsset->GetRetargetChains();
+		const TArray<FBoneChain>& Chains = SourceIKRig->GetRetargetChains();
 		for (const FBoneChain& Chain : Chains)
 		{
 			OutNames.Add(Chain.ChainName);
@@ -131,12 +133,10 @@ void UIKRetargeterController::GetSourceChainNames(TArray<FName>& OutNames) const
 	}
 }
 
-void UIKRetargeterController::CleanChainMapping(const bool bForceReinitialization)
+void UIKRetargeterController::CleanChainMapping(const bool bForceReinitialization) const
 {
-	if (!Asset->TargetIKRigAsset.IsValid())
+	if (Asset->TargetIKRigAsset.IsNull())
 	{
-		// empty the chain mapping
-		Asset->ChainSettings.Empty();
 		return;
 	}
 	
diff --git a/Engine/Plugins/Animation/IKRig/Source/IKRigEditor/Public/RetargetEditor/IKRetargeterController.h b/Engine/Plugins/Animation/IKRig/Source/IKRigEditor/Public/RetargetEditor/IKRetargeterController.h
--- a/Engine/Plugins/Animation/IKRig/Source/IKRigEditor/Public/RetargetEditor/IKRetargeterController.h
+++ b/Engine/Plugins/Animation/IKRig/Source/IKRigEditor/Public/RetargetEditor/IKRetargeterController.h
@@ -60,7 +60,7 @@ public:
 	/** Get names of all the source bone chains. */
 	void GetSourceChainNames(TArray<FName>& OutNames) const;
 	/** Remove invalid chain mappings (no longer existing in currently referenced source/target IK Rig assets) */
-	void CleanChainMapping(const bool bForceReinitialization=true);
+	void CleanChainMapping(const bool bForceReinitialization=true) const;
 	/** Use fuzzy string search to find "best" Source chain to map to each Target chain */
 	void AutoMapChains() const;
 	/** Callback when IK Rig chain is renamed. Retains existing mappings using the new name */
