Fixing issue, Set Transform settings not being serialized.

diff --git a/Engine/Plugins/Animation/IKRig/Source/IKRig/Public/Solvers/IKRig_SetTransform.h b/Engine/Plugins/Animation/IKRig/Source/IKRig/Public/Solvers/IKRig_SetTransform.h
--- a/Engine/Plugins/Animation/IKRig/Source/IKRig/Public/Solvers/IKRig_SetTransform.h
+++ b/Engine/Plugins/Animation/IKRig/Source/IKRig/Public/Solvers/IKRig_SetTransform.h
@@ -42,7 +42,7 @@ public:
 	UPROPERTY(VisibleAnywhere, Category = "Set Transform Settings")
 	FName RootBone;
 	
-	UPROPERTY(Transient)
+	UPROPERTY()
 	TObjectPtr<UIKRig_SetTransformEffector> Effector;
 
 	/** UIKRigSolver interface */
