Force retarget source asset to use target skeletal mesh.

diff --git a/Engine/Plugins/Animation/IKRig/Source/IKRigEditor/Private/RetargetEditor/IKRetargetBatchOperation.cpp b/Engine/Plugins/Animation/IKRig/Source/IKRigEditor/Private/RetargetEditor/IKRetargetBatchOperation.cpp
--- a/Engine/Plugins/Animation/IKRig/Source/IKRigEditor/Private/RetargetEditor/IKRetargetBatchOperation.cpp
+++ b/Engine/Plugins/Animation/IKRig/Source/IKRigEditor/Private/RetargetEditor/IKRetargetBatchOperation.cpp
@@ -167,56 +167,6 @@ void FIKRetargetBatchOperation::RetargetAssets(
 {
 	USkeleton* OldSkeleton = Context.SourceMesh->GetSkeleton();
 	USkeleton* NewSkeleton = Context.TargetMesh->GetSkeleton();
-
-	auto UpdateRetargetSource = [&Context](UAnimSequence* InAnimSequenceToRetarget)
-	{
-#if WITH_EDITORONLY_DATA
-
-		const bool bHasRetargetSource = InAnimSequenceToRetarget->RetargetSource != NAME_None;
-		const bool bHasRetargetSourceAsset = !InAnimSequenceToRetarget->RetargetSourceAsset.IsNull();
-		if (!bHasRetargetSource && !bHasRetargetSourceAsset)
-		{
-			return;
-		}
-		
-		USkeleton* NewSkeleton = Context.TargetMesh->GetSkeleton();
-		const FName TargetRetargetSource = NewSkeleton->GetRetargetSourceForMesh(Context.TargetMesh);
-
-		// NOTE: should we keep the source anim sequence RetargetSource / RetargetSourceAsset if we don't find any
-		// TargetRetargetSource or reset them to default values?
-		
-		// update RetargetSource
-		if (bHasRetargetSource)
-		{
-			InAnimSequenceToRetarget->RetargetSource = TargetRetargetSource;
-		}
-
-		// update RetargetSourceAsset
-		if (bHasRetargetSourceAsset)
-		{
-			const FReferencePose* PoseFound = NewSkeleton->AnimRetargetSources.Find(TargetRetargetSource);
-			if (PoseFound && PoseFound->SourceReferenceMesh)
-			{
-				InAnimSequenceToRetarget->RetargetSourceAsset = PoseFound->SourceReferenceMesh;
-
-				// make sure we update the RetargetSourceAssetReferencePose accordingly
-				const USkeletalMesh* NewReferenceMesh = InAnimSequenceToRetarget->RetargetSourceAsset.LoadSynchronous();
-				if (NewReferenceMesh)
-				{
-					FAnimationRuntime::MakeSkeletonRefPoseFromMesh(NewReferenceMesh, NewSkeleton, InAnimSequenceToRetarget->RetargetSourceAssetReferencePose);
-				}
-				else
-				{
-					InAnimSequenceToRetarget->RetargetSourceAssetReferencePose.Empty();
-				}
-				return;
-			}
-
-			InAnimSequenceToRetarget->RetargetSourceAsset.Reset();
-			InAnimSequenceToRetarget->RetargetSourceAssetReferencePose.Empty();
-		}
-#endif
-	};
 	
 	for (UAnimationAsset* AssetToRetarget : AnimationAssetsToRetarget)
 	{
@@ -233,8 +183,9 @@ void FIKRetargetBatchOperation::RetargetAssets(
 			Controller.RemoveAllCurvesOfType(ERawCurveTrackTypes::RCT_Transform);
 			// clear bone tracks to prevent recompression
 			Controller.RemoveAllBoneTracks(false);
-			// make sure we update the retarget source
-			UpdateRetargetSource(AnimSequenceToRetarget);
+			// set the retarget source to the target skeletal mesh
+			AnimSequenceToRetarget->RetargetSource = NAME_None;
+			AnimSequenceToRetarget->RetargetSourceAsset = Context.TargetMesh;
 			// done editing sequence data, close bracket
 			Controller.CloseBracket(ShouldTransactAnimEdits);
 		}
