Adding root settings to retargeter.
Adds rotation offset to chain settings.

diff --git a/Engine/Plugins/Animation/IKRig/Source/IKRig/Private/Retargeter/IKRetargetProcessor.cpp b/Engine/Plugins/Animation/IKRig/Source/IKRig/Private/Retargeter/IKRetargetProcessor.cpp
--- a/Engine/Plugins/Animation/IKRig/Source/IKRig/Private/Retargeter/IKRetargetProcessor.cpp
+++ b/Engine/Plugins/Animation/IKRig/Source/IKRig/Private/Retargeter/IKRetargetProcessor.cpp
@@ -611,7 +611,7 @@ void FChainDecoderFK::DecodePose(
 				break;
 			case ERetargetTranslationMode::GloballyScaled:
 				{
-					OutPosition = SourceCurrentTransform.GetTranslation() * RootRetargeter.GlobalScale; // todo expose global scale modifier
+					OutPosition = SourceCurrentTransform.GetTranslation() * RootRetargeter.GetGlobalScaleVector();
 				}
 				break;
 			case ERetargetTranslationMode::Absolute:
@@ -825,76 +825,79 @@ void FChainRetargeterIK::DecodePose(
     const TArray<FTransform>& OutGlobalPose,
     FDecodedIKChain& OutResults)
 {
-	FVector EndPosition;
-	if (Settings.DriveIKGoal)
+	// record the end bone rotation on the input pose
+	const FQuat InputEndRotation = OutGlobalPose[Target.BoneIndexC].GetRotation();
+	// we have to "undo" the end bone delta, otherwise we will get double-transformations because the FK pass has already rotated the foot
+	const FQuat InputToInitialDeltaRotation = InputEndRotation * Target.InitialEndRotation.Inverse();
+	const FQuat Rotation = InputToInitialDeltaRotation * Target.InitialEndRotation;
+	
+	if (!Settings.DriveIKGoal)
 	{
-		// set position to length-scaled direction from source limb
-		const FVector Start = OutGlobalPose[Target.BoneIndexA].GetTranslation();
-		EndPosition = Start + (Source.CurrentEndDirectionNormalized * Target.InitialLength);
+		// set goal transform to the input coming from the previous retarget phase (FK if enabled)
+		OutResults.EndEffectorPosition = OutGlobalPose[Target.BoneIndexC].GetTranslation();
+		OutResults.EndEffectorRotation = Rotation;
+		return;
+	}
+	
+	// apply static rotation offset in the local space of the foot
+	const FQuat GoalRotation = Rotation * Settings.StaticRotationOffset.Quaternion();
 
-		// blend to source location
-		if (Settings.BlendToSource > KINDA_SMALL_NUMBER)
-		{
-			const FVector Weight = Settings.BlendToSource * Settings.BlendToSourceWeights;
-			EndPosition.X = FMath::Lerp(EndPosition.X, Source.CurrentEndPosition.X, Weight.X);
-			EndPosition.Y = FMath::Lerp(EndPosition.Y, Source.CurrentEndPosition.Y, Weight.Y);
-			EndPosition.Z = FMath::Lerp(EndPosition.Z, Source.CurrentEndPosition.Z, Weight.Z);
-		}
+	//
+	// calculate position of IK goal ...
+	//
+	
+	// set position to length-scaled direction from source limb
+	const FVector Start = OutGlobalPose[Target.BoneIndexA].GetTranslation();
+	FVector GoalPosition = Start + (Source.CurrentEndDirectionNormalized * Target.InitialLength);
 
-		// apply static offset
-		EndPosition += Settings.StaticOffset;
+	// blend to source location
+	if (Settings.BlendToSource > KINDA_SMALL_NUMBER)
+	{
+		const FVector Weight = Settings.BlendToSource * Settings.BlendToSourceWeights;
+		GoalPosition.X = FMath::Lerp(GoalPosition.X, Source.CurrentEndPosition.X, Weight.X);
+		GoalPosition.Y = FMath::Lerp(GoalPosition.Y, Source.CurrentEndPosition.Y, Weight.Y);
+		GoalPosition.Z = FMath::Lerp(GoalPosition.Z, Source.CurrentEndPosition.Z, Weight.Z);
+	}
 
-		// apply extension
-		if (!FMath::IsNearlyEqual(Settings.Extension, 1.0f))
-		{
-			EndPosition = Start + (EndPosition - Start) * Settings.Extension;	
-		}
+	// apply global static offset
+	GoalPosition += Settings.StaticOffset;
 
-		// match velocity
-		if (Settings.MatchSourceVelocity > KINDA_SMALL_NUMBER)
-		{
-			const FVector SourceVelocity = Source.CurrentEndPosition - Source.PreviousEndPosition;
-			const float SourceSpeed = SourceVelocity.Size();
+	// apply local static offset
+	GoalPosition += GoalRotation.RotateVector(Settings.StaticLocalOffset);
 
-			const FVector TargetVelocity = EndPosition - Target.PrevEndPosition;
-			const float TargetSpeed = TargetVelocity.Size();
+	// apply extension
+	if (!FMath::IsNearlyEqual(Settings.Extension, 1.0f))
+	{
+		GoalPosition = Start + (GoalPosition - Start) * Settings.Extension;	
+	}
 
-			// if target is moving slowly enough, start matching velocity
-			if (TargetSpeed < Settings.TeleportVelocityThreshold)
-			{
-				// match target speed to source
-				const float BlendedSpeed = FMath::Lerp(TargetSpeed, SourceSpeed, Settings.MatchSourceVelocity);
-				EndPosition = Target.PrevEndPosition + TargetVelocity * BlendedSpeed;
+	// match velocity
+	if (Settings.MatchSourceVelocity > KINDA_SMALL_NUMBER)
+	{
+		const FVector SourceVelocity = Source.CurrentEndPosition - Source.PreviousEndPosition;
+		const float SourceSpeed = SourceVelocity.Size();
 
-				// what would happen if we blended into and out of this state?
-				// TimeSinceTargetStartedMatchingVelocity
-				// WasMatchingVelocity
-			}
+		const FVector TargetVelocity = GoalPosition - Target.PrevEndPosition;
+		const float TargetSpeed = TargetVelocity.Size();
+
+		// if target is moving slowly enough, start matching velocity
+		if (TargetSpeed < Settings.TeleportVelocityThreshold)
+		{
+			// match target speed to source
+			const float BlendedSpeed = FMath::Lerp(TargetSpeed, SourceSpeed, Settings.MatchSourceVelocity);
+			GoalPosition = Target.PrevEndPosition + TargetVelocity * BlendedSpeed;
+
+			// what would happen if we blended into and out of this state?
+			// TimeSinceTargetStartedMatchingVelocity
+			// WasMatchingVelocity
 		}
-		
-		// optionally factor in ground height
-		//const bool IsGrounded = false;
-		//if (IsGrounded) // todo expose this
-		//{
-		//	OutResults.EndEffectorPosition.Z = (EncodedChain.HeightFromGroundNormalized * Length) + EndPositionOrig.Z;
-		//}
 	}
-	else
-	{
-		// set goal location to the input coming from the previous retarget phase (FK if enabled)
-		EndPosition = OutGlobalPose[Target.BoneIndexC].GetTranslation();
-	}
-
-	// apply effector position
-	OutResults.EndEffectorPosition = EndPosition;
-	Target.PrevEndPosition = EndPosition;
 	
-	// calculate end effector rotation
-	const FQuat RotationDelta = Source.CurrentEndRotation * Source.InitialEndRotation.Inverse();
-	OutResults.EndEffectorRotation = RotationDelta * Target.InitialEndRotation;
-
-	// TBD calc pole vector position
-	OutResults.PoleVectorPosition = FVector::OneVector;
+	// output transform
+	OutResults.EndEffectorPosition = GoalPosition;
+	OutResults.EndEffectorRotation = GoalRotation;
+	OutResults.PoleVectorPosition = FVector::OneVector; // TODO calc pole vector position
+	Target.PrevEndPosition = GoalPosition;
 }
 
 bool FRetargetChainPair::Initialize(
@@ -1120,6 +1123,7 @@ bool FRootRetargeter::InitializeTarget(
 	const FTransform TargetInitialTransform = TargetSkeleton.RetargetGlobalPose[Target.BoneIndex];
 	Target.InitialHeight = TargetInitialTransform.GetTranslation().Z;
 	Target.InitialRotation = TargetInitialTransform.GetRotation();
+	Target.InitialPosition = TargetInitialTransform.GetTranslation();
 
 	return true;
 }
@@ -1133,23 +1137,30 @@ void FRootRetargeter::Reset()
 void FRootRetargeter::EncodePose(const TArray<FTransform>& SourceGlobalPose)
 {
 	const FTransform& SourceTransform = SourceGlobalPose[Source.BoneIndex];
-	Source.CurrentPositionNormalized = SourceTransform.GetTranslation() * Source.InitialHeightInverse;
+	Source.CurrentPosition = SourceTransform.GetTranslation();
+	Source.CurrentPositionNormalized = Source.CurrentPosition * Source.InitialHeightInverse;
 	Source.CurrentRotation = SourceTransform.GetRotation();	
 }
 
-void FRootRetargeter::DecodePose(
-	TArray<FTransform>& OutTargetGlobalPose,
-	const float StrideScale) const
+void FRootRetargeter::DecodePose(TArray<FTransform>& OutTargetGlobalPose) const
 {
 	// scale normalized position by root height
-	FVector Position = Source.CurrentPositionNormalized * Target.InitialHeight;
-	// scale horizontal displacement by stride scale
-	Position *= FVector(StrideScale, StrideScale, 1.0f);
+	FVector RetargetedPosition = Source.CurrentPositionNormalized * Target.InitialHeight;// * GlobalScaleVertical;
+	RetargetedPosition.Z *= GlobalScaleVertical;
+	// globally scale offset of root
+	const FVector RootOffset = (RetargetedPosition - Target.InitialPosition) * FVector(GlobalScaleHorizontal, GlobalScaleHorizontal, 1.0f);
+	RetargetedPosition = Target.InitialPosition + RootOffset;
+	// blend the retarget root position towards the source retarget root position
+	FVector Position = FMath::Lerp(RetargetedPosition, Source.CurrentPosition, BlendToSource);
+	// apply a static offset
+	Position += StaticOffset;
 
 	// calc offset between initial source/target root rotations
 	const FQuat RotationDelta = Source.CurrentRotation * Source.InitialRotation.Inverse();
-	// add offset to the current source rotation
-	const FQuat Rotation = RotationDelta * Target.InitialRotation;
+	// add retarget pose delta to the current source rotation
+	FQuat Rotation = RotationDelta * Target.InitialRotation;
+	// add static rotation offset
+	Rotation = StaticRotationOffset.Quaternion() * Rotation;
 
 	// apply to target
 	FTransform& TargetRootTransform = OutTargetGlobalPose[Target.BoneIndex];
@@ -1446,8 +1457,7 @@ void UIKRetargetProcessor::RunRootRetarget(
     TArray<FTransform>& OutGlobalTransforms)
 {
 	RootRetargeter.EncodePose(InGlobalTransforms);
-	const float StrideScale = 1.0f;
-	RootRetargeter.DecodePose(OutGlobalTransforms, StrideScale);
+	RootRetargeter.DecodePose(OutGlobalTransforms);
 }
 
 void UIKRetargetProcessor::RunFKRetarget(
@@ -1500,7 +1510,7 @@ void UIKRetargetProcessor::RunIKRetarget(
 			OutIKGoal.EndEffectorPosition,
 			OutIKGoal.EndEffectorRotation,
 			1.0f,
-			0.0f,
+			1.0f,
 			EIKRigGoalSpace::Component,
 			EIKRigGoalSpace::Component);
 		IKRigProcessor->SetIKGoal(Goal);
@@ -1583,6 +1593,14 @@ void UIKRetargetProcessor::CopyAllSettingsFromAsset()
 			}
 		}
 	}
+
+	// copy root settings
+	const URetargetRootSettings& RootSettings = *RetargeterAsset->RootSettings;
+	RootRetargeter.GlobalScaleHorizontal = RootSettings.GlobalScaleHorizontal;
+	RootRetargeter.GlobalScaleVertical = RootSettings.GlobalScaleVertical;
+	RootRetargeter.BlendToSource = RootSettings.BlendToSource;
+	RootRetargeter.StaticOffset = RootSettings.StaticOffset;
+	RootRetargeter.StaticRotationOffset = RootSettings.StaticRotationOffset;
 }
 #endif
 
diff --git a/Engine/Plugins/Animation/IKRig/Source/IKRig/Private/Retargeter/IKRetargeter.cpp b/Engine/Plugins/Animation/IKRig/Source/IKRig/Private/Retargeter/IKRetargeter.cpp
--- a/Engine/Plugins/Animation/IKRig/Source/IKRig/Private/Retargeter/IKRetargeter.cpp
+++ b/Engine/Plugins/Animation/IKRig/Source/IKRig/Private/Retargeter/IKRetargeter.cpp
@@ -9,6 +9,13 @@ const FName UIKRetargeter::GetSourcePreviewMeshPropertyName() { return GET_MEMBE
 const FName UIKRetargeter::GetTargetPreviewMeshPropertyName() { return GET_MEMBER_NAME_STRING_CHECKED(UIKRetargeter, TargetPreviewMesh); };
 #endif
 
+UIKRetargeter::UIKRetargeter(const FObjectInitializer& ObjectInitializer)
+	: Super(ObjectInitializer)
+{
+	RootSettings = CreateDefaultSubobject<URetargetRootSettings>(TEXT("RootSettings"));
+	RootSettings->SetFlags(RF_Transactional);
+}
+
 void UIKRetargeter::PostLoad()
 {
 	Super::PostLoad();
diff --git a/Engine/Plugins/Animation/IKRig/Source/IKRig/Public/Retargeter/IKRetargetProcessor.h b/Engine/Plugins/Animation/IKRig/Source/IKRig/Public/Retargeter/IKRetargetProcessor.h
--- a/Engine/Plugins/Animation/IKRig/Source/IKRig/Public/Retargeter/IKRetargetProcessor.h
+++ b/Engine/Plugins/Animation/IKRig/Source/IKRig/Public/Retargeter/IKRetargetProcessor.h
@@ -112,20 +112,19 @@ struct IKRIG_API FResolvedBoneChain
 struct FRootSource
 {
 	int32 BoneIndex;
-
 	FQuat InitialRotation;
-
 	float InitialHeightInverse;
-
+	FVector InitialPosition;
+	FVector CurrentPosition;
 	FVector CurrentPositionNormalized;
-	
 	FQuat CurrentRotation;
 };
 
 struct FRootTarget
 {
 	int32 BoneIndex;
-	
+
+	FVector InitialPosition;
 	FQuat InitialRotation;
 	
 	float InitialHeight;
@@ -137,7 +136,15 @@ struct FRootRetargeter
 	
 	FRootTarget Target;
 
-	float GlobalScale = 1.0f;
+	float GlobalScaleHorizontal = 1.0f;
+
+	float GlobalScaleVertical = 1.0f;
+
+	FVector BlendToSource = FVector::ZeroVector;
+
+	FVector StaticOffset = FVector::ZeroVector;
+
+	FRotator StaticRotationOffset = FRotator::ZeroRotator;
 
 	void Reset();
 	
@@ -153,7 +160,12 @@ struct FRootRetargeter
 
 	void EncodePose(const TArray<FTransform> &SourceGlobalPose);
 	
-	void DecodePose( TArray<FTransform> &OutTargetGlobalPose, const float StrideScale) const;
+	void DecodePose( TArray<FTransform> &OutTargetGlobalPose) const;
+
+	FVector GetGlobalScaleVector() const
+	{
+		return FVector(GlobalScaleHorizontal, GlobalScaleHorizontal, GlobalScaleVertical);
+	}
 };
 
 struct FRetargetChainSettings
@@ -168,6 +180,8 @@ struct FRetargetChainSettings
 
 	bool DriveIKGoal = true;
 	FVector StaticOffset;
+	FVector StaticLocalOffset;
+	FRotator StaticRotationOffset;
 	float Extension = 1.0f;
 	float BlendToSource = 0.0f;
 	FVector BlendToSourceWeights = FVector::OneVector;
@@ -189,6 +203,8 @@ public:
 		DriveIKGoal = AssetChainSettings->DriveIKGoal;
 		Extension = AssetChainSettings->Extension;
 		StaticOffset = AssetChainSettings->StaticOffset;
+		StaticLocalOffset = AssetChainSettings->StaticLocalOffset;
+		StaticRotationOffset = AssetChainSettings->StaticRotationOffset;
 		BlendToSource = AssetChainSettings->BlendToSource;
 		BlendToSourceWeights = AssetChainSettings->BlendToSourceWeights;
 		MatchSourceVelocity = AssetChainSettings->MatchSourceVelocity;
diff --git a/Engine/Plugins/Animation/IKRig/Source/IKRig/Public/Retargeter/IKRetargeter.h b/Engine/Plugins/Animation/IKRig/Source/IKRig/Public/Retargeter/IKRetargeter.h
--- a/Engine/Plugins/Animation/IKRig/Source/IKRig/Public/Retargeter/IKRetargeter.h
+++ b/Engine/Plugins/Animation/IKRig/Source/IKRig/Public/Retargeter/IKRetargeter.h
@@ -114,6 +114,14 @@ class IKRIG_API URetargetChainSettings: public UObject
 	/** Default 0, 0, 0. Apply static global-space offset to IK goal position. */
 	UPROPERTY(EditAnywhere, Category = "IK Adjustments")
 	FVector StaticOffset;
+
+	/** Default 0, 0, 0. Apply static local-space offset to IK goal position. */
+	UPROPERTY(EditAnywhere, Category = "IK Adjustments")
+	FVector StaticLocalOffset;
+
+	/** Default 0, 0, 0. Apply static local-space offset to IK goal rotation. */
+	UPROPERTY(EditAnywhere, Category = "IK Adjustments")
+	FRotator StaticRotationOffset;
 	
 	/** Range 0 to 5. Default 1. Brings IK goal closer (-) or further (+) from origin of chain.
 	*  At 0 the effector is placed at the origin of the chain.
@@ -134,29 +142,42 @@ class IKRIG_API URetargetChainSettings: public UObject
 	*  Values between 0 and 1 will blend the applied velocity between the retargeted amount and the source bone amount.*/
 	UPROPERTY(EditAnywhere, Category = "Experimental (May Change)", meta = (ClampMin = "0.0", ClampMax = "100.0", UIMin = "0.0", UIMax = "100.0"))
 	float VelocityThreshold = 5.0f;
+};
 
-	/** Range 0 to 1. Default is 1.0. Blend IK effector at the end of this chain towards the original position
-	*  on the source skeleton (0.0) or the position on the retargeted target skeleton (1.0). */
-	//UPROPERTY(EditAnywhere, BlueprintReadWrite, Category = IkMode, meta = (ClampMin = "0.0", ClampMax = "1.0", UIMin = "0.0", UIMax = "1.0"))
-	//float IkToSourceOrTarget = 1.0f;
-
-	/** Range 0 to 1. Default 0. Allow the chain to stretch by translating to reach the IK goal locations.
-	*  At 0 the chain will not stretch at all. At 1 the chain will be allowed to stretch double it's full length to reach IK. */
-	//UPROPERTY(EditAnywhere, BlueprintReadWrite, Category = Stretch, meta = (ClampMin = "0.0", ClampMax = "1.0", UIMin = "0.0", UIMax = "1.0"))
-	//float StretchTolerance = 0.0f;
+UCLASS()
+class IKRIG_API URetargetRootSettings: public UObject
+{
+	GENERATED_BODY()
 
-	/** When true, the source IK position is calculated relative to a source bone and applied relative to a target bone. */
-	//UPROPERTY(EditAnywhere, BlueprintReadWrite, Category = IkMode)
-	//bool bIkRelativeToSourceBone = false;
+	public:
 	
-	/** A bone in the SOURCE skeleton that the IK location is relative to. */
-	//UPROPERTY(EditAnywhere, BlueprintReadWrite, Category = IkMode, meta = (EditCondition="bIkRelativeToSourceBone"))
-	//FName IkRelativeSourceBone;
+	URetargetRootSettings() = default;
+
+	/** Whether to modify the location of the retarget root bone. Default is true. */
+	UPROPERTY(EditAnywhere, Category = "Root Retarget Settings")
+	bool RetargetRootTranslation = true;
+
+	/** Default 1. Scales the motion of the root position in the horizontal plane (X,Y). */
+	UPROPERTY(EditAnywhere, Category = "Root Retarget Settings", meta = (UIMin = "0.0", UIMax = "3.0"))
+	float GlobalScaleHorizontal = 1.0f;
+
+	/** Default 1. Scales the motion of the root position in the vertical direction (Z). */
+	UPROPERTY(EditAnywhere, Category = "Root Retarget Settings", meta = (UIMin = "0.0", UIMax = "3.0"))
+	float GlobalScaleVertical = 1.0f;
 	
-	/** A bone in the TARGET skeleton that the IK location is relative to.
-	 * This is usually the same bone as the source skeleton, but may have a different name in the target skeleton. */
-	//UPROPERTY(EditAnywhere, BlueprintReadWrite, Category = IkMode, meta = (EditCondition="bIkRelativeToSourceBone"))
-	//FName IkRelativeTargetBone;
+	/** Range 0 to 1. Default 0. Blends the retarget root's translation to the exact source location (per axis).
+	*  At 0 the root is placed at the retargeted location.
+	*  At 1 the root is placed at the location of the source's retarget root bone. */
+	UPROPERTY(EditAnywhere, Category = "Root Retarget Settings")
+	FVector BlendToSource = FVector::ZeroVector;
+
+	/** Applies a static component-space translation offset to the retarget root.*/
+	UPROPERTY(EditAnywhere, Category = "Root Retarget Settings")
+	FVector StaticOffset = FVector::ZeroVector;
+
+	/** Applies a static component-space rotation offset to the retarget root.*/
+	UPROPERTY(EditAnywhere, Category = "Root Retarget Settings")
+	FRotator StaticRotationOffset = FRotator::ZeroRotator;
 };
 
 USTRUCT()
@@ -195,7 +216,10 @@ UCLASS(Blueprintable)
 class IKRIG_API UIKRetargeter : public UObject
 {
 	GENERATED_BODY()
+	
 public:
+	
+	UIKRetargeter(const FObjectInitializer& ObjectInitializer);
 
 	/** Get read-only access to the source IK Rig asset */
 	const UIKRigDefinition* GetSourceIKRig() const { return SourceIKRigAsset.LoadSynchronous(); };
@@ -276,6 +300,9 @@ public:
 	/** log warnings and errors */
 	FIKRigLogger Log;
 
+	/** the retarget root settings */
+	TObjectPtr<URetargetRootSettings> RootSettings;
+
 #if WITH_EDITORONLY_DATA
 	UPROPERTY()
 	float TargetActorOffset_DEPRECATED = 0.0f;
diff --git a/Engine/Plugins/Animation/IKRig/Source/IKRigEditor/Private/RetargetEditor/SIKRetargetChainMapList.cpp b/Engine/Plugins/Animation/IKRig/Source/IKRigEditor/Private/RetargetEditor/SIKRetargetChainMapList.cpp
--- a/Engine/Plugins/Animation/IKRig/Source/IKRigEditor/Private/RetargetEditor/SIKRetargetChainMapList.cpp
+++ b/Engine/Plugins/Animation/IKRig/Source/IKRigEditor/Private/RetargetEditor/SIKRetargetChainMapList.cpp
@@ -127,7 +127,7 @@ void SIKRetargetChainMapList::Construct(
 {
 	EditorController = InEditorController;
 	EditorController.Pin()->ChainsView = SharedThis(this);
-
+	
 	ChildSlot
     [
 	    SNew(SVerticalBox)
@@ -135,13 +135,27 @@ void SIKRetargetChainMapList::Construct(
 		.AutoHeight()
 		[
 			SNew(SHorizontalBox)
+
+			+ SHorizontalBox::Slot()
+			.AutoWidth()
+			.Padding(5, 0)
+			.HAlign(HAlign_Left)
+			.VAlign(VAlign_Center)
+			[
+				SNew(SPositiveActionButton)
+				.Icon(FAppStyle::Get().GetBrush("Icons.Settings"))
+				.Text(LOCTEXT("EditRootButtonLabel", "Root Settings"))
+				.ToolTipText(LOCTEXT("EditRootButtonToolTip", "Edit the root retarget settings."))
+				.OnClicked(this, &SIKRetargetChainMapList::OnEditSettingsButtonClicked)
+			]
+
 			+ SHorizontalBox::Slot()
 			.AutoWidth()
 			.HAlign(HAlign_Left)
 			.VAlign(VAlign_Center)
 			[
 				SNew(STextBlock)
-				.Text(LOCTEXT("TargetRootLabel", "Target Root: "))
+				.Text(LOCTEXT("TargetRootLabel", "Target: "))
 				.TextStyle(FEditorStyle::Get(), "NormalText")
 			]
 
@@ -162,7 +176,7 @@ void SIKRetargetChainMapList::Construct(
 			.VAlign(VAlign_Center)
 			[
 				SNew(STextBlock)
-				.Text(LOCTEXT("SourceRootLabel", "Source Root: "))
+				.Text(LOCTEXT("SourceRootLabel", "Source: "))
 				.TextStyle(FEditorStyle::Get(), "NormalText")
 			]
 
@@ -360,4 +374,24 @@ FReply SIKRetargetChainMapList::OnAutoMapButtonClicked() const
 	return FReply::Handled();
 }
 
+FReply SIKRetargetChainMapList::OnEditSettingsButtonClicked() const
+{
+	const TSharedPtr<FIKRetargetEditorController> Controller = EditorController.Pin();
+	if (!Controller.IsValid())
+	{
+		return FReply::Unhandled();
+	}
+
+	UIKRetargeterController* RetargeterController = GetRetargetController();
+	if (!RetargeterController)
+	{
+		return FReply::Unhandled();
+	}
+	
+	Controller->DetailsView->SetObject(RetargeterController->GetAsset()->RootSettings);
+	return FReply::Handled();
+}
+
+
+
 #undef LOCTEXT_NAMESPACE
diff --git a/Engine/Plugins/Animation/IKRig/Source/IKRigEditor/Public/RetargetEditor/SIKRetargetChainMapList.h b/Engine/Plugins/Animation/IKRig/Source/IKRigEditor/Public/RetargetEditor/SIKRetargetChainMapList.h
--- a/Engine/Plugins/Animation/IKRig/Source/IKRigEditor/Public/RetargetEditor/SIKRetargetChainMapList.h
+++ b/Engine/Plugins/Animation/IKRig/Source/IKRigEditor/Public/RetargetEditor/SIKRetargetChainMapList.h
@@ -102,6 +102,8 @@ private:
 	bool IsChainMapEnabled() const;
 	/** when a chain is clicked on in the table view */
 	void OnItemClicked(TSharedPtr<FRetargetChainMapElement> InItem);
+	/** when edit root settings button clicked */
+	FReply OnEditSettingsButtonClicked() const;
 
 	/** update selection / details view */
 	void OnSelectionChanged();
