Fixed issue retarget root settings getting garbage collected

diff --git a/Engine/Plugins/Animation/IKRig/Source/IKRig/Private/Retargeter/IKRetargetProcessor.cpp b/Engine/Plugins/Animation/IKRig/Source/IKRig/Private/Retargeter/IKRetargetProcessor.cpp
--- a/Engine/Plugins/Animation/IKRig/Source/IKRig/Private/Retargeter/IKRetargetProcessor.cpp
+++ b/Engine/Plugins/Animation/IKRig/Source/IKRig/Private/Retargeter/IKRetargetProcessor.cpp
@@ -1595,7 +1595,7 @@ void UIKRetargetProcessor::CopyAllSettingsFromAsset()
 	}
 
 	// copy root settings
-	const URetargetRootSettings& RootSettings = *RetargeterAsset->RootSettings;
+	const URetargetRootSettings& RootSettings = *RetargeterAsset->GetRetargetRootSettings();
 	RootRetargeter.GlobalScaleHorizontal = RootSettings.GlobalScaleHorizontal;
 	RootRetargeter.GlobalScaleVertical = RootSettings.GlobalScaleVertical;
 	RootRetargeter.BlendToSource = RootSettings.BlendToSource;
diff --git a/Engine/Plugins/Animation/IKRig/Source/IKRig/Public/Retargeter/IKRetargeter.h b/Engine/Plugins/Animation/IKRig/Source/IKRig/Public/Retargeter/IKRetargeter.h
--- a/Engine/Plugins/Animation/IKRig/Source/IKRig/Public/Retargeter/IKRetargeter.h
+++ b/Engine/Plugins/Animation/IKRig/Source/IKRig/Public/Retargeter/IKRetargeter.h
@@ -233,6 +233,8 @@ public:
 	UIKRigDefinition* GetTargetIKRigWriteable() const { return TargetIKRigAsset.LoadSynchronous(); };
 	/** Get read-only access to the chain mapping */
 	const TArray<TObjectPtr<URetargetChainSettings>> GetAllChainSettings() const { return ChainSettings; };
+	/** Get access to the root settings */
+	URetargetRootSettings* GetRetargetRootSettings() const { return RootSettings; };
 	/** Get read-only access to a retarget pose */
 	const FIKRetargetPose* GetCurrentRetargetPose() const { return &RetargetPoses[CurrentRetargetPose]; };
 	/** Get the unique ID of this UObject as a name. */
@@ -283,6 +285,10 @@ private:
 	TSoftObjectPtr<USkeletalMesh> TargetPreviewMesh = nullptr;
 #endif
 
+	/** the retarget root settings */
+	UPROPERTY()
+	TObjectPtr<URetargetRootSettings> RootSettings;
+	
 public:
 	
 	/** When false, translational motion of skeleton root is not copied. Useful for debugging.*/
@@ -300,9 +306,6 @@ public:
 	/** log warnings and errors */
 	FIKRigLogger Log;
 
-	/** the retarget root settings */
-	TObjectPtr<URetargetRootSettings> RootSettings;
-
 #if WITH_EDITORONLY_DATA
 	UPROPERTY()
 	float TargetActorOffset_DEPRECATED = 0.0f;
diff --git a/Engine/Plugins/Animation/IKRig/Source/IKRigEditor/Private/RetargetEditor/SIKRetargetChainMapList.cpp b/Engine/Plugins/Animation/IKRig/Source/IKRigEditor/Private/RetargetEditor/SIKRetargetChainMapList.cpp
--- a/Engine/Plugins/Animation/IKRig/Source/IKRigEditor/Private/RetargetEditor/SIKRetargetChainMapList.cpp
+++ b/Engine/Plugins/Animation/IKRig/Source/IKRigEditor/Private/RetargetEditor/SIKRetargetChainMapList.cpp
@@ -388,7 +388,7 @@ FReply SIKRetargetChainMapList::OnEditSettingsButtonClicked() const
 		return FReply::Unhandled();
 	}
 	
-	Controller->DetailsView->SetObject(RetargeterController->GetAsset()->RootSettings);
+	Controller->DetailsView->SetObject(RetargeterController->GetAsset()->GetRetargetRootSettings());
 	return FReply::Handled();
 }
 
