Fixing PCH compile error on Mac.

diff --git a/Engine/Plugins/Animation/IKRig/Source/IKRigEditor/Public/RetargetEditor/IKRetargetEditPoseMode.h b/Engine/Plugins/Animation/IKRig/Source/IKRigEditor/Public/RetargetEditor/IKRetargetEditPoseMode.h
--- a/Engine/Plugins/Animation/IKRig/Source/IKRigEditor/Public/RetargetEditor/IKRetargetEditPoseMode.h
+++ b/Engine/Plugins/Animation/IKRig/Source/IKRigEditor/Public/RetargetEditor/IKRetargetEditPoseMode.h
@@ -12,7 +12,7 @@ class FIKRetargetEditorController;
 class FIKRetargetEditor;
 class FIKRetargetPreviewScene;
 
-enum FIKRetargetTrackingState : int8
+enum class FIKRetargetTrackingState : int8
 {
 	None,
 	RotatingBone,
