Fix crash when clearing target IK rig asset (fix ported from 5.0 release).

diff --git a/Engine/Plugins/Animation/IKRig/Source/IKRigEditor/Private/RetargetEditor/IKRetargeterController.cpp b/Engine/Plugins/Animation/IKRig/Source/IKRigEditor/Private/RetargetEditor/IKRetargeterController.cpp
--- a/Engine/Plugins/Animation/IKRig/Source/IKRigEditor/Private/RetargetEditor/IKRetargeterController.cpp
+++ b/Engine/Plugins/Animation/IKRig/Source/IKRigEditor/Private/RetargetEditor/IKRetargeterController.cpp
@@ -55,12 +55,6 @@ void UIKRetargeterController::SetSourceIKRig(UIKRigDefinition* SourceIKRig)
 	Asset->SourceIKRigAsset = SourceIKRig;
 }
 
-void UIKRetargeterController::SetTargetIKRig(UIKRigDefinition* TargetIKRig)
-{
-	CleanChainMapping();
-	AutoMapChains();
-}
-
 USkeletalMesh* UIKRetargeterController::GetSourcePreviewMesh() const
 {
 	// can't preview anything if target IK Rig is null
@@ -135,56 +129,57 @@ void UIKRetargeterController::GetSourceChainNames(TArray<FName>& OutNames) const
 
 void UIKRetargeterController::CleanChainMapping(const bool bForceReinitialization) const
 {
-	if (Asset->TargetIKRigAsset.IsNull())
+	if (IsValid(Asset->GetTargetIKRig()))
 	{
-		return;
-	}
-	
-	TArray<FName> TargetChainNames;
-	GetTargetChainNames(TargetChainNames);
+		TArray<FName> TargetChainNames;
+		GetTargetChainNames(TargetChainNames);
 
-	// remove all target chains that are no longer in the target IK rig asset
-	TArray<FName> TargetChainsToRemove;
-	for (const URetargetChainSettings* ChainMap : Asset->ChainSettings)
-	{
-		if (!TargetChainNames.Contains(ChainMap->TargetChain))
+		// remove all target chains that are no longer in the target IK rig asset
+		TArray<FName> TargetChainsToRemove;
+		for (const URetargetChainSettings* ChainMap : Asset->ChainSettings)
 		{
-			TargetChainsToRemove.Add(ChainMap->TargetChain);
+			if (!TargetChainNames.Contains(ChainMap->TargetChain))
+			{
+				TargetChainsToRemove.Add(ChainMap->TargetChain);
+			}
 		}
-	}
-	for (FName TargetChainToRemove : TargetChainsToRemove)
-	{
-		Asset->ChainSettings.RemoveAll([&TargetChainToRemove](const URetargetChainSettings* Element)
+		for (FName TargetChainToRemove : TargetChainsToRemove)
 		{
-			return Element->TargetChain == TargetChainToRemove;
-		});
-	}
+			Asset->ChainSettings.RemoveAll([&TargetChainToRemove](const URetargetChainSettings* Element)
+			{
+				return Element->TargetChain == TargetChainToRemove;
+			});
+		}
 
-	// add a mapping for each chain that is in the target IK rig (if it doesn't have one already)
-	for (FName TargetChainName : TargetChainNames)
-	{
-		const bool HasChain = Asset->ChainSettings.ContainsByPredicate([&TargetChainName](const URetargetChainSettings* Element)
+		// add a mapping for each chain that is in the target IK rig (if it doesn't have one already)
+		for (FName TargetChainName : TargetChainNames)
 		{
-			return Element->TargetChain == TargetChainName;
-		});
+			const bool HasChain = Asset->ChainSettings.ContainsByPredicate([&TargetChainName](const URetargetChainSettings* Element)
+			{
+				return Element->TargetChain == TargetChainName;
+			});
 		
-		if (!HasChain)
-		{
-			TObjectPtr<URetargetChainSettings> ChainMap = NewObject<URetargetChainSettings>(Asset, URetargetChainSettings::StaticClass(), NAME_None, RF_Transactional);
-			ChainMap->TargetChain = TargetChainName;
-			Asset->ChainSettings.Add(ChainMap);
+			if (!HasChain)
+			{
+				TObjectPtr<URetargetChainSettings> ChainMap = NewObject<URetargetChainSettings>(Asset, URetargetChainSettings::StaticClass(), NAME_None, RF_Transactional);
+				ChainMap->TargetChain = TargetChainName;
+				Asset->ChainSettings.Add(ChainMap);
+			}
 		}
 	}
 
-	TArray<FName> SourceChainNames;
-	GetSourceChainNames(SourceChainNames);
-	
-	// reset any sources that are no longer present to "None"
-	for (URetargetChainSettings* ChainMap : Asset->ChainSettings)
+	if (IsValid(Asset->GetSourceIKRig()))
 	{
-		if (!SourceChainNames.Contains(ChainMap->SourceChain))
+		TArray<FName> SourceChainNames;
+		GetSourceChainNames(SourceChainNames);
+	
+		// reset any sources that are no longer present to "None"
+		for (URetargetChainSettings* ChainMap : Asset->ChainSettings)
 		{
-			ChainMap->SourceChain = NAME_None;
+			if (!SourceChainNames.Contains(ChainMap->SourceChain))
+			{
+				ChainMap->SourceChain = NAME_None;
+			}
 		}
 	}
 
@@ -213,9 +208,9 @@ void UIKRetargeterController::CleanPoseList(const bool bForceReinitialization)
 	}
 
 	// remove all bone offsets that are no longer part of the target skeleton
-	if (Asset->TargetIKRigAsset.IsValid())
+	if (const UIKRigDefinition* TargetIKRig = Asset->GetTargetIKRig())
 	{
-		const TArray<FName> AllowedBoneNames = Asset->TargetIKRigAsset->Skeleton.BoneNames;
+		const TArray<FName> AllowedBoneNames = TargetIKRig->Skeleton.BoneNames;
 		for (TTuple<FName, FIKRetargetPose>& Pose : Asset->RetargetPoses)
 		{
 			// find bone offsets no longer in target skeleton
@@ -235,7 +230,7 @@ void UIKRetargeterController::CleanPoseList(const bool bForceReinitialization)
 			}
 
 			// sort the pose offset from leaf to root
-			Pose.Value.SortHierarchically(Asset->TargetIKRigAsset->Skeleton);
+			Pose.Value.SortHierarchically(TargetIKRig->Skeleton);
 		}
 	}
 
@@ -291,8 +286,8 @@ void UIKRetargeterController::AutoMapChains() const
 
 void UIKRetargeterController::OnRetargetChainRenamed(UIKRigDefinition* IKRig, FName OldChainName, FName NewChainName) const
 {
-	const bool bIsSourceRig = IKRig == Asset->SourceIKRigAsset;
-	check(bIsSourceRig || IKRig == Asset->TargetIKRigAsset)
+	const bool bIsSourceRig = IKRig == Asset->GetSourceIKRig();
+	check(bIsSourceRig || IKRig == Asset->GetTargetIKRig())
 	for (URetargetChainSettings* ChainMap : Asset->ChainSettings)
 	{
 		FName& ChainNameToUpdate = bIsSourceRig ? ChainMap->SourceChain : ChainMap->TargetChain;
@@ -307,8 +302,8 @@ void UIKRetargeterController::OnRetargetChainRenamed(UIKRigDefinition* IKRig, FN
 
 void UIKRetargeterController::OnRetargetChainRemoved(UIKRigDefinition* IKRig, const FName& InChainRemoved) const
 {
-	const bool bIsSourceRig = IKRig == Asset->SourceIKRigAsset;
-	check(bIsSourceRig || IKRig == Asset->TargetIKRigAsset)
+	const bool bIsSourceRig = IKRig == Asset->GetSourceIKRig();
+	check(bIsSourceRig || IKRig == Asset->GetTargetIKRig())
 
 	// set source chain name to NONE if it has been deleted 
 	if (bIsSourceRig)
@@ -501,10 +496,16 @@ URetargetChainSettings* UIKRetargeterController::GetChainMap(const FName& Target
 
 void UIKRetargeterController::SortChainMapping() const
 {
-	Asset->ChainSettings.Sort([this](const URetargetChainSettings& A, const URetargetChainSettings& B)
+	const UIKRigDefinition* TargetIKRig = Asset->GetTargetIKRig();
+	if (!IsValid(TargetIKRig))
+	{
+		return;
+	}
+	
+	Asset->ChainSettings.Sort([TargetIKRig](const URetargetChainSettings& A, const URetargetChainSettings& B)
 	{
-		const TArray<FBoneChain>& BoneChains = Asset->TargetIKRigAsset->GetRetargetChains();
-		const FIKRigSkeleton& TargetSkeleton = Asset->TargetIKRigAsset->Skeleton;
+		const TArray<FBoneChain>& BoneChains = TargetIKRig->GetRetargetChains();
+		const FIKRigSkeleton& TargetSkeleton = TargetIKRig->Skeleton;
 
 		// look for chains
 		const int32 IndexA = BoneChains.IndexOfByPredicate([&A](const FBoneChain& Chain)
diff --git a/Engine/Plugins/Animation/IKRig/Source/IKRigEditor/Public/RetargetEditor/IKRetargeterController.h b/Engine/Plugins/Animation/IKRig/Source/IKRigEditor/Public/RetargetEditor/IKRetargeterController.h
--- a/Engine/Plugins/Animation/IKRig/Source/IKRigEditor/Public/RetargetEditor/IKRetargeterController.h
+++ b/Engine/Plugins/Animation/IKRig/Source/IKRigEditor/Public/RetargetEditor/IKRetargeterController.h
@@ -40,8 +40,6 @@ public:
 	*/
 	/** Set the IK Rig to use as the source (to copy animation FROM) */
 	void SetSourceIKRig(UIKRigDefinition* SourceIKRig);
-	/** Set the IK Rig to use as the target (to copy animation TO) */
-	void SetTargetIKRig(UIKRigDefinition* TargetIKRig);
 	/** Get source skeletal mesh */
 	USkeletalMesh* GetSourcePreviewMesh() const;
 	/** Get target skeletal mesh */
