Adds curve-based speed IK planting to the IK Retargeter

diff --git a/Engine/Plugins/Animation/IKRig/Source/IKRig/Private/AnimNodes/AnimNode_RetargetPoseFromMesh.cpp b/Engine/Plugins/Animation/IKRig/Source/IKRig/Private/AnimNodes/AnimNode_RetargetPoseFromMesh.cpp
--- a/Engine/Plugins/Animation/IKRig/Source/IKRig/Private/AnimNodes/AnimNode_RetargetPoseFromMesh.cpp
+++ b/Engine/Plugins/Animation/IKRig/Source/IKRig/Private/AnimNodes/AnimNode_RetargetPoseFromMesh.cpp
@@ -57,7 +57,9 @@ void FAnimNode_RetargetPoseFromMesh::Update_AnyThread(const FAnimationUpdateCont
     // this introduces a frame of latency in setting the pin-driven source component,
     // but we cannot do the work to extract transforms on a worker thread as it is not thread safe.
     GetEvaluateGraphExposedInputs().Execute(Context);
+	DeltaTime += Context.GetDeltaTime();
 }
+
 void FAnimNode_RetargetPoseFromMesh::Evaluate_AnyThread(FPoseContext& Output)
 {
 	DECLARE_SCOPE_HIERARCHICAL_COUNTER_ANIMNODE(Evaluate_AnyThread)
@@ -97,7 +99,8 @@ void FAnimNode_RetargetPoseFromMesh::Evaluate_AnyThread(FPoseContext& Output)
 #endif
 
 	// run the retargeter
-	const TArray<FTransform>& RetargetedPose = Processor->RunRetargeter(SourceMeshComponentSpaceBoneTransforms);
+	const TArray<FTransform>& RetargetedPose = Processor->RunRetargeter(SourceMeshComponentSpaceBoneTransforms, SpeedValuesFromCurves, DeltaTime);
+	DeltaTime = 0.0f;
 
 	// copy pose back
 	FCSPose<FCompactPose> ComponentPose;
@@ -119,7 +122,7 @@ void FAnimNode_RetargetPoseFromMesh::Evaluate_AnyThread(FPoseContext& Output)
 	// copy curves over
 	if (bCopyCurves)
 	{
-		for (const TPair<FName, float>& SourceCurve : SourceCurveList)
+		for (const TPair<FName, float>& SourceCurve : SourceCurveValues)
 		{
 			if (const SmartName::UID_Type* UID = CurveNameToUIDMap.Find(SourceCurve.Key))
 			{
@@ -151,13 +154,43 @@ void FAnimNode_RetargetPoseFromMesh::PreUpdate(const UAnimInstance* InAnimInstan
 
 		if(bCopyCurves)
 		{
-			SourceCurveList.Reset();
+			SourceCurveValues.Reset();
 			if (const UAnimInstance* SourceAnimInstance = SourceMeshComponent->GetAnimInstance())
 			{
 				// attribute curve contains all list	
-				SourceCurveList.Append(SourceAnimInstance->GetAnimationCurveList(EAnimCurveType::AttributeCurve));
+				SourceCurveValues.Append(SourceAnimInstance->GetAnimationCurveList(EAnimCurveType::AttributeCurve));
 			}
+
+			UpdateSpeedValuesFromCurves();
+		}
+	}
+}
+
+void FAnimNode_RetargetPoseFromMesh::UpdateSpeedValuesFromCurves()
+{
+	SpeedValuesFromCurves.Reset();
+	
+	if (!IKRetargeterAsset)
+	{
+		return;
+	}
+
+	const TArray<TObjectPtr<URetargetChainSettings>> ChainSettings = IKRetargeterAsset->GetAllChainSettings();
+	for (const TObjectPtr<URetargetChainSettings>& ChainSetting : ChainSettings)
+	{
+		FName CurveName = ChainSetting->SpeedCurveName;
+		if (CurveName == NAME_None)
+		{
+			continue;
+		}
+
+		// value will be negative if the curve was not found (retargeter ignores negative speeds)
+		float CurveValue = -1.0f;
+		if (SourceCurveValues.Contains(CurveName))
+		{
+			CurveValue = SourceCurveValues[CurveName];
 		}
+		SpeedValuesFromCurves.Add(CurveName, CurveValue);
 	}
 }
 
@@ -171,7 +204,7 @@ void FAnimNode_RetargetPoseFromMesh::SetProcessorNeedsInitialized()
 }
 #endif
 
-const UIKRetargetProcessor* FAnimNode_RetargetPoseFromMesh::GetRetargetProcessor() const
+UIKRetargetProcessor* FAnimNode_RetargetPoseFromMesh::GetRetargetProcessor() const
 {
 	return Processor;
 }
diff --git a/Engine/Plugins/Animation/IKRig/Source/IKRig/Private/Retargeter/IKRetargetProcessor.cpp b/Engine/Plugins/Animation/IKRig/Source/IKRig/Private/Retargeter/IKRetargetProcessor.cpp
--- a/Engine/Plugins/Animation/IKRig/Source/IKRig/Private/Retargeter/IKRetargetProcessor.cpp
+++ b/Engine/Plugins/Animation/IKRig/Source/IKRig/Private/Retargeter/IKRetargetProcessor.cpp
@@ -6,6 +6,7 @@
 #include "IKRigLogger.h"
 #include "IKRigProcessor.h"
 #include "Engine/SkeletalMesh.h"
+#include "Kismet/KismetMathLibrary.h"
 #include "Retargeter/IKRetargeter.h"
 
 #define LOCTEXT_NAMESPACE "IKRetargetProcessor"
@@ -810,6 +811,7 @@ bool FChainRetargeterIK::InitializeTarget(
 	Target.InitialEndPosition = Last.GetTranslation();
 	Target.InitialEndRotation = Last.GetRotation();
 	Target.InitialLength = (TargetInitialGlobalPose[Target.BoneIndexA].GetTranslation() - Last.GetTranslation()).Size();
+	ResetThisTick = true;
 
 	if (Target.InitialLength <= KINDA_SMALL_NUMBER)
 	{
@@ -822,6 +824,8 @@ bool FChainRetargeterIK::InitializeTarget(
 	
 void FChainRetargeterIK::DecodePose(
 	const FRetargetChainSettings& Settings,
+	const TMap<FName, float>& SpeedValuesFromCurves,
+	const float DeltaTime,
     const TArray<FTransform>& OutGlobalPose,
     FDecodedIKChain& OutResults)
 {
@@ -870,26 +874,23 @@ void FChainRetargeterIK::DecodePose(
 	{
 		GoalPosition = Start + (GoalPosition - Start) * Settings.Extension;	
 	}
-
+	
 	// match velocity
-	if (Settings.MatchSourceVelocity > KINDA_SMALL_NUMBER)
+	if (!ResetThisTick && Settings.UseSpeedCurveToPlantIK && SpeedValuesFromCurves.Contains(Settings.SpeedCurveName))
 	{
-		const FVector SourceVelocity = Source.CurrentEndPosition - Source.PreviousEndPosition;
-		const float SourceSpeed = SourceVelocity.Size();
-
-		const FVector TargetVelocity = GoalPosition - Target.PrevEndPosition;
-		const float TargetSpeed = TargetVelocity.Size();
-
-		// if target is moving slowly enough, start matching velocity
-		if (TargetSpeed < Settings.TeleportVelocityThreshold)
+		const float SourceSpeed = SpeedValuesFromCurves[Settings.SpeedCurveName];
+		if (SourceSpeed < 0.0f || SourceSpeed > Settings.SpeedThreshold)
 		{
-			// match target speed to source
-			const float BlendedSpeed = FMath::Lerp(TargetSpeed, SourceSpeed, Settings.MatchSourceVelocity);
-			GoalPosition = Target.PrevEndPosition + TargetVelocity * BlendedSpeed;
-
-			// what would happen if we blended into and out of this state?
-			// TimeSinceTargetStartedMatchingVelocity
-			// WasMatchingVelocity
+			GoalPosition = UKismetMathLibrary::VectorSpringInterp(
+				Target.PrevEndPosition, GoalPosition, PlantingSpringState,
+				Settings.UnplantStiffness,
+				Settings.UnplantCriticalDamping,
+				DeltaTime, 1.0f, 0.0f);
+		}
+		else
+		{
+			PlantingSpringState.Reset();
+			GoalPosition = Target.PrevEndPosition;
 		}
 	}
 	
@@ -898,6 +899,7 @@ void FChainRetargeterIK::DecodePose(
 	OutResults.EndEffectorRotation = GoalRotation;
 	OutResults.PoleVectorPosition = FVector::OneVector; // TODO calc pole vector position
 	Target.PrevEndPosition = GoalPosition;
+	ResetThisTick = false;
 }
 
 bool FRetargetChainPair::Initialize(
@@ -1406,7 +1408,10 @@ bool UIKRetargetProcessor::InitializeIKRig(UObject* Outer, const USkeletalMesh*
 	return true;
 }
 
-TArray<FTransform>&  UIKRetargetProcessor::RunRetargeter(const TArray<FTransform>& InSourceGlobalPose)
+TArray<FTransform>&  UIKRetargetProcessor::RunRetargeter(
+	const TArray<FTransform>& InSourceGlobalPose,
+	const TMap<FName, float>& SpeedValuesFromCurves,
+	const float DeltaTime)
 {
 	check(bIsInitialized);
 
@@ -1446,7 +1451,7 @@ TArray<FTransform>&  UIKRetargetProcessor::RunRetargeter(const TArray<FTransform
 	// IK CHAIN retargeting
 	if (RetargeterAsset->bRetargetIK && bAtLeastOneValidBoneChainPair && bIKRigInitialized)
 	{
-		RunIKRetarget(InSourceGlobalPose, TargetSkeleton.OutputGlobalPose);
+		RunIKRetarget(InSourceGlobalPose, TargetSkeleton.OutputGlobalPose, SpeedValuesFromCurves, DeltaTime);
 	}
 
 	return TargetSkeleton.OutputGlobalPose;
@@ -1484,7 +1489,9 @@ void UIKRetargetProcessor::RunFKRetarget(
 
 void UIKRetargetProcessor::RunIKRetarget(
 	const TArray<FTransform>& InSourceGlobalPose,
-    TArray<FTransform>& OutTargeGlobalPose)
+    TArray<FTransform>& OutTargeGlobalPose,
+    const TMap<FName, float>& SpeedValuesFromCurves,
+    const float DeltaTime)
 {
 	if (!IKRigProcessor->IsInitialized())
 	{
@@ -1503,7 +1510,12 @@ void UIKRetargetProcessor::RunIKRetarget(
 		ChainPair.IKChainRetargeter.EncodePose(InSourceGlobalPose);
 		// decode the IK goal and apply to IKRig
 		FDecodedIKChain OutIKGoal;
-		ChainPair.IKChainRetargeter.DecodePose(ChainPair.Settings, OutTargeGlobalPose, OutIKGoal);
+		ChainPair.IKChainRetargeter.DecodePose(
+			ChainPair.Settings,
+			SpeedValuesFromCurves,
+			DeltaTime,
+			OutTargeGlobalPose,
+			OutIKGoal);
 		// set the goal transform on the IK Rig
 		FIKRigGoal Goal = FIKRigGoal(
 			ChainPair.IKGoalName,
@@ -1524,6 +1536,14 @@ void UIKRetargetProcessor::RunIKRetarget(
 	IKRigProcessor->CopyOutputGlobalPoseToArray(OutTargeGlobalPose);
 }
 
+void UIKRetargetProcessor::ResetPlanting()
+{
+	for (FRetargetChainPairIK& ChainPair : ChainPairsIK)
+	{
+		ChainPair.IKChainRetargeter.ResetThisTick = true;
+	}
+}
+
 FTransform UIKRetargetProcessor::GetTargetBoneRetargetPoseGlobalTransform(const int32& TargetBoneIndex) const
 {
 	check(TargetSkeleton.BoneNames.IsValidIndex(TargetBoneIndex))
diff --git a/Engine/Plugins/Animation/IKRig/Source/IKRig/Private/Retargeter/IKRetargeter.cpp b/Engine/Plugins/Animation/IKRig/Source/IKRig/Private/Retargeter/IKRetargeter.cpp
--- a/Engine/Plugins/Animation/IKRig/Source/IKRig/Private/Retargeter/IKRetargeter.cpp
+++ b/Engine/Plugins/Animation/IKRig/Source/IKRig/Private/Retargeter/IKRetargeter.cpp
@@ -6,7 +6,18 @@
 const FName UIKRetargeter::GetSourceIKRigPropertyName() { return GET_MEMBER_NAME_STRING_CHECKED(UIKRetargeter, SourceIKRigAsset); };
 const FName UIKRetargeter::GetTargetIKRigPropertyName() { return GET_MEMBER_NAME_STRING_CHECKED(UIKRetargeter, TargetIKRigAsset); };
 const FName UIKRetargeter::GetSourcePreviewMeshPropertyName() { return GET_MEMBER_NAME_STRING_CHECKED(UIKRetargeter, SourcePreviewMesh); };
-const FName UIKRetargeter::GetTargetPreviewMeshPropertyName() { return GET_MEMBER_NAME_STRING_CHECKED(UIKRetargeter, TargetPreviewMesh); };
+const FName UIKRetargeter::GetTargetPreviewMeshPropertyName() { return GET_MEMBER_NAME_STRING_CHECKED(UIKRetargeter, TargetPreviewMesh); }
+
+void UIKRetargeter::GetSpeedCurveNames(TArray<FName>& OutSpeedCurveNames) const
+{
+	for (const URetargetChainSettings* ChainSetting : ChainSettings)
+	{
+		if (ChainSetting->SpeedCurveName != NAME_None)
+		{
+			OutSpeedCurveNames.Add(ChainSetting->SpeedCurveName);
+		}
+	}
+};
 #endif
 
 UIKRetargeter::UIKRetargeter(const FObjectInitializer& ObjectInitializer)
diff --git a/Engine/Plugins/Animation/IKRig/Source/IKRig/Public/AnimNodes/AnimNode_RetargetPoseFromMesh.h b/Engine/Plugins/Animation/IKRig/Source/IKRig/Public/AnimNodes/AnimNode_RetargetPoseFromMesh.h
--- a/Engine/Plugins/Animation/IKRig/Source/IKRig/Public/AnimNodes/AnimNode_RetargetPoseFromMesh.h
+++ b/Engine/Plugins/Animation/IKRig/Source/IKRig/Public/AnimNodes/AnimNode_RetargetPoseFromMesh.h
@@ -56,14 +56,16 @@ struct IKRIG_API FAnimNode_RetargetPoseFromMesh : public FAnimNode_Base
 	void SetProcessorNeedsInitialized();
 #endif
 
-	/** Read-only access to the runtime processor */
-	const UIKRetargetProcessor* GetRetargetProcessor() const;
+	/** Access to the runtime processor */
+	UIKRetargetProcessor* GetRetargetProcessor() const;
 
 private:
 
 	/** returns true if processor is setup and ready to go, false otherwise */
 	bool EnsureProcessorIsInitialized(const TObjectPtr<USkeletalMeshComponent> TargetMeshComponent);
+	/** copies the source mesh pose (on main thread) */
 	void CopyBoneTransformsFromSource(USkeletalMeshComponent* TargetMeshComponent);
+	/** indirection to account for Master Pose Component setup */
 	TObjectPtr<USkeletalMeshComponent> GetComponentToCopyPoseFrom() const;
 
 	/** the runtime processor used to run the retarget and generate new poses */
@@ -78,5 +80,12 @@ private:
 	// mapping of Skeleton curve names to the UID (for copying curves from the source mesh component)
 	TMap<FName, SmartName::UID_Type> CurveNameToUIDMap;
 	// cached curves, copied on the game thread
-	TMap<FName, float> SourceCurveList;
+	TMap<FName, float> SourceCurveValues;
+
+	/** update map of curve values containing speeds used for IK planting */
+	void UpdateSpeedValuesFromCurves();
+	// map of curve names to values, passed to retargeter for IK planting (copied from source mesh)
+	TMap<FName, float> SpeedValuesFromCurves;
+	// the accumulated delta time this tick
+	float DeltaTime;
 };
diff --git a/Engine/Plugins/Animation/IKRig/Source/IKRig/Public/Retargeter/IKRetargetProcessor.h b/Engine/Plugins/Animation/IKRig/Source/IKRig/Public/Retargeter/IKRetargetProcessor.h
--- a/Engine/Plugins/Animation/IKRig/Source/IKRig/Public/Retargeter/IKRetargetProcessor.h
+++ b/Engine/Plugins/Animation/IKRig/Source/IKRig/Public/Retargeter/IKRetargetProcessor.h
@@ -5,6 +5,7 @@
 #include "CoreMinimal.h"
 #include "IKRetargeter.h"
 #include "IKRigLogger.h"
+#include "Kismet/KismetMathLibrary.h"
 #include "IKRetargetProcessor.generated.h"
 
 class URetargetChainSettings;
@@ -185,8 +186,12 @@ struct FRetargetChainSettings
 	float Extension = 1.0f;
 	float BlendToSource = 0.0f;
 	FVector BlendToSourceWeights = FVector::OneVector;
-	float MatchSourceVelocity = 0.0f;
-	float TeleportVelocityThreshold = 0.0f;
+
+	bool UseSpeedCurveToPlantIK = false;
+	float SpeedThreshold = 0.0f;
+	FName SpeedCurveName;
+	float UnplantStiffness = 250.0f;
+	float UnplantCriticalDamping = 1.0f;
 
 public:
 	
@@ -207,8 +212,12 @@ public:
 		StaticRotationOffset = AssetChainSettings->StaticRotationOffset;
 		BlendToSource = AssetChainSettings->BlendToSource;
 		BlendToSourceWeights = AssetChainSettings->BlendToSourceWeights;
-		MatchSourceVelocity = AssetChainSettings->MatchSourceVelocity;
-		TeleportVelocityThreshold = AssetChainSettings->VelocityThreshold;
+
+		UseSpeedCurveToPlantIK = AssetChainSettings->UseSpeedCurveToPlantIK;
+		UnplantStiffness = AssetChainSettings->UnplantStiffness;
+		UnplantCriticalDamping = AssetChainSettings->UnplantCriticalDamping;
+		SpeedThreshold = AssetChainSettings->VelocityThreshold;
+		SpeedCurveName = AssetChainSettings->SpeedCurveName;
 	};
 };
 
@@ -350,6 +359,9 @@ struct FChainRetargeterIK
 	
 	FTargetChainIK Target;
 
+	bool ResetThisTick;
+	FVectorSpringState PlantingSpringState;
+
 	bool InitializeSource(
 		const TArray<int32>& BoneIndices,
 		const TArray<FTransform> &SourceInitialGlobalPose,
@@ -364,6 +376,8 @@ struct FChainRetargeterIK
 	
 	void DecodePose(
 		const FRetargetChainSettings& Settings,
+		const TMap<FName, float>& SpeedValuesFromCurves,
+		const float DeltaTime,
 		const TArray<FTransform>& OutGlobalPose,
 		FDecodedIKChain& OutResults);
 };
@@ -463,7 +477,7 @@ public:
 	* @param InSourceGlobalPose -  is the source mesh input pose in Component/Global space
 	* @return The retargeted Component/Global space pose for the target skeleton
 	*/
-	TArray<FTransform>& RunRetargeter(const TArray<FTransform>& InSourceGlobalPose);
+	TArray<FTransform>& RunRetargeter(const TArray<FTransform>& InSourceGlobalPose, const TMap<FName, float>& SpeedValuesFromCurves, const float DeltaTime);
 
 	/**
 	* Get the Global transform, in the currently used retarget pose, for a bone in the target skeletal mesh
@@ -495,6 +509,9 @@ public:
 
 	/** Get the currently running IK Rig processor for the target */
 	UIKRigProcessor* GetTargetIKRigProcessor() const { return IKRigProcessor; };
+	
+	/** Reset the IK planting state. */
+	void ResetPlanting();
 
 #if WITH_EDITOR
 	/** Set that this processor needs to be reinitialized. */
@@ -555,5 +572,9 @@ private:
 	void RunFKRetarget(const TArray<FTransform>& InGlobalTransforms, TArray<FTransform>& OutGlobalTransforms);
 
 	/** Internal retarget phase for the IK chains. */
-	void RunIKRetarget(const TArray<FTransform>& InSourceGlobalPose, TArray<FTransform>& OutTargeGlobalPose);
+	void RunIKRetarget(
+		const TArray<FTransform>& InSourceGlobalPose,
+		TArray<FTransform>& OutTargeGlobalPose,
+		const TMap<FName, float>& SpeedValuesFromCurves,
+		const float DeltaTime);
 };
diff --git a/Engine/Plugins/Animation/IKRig/Source/IKRig/Public/Retargeter/IKRetargeter.h b/Engine/Plugins/Animation/IKRig/Source/IKRig/Public/Retargeter/IKRetargeter.h
--- a/Engine/Plugins/Animation/IKRig/Source/IKRig/Public/Retargeter/IKRetargeter.h
+++ b/Engine/Plugins/Animation/IKRig/Source/IKRig/Public/Retargeter/IKRetargeter.h
@@ -43,7 +43,7 @@ enum class ERetargetRotationMode : uint8
 };
 
 UCLASS()
-class IKRIG_API URetargetChainSettings: public UObject
+class IKRIG_API URetargetChainSettings : public UObject
 {
 	GENERATED_BODY()
 
@@ -55,11 +55,11 @@ class IKRIG_API URetargetChainSettings: public UObject
 	
 	/** The chain on the Source IK Rig asset to copy animation FROM. */
 	UPROPERTY(VisibleAnywhere, Category = "Chain Mapping")
-	FName SourceChain = NAME_None;
+	FName SourceChain;
 
 	/** The chain on the Target IK Rig asset to copy animation TO. */
 	UPROPERTY(VisibleAnywhere, Category = "Chain Mapping")
-	FName TargetChain = NAME_None;
+	FName TargetChain;
 
 	/** Whether to copy the shape of the chain from the source skeleton using the Rotation and Translation modes. Default is true.
 	* NOTE: All FK operations run before the IK pass to copy the shape of the FK chain from the source skeleton. */
@@ -129,19 +129,26 @@ class IKRIG_API URetargetChainSettings: public UObject
 	UPROPERTY(EditAnywhere, Category = "IK Adjustments", meta = (ClampMin = "0.0", ClampMax = "5.0", UIMin = "0.1", UIMax = "2.0"))
 	float Extension = 1.0f;
 
-	/** Range 0 to 1. Default 0. Blends IK goal from retargeted velocity (0) to source bone velocity (1).
-	*  At 0 the goal is placed at the retargeted location.
-	*  At 1 the goal is placed at the retargeted location but clamped by the velocity of the source chain's end bone.
-	*  Values between 0 and 1 will blend the applied velocity between the retargeted amount and the source bone amount.*/
-	UPROPERTY(EditAnywhere, Category = "Experimental (May Change)", meta = (ClampMin = "0.0", ClampMax = "1.0", UIMin = "0.0", UIMax = "1.0"))
-	float MatchSourceVelocity = 0.0f;
-
-	/** Range 0 to 1000. Default 5. The maximum speed a source goal can be moving before being considered statically planted.
-	*  At 0 the goal is placed at the retargeted location.
-	*  At 1 the goal is placed at the retargeted location but clamped by the velocity of the source chain's end bone.
-	*  Values between 0 and 1 will blend the applied velocity between the retargeted amount and the source bone amount.*/
-	UPROPERTY(EditAnywhere, Category = "Experimental (May Change)", meta = (ClampMin = "0.0", ClampMax = "100.0", UIMin = "0.0", UIMax = "100.0"))
-	float VelocityThreshold = 5.0f;
+	/** The name of the curve on the source animation that contains the speed of the end effector bone.*/
+	UPROPERTY(EditAnywhere, Category = "Plant IK by Speed", meta = (ClampMin = "0.0", ClampMax = "100.0", UIMin = "0.0", UIMax = "100.0"))
+	bool UseSpeedCurveToPlantIK = false;
+	
+	/** The name of the curve on the source animation that contains the speed of the end effector bone.*/
+	UPROPERTY(EditAnywhere, Category = "Plant IK by Speed", meta = (ClampMin = "0.0", ClampMax = "100.0", UIMin = "0.0", UIMax = "100.0"))
+	FName SpeedCurveName;
+
+	/** Range 0 to 1000. Default 15. The maximum speed a source bone can be moving while being considered 'planted'.
+	*  The target IK goal will not be allowed to move whenever the source bone speed drops below this threshold speed. */
+	UPROPERTY(EditAnywhere, Category = "Plant IK by Speed", meta = (ClampMin = "0.0", ClampMax = "100.0", UIMin = "0.0", UIMax = "100.0"))
+	float VelocityThreshold = 15.0f;
+
+	// How stiff the spring model is that smoothly pulls the IK position after unplanting (more stiffness means more oscillation around the target value)
+	UPROPERTY(EditAnywhere, Category = "Plant IK by Speed", meta = (ClampMin = "0.0", UIMin = "0.0"))
+	float UnplantStiffness = 250.0f;
+
+	// How much damping to apply to the spring (0 means no damping, 1 means critically damped which means no oscillation)
+	UPROPERTY(EditAnywhere, Category = "Plant IK by Speed", meta = (ClampMin = "0.0", ClampMax = "10.0", UIMin = "0.0", UIMax = "1.0"))
+	float UnplantCriticalDamping = 1.0f;
 };
 
 UCLASS()
@@ -256,6 +263,8 @@ public:
 	static const FName GetSourcePreviewMeshPropertyName();
 	/* Get name of Target Preview Mesh property */
 	static const FName GetTargetPreviewMeshPropertyName();
+	/** Get the names of the all the speed curves the retargeter will be looking for */
+	void GetSpeedCurveNames(TArray<FName>& OutSpeedCurveNames) const;
 #endif
 
 #if WITH_EDITOR
diff --git a/Engine/Plugins/Animation/IKRig/Source/IKRigEditor/Private/RetargetEditor/IKRetargetAnimInstance.cpp b/Engine/Plugins/Animation/IKRig/Source/IKRigEditor/Private/RetargetEditor/IKRetargetAnimInstance.cpp
--- a/Engine/Plugins/Animation/IKRig/Source/IKRigEditor/Private/RetargetEditor/IKRetargetAnimInstance.cpp
+++ b/Engine/Plugins/Animation/IKRig/Source/IKRigEditor/Private/RetargetEditor/IKRetargetAnimInstance.cpp
@@ -17,7 +17,7 @@ void UIKRetargetAnimInstance::SetRetargetAssetAndSourceComponent(
 	Proxy.SetRetargetAssetAndSourceComponent(InAsset, InSourceMeshComponent);
 }
 
-const UIKRetargetProcessor* UIKRetargetAnimInstance::GetRetargetProcessor() const
+UIKRetargetProcessor* UIKRetargetAnimInstance::GetRetargetProcessor() const
 {
 	return IKRetargeterNode.GetRetargetProcessor();
 }
diff --git a/Engine/Plugins/Animation/IKRig/Source/IKRigEditor/Private/RetargetEditor/IKRetargetBatchOperation.cpp b/Engine/Plugins/Animation/IKRig/Source/IKRigEditor/Private/RetargetEditor/IKRetargetBatchOperation.cpp
--- a/Engine/Plugins/Animation/IKRig/Source/IKRigEditor/Private/RetargetEditor/IKRetargetBatchOperation.cpp
+++ b/Engine/Plugins/Animation/IKRig/Source/IKRigEditor/Private/RetargetEditor/IKRetargetBatchOperation.cpp
@@ -260,6 +260,23 @@ void FIKRetargetBatchOperation::ConvertAnimation(
 
 	TArray<FTransform> SourceComponentPose;
 	SourceComponentPose.SetNum(NumSourceBones);
+
+	// get names of the curves the retargeter is looking for
+	TArray<FName> SpeedCurveNames;
+	Context.IKRetargetAsset->GetSpeedCurveNames(SpeedCurveNames);
+	// get all curve names in the source skeleton
+	const FSmartNameMapping* SourceContainer = Context.SourceMesh->GetSkeleton()->GetSmartNameContainer(USkeleton::AnimCurveMappingName);
+	TArray<FName> SourceCurveNames;
+	SourceContainer->FillNameArray(SourceCurveNames);
+	// get list of all the source curve IDs that hold speed values the retargeter is looking for
+	TMap<FName, SmartName::UID_Type> SourceSpeedCurveIDs;
+	for (int32 CurveIndex = 0; CurveIndex < SourceCurveNames.Num(); ++CurveIndex)
+	{
+		if (SpeedCurveNames.Contains(SourceCurveNames[CurveIndex]))
+		{
+			SourceSpeedCurveIDs.Add(SourceCurveNames[CurveIndex],SourceContainer->FindUID(SourceCurveNames[CurveIndex]));
+		}
+	}
 	
 	// for each pair of source / target animation sequences
 	for (TPair<UAnimationAsset*, UAnimationAsset*>& Pair : DuplicatedAnimAssets)
@@ -295,6 +312,9 @@ void FIKRetargetBatchOperation::ConvertAnimation(
 		// ensure we evaluate the source animation using the skeletal mesh proportions that were evaluated in the viewport
 		FAnimPoseEvaluationOptions EvaluationOptions = FAnimPoseEvaluationOptions();
 		EvaluationOptions.OptionalSkeletalMesh = SourceSkeleton.SkeletalMesh;
+
+		// reset the planting state
+		Processor->ResetPlanting();
 		
 		// retarget each frame's pose from source to target
 		for (int32 FrameIndex=0; FrameIndex<NumFrames; ++FrameIndex)
@@ -314,8 +334,24 @@ void FIKRetargetBatchOperation::ConvertAnimation(
 			// update goals 
 			Processor->CopyAllSettingsFromAsset();
 			
-			// run the retarget
-			const TArray<FTransform>& TargetComponentPose = Processor->RunRetargeter(SourceComponentPose);
+			// calculate the delta time
+			const float TimeAtCurrentFrame = SourceSequence->GetTimeAtFrame(FrameIndex);
+			float DeltaTime = TimeAtCurrentFrame;
+			if (FrameIndex > 0)
+			{
+				const float TimeAtPrevFrame = SourceSequence->GetTimeAtFrame(FrameIndex-1);
+				DeltaTime = TimeAtCurrentFrame - TimeAtPrevFrame;
+			}
+			
+			// get the curve values from the source sequence (for speed-based IK planting)
+			TMap<FName, float> SpeedCurveValues;
+			for (TPair<FName, SmartName::UID_Type> CurveNameAndID : SourceSpeedCurveIDs)
+			{
+				SpeedCurveValues.Add(CurveNameAndID.Key, SourceSequence->EvaluateCurveData(CurveNameAndID.Value, TimeAtCurrentFrame));
+			}
+
+			// run the retargeter
+			const TArray<FTransform>& TargetComponentPose = Processor->RunRetargeter(SourceComponentPose, SpeedCurveValues, DeltaTime);
 
 			// convert to a local-space pose
 			TArray<FTransform> TargetLocalPose = TargetComponentPose;
diff --git a/Engine/Plugins/Animation/IKRig/Source/IKRigEditor/Private/RetargetEditor/IKRetargetEditor.cpp b/Engine/Plugins/Animation/IKRig/Source/IKRigEditor/Private/RetargetEditor/IKRetargetEditor.cpp
--- a/Engine/Plugins/Animation/IKRig/Source/IKRigEditor/Private/RetargetEditor/IKRetargetEditor.cpp
+++ b/Engine/Plugins/Animation/IKRig/Source/IKRigEditor/Private/RetargetEditor/IKRetargetEditor.cpp
@@ -32,6 +32,7 @@ const FName IKRetargetEditorAppName = FName(TEXT("IKRetargetEditorApp"));
 
 FIKRetargetEditor::FIKRetargetEditor()
 	: EditorController(MakeShared<FIKRetargetEditorController>())
+	, PreviousTime(-1.0f)
 {
 }
 
@@ -216,6 +217,15 @@ void FIKRetargetEditor::Tick(float DeltaTime)
 	// update with latest offsets
 	EditorController->AddOffsetAndUpdatePreviewMeshPosition(FVector::ZeroVector, EditorController->SourceSkelMeshComponent);
 	EditorController->AddOffsetAndUpdatePreviewMeshPosition(FVector::ZeroVector, EditorController->TargetSkelMeshComponent);
+
+	// retargeter IK planting must be reset when time is reversed or playback jumps ahead 
+	const float CurrentTime = EditorController->SourceAnimInstance->GetCurrentTime();
+	constexpr float MaxSkipTimeBeforeReset = 0.25f;
+	if (CurrentTime < PreviousTime || CurrentTime > PreviousTime + MaxSkipTimeBeforeReset)
+	{
+		EditorController->ResetIKPlantingState();
+	}
+	PreviousTime = CurrentTime;
 }
 
 TStatId FIKRetargetEditor::GetStatId() const
diff --git a/Engine/Plugins/Animation/IKRig/Source/IKRigEditor/Private/RetargetEditor/IKRetargetEditorController.cpp b/Engine/Plugins/Animation/IKRig/Source/IKRigEditor/Private/RetargetEditor/IKRetargetEditorController.cpp
--- a/Engine/Plugins/Animation/IKRig/Source/IKRigEditor/Private/RetargetEditor/IKRetargetEditorController.cpp
+++ b/Engine/Plugins/Animation/IKRig/Source/IKRigEditor/Private/RetargetEditor/IKRetargetEditorController.cpp
@@ -191,7 +191,7 @@ bool FIKRetargetEditorController::GetTargetBoneLineSegments(
 	return true;
 }
 
-const UIKRetargetProcessor* FIKRetargetEditorController::GetRetargetProcessor() const
+UIKRetargetProcessor* FIKRetargetEditorController::GetRetargetProcessor() const
 {	
 	if(UIKRetargetAnimInstance* AnimInstance = TargetAnimInstance.Get())
 	{
@@ -201,6 +201,14 @@ const UIKRetargetProcessor* FIKRetargetEditorController::GetRetargetProcessor()
 	return nullptr;	
 }
 
+void FIKRetargetEditorController::ResetIKPlantingState() const
+{
+	if (TargetAnimInstance.IsValid())
+	{
+		TargetAnimInstance->GetRetargetProcessor()->ResetPlanting();
+	}
+}
+
 void FIKRetargetEditorController::ClearOutputLog() const
 {
 	if (OutputLogView.IsValid())
diff --git a/Engine/Plugins/Animation/IKRig/Source/IKRigEditor/Public/RetargetEditor/IKRetargetAnimInstance.h b/Engine/Plugins/Animation/IKRig/Source/IKRigEditor/Public/RetargetEditor/IKRetargetAnimInstance.h
--- a/Engine/Plugins/Animation/IKRig/Source/IKRigEditor/Public/RetargetEditor/IKRetargetAnimInstance.h
+++ b/Engine/Plugins/Animation/IKRig/Source/IKRigEditor/Public/RetargetEditor/IKRetargetAnimInstance.h
@@ -20,7 +20,7 @@ public:
 		UIKRetargeter* InIKRetargetAsset,
 		TWeakObjectPtr<USkeletalMeshComponent> InSourceMeshComponent);
 
-	const UIKRetargetProcessor* GetRetargetProcessor() const;
+	UIKRetargetProcessor* GetRetargetProcessor() const;
 
 	void SetProcessorNeedsInitialized();
 
diff --git a/Engine/Plugins/Animation/IKRig/Source/IKRigEditor/Public/RetargetEditor/IKRetargetEditor.h b/Engine/Plugins/Animation/IKRig/Source/IKRigEditor/Public/RetargetEditor/IKRetargetEditor.h
--- a/Engine/Plugins/Animation/IKRig/Source/IKRigEditor/Public/RetargetEditor/IKRetargetEditor.h
+++ b/Engine/Plugins/Animation/IKRig/Source/IKRigEditor/Public/RetargetEditor/IKRetargetEditor.h
@@ -96,6 +96,9 @@ private:
 	
 	/** centralized management across all views */
 	TSharedRef<FIKRetargetEditorController> EditorController;
+
+	/** record previous playback time of source anim instance to trigger reset when scrubbing / looping */
+	float PreviousTime;
 	
 	friend FIKRetargetApplicationMode;
 };
diff --git a/Engine/Plugins/Animation/IKRig/Source/IKRigEditor/Public/RetargetEditor/IKRetargetEditorController.h b/Engine/Plugins/Animation/IKRig/Source/IKRigEditor/Public/RetargetEditor/IKRetargetEditorController.h
--- a/Engine/Plugins/Animation/IKRig/Source/IKRigEditor/Public/RetargetEditor/IKRetargetEditorController.h
+++ b/Engine/Plugins/Animation/IKRig/Source/IKRigEditor/Public/RetargetEditor/IKRetargetEditorController.h
@@ -86,7 +86,10 @@ public:
 		TArray<FVector>& OutChildren) const;
 
 	/** get the retargeter that is running in the viewport (which is a duplicate of the source asset) */
-	const UIKRetargetProcessor* GetRetargetProcessor() const;
+	UIKRetargetProcessor* GetRetargetProcessor() const;
+
+	/** Reset the planting state of the IK (when scrubbing or animation loops over) */
+	void ResetIKPlantingState() const;
 
 	/** Sequence Browser**/
 	void PlayAnimationAsset(UAnimationAsset* AssetToPlay);
