IK Batch Retargeter: Fix order dependence issues when retargeting AnimationAssets (Contributed by [REDACTED_PERSON])

PR #9119

diff --git a/Engine/Plugins/Animation/IKRig/Source/IKRigEditor/Private/RetargetEditor/IKRetargetBatchOperation.cpp b/Engine/Plugins/Animation/IKRig/Source/IKRigEditor/Private/RetargetEditor/IKRetargetBatchOperation.cpp
--- a/Engine/Plugins/Animation/IKRig/Source/IKRigEditor/Private/RetargetEditor/IKRetargetBatchOperation.cpp
+++ b/Engine/Plugins/Animation/IKRig/Source/IKRigEditor/Private/RetargetEditor/IKRetargetBatchOperation.cpp
@@ -194,6 +194,15 @@ void FIKRetargetBatchOperation::RetargetAssets(
 		AssetToRetarget->ReplaceReferredAnimations(RemappedAnimAssets);
 		AssetToRetarget->SetSkeleton(NewSkeleton);
 		AssetToRetarget->SetPreviewMesh(Context.TargetMesh);
+	}
+
+	// Call PostEditChange after the references of all assets were replaced,
+	// to prevent order dependence of post edit change hooks.
+	// If PostEditChange is called right after ReplaceReferredAnimations,
+	// it can access references that are still queued for retarget and follow
+	// the current asset in the array.
+	for (UAnimationAsset* AssetToRetarget : AnimationAssetsToRetarget)
+	{
 		AssetToRetarget->PostEditChange();
 		AssetToRetarget->MarkPackageDirty();
 	}
