IK Retargeter: Fixes crash when making new asset (regression from adding velocity matching)

diff --git a/Engine/Plugins/Animation/IKRig/Source/IKRigEditor/Private/RetargetEditor/IKRetargetEditorController.cpp b/Engine/Plugins/Animation/IKRig/Source/IKRigEditor/Private/RetargetEditor/IKRetargetEditorController.cpp
--- a/Engine/Plugins/Animation/IKRig/Source/IKRigEditor/Private/RetargetEditor/IKRetargetEditorController.cpp
+++ b/Engine/Plugins/Animation/IKRig/Source/IKRigEditor/Private/RetargetEditor/IKRetargetEditorController.cpp
@@ -203,9 +203,9 @@ UIKRetargetProcessor* FIKRetargetEditorController::GetRetargetProcessor() const
 
 void FIKRetargetEditorController::ResetIKPlantingState() const
 {
-	if (TargetAnimInstance.IsValid())
+	if (UIKRetargetProcessor* Processor = GetRetargetProcessor())
 	{
-		TargetAnimInstance->GetRetargetProcessor()->ResetPlanting();
+		Processor->ResetPlanting();
 	}
 }
 
