IK Retargeter: fixes bug, undo exits edit pose mode.

Also fixes:
- Could not undo chain map edits
- Undo cleared chain map selection
- Undo with selected chain map set details panel to base properties
- Editing bone into straight pose resulted in non-normalized rotations

diff --git a/Engine/Plugins/Animation/IKRig/Source/IKRigEditor/Private/RetargetEditor/IKRetargetDefaultMode.cpp b/Engine/Plugins/Animation/IKRig/Source/IKRigEditor/Private/RetargetEditor/IKRetargetDefaultMode.cpp
--- a/Engine/Plugins/Animation/IKRig/Source/IKRigEditor/Private/RetargetEditor/IKRetargetDefaultMode.cpp
+++ b/Engine/Plugins/Animation/IKRig/Source/IKRigEditor/Private/RetargetEditor/IKRetargetDefaultMode.cpp
@@ -9,6 +9,7 @@
 #include "IPersonaPreviewScene.h"
 #include "Animation/DebugSkelMeshComponent.h"
 #include "RetargetEditor/IKRetargetEditor.h"
+#include "RetargetEditor/SIKRetargetChainMapList.h"
 
 
 #define LOCTEXT_NAMESPACE "IKRetargetDefaultMode"
@@ -279,6 +280,7 @@ void FIKRetargetDefaultMode::SetSelectedComponent(UPrimitiveComponent* InCompone
 	// when clicking in empty space, show global details
 	if (SelectedComponent == nullptr)
 	{
+		Controller->ChainsView->ClearSelection();
 		Controller->DetailsView->SetObject(Controller->AssetController->GetAsset());
 	}
 }
diff --git a/Engine/Plugins/Animation/IKRig/Source/IKRigEditor/Private/RetargetEditor/IKRetargetEditor.cpp b/Engine/Plugins/Animation/IKRig/Source/IKRigEditor/Private/RetargetEditor/IKRetargetEditor.cpp
--- a/Engine/Plugins/Animation/IKRig/Source/IKRigEditor/Private/RetargetEditor/IKRetargetEditor.cpp
+++ b/Engine/Plugins/Animation/IKRig/Source/IKRigEditor/Private/RetargetEditor/IKRetargetEditor.cpp
@@ -235,36 +235,12 @@ TStatId FIKRetargetEditor::GetStatId() const
 
 void FIKRetargetEditor::PostUndo(bool bSuccess)
 {
-	EditorController->ClearOutputLog();
-	
-	const bool WasEditing = EditorController->IsEditingPose();
-	
-	EditorController->AssetController->BroadcastNeedsReinitialized();
-	EditorController->RefreshAllViews();
-
-	// restore pose mode state to avoid stepping out of the edition when undoing things
-	// note that BroadcastNeedsReinitialized will unset it in FIKRetargetEditorController::OnRetargeterNeedsInitialized
-	if (WasEditing)
-	{
-		EditorController->HandleEditPose();
-	}
+	EditorController->OnRetargeterNeedsInitialized(EditorController->AssetController->GetAsset());
 }
 
 void FIKRetargetEditor::PostRedo(bool bSuccess)
 {
-	EditorController->ClearOutputLog();
-	
-	const bool WasEditing = EditorController->IsEditingPose();
-	
-	EditorController->AssetController->BroadcastNeedsReinitialized();
-	EditorController->RefreshAllViews();
-	
-	// restore pose mode state to avoid stepping out of the edition when undoing things
-	// note that BroadcastNeedsReinitialized will unset it in FIKRetargetEditorController::OnRetargeterNeedsInitialized
-	if (WasEditing)
-	{
-		EditorController->HandleEditPose();
-	}
+	EditorController->OnRetargeterNeedsInitialized(EditorController->AssetController->GetAsset());
 }
 
 void FIKRetargetEditor::HandlePreviewSceneCreated(const TSharedRef<IPersonaPreviewScene>& InPersonaPreviewScene)
diff --git a/Engine/Plugins/Animation/IKRig/Source/IKRigEditor/Private/RetargetEditor/IKRetargetEditorController.cpp b/Engine/Plugins/Animation/IKRig/Source/IKRigEditor/Private/RetargetEditor/IKRetargetEditorController.cpp
--- a/Engine/Plugins/Animation/IKRig/Source/IKRigEditor/Private/RetargetEditor/IKRetargetEditorController.cpp
+++ b/Engine/Plugins/Animation/IKRig/Source/IKRigEditor/Private/RetargetEditor/IKRetargetEditorController.cpp
@@ -84,8 +84,6 @@ void FIKRetargetEditorController::OnRetargeterNeedsInitialized(const UIKRetarget
 {
 	// clear the output log
 	ClearOutputLog();
-	// force edit pose mode to be off
-	Editor.Pin()->GetEditorModeManager().DeactivateMode(FIKRetargetEditPoseMode::ModeName);
 	// force reinit the runtime retarget processor
 	TargetAnimInstance->SetProcessorNeedsInitialized();
 	// refresh all the UI views
@@ -141,6 +139,7 @@ FTransform FIKRetargetEditorController::GetTargetBoneGlobalTransform(
 	// scale and offset
 	BoneTransform.ScaleTranslation(AssetController->GetAsset()->TargetMeshScale);
 	BoneTransform.AddToTranslation(AssetController->GetAsset()->TargetMeshOffset);
+	BoneTransform.NormalizeRotation();
 
 	return BoneTransform;
 }
@@ -276,6 +275,7 @@ void FIKRetargetEditorController::HandleEditPose() const
 		// stop pose editing
 		Editor.Pin()->GetEditorModeManager().DeactivateMode(FIKRetargetEditPoseMode::ModeName);
 		Editor.Pin()->GetEditorModeManager().ActivateMode(FIKRetargetDefaultMode::ModeName);
+		AssetController->GetAsset()->SetOutputMode(ERetargeterOutputMode::RunRetarget);
 		
 		// must reinitialize after editing the retarget pose
 		AssetController->BroadcastNeedsReinitialized();
@@ -287,6 +287,7 @@ void FIKRetargetEditorController::HandleEditPose() const
 		// start pose editing
 		Editor.Pin()->GetEditorModeManager().DeactivateMode(FIKRetargetDefaultMode::ModeName);
 		Editor.Pin()->GetEditorModeManager().ActivateMode(FIKRetargetEditPoseMode::ModeName);
+		AssetController->GetAsset()->SetOutputMode(ERetargeterOutputMode::EditRetargetPose);
 	}
 }
 
diff --git a/Engine/Plugins/Animation/IKRig/Source/IKRigEditor/Private/RetargetEditor/IKRetargeterController.cpp b/Engine/Plugins/Animation/IKRig/Source/IKRigEditor/Private/RetargetEditor/IKRetargeterController.cpp
--- a/Engine/Plugins/Animation/IKRig/Source/IKRigEditor/Private/RetargetEditor/IKRetargeterController.cpp
+++ b/Engine/Plugins/Animation/IKRig/Source/IKRigEditor/Private/RetargetEditor/IKRetargeterController.cpp
@@ -335,11 +335,12 @@ void UIKRetargeterController::OnRetargetChainRemoved(UIKRigDefinition* IKRig, co
 
 void UIKRetargeterController::SetSourceChainForTargetChain(URetargetChainSettings* ChainMap, FName SourceChainToMapTo) const
 {
-	FScopedTransaction Transaction(LOCTEXT("SetRetargetChainSource", "Set Retarget Chain Source"));
-	Asset->Modify();
-	
 	check(ChainMap)
+	
+	FScopedTransaction Transaction(LOCTEXT("SetRetargetChainSource", "Set Retarget Chain Source"));
+	ChainMap->Modify();
 	ChainMap->SourceChain = SourceChainToMapTo;
+	
 	BroadcastNeedsReinitialized();
 }
 
diff --git a/Engine/Plugins/Animation/IKRig/Source/IKRigEditor/Private/RetargetEditor/SIKRetargetChainMapList.cpp b/Engine/Plugins/Animation/IKRig/Source/IKRigEditor/Private/RetargetEditor/SIKRetargetChainMapList.cpp
--- a/Engine/Plugins/Animation/IKRig/Source/IKRigEditor/Private/RetargetEditor/SIKRetargetChainMapList.cpp
+++ b/Engine/Plugins/Animation/IKRig/Source/IKRigEditor/Private/RetargetEditor/SIKRetargetChainMapList.cpp
@@ -235,6 +235,11 @@ void SIKRetargetChainMapList::Construct(
 	RefreshView();
 }
 
+void SIKRetargetChainMapList::ClearSelection() const
+{
+	ListView->ClearSelection();
+}
+
 UIKRetargeterController* SIKRetargetChainMapList::GetRetargetController() const
 {
 	const TSharedPtr<FIKRetargetEditorController> Controller = EditorController.Pin();
@@ -293,15 +298,28 @@ void SIKRetargetChainMapList::RefreshView()
 		return; 
 	}
 	
+	TArray<TSharedPtr<FRetargetChainMapElement>> PreviouslySelectedItems = ListView->GetSelectedItems();
+	TArray<TObjectPtr<URetargetChainSettings>> SelectedChainMaps;
+	for (TSharedPtr<FRetargetChainMapElement> SelectedItem : PreviouslySelectedItems)
+	{
+		SelectedChainMaps.Add(SelectedItem->ChainMap.Get());
+	}
+	
 	// refresh list of chains
+	TArray< TSharedPtr<FRetargetChainMapElement> > SelectedItems;
 	ListViewItems.Reset();
 	const TArray<TObjectPtr<URetargetChainSettings>>& ChainMappings = RetargeterController->GetChainMappings();
 	for (const TObjectPtr<URetargetChainSettings> ChainMap : ChainMappings)
 	{
 		TSharedPtr<FRetargetChainMapElement> ChainItem = FRetargetChainMapElement::Make(ChainMap);
 		ListViewItems.Add(ChainItem);
-	}
 
+		if (SelectedChainMaps.Contains(ChainMap))
+		{
+			ListView->SetItemSelection(ChainItem, true, ESelectInfo::Direct);
+		}
+	}
+	
 	ListView->RequestListRefresh();
 }
 
@@ -328,8 +346,8 @@ void SIKRetargetChainMapList::OnSelectionChanged()
 		return;
 	}
 	
-	// update selected settings
-	SelectedChainSettings.Empty();
+	// get selected chain settings
+	TArray<UObject*> SelectedChainSettings;
 	TArray<TSharedPtr<FRetargetChainMapElement>> SelectedItems = ListView.Get()->GetSelectedItems();
 	for (const TSharedPtr<FRetargetChainMapElement>& Item : SelectedItems)
 	{
diff --git a/Engine/Plugins/Animation/IKRig/Source/IKRigEditor/Public/RetargetEditor/SIKRetargetChainMapList.h b/Engine/Plugins/Animation/IKRig/Source/IKRigEditor/Public/RetargetEditor/SIKRetargetChainMapList.h
--- a/Engine/Plugins/Animation/IKRig/Source/IKRigEditor/Public/RetargetEditor/SIKRetargetChainMapList.h
+++ b/Engine/Plugins/Animation/IKRig/Source/IKRigEditor/Public/RetargetEditor/SIKRetargetChainMapList.h
@@ -75,12 +75,16 @@ typedef SListView< TSharedPtr<FRetargetChainMapElement> > SRetargetChainMapListV
 
 class SIKRetargetChainMapList : public SCompoundWidget, public FEditorUndoClient
 {
+	
 public:
+	
 	SLATE_BEGIN_ARGS(SIKRetargetChainMapList) {}
 	SLATE_END_ARGS()
 
 	void Construct(const FArguments& InArgs, TSharedRef<FIKRetargetEditorController> InEditorController);
 
+	void ClearSelection() const;
+	
 private:
 	
 	/** menu for adding new solver commands */
@@ -107,7 +111,6 @@ private:
 
 	/** update selection / details view */
 	void OnSelectionChanged();
-	TArray<UObject*> SelectedChainSettings;
 
 	/** auto-map chain button*/
 	EVisibility IsAutoMapButtonVisible() const;
