Fix crash if target IK Rig missing a mapped chain.

diff --git a/Engine/Plugins/Animation/IKRig/Source/IKRig/Private/Retargeter/IKRetargetProcessor.cpp b/Engine/Plugins/Animation/IKRig/Source/IKRig/Private/Retargeter/IKRetargetProcessor.cpp
--- a/Engine/Plugins/Animation/IKRig/Source/IKRig/Private/Retargeter/IKRetargetProcessor.cpp
+++ b/Engine/Plugins/Animation/IKRig/Source/IKRig/Private/Retargeter/IKRetargetProcessor.cpp
@@ -1302,6 +1302,7 @@ bool UIKRetargetProcessor::InitializeBoneChainPairs()
 			RetargeterAsset->Log.LogWarning( FText::Format(
 			LOCTEXT("MissingTargetChain", "IK Retargeter missing target bone chain: {0}. Please update the mapping."),
 			FText::FromString(ChainMap->TargetChain.ToString())));
+			continue;
 		}
 		
 		// user opted to not map this to anything, we don't need to spam a warning about it
