Exported, retargeted animations no longer have to click "Update" on Retarget Source Asset to get correct translation retargeting.

diff --git a/Engine/Plugins/Animation/IKRig/Source/IKRigEditor/Private/RetargetEditor/IKRetargetBatchOperation.cpp b/Engine/Plugins/Animation/IKRig/Source/IKRigEditor/Private/RetargetEditor/IKRetargetBatchOperation.cpp
--- a/Engine/Plugins/Animation/IKRig/Source/IKRigEditor/Private/RetargetEditor/IKRetargetBatchOperation.cpp
+++ b/Engine/Plugins/Animation/IKRig/Source/IKRigEditor/Private/RetargetEditor/IKRetargetBatchOperation.cpp
@@ -196,13 +196,20 @@ void FIKRetargetBatchOperation::RetargetAssets(
 		AssetToRetarget->SetPreviewMesh(Context.TargetMesh);
 	}
 
-	// Call PostEditChange after the references of all assets were replaced,
-	// to prevent order dependence of post edit change hooks.
-	// If PostEditChange is called right after ReplaceReferredAnimations,
-	// it can access references that are still queued for retarget and follow
-	// the current asset in the array.
+	// Call PostEditChange after the references of all assets were replaced, to prevent order dependence of post edit
+	// change hooks. If PostEditChange is called right after ReplaceReferredAnimations it can access references that are
+	// still queued for retarget and follow the current asset in the array.
+	static const FName RetargetSourceAssetPropertyName =  GET_MEMBER_NAME_STRING_CHECKED(UAnimSequence, RetargetSourceAsset);
+	static FProperty* RetargetAssetProperty = UAnimSequence::StaticClass()->FindPropertyByName(RetargetSourceAssetPropertyName);
 	for (UAnimationAsset* AssetToRetarget : AnimationAssetsToRetarget)
 	{
+		// force updating of the retarget pose, this is normally done on PreSave() but is guarded against procedural saves
+		if (UAnimSequence* AnimSequenceToRetarget = Cast<UAnimSequence>(AssetToRetarget))
+		{
+			FPropertyChangedEvent RetargetAssetPropertyChangedEvent(RetargetAssetProperty);
+			AnimSequenceToRetarget->PostEditChangeProperty(RetargetAssetPropertyChangedEvent);
+		}
+		
 		AssetToRetarget->PostEditChange();
 		AssetToRetarget->MarkPackageDirty();
 	}
