Adding ability to import/export retarget poses.

diff --git a/Engine/Plugins/Animation/IKRig/Source/IKRig/Private/Retargeter/IKRetargetProcessor.cpp b/Engine/Plugins/Animation/IKRig/Source/IKRig/Private/Retargeter/IKRetargetProcessor.cpp
--- a/Engine/Plugins/Animation/IKRig/Source/IKRig/Private/Retargeter/IKRetargetProcessor.cpp
+++ b/Engine/Plugins/Animation/IKRig/Source/IKRig/Private/Retargeter/IKRetargetProcessor.cpp
@@ -140,6 +140,34 @@ void FRetargetSkeleton::GetChildrenIndices(const int32 BoneIndex, TArray<int32>&
 	}
 }
 
+void FRetargetSkeleton::GetChildrenIndicesRecursive(const int32 BoneIndex, TArray<int32>& OutChildren) const
+{
+	const int32 NumBones = BoneNames.Num();
+	for (int32 ChildIndex=BoneIndex+1; ChildIndex<NumBones; ++ChildIndex)
+	{
+		if (!OutChildren.Contains(ChildIndex) && IsParentOfChild(BoneIndex, ChildIndex))
+		{
+			OutChildren.Add(ChildIndex);
+		}
+	}
+}
+
+bool FRetargetSkeleton::IsParentOfChild(const int32 PotentialParentIndex, const int32 ChildBoneIndex) const
+{
+	int32 ParentIndex = GetParentIndex(ChildBoneIndex);
+	while (ParentIndex != INDEX_NONE)
+	{
+		if (ParentIndex == PotentialParentIndex)
+		{
+			return true;
+		}
+		
+		ParentIndex = GetParentIndex(ParentIndex);
+	}
+	
+	return false;
+}
+
 int32 FRetargetSkeleton::GetParentIndex(const int32 BoneIndex) const
 {
 	if (BoneIndex < 0 || BoneIndex>ParentIndices.Num() || BoneIndex == INDEX_NONE)
diff --git a/Engine/Plugins/Animation/IKRig/Source/IKRig/Public/Retargeter/IKRetargetProcessor.h b/Engine/Plugins/Animation/IKRig/Source/IKRig/Public/Retargeter/IKRetargetProcessor.h
--- a/Engine/Plugins/Animation/IKRig/Source/IKRig/Public/Retargeter/IKRetargetProcessor.h
+++ b/Engine/Plugins/Animation/IKRig/Source/IKRig/Public/Retargeter/IKRetargetProcessor.h
@@ -65,6 +65,10 @@ struct IKRIG_API FRetargetSkeleton
 		const TArray<FTransform>& InGlobalPose) const;
 
 	void GetChildrenIndices(const int32 BoneIndex, TArray<int32>& OutChildren) const;
+
+	void GetChildrenIndicesRecursive(const int32 BoneIndex, TArray<int32>& OutChildren) const;
+	
+	bool IsParentOfChild(const int32 PotentialParentIndex, const int32 ChildBoneIndex) const;
 };
 
 struct FTargetSkeleton : public FRetargetSkeleton
diff --git a/Engine/Plugins/Animation/IKRig/Source/IKRig/Public/Retargeter/IKRetargeter.h b/Engine/Plugins/Animation/IKRig/Source/IKRig/Public/Retargeter/IKRetargeter.h
--- a/Engine/Plugins/Animation/IKRig/Source/IKRig/Public/Retargeter/IKRetargeter.h
+++ b/Engine/Plugins/Animation/IKRig/Source/IKRig/Public/Retargeter/IKRetargeter.h
@@ -210,6 +210,26 @@ public:
 	void SortHierarchically(const FIKRigSkeleton& Skeleton);
 };
 
+UCLASS(Blueprintable)
+class IKRIG_API URetargetPose: public UObject
+{
+	GENERATED_BODY()
+
+public:
+
+	UPROPERTY(EditAnywhere, Category = RetargetPose)
+	FVector RootTranslationOffset = FVector::ZeroVector;
+	
+	UPROPERTY(EditAnywhere, Category = RetargetPose)
+	TMap<FName, FQuat> BoneRotationOffsets;
+
+	void GetAsRetargetPose(FIKRetargetPose& OutPose) const
+	{
+		OutPose.RootTranslationOffset = RootTranslationOffset;
+		OutPose.BoneRotationOffsets = BoneRotationOffsets;
+	}
+};
+
 #if WITH_EDITOR
 enum class ERetargeterOutputMode : uint8
 {
diff --git a/Engine/Plugins/Animation/IKRig/Source/IKRigEditor/Private/IKRigEditor.cpp b/Engine/Plugins/Animation/IKRig/Source/IKRigEditor/Private/IKRigEditor.cpp
--- a/Engine/Plugins/Animation/IKRig/Source/IKRigEditor/Private/IKRigEditor.cpp
+++ b/Engine/Plugins/Animation/IKRig/Source/IKRigEditor/Private/IKRigEditor.cpp
@@ -10,6 +10,7 @@
 #include "IKRigDefinition.h"
 #include "RigEditor/AssetTypeActions_IKRigDefinition.h"
 #include "RetargetEditor/AssetTypeActions_IKRetargeter.h"
+#include "RetargetEditor/AssetTypeActions_RetargetPose.h"
 #include "RetargetEditor/IKRetargetCommands.h"
 #include "RetargetEditor/IKRetargetDefaultMode.h"
 #include "RetargetEditor/IKRetargetDetails.h"
@@ -34,11 +35,16 @@ void FIKRigEditor::StartupModule()
 	FIKRetargetCommands::Register();
 	
 	// register custom asset type actions
+	IAssetTools& ToolsModule = FModuleManager::LoadModuleChecked<FAssetToolsModule>("AssetTools").Get();
+	//
 	IKRigDefinitionAssetAction = MakeShareable(new FAssetTypeActions_IKRigDefinition);
-	FModuleManager::LoadModuleChecked<FAssetToolsModule>("AssetTools").Get().RegisterAssetTypeActions(IKRigDefinitionAssetAction.ToSharedRef());
+	ToolsModule.RegisterAssetTypeActions(IKRigDefinitionAssetAction.ToSharedRef());
 	//
 	IKRetargeterAssetAction = MakeShareable(new FAssetTypeActions_IKRetargeter);
-	FModuleManager::LoadModuleChecked<FAssetToolsModule>("AssetTools").Get().RegisterAssetTypeActions(IKRetargeterAssetAction.ToSharedRef());
+	ToolsModule.RegisterAssetTypeActions(IKRetargeterAssetAction.ToSharedRef());
+	//
+	RetargetPoseAssetAction = MakeShareable(new FAssetTypeActions_RetargetPose);
+	ToolsModule.RegisterAssetTypeActions(RetargetPoseAssetAction.ToSharedRef());
 
 	// extend the content browser menu
 	FAssetTypeActions_IKRetargeter::ExtendAnimSequenceToolMenu();
@@ -68,24 +74,24 @@ void FIKRigEditor::ShutdownModule()
 	FEditorModeRegistry::Get().UnregisterMode(FIKRetargetDefaultMode::ModeName);
 	FEditorModeRegistry::Get().UnregisterMode(FIKRetargetEditPoseMode::ModeName);
 
-	// unregister IKRigDefinition asset action
-	if (IKRigDefinitionAssetAction.IsValid())
+	// unregister asset actions
+	IAssetTools& ToolsModule = FModuleManager::LoadModuleChecked<FAssetToolsModule>("AssetTools").Get();
+	if (FModuleManager::Get().IsModuleLoaded("AssetTools"))
 	{
-		if (FModuleManager::Get().IsModuleLoaded("AssetTools"))
+		if (IKRigDefinitionAssetAction.IsValid())
 		{
-			FModuleManager::GetModuleChecked<FAssetToolsModule>("AssetTools").Get().UnregisterAssetTypeActions(IKRigDefinitionAssetAction.ToSharedRef());
+			ToolsModule.UnregisterAssetTypeActions(IKRigDefinitionAssetAction.ToSharedRef());
 		}
-		IKRigDefinitionAssetAction.Reset();
-	}
-
-	// unregister IKRetargeter asset action
-	if (IKRetargeterAssetAction.IsValid())
-	{
-		if (FModuleManager::Get().IsModuleLoaded("AssetTools"))
+		
+		if (IKRetargeterAssetAction.IsValid())
+		{
+			ToolsModule.UnregisterAssetTypeActions(IKRetargeterAssetAction.ToSharedRef());
+		}
+		
+		if (RetargetPoseAssetAction.IsValid())
 		{
-			FModuleManager::GetModuleChecked<FAssetToolsModule>("AssetTools").Get().UnregisterAssetTypeActions(IKRetargeterAssetAction.ToSharedRef());
+			ToolsModule.UnregisterAssetTypeActions(RetargetPoseAssetAction.ToSharedRef());
 		}
-		IKRetargeterAssetAction.Reset();
 	}
 }
 
diff --git a/Engine/Plugins/Animation/IKRig/Source/IKRigEditor/Private/RetargetEditor/AssetTypeActions_RetargetPose.cpp b/Engine/Plugins/Animation/IKRig/Source/IKRigEditor/Private/RetargetEditor/AssetTypeActions_RetargetPose.cpp
new file mode 100644
--- /dev/null
+++ b/Engine/Plugins/Animation/IKRig/Source/IKRigEditor/Private/RetargetEditor/AssetTypeActions_RetargetPose.cpp
@@ -0,0 +1,38 @@
+﻿// Copyright Epic Games, Inc. All Rights Reserved.
+
+#include "RetargetEditor/AssetTypeActions_RetargetPose.h"
+#include "Retargeter/IKRetargeter.h"
+
+#include "Animation/AnimationAsset.h"
+#include "ContentBrowserMenuContexts.h"
+#include "ToolMenus.h"
+#include "ThumbnailRendering/SceneThumbnailInfo.h"
+
+#define LOCTEXT_NAMESPACE "AssetTypeActions"
+
+UClass* FAssetTypeActions_RetargetPose::GetSupportedClass() const
+{
+	return URetargetPose::StaticClass();
+}
+
+const TArray<FText>& FAssetTypeActions_RetargetPose::GetSubMenus() const
+{
+	static const TArray<FText> SubMenus
+	{
+		LOCTEXT("AnimIKRigSubMenu", "IK Rig")
+	};
+	return SubMenus;
+}
+
+UThumbnailInfo* FAssetTypeActions_RetargetPose::GetThumbnailInfo(UObject* Asset) const
+{
+	URetargetPose* RetargetPose = CastChecked<URetargetPose>(Asset);
+	return NewObject<USceneThumbnailInfo>(RetargetPose, NAME_None, RF_Transactional);
+}
+
+void FAssetTypeActions_RetargetPose::GetActions(const TArray<UObject*>& InObjects, FToolMenuSection& Section)
+{
+	FAssetTypeActions_Base::GetActions(InObjects, Section);
+}
+
+#undef LOCTEXT_NAMESPACE
diff --git a/Engine/Plugins/Animation/IKRig/Source/IKRigEditor/Private/RetargetEditor/IKRetargetCommands.cpp b/Engine/Plugins/Animation/IKRig/Source/IKRigEditor/Private/RetargetEditor/IKRetargetCommands.cpp
--- a/Engine/Plugins/Animation/IKRig/Source/IKRigEditor/Private/RetargetEditor/IKRetargetCommands.cpp
+++ b/Engine/Plugins/Animation/IKRig/Source/IKRigEditor/Private/RetargetEditor/IKRetargetCommands.cpp
@@ -9,11 +9,17 @@ void FIKRetargetCommands::RegisterCommands()
 	UI_COMMAND(GoToRetargetPose, "Go to Retarget Pose", "Stop playback and set Source and Target meshes to the Retarget Pose.", EUserInterfaceActionType::Button, FInputChord());
 	
 	UI_COMMAND(EditRetargetPose, "Edit Mode", "Enter into mode allowing manual editing of the target skeleton pose in the viewport.", EUserInterfaceActionType::ToggleButton, FInputChord());
-	UI_COMMAND(SetToRefPose, "Set to Ref Pose", "Sets the retarget pose to the mesh reference pose.", EUserInterfaceActionType::Button, FInputChord());
+	UI_COMMAND(ResetAllBones, "Reset All", "Sets the selected bones to the mesh reference pose.", EUserInterfaceActionType::Button, FInputChord());
+	UI_COMMAND(ResetSelectedBones, "Reset Selected Bones", "Sets the selected bones to the mesh reference pose.", EUserInterfaceActionType::Button, FInputChord());
+	UI_COMMAND(ResetSelectedAndChildrenBones, "Reset Selected and Children Bones", "Sets the selected bones (and all children recursively) to the mesh reference pose.", EUserInterfaceActionType::Button, FInputChord());
 
-	UI_COMMAND(NewRetargetPose, "New", "Create a new retarget pose for the target mesh.", EUserInterfaceActionType::Button, FInputChord());
+	UI_COMMAND(NewRetargetPose, "Create New", "Create a new retarget pose.", EUserInterfaceActionType::Button, FInputChord());
+	UI_COMMAND(DuplicateRetargetPose, "Duplicate Current", "Duplicate the current retarget pose.", EUserInterfaceActionType::Button, FInputChord());
 	UI_COMMAND(DeleteRetargetPose, "Delete", "Delete current retarget pose.", EUserInterfaceActionType::Button, FInputChord());
 	UI_COMMAND(RenameRetargetPose, "Rename", "Rename current retarget pose.", EUserInterfaceActionType::Button, FInputChord());
+	UI_COMMAND(ImportRetargetPose, "Import from Retarget Pose Asset", "Import a retarget pose asset.", EUserInterfaceActionType::Button, FInputChord());
+	UI_COMMAND(ImportRetargetPoseFromAnim, "Import from Animation Sequence", "Import a retarget pose from an animation sequence.", EUserInterfaceActionType::Button, FInputChord());
+	UI_COMMAND(ExportRetargetPose, "Export Retarget Pose Asset", "Export current retarget pose as an asset.", EUserInterfaceActionType::Button, FInputChord());
 }
 
 #undef LOCTEXT_NAMESPACE
diff --git a/Engine/Plugins/Animation/IKRig/Source/IKRigEditor/Private/RetargetEditor/IKRetargetDetails.cpp b/Engine/Plugins/Animation/IKRig/Source/IKRigEditor/Private/RetargetEditor/IKRetargetDetails.cpp
--- a/Engine/Plugins/Animation/IKRig/Source/IKRigEditor/Private/RetargetEditor/IKRetargetDetails.cpp
+++ b/Engine/Plugins/Animation/IKRig/Source/IKRigEditor/Private/RetargetEditor/IKRetargetDetails.cpp
@@ -33,6 +33,9 @@ void FIKRetargeterDetails::CustomizeDetails(IDetailLayoutBuilder& DetailBuilder)
 		return;
 	}
 
+	// the commands for the menus
+	TSharedPtr<FUICommandList> Commands = Controller->GetEditorController()->Editor.Pin()->GetToolkitCommands();
+
 	// add a new category at the top to edit the retarget pose
 	IDetailCategoryBuilder& EditPoseCategoryBuilder = DetailBuilder.EditCategory( "Edit Retarget Pose", LOCTEXT("EditPoseLabel", "Edit Retarget Pose"), ECategoryPriority::Default );
 
@@ -65,12 +68,7 @@ void FIKRetargeterDetails::CustomizeDetails(IDetailLayoutBuilder& DetailBuilder)
 		]
 	];
 
-	// add the bone size slider
-	TSharedRef<IPropertyHandle> BoneSizeHandle = DetailBuilder.GetProperty(GET_MEMBER_NAME_CHECKED(UIKRetargeter, BoneDrawSize));
-	EditPoseCategoryBuilder.AddProperty(BoneSizeHandle);
-
 	// add pose editing toolbar
-	const TSharedRef<FUICommandList>& Commands = Controller->GetEditorController()->Editor.Pin()->GetToolkitCommands();
 	FDetailWidgetRow& ToolbarRow = EditPoseCategoryBuilder.AddCustomRow(LOCTEXT("CurrentPoseLabel", "Edit Pose"))
 	.WholeRowWidget
 	[
@@ -82,9 +80,13 @@ void FIKRetargeterDetails::CustomizeDetails(IDetailLayoutBuilder& DetailBuilder)
 			MakeToolbar(Commands)
 		]
 	];
+
+	// add the bone size slider
+	TSharedRef<IPropertyHandle> BoneSizeHandle = DetailBuilder.GetProperty(GET_MEMBER_NAME_CHECKED(UIKRetargeter, BoneDrawSize));
+	EditPoseCategoryBuilder.AddProperty(BoneSizeHandle);
 }
 
-TSharedRef<SWidget>  FIKRetargeterDetails::MakeToolbar(const TSharedRef<FUICommandList>& Commands)
+TSharedRef<SWidget>  FIKRetargeterDetails::MakeToolbar(TSharedPtr<FUICommandList> Commands)
 {
 	FToolBarBuilder ToolbarBuilder(Commands, FMultiBoxCustomization::None);
 	
@@ -97,22 +99,23 @@ TSharedRef<SWidget>  FIKRetargeterDetails::MakeToolbar(const TSharedRef<FUIComma
 		TAttribute<FText>(),
 		FSlateIcon(FAppStyle::Get().GetStyleSetName(),"Icons.Edit"));
 
-	ToolbarBuilder.AddToolBarButton(
-		FIKRetargetCommands::Get().SetToRefPose,
-		NAME_None,
-		TAttribute<FText>(),
-		TAttribute<FText>(),
-		FSlateIcon(FAppStyle::Get().GetStyleSetName(),"Icons.Refresh"));
+	ToolbarBuilder.AddComboButton(
+		FUIAction(),
+		FOnGetContent::CreateSP(this, &FIKRetargeterDetails::GenerateResetMenuContent, Commands),
+		LOCTEXT("ResetPose_Label", "Reset"),
+		LOCTEXT("ResetPoseToolTip_Label", "Reset bones to reference pose."),
+		FSlateIcon(FAppStyle::GetAppStyleSetName(), "Icons.Refresh"));
 
 	ToolbarBuilder.EndSection();
+
 	ToolbarBuilder.BeginSection("Create Poses");
 
-	ToolbarBuilder.AddToolBarButton(
-		FIKRetargetCommands::Get().NewRetargetPose,
-		NAME_None,
-		TAttribute<FText>(),
+	ToolbarBuilder.AddComboButton(
+		FUIAction(),
+		FOnGetContent::CreateSP(this, &FIKRetargeterDetails::GenerateNewMenuContent, Commands),
+		LOCTEXT("CreatePose_Label", "Create"),
 		TAttribute<FText>(),
-		FSlateIcon(FAppStyle::Get().GetStyleSetName(),"Icons.Plus"));
+		FSlateIcon(FAppStyle::GetAppStyleSetName(), "Icons.Plus"));
 
 	ToolbarBuilder.AddToolBarButton(
 		FIKRetargetCommands::Get().DeleteRetargetPose,
@@ -133,6 +136,80 @@ TSharedRef<SWidget>  FIKRetargeterDetails::MakeToolbar(const TSharedRef<FUIComma
 	return ToolbarBuilder.MakeWidget();
 }
 
+TSharedRef<SWidget> FIKRetargeterDetails::GenerateResetMenuContent(TSharedPtr<FUICommandList> Commands)
+{
+	FMenuBuilder MenuBuilder(true, Commands);
+
+	MenuBuilder.AddMenuEntry(
+		FIKRetargetCommands::Get().ResetSelectedBones,
+		TEXT("Reset Selected"),
+		TAttribute<FText>(),
+		TAttribute<FText>());
+
+	MenuBuilder.AddMenuEntry(
+		FIKRetargetCommands::Get().ResetSelectedAndChildrenBones,
+		TEXT("Reset Selected And Children"),
+		TAttribute<FText>(),
+		TAttribute<FText>());
+
+	MenuBuilder.AddMenuEntry(
+		FIKRetargetCommands::Get().ResetAllBones,
+		TEXT("Reset All"),
+		TAttribute<FText>(),
+		TAttribute<FText>());
+
+	return MenuBuilder.MakeWidget();
+}
+
+TSharedRef<SWidget> FIKRetargeterDetails::GenerateNewMenuContent(TSharedPtr<FUICommandList> Commands)
+{
+	FMenuBuilder MenuBuilder(true, Commands);
+
+	MenuBuilder.BeginSection("Create", LOCTEXT("CreatePoseOperations", "Create New Retarget Pose"));
+	{
+		MenuBuilder.AddMenuEntry(
+		FIKRetargetCommands::Get().NewRetargetPose,
+		TEXT("Create"),
+		TAttribute<FText>(),
+		TAttribute<FText>());
+
+		MenuBuilder.AddMenuEntry(
+			FIKRetargetCommands::Get().DuplicateRetargetPose,
+			TEXT("Create"),
+			TAttribute<FText>(),
+			TAttribute<FText>());
+	}
+	MenuBuilder.EndSection();
+
+	MenuBuilder.BeginSection("Import",LOCTEXT("ImportPoseOperations", "Import Retarget Pose"));
+	{
+		MenuBuilder.AddMenuEntry(
+		FIKRetargetCommands::Get().ImportRetargetPose,
+		TEXT("Import"),
+		TAttribute<FText>(),
+		TAttribute<FText>());
+	
+		MenuBuilder.AddMenuEntry(
+			FIKRetargetCommands::Get().ImportRetargetPoseFromAnim,
+			TEXT("ImportFromSequence"),
+			TAttribute<FText>(),
+			TAttribute<FText>());
+	}
+	MenuBuilder.EndSection();
+
+	MenuBuilder.BeginSection("Export",LOCTEXT("EmportPoseOperations", "Export Retarget Pose"));
+	{
+		MenuBuilder.AddMenuEntry(
+			FIKRetargetCommands::Get().ExportRetargetPose,
+			TEXT("Export"),
+			TAttribute<FText>(),
+			TAttribute<FText>());
+	}
+	MenuBuilder.EndSection();
+
+	return MenuBuilder.MakeWidget();
+}
+
 TObjectPtr<UIKRetargeterController> FIKRetargeterDetails::GetAssetControllerFromSelectedObjects(IDetailLayoutBuilder& DetailBuilder) const
 {
 	TArray< TWeakObjectPtr<UObject> > OutObjects;
diff --git a/Engine/Plugins/Animation/IKRig/Source/IKRigEditor/Private/RetargetEditor/IKRetargetEditPoseMode.cpp b/Engine/Plugins/Animation/IKRig/Source/IKRigEditor/Private/RetargetEditor/IKRetargetEditPoseMode.cpp
--- a/Engine/Plugins/Animation/IKRig/Source/IKRigEditor/Private/RetargetEditor/IKRetargetEditPoseMode.cpp
+++ b/Engine/Plugins/Animation/IKRig/Source/IKRigEditor/Private/RetargetEditor/IKRetargetEditPoseMode.cpp
@@ -432,6 +432,11 @@ void FIKRetargetEditPoseMode::Exit()
 	IPersonaEditMode::Exit();
 }
 
+const TArray<FName>& FIKRetargetEditPoseMode::GetSelectedBones() const
+{
+	return BoneEdit.SelectedBones;
+}
+
 void FIKRetargetEditPoseMode::UpdateWidgetTransform()
 {
 	const TSharedPtr<FIKRetargetEditorController> Controller = EditorController.Pin();
diff --git a/Engine/Plugins/Animation/IKRig/Source/IKRigEditor/Private/RetargetEditor/IKRetargetEditor.cpp b/Engine/Plugins/Animation/IKRig/Source/IKRigEditor/Private/RetargetEditor/IKRetargetEditor.cpp
--- a/Engine/Plugins/Animation/IKRig/Source/IKRigEditor/Private/RetargetEditor/IKRetargetEditor.cpp
+++ b/Engine/Plugins/Animation/IKRig/Source/IKRigEditor/Private/RetargetEditor/IKRetargetEditor.cpp
@@ -125,15 +125,34 @@ void FIKRetargetEditor::BindCommands()
 		EUIActionRepeatMode::RepeatDisabled);
 
 	ToolkitCommands->MapAction(
-		Commands.SetToRefPose,
-		FExecuteAction::CreateSP(EditorController, &FIKRetargetEditorController::HandleResetPose),
+		Commands.ResetAllBones,
+		FExecuteAction::CreateSP(EditorController, &FIKRetargetEditorController::HandleResetAllBones),
 		FCanExecuteAction::CreateSP(EditorController, &FIKRetargetEditorController::CanResetPose),
 		EUIActionRepeatMode::RepeatDisabled);
 
+	ToolkitCommands->MapAction(
+		Commands.ResetSelectedBones,
+		FExecuteAction::CreateSP(EditorController, &FIKRetargetEditorController::HandleResetSelectedBones),
+		FCanExecuteAction::CreateSP(EditorController, &FIKRetargetEditorController::CanResetSelected),
+		EUIActionRepeatMode::RepeatDisabled);
+
+	ToolkitCommands->MapAction(
+		Commands.ResetSelectedAndChildrenBones,
+		FExecuteAction::CreateSP(EditorController, &FIKRetargetEditorController::HandleResetSelectedAndChildrenBones),
+		FCanExecuteAction::CreateSP(EditorController, &FIKRetargetEditorController::CanResetSelected),
+		EUIActionRepeatMode::RepeatDisabled);
+
 	ToolkitCommands->MapAction(
 		Commands.NewRetargetPose,
 		FExecuteAction::CreateSP(EditorController, &FIKRetargetEditorController::HandleNewPose),
-		FCanExecuteAction::CreateSP(EditorController, &FIKRetargetEditorController::CanNewPose),
+		FCanExecuteAction::CreateSP(EditorController, &FIKRetargetEditorController::CanCreatePose),
+		FCanExecuteAction(),
+		EUIActionRepeatMode::RepeatDisabled);
+
+	ToolkitCommands->MapAction(
+		Commands.DuplicateRetargetPose,
+		FExecuteAction::CreateSP(EditorController, &FIKRetargetEditorController::HandleDuplicatePose),
+		FCanExecuteAction::CreateSP(EditorController, &FIKRetargetEditorController::CanCreatePose),
 		FCanExecuteAction(),
 		EUIActionRepeatMode::RepeatDisabled);
 
@@ -150,6 +169,27 @@ void FIKRetargetEditor::BindCommands()
 		FCanExecuteAction::CreateSP(EditorController, &FIKRetargetEditorController::CanRenamePose),
 		FCanExecuteAction(),
 		EUIActionRepeatMode::RepeatDisabled);
+
+	ToolkitCommands->MapAction(
+		Commands.ImportRetargetPose,
+		FExecuteAction::CreateSP(EditorController, &FIKRetargetEditorController::HandleImportPose),
+		FCanExecuteAction::CreateSP(EditorController, &FIKRetargetEditorController::CanCreatePose),
+		FCanExecuteAction(),
+		EUIActionRepeatMode::RepeatDisabled);
+
+	ToolkitCommands->MapAction(
+		Commands.ImportRetargetPoseFromAnim,
+		FExecuteAction::CreateSP(EditorController, &FIKRetargetEditorController::HandleImportPoseFromSequence),
+		FCanExecuteAction::CreateSP(EditorController, &FIKRetargetEditorController::CanCreatePose),
+		FCanExecuteAction(),
+		EUIActionRepeatMode::RepeatDisabled);
+
+	ToolkitCommands->MapAction(
+		Commands.ExportRetargetPose,
+		FExecuteAction::CreateSP(EditorController, &FIKRetargetEditorController::HandleExportPose),
+		FCanExecuteAction::CreateSP(EditorController, &FIKRetargetEditorController::CanCreatePose),
+		FCanExecuteAction(),
+		EUIActionRepeatMode::RepeatDisabled);
 }
 
 void FIKRetargetEditor::ExtendToolbar()
diff --git a/Engine/Plugins/Animation/IKRig/Source/IKRigEditor/Private/RetargetEditor/IKRetargetEditorController.cpp b/Engine/Plugins/Animation/IKRig/Source/IKRigEditor/Private/RetargetEditor/IKRetargetEditorController.cpp
--- a/Engine/Plugins/Animation/IKRig/Source/IKRigEditor/Private/RetargetEditor/IKRetargetEditorController.cpp
+++ b/Engine/Plugins/Animation/IKRig/Source/IKRigEditor/Private/RetargetEditor/IKRetargetEditorController.cpp
@@ -2,13 +2,18 @@
 
 #include "RetargetEditor/IKRetargetEditorController.h"
 
+#include "AnimPose.h"
+#include "AssetToolsModule.h"
+#include "ContentBrowserModule.h"
 #include "EditorModeManager.h"
+#include "IContentBrowserSingleton.h"
 #include "Widgets/Input/SButton.h"
 #include "Animation/DebugSkelMeshComponent.h"
 #include "RetargetEditor/IKRetargetAnimInstance.h"
 #include "RetargetEditor/IKRetargetDefaultMode.h"
 #include "RetargetEditor/IKRetargetEditPoseMode.h"
 #include "RetargetEditor/IKRetargetEditor.h"
+#include "RetargetEditor/IKRetargetFactory.h"
 #include "RetargetEditor/SIKRetargetChainMapList.h"
 #include "Retargeter/IKRetargeter.h"
 #include "RigEditor/SIKRigOutputLog.h"
@@ -127,6 +132,37 @@ USkeletalMesh* FIKRetargetEditorController::GetTargetSkeletalMesh() const
 	return AssetController ? AssetController->GetTargetPreviewMesh() : nullptr;
 }
 
+const USkeleton* FIKRetargetEditorController::GetSourceSkeleton() const
+{
+	if (const USkeletalMesh* Mesh = GetSourceSkeletalMesh())
+	{
+		return Mesh->GetSkeleton();
+	}
+	
+	return nullptr;
+}
+
+const USkeleton* FIKRetargetEditorController::GetTargetSkeleton() const
+{
+	if (const USkeletalMesh* Mesh = GetTargetSkeletalMesh())
+	{
+		return Mesh->GetSkeleton();
+	}
+	
+	return nullptr;
+}
+
+const TArray<FName>& FIKRetargetEditorController::GetSelectedBones() const
+{
+	if (const FIKRetargetEditPoseMode* EditPoseMode = Editor.Pin()->GetEditorModeManager().GetActiveModeTyped<FIKRetargetEditPoseMode>(FIKRetargetEditPoseMode::ModeName))
+	{
+		return EditPoseMode->GetSelectedBones();
+	}
+	
+	static TArray<FName> Empty;
+	return Empty;
+}
+
 FTransform FIKRetargetEditorController::GetTargetBoneGlobalTransform(
 	const UIKRetargetProcessor* RetargetProcessor,
 	const int32& TargetBoneIndex) const
@@ -310,12 +346,12 @@ bool FIKRetargetEditorController::IsEditingPose() const
 void FIKRetargetEditorController::HandleNewPose()
 {
 	// get a unique pose name to use as suggestion
-	const FName DefaultNewPoseName = FName(LOCTEXT("NewRetargetPoseName", "CustomRetargetPose").ToString());
+	const FString DefaultNewPoseName = LOCTEXT("NewRetargetPoseName", "CustomRetargetPose").ToString();
 	const FName UniqueNewPoseName = AssetController->MakePoseNameUnique(DefaultNewPoseName);
 	
 	SAssignNew(NewPoseWindow, SWindow)
 	.Title(LOCTEXT("NewRetargetPoseOptions", "Create New Retarget Pose"))
-	.ClientSize(FVector2D(250, 80))
+	.ClientSize(FVector2D(300, 80))
 	.HasCloseButton(true)
 	.SupportsMinimize(false) .SupportsMaximize(false)
 	[
@@ -323,48 +359,49 @@ void FIKRetargetEditorController::HandleNewPose()
 		.BorderImage( FAppStyle::GetBrush("Menu.Background") )
 		[
 			SNew(SVerticalBox)
-
 			+ SVerticalBox::Slot()
 			.Padding(4)
 			.AutoHeight()
 			[
 				SAssignNew(NewPoseEditableText, SEditableTextBox)
+				.MinDesiredWidth(275)
 				.Text(FText::FromName(UniqueNewPoseName))
 			]
 
 			+ SVerticalBox::Slot()
 			.Padding(4)
+			.HAlign(HAlign_Right)
 			.AutoHeight()
 			[
 				SNew(SHorizontalBox)
 				+ SHorizontalBox::Slot()
-				.HAlign(HAlign_Center)
+				.AutoWidth()
+				.Padding(4)
+				.HAlign(HAlign_Right)
 				[
 					SNew(SButton)
 					.ButtonStyle(FAppStyle::Get(), "Button")
 					.TextStyle( FAppStyle::Get(), "DialogButtonText" )
-					.HAlign(HAlign_Center)
-					.VAlign(VAlign_Center)
 					.Text(LOCTEXT("OkButtonLabel", "Ok"))
 					.OnClicked(this, &FIKRetargetEditorController::CreateNewPose)
 				]
 				
 				+ SHorizontalBox::Slot()
-				.HAlign(HAlign_Center)
+				.AutoWidth()
+				.Padding(4)
+				.HAlign(HAlign_Right)
 				[
 					SNew(SButton)
 					.ButtonStyle(FAppStyle::Get(), "Button")
 					.TextStyle( FAppStyle::Get(), "DialogButtonText" )
-					.HAlign(HAlign_Center)
-					.VAlign(VAlign_Center)
 					.Text(LOCTEXT("CancelButtonLabel", "Cancel"))
-					.OnClicked_Lambda( [=]()
+					.OnClicked_Lambda( [Window=NewPoseWindow]()
 					{
-						NewPoseWindow->RequestDestroyWindow();
+						Window->RequestDestroyWindow();
 						return FReply::Handled();
 					})
 				]
-			]	
+			]
 		]
 	];
 
@@ -372,7 +409,7 @@ void FIKRetargetEditorController::HandleNewPose()
 	NewPoseWindow.Reset();
 }
 
-bool FIKRetargetEditorController::CanNewPose() const
+bool FIKRetargetEditorController::CanCreatePose() const
 {
 	return !IsEditingPose();
 }
@@ -386,6 +423,512 @@ FReply FIKRetargetEditorController::CreateNewPose() const
 	return FReply::Handled();
 }
 
+void FIKRetargetEditorController::HandleDuplicatePose()
+{
+	// get a unique pose name to use as suggestion for duplicate
+	const FString DuplicateSuffix = LOCTEXT("DuplicateSuffix", "_Copy").ToString();
+	FString CurrentPoseName = GetCurrentPoseName().ToString();
+	const FString DefaultDuplicatePoseName = CurrentPoseName.Append(*DuplicateSuffix);
+	const FName UniqueNewPoseName = AssetController->MakePoseNameUnique(DefaultDuplicatePoseName);
+	
+	SAssignNew(NewPoseWindow, SWindow)
+	.Title(LOCTEXT("DuplicateRetargetPoseOptions", "Duplicate Retarget Pose"))
+	.ClientSize(FVector2D(300, 80))
+	.HasCloseButton(true)
+	.SupportsMinimize(false) .SupportsMaximize(false)
+	[
+		SNew(SBorder)
+		.BorderImage( FAppStyle::GetBrush("Menu.Background") )
+		[
+			SNew(SVerticalBox)
+			+ SVerticalBox::Slot()
+			.Padding(4)
+			.HAlign(HAlign_Right)
+			.AutoHeight()
+			[
+				SAssignNew(NewPoseEditableText, SEditableTextBox)
+				.MinDesiredWidth(275)
+				.Text(FText::FromName(UniqueNewPoseName))
+			]
+
+			+ SVerticalBox::Slot()
+			.Padding(4)
+			.HAlign(HAlign_Right)
+			.AutoHeight()
+			[
+				SNew(SHorizontalBox)
+				+ SHorizontalBox::Slot()
+				.AutoWidth()
+				.Padding(4)
+				.HAlign(HAlign_Right)
+				[
+					SNew(SButton)
+					.ButtonStyle(FAppStyle::Get(), "Button")
+					.TextStyle( FAppStyle::Get(), "DialogButtonText" )
+					.Text(LOCTEXT("OkButtonLabel", "Ok"))
+					.OnClicked(this, &FIKRetargetEditorController::CreateDuplicatePose)
+				]
+				
+				+ SHorizontalBox::Slot()
+				.AutoWidth()
+				.Padding(4)
+				.HAlign(HAlign_Right)
+				[
+					SNew(SButton)
+					.ButtonStyle(FAppStyle::Get(), "Button")
+					.TextStyle( FAppStyle::Get(), "DialogButtonText" )
+					.Text(LOCTEXT("CancelButtonLabel", "Cancel"))
+					.OnClicked_Lambda( [Window=NewPoseWindow]()
+					{
+						Window->RequestDestroyWindow();
+						return FReply::Handled();
+					})
+				]
+			]	
+		]
+	];
+
+	GEditor->EditorAddModalWindow(NewPoseWindow.ToSharedRef());
+	NewPoseWindow.Reset();
+}
+
+FReply FIKRetargetEditorController::CreateDuplicatePose() const
+{
+	const FIKRetargetPose& PoseToDuplicate = AssetController->GetCurrentRetargetPose();
+	const FName NewPoseName = FName(NewPoseEditableText.Get()->GetText().ToString());
+	AssetController->AddRetargetPose(NewPoseName, &PoseToDuplicate);
+	NewPoseWindow->RequestDestroyWindow();
+	DetailsView->ForceRefresh();
+	return FReply::Handled();
+}
+
+void FIKRetargetEditorController::HandleImportPose()
+{
+	RetargetPoseToImport = nullptr;
+
+	// load the content browser module to display an asset picker
+	const FContentBrowserModule& ContentBrowserModule = FModuleManager::LoadModuleChecked<FContentBrowserModule>("ContentBrowser");
+
+	// the asset picker will only show animation sequences
+	FAssetPickerConfig AssetPickerConfig;
+	AssetPickerConfig.Filter.ClassNames.Add(URetargetPose::StaticClass()->GetFName());
+	AssetPickerConfig.Filter.bRecursiveClasses = true;
+	AssetPickerConfig.OnAssetSelected = FOnAssetSelected::CreateSP(this, &FIKRetargetEditorController::OnRetargetPoseSelected);
+	AssetPickerConfig.InitialAssetViewType = EAssetViewType::Tile;
+
+	ImportPoseWindow = SNew(SWindow)
+	.Title(LOCTEXT("ImportRetargetPose", "Import Retarget Pose"))
+	.ClientSize(FVector2D(500, 600))
+	.SupportsMinimize(false) .SupportsMaximize(false)
+	[
+		SNew(SBorder)
+		.BorderImage( FAppStyle::GetBrush("Menu.Background") )
+		[
+			SNew(SVerticalBox)
+			
+			+ SVerticalBox::Slot()
+			.Padding(4)
+			[
+				ContentBrowserModule.Get().CreateAssetPicker(AssetPickerConfig)
+			]
+			
+			+ SVerticalBox::Slot()
+			.Padding(4)
+			.AutoHeight()
+			.HAlign(HAlign_Right)
+			[
+				SNew(SHorizontalBox)
+				+ SHorizontalBox::Slot()
+				.AutoWidth()
+				.Padding(4)
+				[
+					SNew(SButton)
+					.ButtonStyle(FAppStyle::Get(), "Button")
+					.TextStyle( FAppStyle::Get(), "DialogButtonText" )
+					.HAlign(HAlign_Center)
+					.VAlign(VAlign_Center)
+					.Text(LOCTEXT("ImportButtonLabel", "Import New Retarget Pose"))
+					.OnClicked(this, &FIKRetargetEditorController::ImportRetargetPose)
+				]
+				
+				+ SHorizontalBox::Slot()
+				.AutoWidth()
+				.Padding(4)
+				[
+					SNew(SButton)
+					.ButtonStyle(FAppStyle::Get(), "Button")
+					.TextStyle( FAppStyle::Get(), "DialogButtonText" )
+					.HAlign(HAlign_Center)
+					.VAlign(VAlign_Center)
+					.Text(LOCTEXT("CancelButtonLabel", "Cancel"))
+					.OnClicked_Lambda( [Window=ImportPoseWindow]()
+					{
+						Window->RequestDestroyWindow();
+						return FReply::Handled();
+					})
+				]
+			]	
+		]
+	];
+
+	GEditor->EditorAddModalWindow(ImportPoseWindow.ToSharedRef());
+	ImportPoseWindow.Reset();
+}
+
+FReply FIKRetargetEditorController::ImportRetargetPose() const
+{
+	ImportPoseWindow->RequestDestroyWindow();
+	
+	if (RetargetPoseToImport.IsNull())
+	{
+		return FReply::Handled();
+	}
+
+	const TObjectPtr<URetargetPose> RetargetPose = Cast<URetargetPose>(RetargetPoseToImport.TryLoad());
+	if (!RetargetPose)
+	{
+		return FReply::Handled();
+	}
+
+	// create a new pose with the data from the selected retarget pose asset
+	FIKRetargetPose Pose;
+	RetargetPose->GetAsRetargetPose(Pose);
+	AssetController->AddRetargetPose(FName(RetargetPose->GetName()), &Pose);
+	
+	RefreshAllViews();
+	
+	return FReply::Unhandled();
+}
+
+void FIKRetargetEditorController::OnRetargetPoseSelected(const FAssetData& SelectedAsset)
+{
+	RetargetPoseToImport = SelectedAsset.ToSoftObjectPath();
+}
+
+void FIKRetargetEditorController::HandleImportPoseFromSequence()
+{
+	SequenceToImportAsPose = nullptr;
+
+	// get a unique pose name to use as suggestion
+	const FString DefaultImportedPoseName = LOCTEXT("ImportedRetargetPoseName", "ImportedRetargetPose").ToString();
+	const FName UniqueNewPoseName = AssetController->MakePoseNameUnique(DefaultImportedPoseName);
+	ImportedPoseName = FText::FromName(UniqueNewPoseName);
+
+	// load the content browser module to display an asset picker
+	const FContentBrowserModule& ContentBrowserModule = FModuleManager::LoadModuleChecked<FContentBrowserModule>("ContentBrowser");
+
+	// the asset picker will only show animation sequences compatible with the preview mesh
+	FAssetPickerConfig AssetPickerConfig;
+	AssetPickerConfig.Filter.ClassNames.Add(UAnimSequence::StaticClass()->GetFName());
+	AssetPickerConfig.Filter.ClassNames.Add(UAnimMontage::StaticClass()->GetFName());
+	AssetPickerConfig.Filter.ClassNames.Add(UPoseAsset::StaticClass()->GetFName());
+	AssetPickerConfig.InitialAssetViewType = EAssetViewType::Column;
+	AssetPickerConfig.bAddFilterUI = true;
+	AssetPickerConfig.bShowPathInColumnView = true;
+	AssetPickerConfig.bShowTypeInColumnView = true;
+	AssetPickerConfig.DefaultFilterMenuExpansion = EAssetTypeCategories::Animation;
+	AssetPickerConfig.OnShouldFilterAsset = FOnShouldFilterAsset::CreateSP(this, &FIKRetargetEditorController::OnShouldFilterSequenceToImport);
+	AssetPickerConfig.OnAssetSelected = FOnAssetSelected::CreateSP(this, &FIKRetargetEditorController::OnSequenceSelectedForPose);
+	AssetPickerConfig.bAllowNullSelection = false;
+
+	// hide all asset registry columns by default (we only really want the name and path)
+	TArray<UObject::FAssetRegistryTag> AssetRegistryTags;
+	UAnimSequence::StaticClass()->GetDefaultObject()->GetAssetRegistryTags(AssetRegistryTags);
+	FName ColumnToKeep = FName("Number of Frames");
+	for(UObject::FAssetRegistryTag& AssetRegistryTag : AssetRegistryTags)
+	{
+		if (AssetRegistryTag.Name != ColumnToKeep)
+		{
+			AssetPickerConfig.HiddenColumnNames.Add(AssetRegistryTag.Name.ToString());
+		}
+	}
+
+	// Also hide the type column by default (but allow users to enable it, so don't use bShowTypeInColumnView)
+	AssetPickerConfig.HiddenColumnNames.Add(TEXT("Class"));
+	AssetPickerConfig.HiddenColumnNames.Add(TEXT("HasVirtualizedData"));
+	AssetPickerConfig.HiddenColumnNames.Add(TEXT("DiskSize"));
+
+	// create pop-up window for user to select animation sequence asset to import as a retarget pose]
+	ImportPoseFromSequenceWindow = SNew(SWindow)
+	.Title(LOCTEXT("ImportRetargetPose", "Import Retarget Pose from Sequence Asset"))
+	.ClientSize(FVector2D(500, 600))
+	.SupportsMinimize(false) .SupportsMaximize(false)
+	[
+		SNew(SBorder)
+		.BorderImage( FAppStyle::GetBrush("Menu.Background") )
+		[
+			SNew(SVerticalBox)
+			
+			+ SVerticalBox::Slot()
+			.Padding(4)
+			[
+				ContentBrowserModule.Get().CreateAssetPicker(AssetPickerConfig)
+			]
+
+			+ SVerticalBox::Slot()
+			.Padding(4)
+			.AutoHeight()
+			.HAlign(HAlign_Right)
+			[
+				SNew(SHorizontalBox)
+				+SHorizontalBox::Slot()
+				.AutoWidth()
+				.VAlign(VAlign_Center)
+				.Padding(4.0f, 0.0f)
+				[
+					SNew(STextBlock)
+					.Text(LOCTEXT("ImportFrame_Label", "Sequence Frame: "))
+				]
+				+SHorizontalBox::Slot()
+				.AutoWidth()
+				.HAlign(HAlign_Right)
+				.Padding(2.0f, 0.0f)
+				[
+					SNew(SNumericEntryBox<int32>)
+					.ToolTipText(LOCTEXT("ArrayIndex", "Frame of sequence to import pose from."))
+					.AllowSpin(true)
+					.Font(IDetailLayoutBuilder::GetDetailFont())
+					.MinValue(0)
+					.Value_Lambda([this]
+					{
+						return FrameOfSequenceToImport;
+					})
+					.OnValueChanged(SNumericEntryBox<int32>::FOnValueChanged::CreateLambda([this](int32 Value)
+					{
+						FrameOfSequenceToImport = Value;
+					}))
+				]
+			]
+
+			+ SVerticalBox::Slot()
+			.Padding(4)
+			.AutoHeight()
+			.HAlign(HAlign_Right)
+			[
+				SNew(SHorizontalBox)
+				+SHorizontalBox::Slot()
+				.AutoWidth()
+				.VAlign(VAlign_Center)
+				.Padding(4.0f, 0.0f)
+				[
+					SNew(STextBlock)
+					.Text(LOCTEXT("ImportName_Label", "Pose Name: "))
+				]
+				
+				+SHorizontalBox::Slot()
+				.AutoWidth()
+				[
+					SNew(SEditableTextBox)
+					.Text(FText::FromName(UniqueNewPoseName))
+					.OnTextChanged(FOnTextChanged::CreateLambda([this](const FText InText)
+					{
+						ImportedPoseName = InText;
+					}))
+				]
+			]
+			
+			+ SVerticalBox::Slot()
+			.Padding(4)
+			.AutoHeight()
+			.HAlign(HAlign_Right)
+			[
+				SNew(SHorizontalBox)
+				+ SHorizontalBox::Slot()
+				.AutoWidth()
+				.Padding(4)
+				[
+					SNew(SButton)
+					.ButtonStyle(FAppStyle::Get(), "Button")
+					.TextStyle( FAppStyle::Get(), "DialogButtonText" )
+					.HAlign(HAlign_Center)
+					.VAlign(VAlign_Center)
+					.Text(LOCTEXT("ImportButtonLabel", "Import As Retarget Pose"))
+					.OnClicked(this, &FIKRetargetEditorController::OnImportPoseFromSequence)
+				]
+				
+				+ SHorizontalBox::Slot()
+				.AutoWidth()
+				.Padding(4)
+				[
+					SNew(SButton)
+					.ButtonStyle(FAppStyle::Get(), "Button")
+					.TextStyle( FAppStyle::Get(), "DialogButtonText" )
+					.HAlign(HAlign_Center)
+					.VAlign(VAlign_Center)
+					.Text(LOCTEXT("CancelButtonLabel", "Cancel"))
+					.OnClicked_Lambda( [Window=ImportPoseFromSequenceWindow]()
+					{
+						Window->RequestDestroyWindow();
+						return FReply::Handled();
+					})
+				]
+			]	
+		]
+	];
+
+	GEditor->EditorAddModalWindow(ImportPoseFromSequenceWindow.ToSharedRef());
+	ImportPoseFromSequenceWindow.Reset();
+}
+
+bool FIKRetargetEditorController::OnShouldFilterSequenceToImport(const struct FAssetData& AssetData) const
+{
+	// is this an animation asset?
+	if (!AssetData.IsInstanceOf(UAnimationAsset::StaticClass()))
+	{
+		return true;
+	}
+
+	// get target skeleton
+	constexpr bool bEditingTarget = true; // TODO in future allow editing source or target
+	const USkeleton* DesiredSkeleton = bEditingTarget ? GetTargetSkeleton() : GetSourceSkeleton();
+	if (!DesiredSkeleton)
+	{
+		return true;
+	}
+
+	return !DesiredSkeleton->IsCompatibleSkeletonByAssetData(AssetData);
+}
+
+FReply FIKRetargetEditorController::OnImportPoseFromSequence()
+{
+	ImportPoseFromSequenceWindow->RequestDestroyWindow();
+	
+	if (SequenceToImportAsPose.IsNull())
+	{
+		return FReply::Handled();
+	}
+
+	const TObjectPtr<UAnimSequence> AnimSequence = Cast<UAnimSequence>(SequenceToImportAsPose.TryLoad());
+	if (!AnimSequence)
+	{
+		return FReply::Handled();
+	}
+	
+	USkeletalMesh* Mesh = GetTargetSkeletalMesh();
+	if (!Mesh)
+	{
+		return FReply::Handled();
+	}
+
+	UIKRetargetProcessor* Retargeter = GetRetargetProcessor();
+	if (!Retargeter)
+	{
+		return FReply::Handled();
+	}
+	
+	// ensure we evaluate the source animation using the skeletal mesh proportions that were evaluated in the viewport
+	FAnimPoseEvaluationOptions EvaluationOptions = FAnimPoseEvaluationOptions();
+	EvaluationOptions.OptionalSkeletalMesh = Mesh;
+
+	FrameOfSequenceToImport = FMath::Clamp(FrameOfSequenceToImport, 0, AnimSequence->GetNumberOfSampledKeys());
+	
+	// evaluate the sequence at the desired frame
+	FAnimPose ImportedPose;
+	UAnimPoseExtensions::GetAnimPoseAtFrame(AnimSequence, FrameOfSequenceToImport, EvaluationOptions, ImportedPose);
+
+	// record delta pose for all bones being retargeted
+	FIKRetargetPose ImportedRetargetPose;
+	
+	// get all imported bone transforms and record them in the retarget pose
+	FReferenceSkeleton& RefSkeleton = Mesh->GetRefSkeleton();
+	const TArray<FTransform>& RefPose = RefSkeleton.GetRefBonePose();
+	const FTargetSkeleton& TargetSkeleton = Retargeter->GetTargetSkeleton();
+	int32 NumBones = RefSkeleton.GetNum();
+	const int32 RootBoneIndex = Retargeter->GetTargetSkeletonRootBone();
+	for (int32 BoneIndex = 0; BoneIndex < NumBones; BoneIndex++)
+	{
+		const FName& BoneName = RefSkeleton.GetBoneName(BoneIndex);
+		const int32 RetargetBoneIndex = TargetSkeleton.FindBoneIndexByName(BoneName);
+		const bool bIsRetargetRoot = RetargetBoneIndex == RootBoneIndex;
+		
+		// if this is the retarget root, we want to record the translation delta as well
+		if (bIsRetargetRoot)
+		{
+			const FTransform GlobalTransformImported = UAnimPoseExtensions::GetBonePose(ImportedPose, BoneName, EAnimPoseSpaces::World);
+			const FTransform GlobalTransformReference = UAnimPoseExtensions::GetRefBonePose(ImportedPose, BoneName, EAnimPoseSpaces::World);
+			const FVector TranslationDelta = GlobalTransformImported.GetLocation() - GlobalTransformReference.GetLocation();
+			ImportedRetargetPose.RootTranslationOffset = TranslationDelta;
+
+			// rotation offsets are interpreted as relative to the parent (local), but in the case of the retarget root bone,
+			// when we generate the retarget pose, it's parents will be left at ref pose, so we need to generate a local
+			// rotation offset relative to the ref pose parent, NOT the (potentially) posed parent transform from the animation.
+			FTransform GlobalParentTransformInRefPose = FTransform::Identity;
+			const int32 ParentIndex = RefSkeleton.GetParentIndex(BoneIndex);
+			if (ParentIndex != INDEX_NONE)
+			{
+				const FName& ParentBoneName = RefSkeleton.GetBoneName(ParentIndex);
+				GlobalParentTransformInRefPose = UAnimPoseExtensions::GetRefBonePose(ImportedPose, ParentBoneName, EAnimPoseSpaces::World);
+			}
+
+			// this is a bit crazy, but we have to generate a delta rotation in the local space of the retarget root
+			// bone while treating the root bone as being in global space since the retarget pose does not consider any
+			// bones above it.
+			const FQuat GlobalDeltaRotation = GlobalTransformImported.GetRotation() * GlobalTransformReference.GetRotation().Inverse();
+			const FQuat BoneGlobalOrig = GlobalTransformReference.GetRotation();
+			const FQuat BoneGlobalPlusOffset = GlobalDeltaRotation * BoneGlobalOrig;
+			const FQuat ParentInv = GlobalParentTransformInRefPose.GetRotation().Inverse();
+			const FQuat BoneLocal = ParentInv * BoneGlobalOrig;
+			const FQuat BoneLocalPlusOffset = ParentInv * BoneGlobalPlusOffset;
+			const FQuat BoneLocalOffset = BoneLocal * BoneLocalPlusOffset.Inverse();
+			
+			ImportedRetargetPose.BoneRotationOffsets.Add(BoneName, BoneLocalOffset.Inverse());
+		}
+		else
+		{
+			// record the delta rotation
+			const FTransform LocalTransformImported = UAnimPoseExtensions::GetBonePose(ImportedPose, BoneName, EAnimPoseSpaces::Local);
+			const FTransform LocalTransformReference = RefPose[BoneIndex];
+			const FQuat DeltaRotation = LocalTransformImported.GetRotation() * LocalTransformReference.GetRotation().Inverse();
+			// only if it's difference than the ref pose
+			if (DeltaRotation.GetAngle() > KINDA_SMALL_NUMBER)
+			{
+				ImportedRetargetPose.BoneRotationOffsets.Add(BoneName, DeltaRotation);
+			}
+		}
+	}
+
+	// store the newly imported retarget pose in the asset
+	AssetController->AddRetargetPose(FName(ImportedPoseName.ToString()), &ImportedRetargetPose);
+	
+	return FReply::Unhandled();
+}
+
+void FIKRetargetEditorController::OnSequenceSelectedForPose(const FAssetData& SelectedAsset)
+{
+	SequenceToImportAsPose = SelectedAsset.ToSoftObjectPath();
+}
+
+void FIKRetargetEditorController::HandleExportPose()
+{
+	FSaveAssetDialogConfig SaveAssetDialogConfig;
+	SaveAssetDialogConfig.DefaultPath = AssetController->GetAsset()->GetPackage()->GetPathName();
+	SaveAssetDialogConfig.DefaultAssetName = GetCurrentPoseName().ToString();
+	SaveAssetDialogConfig.AssetClassNames.Add(URetargetPose::StaticClass()->GetFName());
+	SaveAssetDialogConfig.ExistingAssetPolicy = ESaveAssetDialogExistingAssetPolicy::AllowButWarn;
+	SaveAssetDialogConfig.DialogTitleOverride = LOCTEXT("ExportRetargetPoseDialogTitle", "Export Retarget Pose");
+
+	const FContentBrowserModule& ContentBrowserModule = FModuleManager::LoadModuleChecked<FContentBrowserModule>("ContentBrowser");
+	FString SaveObjectPath = ContentBrowserModule.Get().CreateModalSaveAssetDialog(SaveAssetDialogConfig);
+	if (SaveObjectPath.IsEmpty())
+	{
+		return;
+	}
+
+	const FString PackagePath = FPackageName::ObjectPathToPackageName(SaveObjectPath);
+	const FString AssetName = FPaths::GetBaseFilename(PackagePath, true);
+	URetargetPoseFactory* NewFactory = NewObject<URetargetPoseFactory>();
+	const FAssetToolsModule& AssetToolsModule = FAssetToolsModule::GetModule();
+	URetargetPose* NewPoseAsset =  Cast<URetargetPose>(AssetToolsModule.Get().CreateAsset(AssetName, PackagePath, URetargetPose::StaticClass(), NewFactory));
+	TArray<UObject*> ObjectsToSync;
+	ObjectsToSync.Add(NewPoseAsset);
+	GEditor->SyncBrowserToObjects(ObjectsToSync);
+	
+	// fill new pose asset with existing pose data
+	const FIKRetargetPose& CurrentRetargetPose = AssetController->GetCurrentRetargetPose();
+	NewPoseAsset->RootTranslationOffset = CurrentRetargetPose.RootTranslationOffset;
+	NewPoseAsset->BoneRotationOffsets = CurrentRetargetPose.BoneRotationOffsets;
+}
+
 void FIKRetargetEditorController::HandleDeletePose() const
 {
 	const FName CurrentPose = AssetController->GetCurrentRetargetPoseName();
@@ -401,10 +944,48 @@ bool FIKRetargetEditorController::CanDeletePose() const
 	return bNotUsingDefaultPose && !IsEditingPose();
 }
 
-void FIKRetargetEditorController::HandleResetPose() const
+void FIKRetargetEditorController::HandleResetAllBones() const
+{
+	const FName CurrentPose = AssetController->GetCurrentRetargetPoseName();
+	static TArray<FName> Empty; // empty list will reset all bones
+	AssetController->ResetRetargetPose(CurrentPose, Empty);
+}
+
+void FIKRetargetEditorController::HandleResetSelectedBones() const
 {
 	const FName CurrentPose = AssetController->GetCurrentRetargetPoseName();
-	AssetController->ResetRetargetPose(CurrentPose);
+	AssetController->ResetRetargetPose(CurrentPose, GetSelectedBones());
+}
+
+void FIKRetargetEditorController::HandleResetSelectedAndChildrenBones() const
+{
+	const FName CurrentPose = AssetController->GetCurrentRetargetPoseName();
+	const TArray<FName>& SelectedBones = GetSelectedBones();
+	
+	// get list of all children of selected bones
+	const UIKRetargetProcessor* Processor = GetRetargetProcessor();
+	const FTargetSkeleton& Skeleton = Processor->GetTargetSkeleton();
+	TArray<int32> AllChildrenIndices;
+	for (const FName& SelectedBone : SelectedBones)
+	{
+		const int32 SelectedBoneIndex = Skeleton.FindBoneIndexByName(SelectedBone);
+		Skeleton.GetChildrenIndicesRecursive(SelectedBoneIndex, AllChildrenIndices);
+	}
+
+	// merge total list of all selected bones and their children
+	TArray<FName> BonesToReset = SelectedBones;
+	for (const int32 ChildIndex : AllChildrenIndices)
+	{
+		BonesToReset.AddUnique(Skeleton.BoneNames[ChildIndex]);
+	}
+	
+	// reset the bones
+	AssetController->ResetRetargetPose(CurrentPose, BonesToReset);
+}
+
+bool FIKRetargetEditorController::CanResetSelected() const
+{
+	return CanResetPose() && !GetSelectedBones().IsEmpty();
 }
 
 bool FIKRetargetEditorController::CanResetPose() const
@@ -464,9 +1045,9 @@ void FIKRetargetEditorController::HandleRenamePose()
 					.HAlign(HAlign_Center)
 					.VAlign(VAlign_Center)
 					.Text(LOCTEXT("CancelButtonLabel", "Cancel"))
-					.OnClicked_Lambda( [=]()
+					.OnClicked_Lambda( [Window=RenamePoseWindow]()
 					{
-						RenamePoseWindow->RequestDestroyWindow();
+						Window->RequestDestroyWindow();
 						return FReply::Handled();
 					})
 				]
diff --git a/Engine/Plugins/Animation/IKRig/Source/IKRigEditor/Private/RetargetEditor/IKRetargetFactory.cpp b/Engine/Plugins/Animation/IKRig/Source/IKRigEditor/Private/RetargetEditor/IKRetargetFactory.cpp
--- a/Engine/Plugins/Animation/IKRig/Source/IKRigEditor/Private/RetargetEditor/IKRetargetFactory.cpp
+++ b/Engine/Plugins/Animation/IKRig/Source/IKRigEditor/Private/RetargetEditor/IKRetargetFactory.cpp
@@ -117,4 +117,46 @@ FString UIKRetargetFactory::GetDefaultNewAssetName() const
 {
 	return FString(TEXT("NewIKRetargeter"));
 }
+
+/** Start URetargetPose Factory ...*/
+
+URetargetPoseFactory::URetargetPoseFactory()
+{
+	bCreateNew = true;
+	bEditAfterNew = true;
+	SupportedClass = URetargetPose::StaticClass();
+}
+
+FText URetargetPoseFactory::GetDisplayName() const
+{
+	return LOCTEXT("RetargetPose_DisplayName", "Retarget Pose");
+}
+
+uint32 URetargetPoseFactory::GetMenuCategories() const
+{
+	return EAssetTypeCategories::Animation;
+}
+
+FText URetargetPoseFactory::GetToolTip() const
+{
+	return LOCTEXT("RetargetPose_Tooltip", "Defines a set of bone transform offsets to create a pose for retargeting.");
+}
+
+FString URetargetPoseFactory::GetDefaultNewAssetName() const
+{
+	return FString(TEXT("NewRetargetPose"));
+}
+
+UObject* URetargetPoseFactory::FactoryCreateNew(
+	UClass* Class, UObject* InParent, FName Name, EObjectFlags Flags,
+	UObject* Context, FFeedbackContext* Warn)
+{
+	return NewObject<URetargetPose>(InParent, Class, Name, Flags);
+}
+
+bool URetargetPoseFactory::ShouldShowInNewMenu() const
+{
+	return true;
+}
+
 #undef LOCTEXT_NAMESPACE
diff --git a/Engine/Plugins/Animation/IKRig/Source/IKRigEditor/Private/RetargetEditor/IKRetargeterController.cpp b/Engine/Plugins/Animation/IKRig/Source/IKRigEditor/Private/RetargetEditor/IKRetargeterController.cpp
--- a/Engine/Plugins/Animation/IKRig/Source/IKRigEditor/Private/RetargetEditor/IKRetargeterController.cpp
+++ b/Engine/Plugins/Animation/IKRig/Source/IKRigEditor/Private/RetargetEditor/IKRetargeterController.cpp
@@ -349,13 +349,18 @@ const TArray<TObjectPtr<URetargetChainSettings>>& UIKRetargeterController::GetCh
 	return Asset->ChainSettings;
 }
 
-void UIKRetargeterController::AddRetargetPose(FName NewPoseName) const
+void UIKRetargeterController::AddRetargetPose(FName NewPoseName, const FIKRetargetPose* ToDuplicate) const
 {
 	FScopedTransaction Transaction(LOCTEXT("AddRetargetPose", "Add Retarget Pose"));
 	Asset->Modify();
 	
-	NewPoseName = MakePoseNameUnique(NewPoseName);
-	Asset->RetargetPoses.Add(NewPoseName);
+	NewPoseName = MakePoseNameUnique(NewPoseName.ToString());
+	FIKRetargetPose& NewPose = Asset->RetargetPoses.Add(NewPoseName);
+	if (ToDuplicate)
+	{
+		NewPose.RootTranslationOffset = ToDuplicate->RootTranslationOffset;
+		NewPose.BoneRotationOffsets = ToDuplicate->BoneRotationOffsets;
+	}
 	Asset->CurrentRetargetPose = NewPoseName;
 
 	BroadcastNeedsReinitialized();
@@ -411,18 +416,42 @@ void UIKRetargeterController::RemoveRetargetPose(FName PoseToRemove) const
 	BroadcastNeedsReinitialized();
 }
 
-void UIKRetargeterController::ResetRetargetPose(FName PoseToReset) const
+void UIKRetargeterController::ResetRetargetPose(FName PoseToReset, const TArray<FName>& BonesToReset) const
 {
 	if (!Asset->RetargetPoses.Contains(PoseToReset))
 	{
 		return; // cannot reset pose that doesn't exist
 	}
-
-	FScopedTransaction Transaction(LOCTEXT("ResetRetargetPose", "Reset Retarget Pose"));
-	Asset->Modify();
 	
-	Asset->RetargetPoses[PoseToReset].BoneRotationOffsets.Reset();
-	Asset->RetargetPoses[PoseToReset].RootTranslationOffset = FVector::ZeroVector;
+	FIKRetargetPose& PoseToEdit = Asset->RetargetPoses[PoseToReset];
+	
+	if (BonesToReset.IsEmpty())
+	{
+		FScopedTransaction Transaction(LOCTEXT("ResetRetargetPose", "Reset Retarget Pose"));
+		Asset->Modify();
+		
+		PoseToEdit.BoneRotationOffsets.Reset();
+		PoseToEdit.RootTranslationOffset = FVector::ZeroVector;
+	}
+	else
+	{
+		FScopedTransaction Transaction(LOCTEXT("ResetRetargetBonePose", "Reset Bone Pose"));
+		Asset->Modify();
+		
+		const FName RootBoneName = GetTargetRootBone();
+		for (const FName& BoneToReset : BonesToReset)
+		{
+			if (PoseToEdit.BoneRotationOffsets.Contains(BoneToReset))
+			{
+				PoseToEdit.BoneRotationOffsets.Remove(BoneToReset);
+			}
+
+			if (BoneToReset == RootBoneName)
+			{
+				PoseToEdit.RootTranslationOffset = FVector::ZeroVector;	
+			}
+		}
+	}
 
 	BroadcastNeedsReinitialized();
 }
@@ -448,6 +477,11 @@ const TMap<FName, FIKRetargetPose>& UIKRetargeterController::GetRetargetPoses()
 	return GetAsset()->RetargetPoses;
 }
 
+const FIKRetargetPose& UIKRetargeterController::GetCurrentRetargetPose() const
+{
+	return GetAsset()->RetargetPoses[GetCurrentRetargetPoseName()];
+}
+
 void UIKRetargeterController::SetRotationOffsetForRetargetPoseBone(FName BoneName, FQuat RotationOffset) const
 {
 	const FIKRigSkeleton& Skeleton = Asset->GetTargetIKRig()->Skeleton;
@@ -470,16 +504,16 @@ void UIKRetargeterController::AddTranslationOffsetToRetargetRootBone(FVector Tra
 	Asset->RetargetPoses[Asset->CurrentRetargetPose].AddTranslationDeltaToRoot(TranslationOffset);
 }
 
-FName UIKRetargeterController::MakePoseNameUnique(FName PoseName) const
+FName UIKRetargeterController::MakePoseNameUnique(FString PoseName) const
 {
-	FName UniqueName = PoseName;
+	FString UniqueName = PoseName;
 	int32 Suffix = 1;
-	while (Asset->RetargetPoses.Contains(UniqueName))
+	while (Asset->RetargetPoses.Contains(FName(UniqueName)))
 	{
-		UniqueName = FName(PoseName.ToString() + "_" + FString::FromInt(Suffix));
+		UniqueName = PoseName + "_" + FString::FromInt(Suffix);
 		++Suffix;
 	}
-	return UniqueName;
+	return FName(UniqueName);
 }
 
 URetargetChainSettings* UIKRetargeterController::GetChainMap(const FName& TargetChainName) const
diff --git a/Engine/Plugins/Animation/IKRig/Source/IKRigEditor/Private/RetargetEditor/SIKRetargetAssetBrowser.cpp b/Engine/Plugins/Animation/IKRig/Source/IKRigEditor/Private/RetargetEditor/SIKRetargetAssetBrowser.cpp
--- a/Engine/Plugins/Animation/IKRig/Source/IKRigEditor/Private/RetargetEditor/SIKRetargetAssetBrowser.cpp
+++ b/Engine/Plugins/Animation/IKRig/Source/IKRigEditor/Private/RetargetEditor/SIKRetargetAssetBrowser.cpp
@@ -134,7 +134,6 @@ FReply SIKRetargetAssetBrowser::OnExportButtonClicked() const
 	BatchOperation.RunRetarget(BatchContext);
 
 	return FReply::Handled();
-	
 }
 
 bool SIKRetargetAssetBrowser::IsExportButtonEnabled() const
@@ -187,21 +186,14 @@ bool SIKRetargetAssetBrowser::OnShouldFilterAsset(const struct FAssetData& Asset
 	}
 	
 	// controller setup
-	FIKRetargetEditorController* Controller = EditorController.Pin().Get();
+	const FIKRetargetEditorController* Controller = EditorController.Pin().Get();
 	if (!Controller)
 	{
 		return true;
 	}
 
-	// get source mesh
-	USkeletalMesh* SourceMesh = Controller->GetSourceSkeletalMesh();
-	if (!SourceMesh)
-	{
-		return true;
-	}
-
 	// get source skeleton
-	USkeleton* DesiredSkeleton = SourceMesh->GetSkeleton();
+	const USkeleton* DesiredSkeleton = Controller->GetSourceSkeleton();
 	if (!DesiredSkeleton)
 	{
 		return true;
diff --git a/Engine/Plugins/Animation/IKRig/Source/IKRigEditor/Public/IKRigEditor.h b/Engine/Plugins/Animation/IKRig/Source/IKRigEditor/Public/IKRigEditor.h
--- a/Engine/Plugins/Animation/IKRig/Source/IKRigEditor/Public/IKRigEditor.h
+++ b/Engine/Plugins/Animation/IKRig/Source/IKRigEditor/Public/IKRigEditor.h
@@ -15,6 +15,7 @@ private:
 	TSharedPtr<class FAssetTypeActions_AnimationAssetRetarget> RetargetAnimationAssetAction;
 	TSharedPtr<class FAssetTypeActions_IKRigDefinition> IKRigDefinitionAssetAction;
 	TSharedPtr<class FAssetTypeActions_IKRetargeter> IKRetargeterAssetAction;
+	TSharedPtr<class FAssetTypeActions_RetargetPose> RetargetPoseAssetAction;
 
 	TArray<FName> ClassesToUnregisterOnShutdown;
 };
diff --git a/Engine/Plugins/Animation/IKRig/Source/IKRigEditor/Public/RetargetEditor/AssetTypeActions_RetargetPose.h b/Engine/Plugins/Animation/IKRig/Source/IKRigEditor/Public/RetargetEditor/AssetTypeActions_RetargetPose.h
new file mode 100644
--- /dev/null
+++ b/Engine/Plugins/Animation/IKRig/Source/IKRigEditor/Public/RetargetEditor/AssetTypeActions_RetargetPose.h
@@ -0,0 +1,23 @@
+﻿// Copyright Epic Games, Inc. All Rights Reserved.
+
+#pragma once
+
+#include "CoreMinimal.h"
+#include "Toolkits/IToolkitHost.h"
+#include "AssetTypeActions_Base.h"
+#include "EditorAnimUtils.h"
+
+class URetargetPose;
+
+class FAssetTypeActions_RetargetPose : public FAssetTypeActions_Base
+{
+public:
+	// IAssetTypeActions Implementation
+	virtual FText GetName() const override { return NSLOCTEXT("AssetTypeActions", "AssetTypeActions_RetargetPose", "Retarget Pose"); }
+	virtual FColor GetTypeColor() const override { return FColor(0, 40, 255); }
+	virtual UClass* GetSupportedClass() const override;
+	virtual uint32 GetCategories() override { return EAssetTypeCategories::Animation; }
+	const TArray<FText>& GetSubMenus() const override;
+	virtual UThumbnailInfo* GetThumbnailInfo(UObject* Asset) const override;
+	virtual void GetActions(const TArray<UObject*>& InObjects, FToolMenuSection& Section) override;
+};
diff --git a/Engine/Plugins/Animation/IKRig/Source/IKRigEditor/Public/RetargetEditor/IKRetargetCommands.h b/Engine/Plugins/Animation/IKRig/Source/IKRigEditor/Public/RetargetEditor/IKRetargetCommands.h
--- a/Engine/Plugins/Animation/IKRig/Source/IKRigEditor/Public/RetargetEditor/IKRetargetCommands.h
+++ b/Engine/Plugins/Animation/IKRig/Source/IKRigEditor/Public/RetargetEditor/IKRetargetCommands.h
@@ -18,24 +18,40 @@ public:
 	)
 	{}
 
-	/** reset reference pose */
+	/** reset retarget pose */
 	TSharedPtr< FUICommandInfo > GoToRetargetPose;
 	
-	/** edit reference pose */
+	/** edit retarget pose */
 	TSharedPtr< FUICommandInfo > EditRetargetPose;
 
-	/** reset reference pose */
-	TSharedPtr< FUICommandInfo > SetToRefPose;
+	/** reset retarget pose to ref pose */
+	TSharedPtr< FUICommandInfo > ResetAllBones;
+	/** reset retarget pose to ref pose */
+	TSharedPtr< FUICommandInfo > ResetSelectedBones;
+	/** reset retarget pose to ref pose */
+	TSharedPtr< FUICommandInfo > ResetSelectedAndChildrenBones;
 
-	/** new reference pose */
-	TSharedPtr< FUICommandInfo > NewRetargetPose;
-
-	/** delete reference pose */
+	/** delete retarget pose */
 	TSharedPtr< FUICommandInfo > DeleteRetargetPose;
 
-	/** rename reference pose */
+	/** rename retarget pose */
 	TSharedPtr< FUICommandInfo > RenameRetargetPose;
-	
+
+	/** new retarget pose */
+	TSharedPtr< FUICommandInfo > NewRetargetPose;
+
+	/** duplicate retarget pose */
+	TSharedPtr< FUICommandInfo > DuplicateRetargetPose;
+
+	/** import retarget pose */
+	TSharedPtr< FUICommandInfo > ImportRetargetPose;
+
+	/** import retarget pose from anim */
+	TSharedPtr< FUICommandInfo > ImportRetargetPoseFromAnim;
+
+	/** export retarget pose */
+	TSharedPtr< FUICommandInfo > ExportRetargetPose;
+
 	/** export animation */
 	TSharedPtr< FUICommandInfo > ExportAnimation;
 
diff --git a/Engine/Plugins/Animation/IKRig/Source/IKRigEditor/Public/RetargetEditor/IKRetargetDetails.h b/Engine/Plugins/Animation/IKRig/Source/IKRigEditor/Public/RetargetEditor/IKRetargetDetails.h
--- a/Engine/Plugins/Animation/IKRig/Source/IKRigEditor/Public/RetargetEditor/IKRetargetDetails.h
+++ b/Engine/Plugins/Animation/IKRig/Source/IKRigEditor/Public/RetargetEditor/IKRetargetDetails.h
@@ -21,9 +21,13 @@ public:
 
 private:
 
-	TSharedRef<SWidget> MakeToolbar(const TSharedRef<FUICommandList>& UICommandList);
+	TSharedRef<SWidget> MakeToolbar(TSharedPtr<FUICommandList> Commands);
+
+	TSharedRef<SWidget> GenerateResetMenuContent(TSharedPtr<FUICommandList> Commands);
+	
+	TSharedRef<SWidget> GenerateNewMenuContent(TSharedPtr<FUICommandList> Commands);
 	
 	TObjectPtr<UIKRetargeterController> GetAssetControllerFromSelectedObjects(IDetailLayoutBuilder& DetailBuilder) const;
-
+	
 	TArray<TSharedPtr<FName>> PoseNames;
 };
diff --git a/Engine/Plugins/Animation/IKRig/Source/IKRigEditor/Public/RetargetEditor/IKRetargetEditPoseMode.h b/Engine/Plugins/Animation/IKRig/Source/IKRigEditor/Public/RetargetEditor/IKRetargetEditPoseMode.h
--- a/Engine/Plugins/Animation/IKRig/Source/IKRigEditor/Public/RetargetEditor/IKRetargetEditPoseMode.h
+++ b/Engine/Plugins/Animation/IKRig/Source/IKRigEditor/Public/RetargetEditor/IKRetargetEditPoseMode.h
@@ -68,6 +68,9 @@ public:
 	virtual void Exit() override;
 	/** END FEdMode interface */
 
+	// get a list of the names of the currently selected bones (empty if non selected)
+	const TArray<FName>& GetSelectedBones() const;
+	
 private:
 
 	void GetAffectedBones(
diff --git a/Engine/Plugins/Animation/IKRig/Source/IKRigEditor/Public/RetargetEditor/IKRetargetEditorController.h b/Engine/Plugins/Animation/IKRig/Source/IKRigEditor/Public/RetargetEditor/IKRetargetEditorController.h
--- a/Engine/Plugins/Animation/IKRig/Source/IKRigEditor/Public/RetargetEditor/IKRetargetEditorController.h
+++ b/Engine/Plugins/Animation/IKRig/Source/IKRigEditor/Public/RetargetEditor/IKRetargetEditorController.h
@@ -74,7 +74,14 @@ public:
 	USkeletalMesh* GetSourceSkeletalMesh() const;
 	/** get the target skeletal mesh we are copying TO */
 	USkeletalMesh* GetTargetSkeletalMesh() const;
+	/** get the source skeleton asset */
+	const USkeleton* GetSourceSkeleton() const;
+	/** get the target skeleton asset */
+	const USkeleton* GetTargetSkeleton() const;
 
+	/** get list of currently selected bones */
+	const TArray<FName>& GetSelectedBones() const;
+	
 	/** get current chain pose */
 	FTransform GetTargetBoneGlobalTransform(const UIKRetargetProcessor* RetargetProcessor, const int32& TargetBoneIndex) const;
 	FTransform GetTargetBoneLocalTransform(const UIKRetargetProcessor* RetargetProcessor, const int32& TargetBoneIndex) const;
@@ -87,7 +94,6 @@ public:
 
 	/** get the retargeter that is running in the viewport (which is a duplicate of the source asset) */
 	UIKRetargetProcessor* GetRetargetProcessor() const;
-
 	/** Reset the planting state of the IK (when scrubbing or animation loops over) */
 	void ResetIKPlantingState() const;
 
@@ -97,34 +103,69 @@ public:
 	UAnimationAsset* PreviousAsset = nullptr;
 	/** END Sequence Browser */
 
-	// go to retarget pose
+	/* START RETARGET POSES */
+	
+	/** go to retarget pose */
 	void HandleGoToRetargetPose() const;
 	
+	/** toggle current retarget pose */
+	TArray<TSharedPtr<FName>> PoseNames;
+	FText GetCurrentPoseName() const;
+	void OnPoseSelected(TSharedPtr<FName> InPoseName, ESelectInfo::Type SelectInfo) const;
+	
 	/** edit retarget poses */
 	void HandleEditPose() const;
 	bool CanEditPose() const;
 	bool IsEditingPose() const;
 
-	void HandleResetPose() const;
+	/** reset retarget pose */
+	void HandleResetAllBones() const;
+	void HandleResetSelectedBones() const;
+	void HandleResetSelectedAndChildrenBones() const;
 	bool CanResetPose() const;
+	bool CanResetSelected() const;
 
+	/** create new retarget pose */
 	void HandleNewPose();
-	bool CanNewPose() const;
+	bool CanCreatePose() const;
 	FReply CreateNewPose() const;
 	TSharedPtr<SWindow> NewPoseWindow;
 	TSharedPtr<SEditableTextBox> NewPoseEditableText;
-	
+
+	/** duplicate current retarget pose */
+	void HandleDuplicatePose();
+	FReply CreateDuplicatePose() const;
+
+	/** import retarget pose from asset*/
+	void HandleImportPose();
+	FReply ImportRetargetPose() const;
+	void OnRetargetPoseSelected(const FAssetData& SelectedAsset);
+	FSoftObjectPath RetargetPoseToImport;
+	TSharedPtr<SWindow> ImportPoseWindow;
+
+	/** import retarget pose from animation sequence*/
+	void HandleImportPoseFromSequence();
+	bool OnShouldFilterSequenceToImport(const struct FAssetData& AssetData) const;
+	FReply OnImportPoseFromSequence();
+	void OnSequenceSelectedForPose(const FAssetData& SelectedAsset);
+	TSharedPtr<SWindow> ImportPoseFromSequenceWindow;
+	FSoftObjectPath SequenceToImportAsPose;
+	int32 FrameOfSequenceToImport;
+	FText ImportedPoseName;
+
+	/** export retarget pose to asset*/
+	void HandleExportPose();
+
+	/** delete retarget pose */
 	void HandleDeletePose() const;
 	bool CanDeletePose() const;
-	
+
+	/** rename retarget pose */
 	void HandleRenamePose();
 	FReply RenamePose() const;
 	bool CanRenamePose() const;
 	TSharedPtr<SWindow> RenamePoseWindow;
 	TSharedPtr<SEditableTextBox> NewNameEditableText;
 	
-	TArray<TSharedPtr<FName>> PoseNames;
-	FText GetCurrentPoseName() const;
-	void OnPoseSelected(TSharedPtr<FName> InPoseName, ESelectInfo::Type SelectInfo) const;
-	/* END edit retarget poses */
+	/* END RETARGET POSES */
 };
diff --git a/Engine/Plugins/Animation/IKRig/Source/IKRigEditor/Public/RetargetEditor/IKRetargetFactory.h b/Engine/Plugins/Animation/IKRig/Source/IKRigEditor/Public/RetargetEditor/IKRetargetFactory.h
--- a/Engine/Plugins/Animation/IKRig/Source/IKRigEditor/Public/RetargetEditor/IKRetargetFactory.h
+++ b/Engine/Plugins/Animation/IKRig/Source/IKRigEditor/Public/RetargetEditor/IKRetargetFactory.h
@@ -36,3 +36,22 @@ private:
 
 	TSharedPtr<SWindow> PickerWindow;
 };
+
+
+UCLASS(hidecategories=Object, MinimalAPI)
+class URetargetPoseFactory : public UFactory
+{
+	GENERATED_BODY()
+
+public:
+
+	URetargetPoseFactory();
+
+	// UFactory Interface
+	virtual FText GetDisplayName() const override;
+	virtual uint32 GetMenuCategories() const override;
+	virtual FText GetToolTip() const override;
+	virtual FString GetDefaultNewAssetName() const override;
+	virtual UObject* FactoryCreateNew(UClass* Class, UObject* InParent, FName Name, EObjectFlags Flags, UObject* Context, FFeedbackContext* Warn) override;
+	virtual bool ShouldShowInNewMenu() const override;
+};
diff --git a/Engine/Plugins/Animation/IKRig/Source/IKRigEditor/Public/RetargetEditor/IKRetargeterController.h b/Engine/Plugins/Animation/IKRig/Source/IKRigEditor/Public/RetargetEditor/IKRetargeterController.h
--- a/Engine/Plugins/Animation/IKRig/Source/IKRigEditor/Public/RetargetEditor/IKRetargeterController.h
+++ b/Engine/Plugins/Animation/IKRig/Source/IKRigEditor/Public/RetargetEditor/IKRetargeterController.h
@@ -77,19 +77,22 @@ public:
 	/** Remove bones from retarget poses that are no longer in skeleton */
 	void CleanPoseList(const bool bForceReinitialization=true);
 	/** Add new retarget pose. */
-	void AddRetargetPose(FName NewPoseName) const;
+	void AddRetargetPose(FName NewPoseName, const FIKRetargetPose* ToDuplicate = nullptr) const;
 	/** Rename current retarget pose. */
 	void RenameCurrentRetargetPose(FName NewPoseName) const;
 	/** Remove a retarget pose. */
 	void RemoveRetargetPose(FName PoseToRemove) const;
-	/** Reset a retarget pose (removes all stored deltas, returning pose to reference pose */
-	void ResetRetargetPose(FName PoseToReset) const;
+	/** Reset a retarget pose for the specified bones.
+	 *If BonesToReset is Empty, will removes all stored deltas, returning pose to reference pose */
+	void ResetRetargetPose(FName PoseToReset, const TArray<FName>& BonesToReset) const;
 	/** Get the current retarget pose */
     FName GetCurrentRetargetPoseName() const;
 	/** Change which retarget pose is used by the retargeter at runtime */
 	void SetCurrentRetargetPose(FName CurrentPose) const;
 	/** Get read-only access to list of retarget poses */
 	const TMap<FName, FIKRetargetPose>& GetRetargetPoses();
+	/** Get the current retarget pose */
+	const FIKRetargetPose& GetCurrentRetargetPose() const;
 	/** Set a delta rotation for a given bone for the current retarget pose (used in Edit Mode in the retarget editor) */
 	void SetRotationOffsetForRetargetPoseBone(FName BoneName, FQuat RotationOffset) const;
 	/** Get a delta rotation for a given bone for the current retarget pose (used in Edit Mode in the retarget editor) */
@@ -97,7 +100,7 @@ public:
 	/** Add a delta translation to the root bone (used in Edit Mode in the retarget editor) */
 	void AddTranslationOffsetToRetargetRootBone(FVector TranslationOffset) const;
 	/** Add a numbered suffix to the given pose name to make it unique. */
-	FName MakePoseNameUnique(FName PoseName) const;
+	FName MakePoseNameUnique(FString PoseName) const;
 	/** END RETARGET POSE EDITING */
 
 private:
