IK Retargeter, fixed bug bones remain visible but not editable when playing animation while in edit mode.

Also fixes bug, animation resumes playing when going from "Show Retarget Pose" mode  back to "Edit Pose" mode.

diff --git a/Engine/Plugins/Animation/IKRig/Source/IKRigEditor/Private/RetargetEditor/IKRetargetDefaultMode.cpp b/Engine/Plugins/Animation/IKRig/Source/IKRigEditor/Private/RetargetEditor/IKRetargetDefaultMode.cpp
--- a/Engine/Plugins/Animation/IKRig/Source/IKRigEditor/Private/RetargetEditor/IKRetargetDefaultMode.cpp
+++ b/Engine/Plugins/Animation/IKRig/Source/IKRigEditor/Private/RetargetEditor/IKRetargetDefaultMode.cpp
@@ -231,7 +231,6 @@
 	Controller->Editor.Pin()->GetPersonaToolkit()->GetPreviewScene()->SetAllowMeshHitProxies(true);
 	Controller->SourceSkelMeshComponent->bSelectable = true;
 	Controller->TargetSkelMeshComponent->bSelectable = true;
-	Controller->AssetController->GetAsset()->SetOutputMode(ERetargeterOutputMode::RunRetarget);
 
 	// deselect when entering mode
 	SetSelectedComponent(nullptr);
diff --git a/Engine/Plugins/Animation/IKRig/Source/IKRigEditor/Private/RetargetEditor/IKRetargetEditPoseMode.cpp b/Engine/Plugins/Animation/IKRig/Source/IKRigEditor/Private/RetargetEditor/IKRetargetEditPoseMode.cpp
--- a/Engine/Plugins/Animation/IKRig/Source/IKRigEditor/Private/RetargetEditor/IKRetargetEditPoseMode.cpp
+++ b/Engine/Plugins/Animation/IKRig/Source/IKRigEditor/Private/RetargetEditor/IKRetargetEditPoseMode.cpp
@@ -427,8 +427,6 @@
 		return;
 	}
 
-	Controller->AssetController->GetAsset()->SetOutputMode(ERetargeterOutputMode::RunRetarget);
-
 	IPersonaEditMode::Exit();
 }
 
diff --git a/Engine/Plugins/Animation/IKRig/Source/IKRigEditor/Private/RetargetEditor/IKRetargetEditorController.cpp b/Engine/Plugins/Animation/IKRig/Source/IKRigEditor/Private/RetargetEditor/IKRetargetEditorController.cpp
--- a/Engine/Plugins/Animation/IKRig/Source/IKRigEditor/Private/RetargetEditor/IKRetargetEditorController.cpp
+++ b/Engine/Plugins/Animation/IKRig/Source/IKRigEditor/Private/RetargetEditor/IKRetargetEditorController.cpp
@@ -278,55 +278,77 @@
 	if (AssetToPlay && SourceAnimInstance.IsValid())
 	{
 		SourceAnimInstance->SetAnimationAsset(AssetToPlay);
-		PreviousAsset = AssetToPlay;
-		// tell asset to output the retargeted pose
-		AssetController->GetAsset()->SetOutputMode(ERetargeterOutputMode::RunRetarget);
+		SourceAnimInstance->SetPlaying(true);
+		AnimThatWasPlaying = AssetToPlay;
+		// ensure we are running the retargeter so you can see the animation
+		SetRetargeterMode(ERetargeterOutputMode::RunRetarget);
 	}
 }
 
-void FIKRetargetEditorController::PlayPreviousAnimationAsset() const
+void FIKRetargetEditorController::SetRetargeterMode(ERetargeterOutputMode Mode)
 {
-	if (PreviousAsset)
+	const bool bWasEditingPose = IsEditingPose();
+	FEditorModeTools& EditorModeManager = Editor.Pin()->GetEditorModeManager();
+	UIKRetargeter* Asset = AssetController->GetAsset();
+	
+	switch (Mode)
 	{
-		SourceAnimInstance->SetAnimationAsset(PreviousAsset);
-		// tell asset to output the retarget
-		AssetController->GetAsset()->SetOutputMode(ERetargeterOutputMode::RunRetarget);
+		case ERetargeterOutputMode::EditRetargetPose:
+			bWasPlayingAnim = SourceAnimInstance->IsPlaying();
+			EditorModeManager.DeactivateMode(FIKRetargetDefaultMode::ModeName);
+			EditorModeManager.ActivateMode(FIKRetargetEditPoseMode::ModeName);
+			Asset->SetOutputMode(ERetargeterOutputMode::EditRetargetPose);
+			break;
+		
+		case ERetargeterOutputMode::RunRetarget:
+			EditorModeManager.DeactivateMode(FIKRetargetEditPoseMode::ModeName);
+			EditorModeManager.ActivateMode(FIKRetargetDefaultMode::ModeName);
+			Asset->SetOutputMode(ERetargeterOutputMode::RunRetarget);
+			if (bWasEditingPose)
+			{
+				// must reinitialize after editing the retarget pose
+				AssetController->BroadcastNeedsReinitialized();
+				if (bWasPlayingAnim)
+				{
+					// continue playing whatever animation asset was last used
+					SourceAnimInstance->SetAnimationAsset(AnimThatWasPlaying);
+				}
+			}
+			break;
+
+		case ERetargeterOutputMode::ShowRetargetPose:
+			EditorModeManager.DeactivateMode(FIKRetargetEditPoseMode::ModeName);
+			EditorModeManager.ActivateMode(FIKRetargetDefaultMode::ModeName);
+			Asset->SetOutputMode(ERetargeterOutputMode::ShowRetargetPose);
+			// stop playback of animations
+			SourceAnimInstance->SetPlaying(false);
+			// put source back in ref pose
+			SourceSkelMeshComponent->ShowReferencePose(true);
+			// have to move component back to offset position because ShowReferencePose() sets it back to origin
+			AddOffsetAndUpdatePreviewMeshPosition(FVector::ZeroVector, SourceSkelMeshComponent);
+			break;
+
+		default:
+			checkNoEntry();
 	}
 }
 
-void FIKRetargetEditorController::HandleGoToRetargetPose() const
+void FIKRetargetEditorController::HandleGoToRetargetPose()
 {
-	Editor.Pin()->GetEditorModeManager().DeactivateMode(FIKRetargetEditPoseMode::ModeName);
-	Editor.Pin()->GetEditorModeManager().ActivateMode(FIKRetargetDefaultMode::ModeName);
-
-	// put source back in ref pose
-	SourceSkelMeshComponent->ShowReferencePose(true);
-	// have to move component back to offset position because ShowReferencePose() sets it back to origin
-	AddOffsetAndUpdatePreviewMeshPosition(FVector::ZeroVector, SourceSkelMeshComponent);
-	// tell asset to output the retarget pose
-	AssetController->GetAsset()->SetOutputMode(ERetargeterOutputMode::ShowRetargetPose);
+	SetRetargeterMode(ERetargeterOutputMode::ShowRetargetPose);
 }
 
-void FIKRetargetEditorController::HandleEditPose() const
+void FIKRetargetEditorController::HandleEditPose()
 {
 	if (IsEditingPose())
 	{
 		// stop pose editing
-		Editor.Pin()->GetEditorModeManager().DeactivateMode(FIKRetargetEditPoseMode::ModeName);
-		Editor.Pin()->GetEditorModeManager().ActivateMode(FIKRetargetDefaultMode::ModeName);
-		AssetController->GetAsset()->SetOutputMode(ERetargeterOutputMode::RunRetarget);
-		
-		// must reinitialize after editing the retarget pose
-		AssetController->BroadcastNeedsReinitialized();
-		// continue playing whatever animation asset was last used
-		PlayPreviousAnimationAsset();
+		SetRetargeterMode(ERetargeterOutputMode::RunRetarget);
 	}
 	else
 	{
 		// start pose editing
-		Editor.Pin()->GetEditorModeManager().DeactivateMode(FIKRetargetDefaultMode::ModeName);
-		Editor.Pin()->GetEditorModeManager().ActivateMode(FIKRetargetEditPoseMode::ModeName);
-		AssetController->GetAsset()->SetOutputMode(ERetargeterOutputMode::EditRetargetPose);
+		SetRetargeterMode(ERetargeterOutputMode::EditRetargetPose);
 	}
 }
 
diff --git a/Engine/Plugins/Animation/IKRig/Source/IKRigEditor/Public/RetargetEditor/IKRetargetEditorController.h b/Engine/Plugins/Animation/IKRig/Source/IKRigEditor/Public/RetargetEditor/IKRetargetEditorController.h
--- a/Engine/Plugins/Animation/IKRig/Source/IKRigEditor/Public/RetargetEditor/IKRetargetEditorController.h
+++ b/Engine/Plugins/Animation/IKRig/Source/IKRigEditor/Public/RetargetEditor/IKRetargetEditorController.h
@@ -99,14 +99,18 @@
 
 	/** Sequence Browser**/
 	void PlayAnimationAsset(UAnimationAsset* AssetToPlay);
-	void PlayPreviousAnimationAsset() const;
-	UAnimationAsset* PreviousAsset = nullptr;
 	/** END Sequence Browser */
 
+	/** Set viewport / editor tool mode */
+	void SetRetargeterMode(ERetargeterOutputMode Mode);
+	UAnimationAsset* AnimThatWasPlaying = nullptr;
+	bool bWasPlayingAnim = false;
+	/** END viewport / editor tool mode */
+
 	/* START RETARGET POSES */
 	
 	/** go to retarget pose */
-	void HandleGoToRetargetPose() const;
+	void HandleGoToRetargetPose();
 	
 	/** toggle current retarget pose */
 	TArray<TSharedPtr<FName>> PoseNames;
@@ -114,7 +118,7 @@
 	void OnPoseSelected(TSharedPtr<FName> InPoseName, ESelectInfo::Type SelectInfo) const;
 	
 	/** edit retarget poses */
-	void HandleEditPose() const;
+	void HandleEditPose();
 	bool CanEditPose() const;
 	bool IsEditingPose() const;
 
