Adding memory profiling scopes to IK Rig and FBIK

diff --git a/Engine/Plugins/Animation/IKRig/Source/IKRig/Private/IKRigDefinition.cpp b/Engine/Plugins/Animation/IKRig/Source/IKRig/Private/IKRigDefinition.cpp
--- a/Engine/Plugins/Animation/IKRig/Source/IKRig/Private/IKRigDefinition.cpp
+++ b/Engine/Plugins/Animation/IKRig/Source/IKRig/Private/IKRigDefinition.cpp
@@ -311,6 +311,8 @@ USkeletalMesh* UIKRigDefinition::GetPreviewMesh() const
 
 FName UIKRigDefinition::GetUniqueIDAsName() const
 {
+	LLM_SCOPE_BYNAME(TEXT("Animation/IKRig"));
+	
 	static const FString IKRigIDPrefix("IKRig_");
 	return FName(IKRigIDPrefix + FString::FromInt(GetUniqueID()));
 }
diff --git a/Engine/Plugins/Animation/IKRig/Source/IKRig/Private/Retargeter/IKRetargeter.cpp b/Engine/Plugins/Animation/IKRig/Source/IKRig/Private/Retargeter/IKRetargeter.cpp
--- a/Engine/Plugins/Animation/IKRig/Source/IKRig/Private/Retargeter/IKRetargeter.cpp
+++ b/Engine/Plugins/Animation/IKRig/Source/IKRig/Private/Retargeter/IKRetargeter.cpp
@@ -129,6 +129,7 @@ void FIKRetargetPose::SortHierarchically(const FIKRigSkeleton& Skeleton)
 
 FName UIKRetargeter::GetUniqueIDAsName() const
 {
+	LLM_SCOPE_BYNAME(TEXT("Animation/IKRig"));
 	static const FString IKRetargetIDPrefix("IKRetargeter_");
 	return FName(IKRetargetIDPrefix + FString::FromInt(GetUniqueID()));
 }
diff --git a/Engine/Plugins/Animation/IKRig/Source/IKRigDeveloper/Private/AnimGraphNode_IKRig.cpp b/Engine/Plugins/Animation/IKRig/Source/IKRigDeveloper/Private/AnimGraphNode_IKRig.cpp
--- a/Engine/Plugins/Animation/IKRig/Source/IKRigDeveloper/Private/AnimGraphNode_IKRig.cpp
+++ b/Engine/Plugins/Animation/IKRig/Source/IKRigDeveloper/Private/AnimGraphNode_IKRig.cpp
@@ -431,6 +431,8 @@ void UAnimGraphNode_IKRig::CreateCustomPins(TArray<UEdGraphPin*>* InOldPins)
 
 void UAnimGraphNode_IKRig::SetPinDefaultValue(UEdGraphPin* InPin, const FName& InPropertyName)
 {
+	LLM_SCOPE_BYNAME(TEXT("Animation/IKRig"));
+	
 	// default FIKRigGoal structure 
 	static FIKRigGoal DefaultGoal;
 	static const TSharedPtr<FStructOnScope> StructOnScope =
@@ -748,6 +750,8 @@ void UAnimGraphNode_IKRig::CustomizePinData(UEdGraphPin* Pin, FName SourceProper
 
 void UAnimGraphNode_IKRig::BindPropertyChanges()
 {
+	LLM_SCOPE_BYNAME(TEXT("Animation/IKRig"));
+	
 	// already bound
 	if (OnAssetPropertyChangedHandle.IsValid())
 	{
diff --git a/Engine/Plugins/Animation/IKRig/Source/IKRigDeveloper/Private/IKRigDeveloper.cpp b/Engine/Plugins/Animation/IKRig/Source/IKRigDeveloper/Private/IKRigDeveloper.cpp
--- a/Engine/Plugins/Animation/IKRig/Source/IKRigDeveloper/Private/IKRigDeveloper.cpp
+++ b/Engine/Plugins/Animation/IKRig/Source/IKRigDeveloper/Private/IKRigDeveloper.cpp
@@ -5,13 +5,7 @@
 
 #define LOCTEXT_NAMESPACE "IKRigDeveloperModule"
 
-class FIKRigDeveloperModule : public IIKRigDeveloperModule
-{
-public:
-	/** IModuleInterface implementation */
-	virtual void StartupModule() override;
-	virtual void ShutdownModule() override;
-};
+IMPLEMENT_MODULE(FIKRigDeveloperModule, IKRigDeveloper)
 
 void FIKRigDeveloperModule::StartupModule()
 {
@@ -21,5 +15,4 @@ void FIKRigDeveloperModule::ShutdownModule()
 {
 }
 
-IMPLEMENT_MODULE(FIKRigDeveloperModule, IKRigDeveloper)
 #undef LOCTEXT_NAMESPACE
diff --git a/Engine/Plugins/Animation/IKRig/Source/IKRigDeveloper/Public/IKRigDeveloper.h b/Engine/Plugins/Animation/IKRig/Source/IKRigDeveloper/Public/IKRigDeveloper.h
--- a/Engine/Plugins/Animation/IKRig/Source/IKRigDeveloper/Public/IKRigDeveloper.h
+++ b/Engine/Plugins/Animation/IKRig/Source/IKRigDeveloper/Public/IKRigDeveloper.h
@@ -5,8 +5,9 @@
 #include "CoreMinimal.h"
 #include "Modules/ModuleInterface.h"
 
-class IIKRigDeveloperModule : public IModuleInterface
+class IKRIGDEVELOPER_API FIKRigDeveloperModule : public IModuleInterface
 {
 public:
-};
-
+	void StartupModule() override;
+	void ShutdownModule() override;
+};
\ No newline at end of file
diff --git a/Engine/Plugins/Animation/IKRig/Source/IKRigEditor/Private/RetargetEditor/IKRetargetAnimInstance.cpp b/Engine/Plugins/Animation/IKRig/Source/IKRigEditor/Private/RetargetEditor/IKRetargetAnimInstance.cpp
--- a/Engine/Plugins/Animation/IKRig/Source/IKRigEditor/Private/RetargetEditor/IKRetargetAnimInstance.cpp
+++ b/Engine/Plugins/Animation/IKRig/Source/IKRigEditor/Private/RetargetEditor/IKRetargetAnimInstance.cpp
@@ -29,5 +29,6 @@ void UIKRetargetAnimInstance::SetProcessorNeedsInitialized()
 
 FAnimInstanceProxy* UIKRetargetAnimInstance::CreateAnimInstanceProxy()
 {
+	LLM_SCOPE_BYNAME(TEXT("Animation/IKRig"));
 	return new FIKRetargetAnimInstanceProxy(this, &IKRetargeterNode);
 }
diff --git a/Engine/Plugins/Experimental/FullBodyIK/Source/FullBodyIK/Private/JacobianSolver.cpp b/Engine/Plugins/Experimental/FullBodyIK/Source/FullBodyIK/Private/JacobianSolver.cpp
--- a/Engine/Plugins/Experimental/FullBodyIK/Source/FullBodyIK/Private/JacobianSolver.cpp
+++ b/Engine/Plugins/Experimental/FullBodyIK/Source/FullBodyIK/Private/JacobianSolver.cpp
@@ -107,6 +107,7 @@ void CalculateRotationAxisBasedOnEffectorRotation(TArray<FFBIKLinkData>& InOutLi
 ///////////////////////////////////////////////////////////////////////////////
 FJacobianSolverBase::FJacobianSolverBase()
 {
+	LLM_SCOPE_BYNAME(TEXT("Animation/FBIK"));
 	OnCalculatePartialDerivativesDelegate = FCalculatePartialDerivativesDelegate::CreateRaw(this, &FJacobianSolverBase::ComputePartialDerivative);
 	OnCalculateTargetVectorDelegate = FCalculateTargetVectorDelegate::CreateRaw(this, &FJacobianSolverBase::ComputeTargetVector);
 }
diff --git a/Engine/Plugins/Experimental/FullBodyIK/Source/PBIK/Private/Core/PBIKSolver.cpp b/Engine/Plugins/Experimental/FullBodyIK/Source/PBIK/Private/Core/PBIKSolver.cpp
--- a/Engine/Plugins/Experimental/FullBodyIK/Source/PBIK/Private/Core/PBIKSolver.cpp
+++ b/Engine/Plugins/Experimental/FullBodyIK/Source/PBIK/Private/Core/PBIKSolver.cpp
@@ -412,6 +412,8 @@ void FPBIKSolver::UpdateBonesFromBodies()
 
 bool FPBIKSolver::Initialize()
 {
+	LLM_SCOPE_BYNAME(TEXT("Animation/FBIK"));
+	
 	if (bReadyToSimulate)
 	{
 		return true;
@@ -683,6 +685,8 @@ PBIK::FDebugDraw* FPBIKSolver::GetDebugDraw()
 
 void FPBIKSolver::Reset()
 {
+	LLM_SCOPE_BYNAME(TEXT("Animation/FBIK"));
+	
 	bReadyToSimulate = false;
 	SolverRoot = nullptr;
 	RootPin = nullptr;
diff --git a/Engine/Plugins/Experimental/FullBodyIK/Source/PBIK/Private/RigUnit_PBIK.cpp b/Engine/Plugins/Experimental/FullBodyIK/Source/PBIK/Private/RigUnit_PBIK.cpp
--- a/Engine/Plugins/Experimental/FullBodyIK/Source/PBIK/Private/RigUnit_PBIK.cpp
+++ b/Engine/Plugins/Experimental/FullBodyIK/Source/PBIK/Private/RigUnit_PBIK.cpp
@@ -6,6 +6,7 @@
 FRigUnit_PBIK_Execute()
 {
 	DECLARE_SCOPE_HIERARCHICAL_COUNTER_RIGUNIT()
+	LLM_SCOPE_BYNAME(TEXT("Animation/FBIK"));
 
 	if (Context.State == EControlRigState::Init)
 	{
