Fix non-unity
#rnx

diff --git a/Engine/Plugins/Animation/IKRig/Source/IKRigEditor/Public/SBaseHierarchyTreeView.h b/Engine/Plugins/Animation/IKRig/Source/IKRigEditor/Public/SBaseHierarchyTreeView.h
--- a/Engine/Plugins/Animation/IKRig/Source/IKRigEditor/Public/SBaseHierarchyTreeView.h
+++ b/Engine/Plugins/Animation/IKRig/Source/IKRigEditor/Public/SBaseHierarchyTreeView.h
@@ -2,6 +2,8 @@
 
 #pragma once
 
+#include "Widgets/Views/STreeView.h"
+
 template<class ItemType>
 class SBaseHierarchyTreeView : public STreeView<TSharedPtr<ItemType>>
 {
