Fixed name of IK Rig output log tab.

diff --git a/Engine/Plugins/Animation/IKRig/Source/IKRigEditor/Private/RigEditor/IKRigOutputLogTabSummoner.cpp b/Engine/Plugins/Animation/IKRig/Source/IKRigEditor/Private/RigEditor/IKRigOutputLogTabSummoner.cpp
--- a/Engine/Plugins/Animation/IKRig/Source/IKRigEditor/Private/RigEditor/IKRigOutputLogTabSummoner.cpp
+++ b/Engine/Plugins/Animation/IKRig/Source/IKRigEditor/Private/RigEditor/IKRigOutputLogTabSummoner.cpp
@@ -17,10 +17,10 @@ FIKRigOutputLogTabSummoner::FIKRigOutputLogTabSummoner(
 {
 	bIsSingleton = true; // only allow a single instance of this tab
 	
-	TabLabel = LOCTEXT("IKRigOutputLogTabLabel", "IK Rig Output");
+	TabLabel = LOCTEXT("IKRigOutputLogTabLabel", "IK Rig Output Log");
 	TabIcon = FSlateIcon(FIKRigEditorStyle::Get().GetStyleSetName(), "IKRig.TabIcon");
 
-	ViewMenuDescription = LOCTEXT("IKRigOutputLog_ViewMenu_Desc", "IK Rig Output");
+	ViewMenuDescription = LOCTEXT("IKRigOutputLog_ViewMenu_Desc", "IK Rig Output Log");
 	ViewMenuTooltip = LOCTEXT("IKRigOutputLog_ViewMenu_ToolTip", "Show the IK Rig Output Log Tab");
 }
 
