Various fixes to batch retarget export window.

diff --git a/Engine/Plugins/Animation/IKRig/Source/IKRigEditor/Private/RetargetEditor/SIKRetargetAssetBrowser.cpp b/Engine/Plugins/Animation/IKRig/Source/IKRigEditor/Private/RetargetEditor/SIKRetargetAssetBrowser.cpp
--- a/Engine/Plugins/Animation/IKRig/Source/IKRigEditor/Private/RetargetEditor/SIKRetargetAssetBrowser.cpp
+++ b/Engine/Plugins/Animation/IKRig/Source/IKRigEditor/Private/RetargetEditor/SIKRetargetAssetBrowser.cpp
@@ -16,6 +16,7 @@
 #include "RetargetEditor/SRetargetAnimAssetsWindow.h"
 #include "Retargeter/IKRetargeter.h"
 #include "Retargeter/IKRetargetProcessor.h"
+#include "Widgets/Layout/SUniformGridPanel.h"
 
 #define LOCTEXT_NAMESPACE "IKRetargeterAssetBrowser"
 
@@ -86,49 +87,41 @@ void SIKRetargetAssetBrowser::RefreshView()
 	AssetBrowserBox->SetContent(ContentBrowserModule.Get().CreateAssetPicker(AssetPickerConfig));
 }
 
-void SIKRetargetAssetBrowser::OnPathChange(const FString& NewPath)
-{
-	BatchOutputPath = NewPath;
-}
-
-FReply SIKRetargetAssetBrowser::OnExportButtonClicked() const
+FReply SIKRetargetAssetBrowser::OnExportButtonClicked()
 {
 	FIKRetargetEditorController* Controller = EditorController.Pin().Get();
 	if (!Controller)
 	{
+		checkNoEntry();
 		return FReply::Handled();
 	}
-
-	// assemble the data for the assets we want to batch duplicate/retarget
-	FIKRetargetBatchOperationContext BatchContext;
-
-	TSharedRef<SSelectExportPathDialog> Dialog = SNew(SSelectExportPathDialog).DefaultAssetPath(FText::FromString(BatchContext.FolderPath));
-	if(Dialog->ShowModal() != EAppReturnType::Cancel)
+	
+	// prompt user for path to export animations to
+	TSharedRef<SBatchExportDialog> Dialog = SNew(SBatchExportDialog).DefaultAssetPath(FText::FromString(PrevBatchOutputPath));
+	if(Dialog->ShowModal() == EAppReturnType::Cancel)
 	{
-		BatchContext.NameRule.FolderPath = Dialog->GetAssetPath();
+		return FReply::Handled();
 	}
 
-	// set the path to export to
-	if (!BatchOutputPath.IsEmpty())
-	{
-		BatchContext.FolderPath = BatchOutputPath;
-	}
+	// store path for next time
+	PrevBatchOutputPath = Dialog->GetAssetPath();
+
+	// assemble the data for the assets we want to batch duplicate/retarget
+	FIKRetargetBatchOperationContext BatchContext = Dialog.Get().BatchContext;
+	BatchContext.NameRule.FolderPath = PrevBatchOutputPath;
+	BatchContext.SourceMesh = Controller->GetSourceSkeletalMesh();
+	BatchContext.TargetMesh = Controller->GetTargetSkeletalMesh();
+	BatchContext.IKRetargetAsset = Controller->AssetController->GetAsset();
+	BatchContext.bRemapReferencedAssets = false;
 
 	// add selected assets to dup/retarget
 	TArray<FAssetData> SelectedAssets = GetCurrentSelectionDelegate.Execute();
 	for (const FAssetData& Asset : SelectedAssets)
 	{
 		UE_LOG(LogIKRigEditor, Display, TEXT("Duplicating and Retargeting: %s"), *Asset.GetFullName());
-
 		BatchContext.AssetsToRetarget.Add(Asset.GetAsset());
 	}
 
-	BatchContext.SourceMesh = Controller->GetSourceSkeletalMesh();
-	BatchContext.TargetMesh = Controller->GetTargetSkeletalMesh();
-	BatchContext.IKRetargetAsset = Controller->AssetController->GetAsset();
-	BatchContext.bRemapReferencedAssets = false;
-	BatchContext.NameRule.Suffix = "_Retargeted";
-
 	// actually run the retarget
 	FIKRetargetBatchOperation BatchOperation;
 	BatchOperation.RunRetarget(BatchContext);
@@ -202,4 +195,294 @@ bool SIKRetargetAssetBrowser::OnShouldFilterAsset(const struct FAssetData& Asset
 	return !DesiredSkeleton->IsCompatibleSkeletonByAssetData(AssetData);
 }
 
+// ------------------------------------------BEGIN  SBatchExportDialog ----------------------------
+
+void SBatchExportDialog::Construct(const FArguments& InArgs)
+{
+	AssetPath = FText::FromString(FPackageName::GetLongPackagePath(InArgs._DefaultAssetPath.ToString()));
+
+	if(AssetPath.IsEmpty())
+	{
+		AssetPath = FText::FromString(TEXT("/Game"));
+	}
+
+	FPathPickerConfig PathPickerConfig;
+	PathPickerConfig.DefaultPath = AssetPath.ToString();
+	PathPickerConfig.OnPathSelected = FOnPathSelected::CreateSP(this, &SBatchExportDialog::OnPathChange);
+	PathPickerConfig.bAddDefaultPath = true;
+
+	FContentBrowserModule& ContentBrowserModule = FModuleManager::LoadModuleChecked<FContentBrowserModule>("ContentBrowser");
+
+	constexpr int32 TextInputWidths = 200;
+
+	SWindow::Construct(
+		SWindow::FArguments()
+		.Title(LOCTEXT("SBatchExportDialog_Title", "Batch Export Retargeted Animations"))
+		.SupportsMinimize(false)
+		.SupportsMaximize(false)
+		.IsTopmostWindow(true)
+		.ClientSize(FVector2D(350, 600))
+		[
+			SNew(SVerticalBox)
+			
+			+ SVerticalBox::Slot()
+			.Padding(2)
+			[
+				SNew(SBorder)
+				.BorderImage(FAppStyle::GetBrush("ToolPanel.GroupBorder"))
+				[
+					SNew(SVerticalBox)
+
+					+SVerticalBox::Slot()
+					.Padding(2, 3)
+					.AutoHeight()
+					.HAlign(HAlign_Center)
+					[
+						SNew(STextBlock)
+						.Text(LOCTEXT("SelectPath", "Select Path"))
+						.Font(FCoreStyle::GetDefaultFontStyle("Regular", 14))
+					]
+
+					+SVerticalBox::Slot()
+					.FillHeight(1)
+					.Padding(3)
+					[
+						ContentBrowserModule.Get().CreatePathPicker(PathPickerConfig)
+					]
+				]
+			]
+
+			+ SVerticalBox::Slot()
+			.Padding(2)
+			.AutoHeight()
+			[
+				SNew(SBorder)
+				.BorderImage(FAppStyle::GetBrush("ToolPanel.GroupBorder"))
+				[
+					SNew(SVerticalBox)
+
+					+SVerticalBox::Slot()
+					.Padding(2, 3)
+					.AutoHeight()
+					.HAlign(HAlign_Center)
+					[
+						SNew(STextBlock)
+						.Text(LOCTEXT("DuplicateAndRetarget_RenameLabel", "Rename New Assets"))
+						.Font(FCoreStyle::GetDefaultFontStyle("Regular", 14))
+					]
+
+					+SVerticalBox::Slot()
+					.AutoHeight()
+					.Padding(2, 1)
+					.HAlign(HAlign_Right)
+					[
+						SNew(SHorizontalBox)
+						+SHorizontalBox::Slot()
+						.AutoWidth()
+						.VAlign(VAlign_Center)
+						.Padding(2, 1)
+						[
+							SNew(STextBlock).Text(LOCTEXT("DuplicateAndRetarget_Prefix", "Prefix:"))
+						]
+
+						+SHorizontalBox::Slot()
+						.AutoWidth()
+						[
+							SNew(SEditableTextBox)
+							.Text_Lambda([this]
+							{
+								return FText::FromString(BatchContext.NameRule.Prefix);
+							})
+							.MinDesiredWidth(TextInputWidths)
+							.OnTextChanged_Lambda([this](const FText& InText)
+							{
+								BatchContext.NameRule.Prefix = InText.ToString();
+								UpdateExampleText();
+							})
+							.IsReadOnly(false)
+							.RevertTextOnEscape(true)
+						]
+					]
+
+					+SVerticalBox::Slot()
+					.AutoHeight()
+					.Padding(2, 1)
+					.HAlign(HAlign_Right)
+					[
+						SNew(SHorizontalBox)
+						+SHorizontalBox::Slot()
+						.AutoWidth()
+						.VAlign(VAlign_Center)
+						.Padding(2, 1)
+						[
+							SNew(STextBlock).Text(LOCTEXT("DuplicateAndRetarget_Suffix", "Suffix:"))
+						]
+
+						+SHorizontalBox::Slot()
+						.AutoWidth()
+						[
+							SNew(SEditableTextBox)
+							.Text_Lambda([this]
+							{
+								return FText::FromString(BatchContext.NameRule.Suffix);
+							})
+							.MinDesiredWidth(TextInputWidths)
+							.OnTextChanged_Lambda([this](const FText& InText)
+							{
+								BatchContext.NameRule.Suffix = InText.ToString();
+								UpdateExampleText();
+							})
+							.IsReadOnly(false)
+							.RevertTextOnEscape(true)
+						]
+					]
+
+					+SVerticalBox::Slot()
+					.AutoHeight()
+					.Padding(2, 1)
+					.HAlign(HAlign_Right)
+					[
+						SNew(SHorizontalBox)
+						+SHorizontalBox::Slot()
+						.AutoWidth()
+						.VAlign(VAlign_Center)
+						.Padding(2, 1)
+						[
+							SNew(STextBlock).Text(LOCTEXT("DuplicateAndRetarget_Search", "Search:"))
+						]
+
+						+SHorizontalBox::Slot()
+						.AutoWidth()
+						[
+							SNew(SEditableTextBox)
+							.Text_Lambda([this]
+							{
+								return FText::FromString(BatchContext.NameRule.ReplaceFrom);
+							})
+							.MinDesiredWidth(TextInputWidths)
+							.OnTextChanged_Lambda([this](const FText& InText)
+							{
+								BatchContext.NameRule.ReplaceFrom = InText.ToString();
+								UpdateExampleText();
+							})
+							.IsReadOnly(false)
+							.RevertTextOnEscape(true)
+						]
+					]
+
+					+SVerticalBox::Slot()
+					.AutoHeight()
+					.Padding(2, 1)
+					.HAlign(HAlign_Right)
+					[
+						SNew(SHorizontalBox)
+						+SHorizontalBox::Slot()
+						.AutoWidth()
+						.VAlign(VAlign_Center)
+						.Padding(2, 1)
+						[
+							SNew(STextBlock).Text(LOCTEXT("DuplicateAndRetarget_Replace", "Replace:"))
+						]
+
+						+SHorizontalBox::Slot()
+						.AutoWidth()
+						[
+							SNew(SEditableTextBox)
+							.Text_Lambda([this]
+							{
+								return FText::FromString(BatchContext.NameRule.ReplaceTo);
+							})
+							.MinDesiredWidth(TextInputWidths)
+							.OnTextChanged_Lambda([this](const FText& InText)
+							{
+								BatchContext.NameRule.ReplaceTo = InText.ToString();
+								UpdateExampleText();
+							})
+							.IsReadOnly(false)
+							.RevertTextOnEscape(true)
+						]
+					]
+
+					+SVerticalBox::Slot()
+					.AutoHeight()
+					.Padding(2, 3)
+					.HAlign(HAlign_Right)
+					[
+						SNew(SHorizontalBox)
+						+SHorizontalBox::Slot()
+						.Padding(5, 5)
+						.AutoWidth()
+						.VAlign(VAlign_Center)
+						[
+							SNew(STextBlock)
+							.Text_Lambda([this]{ return ExampleText; })
+							.Font(FAppStyle::GetFontStyle("Persona.RetargetManager.ItalicFont"))
+						]
+					]
+				]
+			]
+			
+			+SVerticalBox::Slot()
+			.AutoHeight()
+			.HAlign(HAlign_Right)
+			.Padding(5)
+			[
+				SNew(SUniformGridPanel)
+				.SlotPadding(FAppStyle::GetMargin("StandardDialog.SlotPadding"))
+				.MinDesiredSlotWidth(FAppStyle::GetFloat("StandardDialog.MinDesiredSlotWidth"))
+				.MinDesiredSlotHeight(FAppStyle::GetFloat("StandardDialog.MinDesiredSlotHeight"))
+				+SUniformGridPanel::Slot(0, 0)
+				[
+					SNew(SButton)
+					.HAlign(HAlign_Center)
+					.ContentPadding(FAppStyle::GetMargin("StandardDialog.ContentPadding"))
+					.Text(LOCTEXT("Export", "Export"))
+					.OnClicked(this, &SBatchExportDialog::OnButtonClick, EAppReturnType::Ok)
+				]
+				+SUniformGridPanel::Slot(1, 0)
+				[
+					SNew(SButton)
+					.HAlign(HAlign_Center)
+					.ContentPadding(FAppStyle::GetMargin("StandardDialog.ContentPadding"))
+					.Text(LOCTEXT("Cancel", "Cancel"))
+					.OnClicked(this, &SBatchExportDialog::OnButtonClick, EAppReturnType::Cancel)
+				]
+			]
+		]);
+}
+
+void SBatchExportDialog::UpdateExampleText()
+{
+	const FString ReplaceFrom = FString::Printf(TEXT("Old Name : ###%s###"), *BatchContext.NameRule.ReplaceFrom);
+	const FString ReplaceTo = FString::Printf(TEXT("New Name : %s###%s###%s"), *BatchContext.NameRule.Prefix, *BatchContext.NameRule.ReplaceTo, *BatchContext.NameRule.Suffix);
+
+	ExampleText = FText::FromString(FString::Printf(TEXT("%s\n%s"), *ReplaceFrom, *ReplaceTo));
+}
+
+void SBatchExportDialog::OnPathChange(const FString& NewPath)
+{
+	AssetPath = FText::FromString(NewPath);
+}
+
+FReply SBatchExportDialog::OnButtonClick(EAppReturnType::Type ButtonID)
+{
+	UserResponse = ButtonID;
+
+	RequestDestroyWindow();
+
+	return FReply::Handled();
+}
+
+EAppReturnType::Type SBatchExportDialog::ShowModal()
+{
+	GEditor->EditorAddModalWindow(SharedThis(this));
+	return UserResponse;
+}
+
+FString SBatchExportDialog::GetAssetPath()
+{
+	return AssetPath.ToString();
+}
+
+
 #undef LOCTEXT_NAMESPACE
diff --git a/Engine/Plugins/Animation/IKRig/Source/IKRigEditor/Private/RetargetEditor/SRetargetAnimAssetsWindow.cpp b/Engine/Plugins/Animation/IKRig/Source/IKRigEditor/Private/RetargetEditor/SRetargetAnimAssetsWindow.cpp
--- a/Engine/Plugins/Animation/IKRig/Source/IKRigEditor/Private/RetargetEditor/SRetargetAnimAssetsWindow.cpp
+++ b/Engine/Plugins/Animation/IKRig/Source/IKRigEditor/Private/RetargetEditor/SRetargetAnimAssetsWindow.cpp
@@ -743,7 +743,7 @@ FText SRetargetAnimAssetsWindow::GetFolderPath() const
 FReply SRetargetAnimAssetsWindow::GetExportFolder()
 {
 	TSharedRef<SSelectExportPathDialog> Dialog = SNew(SSelectExportPathDialog)
-	.DefaultAssetPath(FText::FromString(BatchContext.FolderPath));
+	.DefaultAssetPath(FText::FromString(BatchContext.NameRule.FolderPath));
 	
 	if(Dialog->ShowModal() != EAppReturnType::Cancel)
 	{
diff --git a/Engine/Plugins/Animation/IKRig/Source/IKRigEditor/Public/RetargetEditor/IKRetargetBatchOperation.h b/Engine/Plugins/Animation/IKRig/Source/IKRigEditor/Public/RetargetEditor/IKRetargetBatchOperation.h
--- a/Engine/Plugins/Animation/IKRig/Source/IKRigEditor/Public/RetargetEditor/IKRetargetBatchOperation.h
+++ b/Engine/Plugins/Animation/IKRig/Source/IKRigEditor/Public/RetargetEditor/IKRetargetBatchOperation.h
@@ -31,9 +31,6 @@ public:
 	/* Rename rules for duplicated assets */
 	EditorAnimUtils::FNameDuplicationRule NameRule;
 
-	/* The base folder path to put newly duplicated assets */
-	FString FolderPath = "";
-
 	/* Reset all data (called when window re-opened */
 	void Reset()
 	{
diff --git a/Engine/Plugins/Animation/IKRig/Source/IKRigEditor/Public/RetargetEditor/SIKRetargetAssetBrowser.h b/Engine/Plugins/Animation/IKRig/Source/IKRigEditor/Public/RetargetEditor/SIKRetargetAssetBrowser.h
--- a/Engine/Plugins/Animation/IKRig/Source/IKRigEditor/Public/RetargetEditor/SIKRetargetAssetBrowser.h
+++ b/Engine/Plugins/Animation/IKRig/Source/IKRigEditor/Public/RetargetEditor/SIKRetargetAssetBrowser.h
@@ -5,6 +5,7 @@
 #include "CoreMinimal.h"
 
 #include "ContentBrowserDelegates.h"
+#include "IKRetargetBatchOperation.h"
 #include "Widgets/Layout/SBox.h"
 
 
@@ -22,9 +23,7 @@ public:
 	
 private:
 
-	void OnPathChange(const FString& NewPath);
-
-	FReply OnExportButtonClicked() const;
+	FReply OnExportButtonClicked();
 	bool IsExportButtonEnabled() const;
 	
 	void OnAssetDoubleClicked(const FAssetData& AssetData);
@@ -39,8 +38,44 @@ private:
 	/** the animation asset browser */
 	TSharedPtr<SBox> AssetBrowserBox;
 
-	/** the path to export assets to */
-	FString BatchOutputPath = "";
+	/** remember the path the user chose between exports */
+	FString PrevBatchOutputPath = "";
 
 	friend FIKRetargetEditorController;
 };
+
+//** Dialog to select path to export to */
+class SBatchExportDialog: public SWindow
+{
+public:
+	SLATE_BEGIN_ARGS(SBatchExportDialog){}
+	SLATE_ARGUMENT(FText, DefaultAssetPath)
+	SLATE_END_ARGS()
+
+	SBatchExportDialog() : UserResponse(EAppReturnType::Cancel)
+	{
+	}
+
+	void Construct(const FArguments& InArgs);
+
+public:
+	/** Displays the dialog in a blocking fashion */
+	EAppReturnType::Type ShowModal();
+
+	/** Gets the resulting asset path */
+	FString GetAssetPath();
+
+	FIKRetargetBatchOperationContext BatchContext;
+
+protected:
+
+	/** The rename rule sample text */
+	FText ExampleText;
+	void UpdateExampleText();
+
+	void OnPathChange(const FString& NewPath);
+	FReply OnButtonClick(EAppReturnType::Type ButtonID);
+
+	EAppReturnType::Type UserResponse;
+	FText AssetPath;
+};
