IKRig: Fix for non-unity builds

#rnx

diff --git a/Engine/Plugins/Animation/IKRig/Source/IKRigEditor/Public/RetargetEditor/SIKRetargetAssetBrowser.h b/Engine/Plugins/Animation/IKRig/Source/IKRigEditor/Public/RetargetEditor/SIKRetargetAssetBrowser.h
--- a/Engine/Plugins/Animation/IKRig/Source/IKRigEditor/Public/RetargetEditor/SIKRetargetAssetBrowser.h
+++ b/Engine/Plugins/Animation/IKRig/Source/IKRigEditor/Public/RetargetEditor/SIKRetargetAssetBrowser.h
@@ -7,6 +7,7 @@
 #include "ContentBrowserDelegates.h"
 #include "IKRetargetBatchOperation.h"
 #include "Widgets/Layout/SBox.h"
+#include "Widgets/SWindow.h"
 
 
 class FIKRetargetEditorController;
