Added ability to double-click RPFM node to open retarget asset.

diff --git a/Engine/Plugins/Animation/IKRig/Source/IKRigDeveloper/Private/AnimGraphNode_RetargetPoseFromMesh.cpp b/Engine/Plugins/Animation/IKRig/Source/IKRigDeveloper/Private/AnimGraphNode_RetargetPoseFromMesh.cpp
--- a/Engine/Plugins/Animation/IKRig/Source/IKRigDeveloper/Private/AnimGraphNode_RetargetPoseFromMesh.cpp
+++ b/Engine/Plugins/Animation/IKRig/Source/IKRigDeveloper/Private/AnimGraphNode_RetargetPoseFromMesh.cpp
@@ -31,6 +31,11 @@ void UAnimGraphNode_RetargetPoseFromMesh::CustomizePinData(UEdGraphPin* Pin, FNa
 	Super::CustomizePinData(Pin, SourcePropertyName, ArrayIndex);
 }
 
+UObject* UAnimGraphNode_RetargetPoseFromMesh::GetJumpTargetForDoubleClick() const
+{
+	return Node.IKRetargeterAsset;
+}
+
 void UAnimGraphNode_RetargetPoseFromMesh::PostEditChangeProperty(struct FPropertyChangedEvent& PropertyChangedEvent)
 {
 	const FName PropertyName = (PropertyChangedEvent.Property ? PropertyChangedEvent.Property->GetFName() : NAME_None);
diff --git a/Engine/Plugins/Animation/IKRig/Source/IKRigDeveloper/Public/AnimGraphNode_RetargetPoseFromMesh.h b/Engine/Plugins/Animation/IKRig/Source/IKRigDeveloper/Public/AnimGraphNode_RetargetPoseFromMesh.h
--- a/Engine/Plugins/Animation/IKRig/Source/IKRigDeveloper/Public/AnimGraphNode_RetargetPoseFromMesh.h
+++ b/Engine/Plugins/Animation/IKRig/Source/IKRigDeveloper/Public/AnimGraphNode_RetargetPoseFromMesh.h
@@ -36,5 +36,9 @@ public:
 	virtual bool UsingCopyPoseFromMesh() const override { return true; };
 	// End of UAnimGraphNode_Base interface
 
+	// UK2Node interface
+	virtual UObject* GetJumpTargetForDoubleClick() const override;
+	// End of UK2Node interface
+
 	static const FName AnimModeName;
 };
