Fixing crash in retarget editor when deleting the target IK rig.

diff --git a/Engine/Plugins/Animation/IKRig/Source/IKRig/Private/Retargeter/IKRetargetProcessor.cpp b/Engine/Plugins/Animation/IKRig/Source/IKRig/Private/Retargeter/IKRetargetProcessor.cpp
--- a/Engine/Plugins/Animation/IKRig/Source/IKRig/Private/Retargeter/IKRetargetProcessor.cpp
+++ b/Engine/Plugins/Animation/IKRig/Source/IKRig/Private/Retargeter/IKRetargetProcessor.cpp
@@ -1670,7 +1670,13 @@ void UIKRetargetProcessor::SetNeedsInitialized()
 
 void UIKRetargetProcessor::CopyAllSettingsFromAsset()
 {
-	IKRigProcessor->CopyAllInputsFromSourceAssetAtRuntime(RetargeterAsset->GetTargetIKRig());
+	const UIKRigDefinition* TargetIKRig = RetargeterAsset->GetTargetIKRig();
+	if (!TargetIKRig)
+	{
+		return;
+	}
+	
+	IKRigProcessor->CopyAllInputsFromSourceAssetAtRuntime(TargetIKRig);
 
 	// copy most recent settings from asset for each chain
 	const TArray<TObjectPtr<URetargetChainSettings>>& AllChainSettings = RetargeterAsset->GetAllChainSettings();
