Renaming goal with a non-unique name now makes new name unique instead of doing nothing.

diff --git a/Engine/Plugins/Animation/IKRig/Source/IKRigEditor/Private/RigEditor/IKRigController.cpp b/Engine/Plugins/Animation/IKRig/Source/IKRigEditor/Private/RigEditor/IKRigController.cpp
--- a/Engine/Plugins/Animation/IKRig/Source/IKRigEditor/Private/RigEditor/IKRigController.cpp
+++ b/Engine/Plugins/Animation/IKRig/Source/IKRigEditor/Private/RigEditor/IKRigController.cpp
@@ -196,17 +196,15 @@ FName UIKRigController::RenameRetargetChain(const FName& ChainName, const FName&
 		return ChainName; // chain doesn't exist to rename
 	}
 
-	if (Asset->GetRetargetChainByName(NewChainName))
-	{
-		return ChainName; // bone chain already exists with the new name
-	}
+	// make sure it's unique
+	const FName UniqueChainName = GetUniqueRetargetChainName(NewChainName);
 	
 	FScopedTransaction Transaction(LOCTEXT("RenameRetargetChain_Label", "Rename Retarget Chain"));
 	Asset->Modify();
-	Chain->ChainName = NewChainName;
-	RetargetChainRenamed.Broadcast(GetAsset(), ChainName, NewChainName);
+	Chain->ChainName = UniqueChainName;
+	RetargetChainRenamed.Broadcast(GetAsset(), ChainName, UniqueChainName);
 	BroadcastNeedsReinitialized();
-	return NewChainName;
+	return UniqueChainName;
 }
 
 bool UIKRigController::SetRetargetChainStartBone(const FName& ChainName, const FName& StartBoneName) const
@@ -335,11 +333,6 @@ void UIKRigController::SortRetargetChains() const
 
 FName UIKRigController::GetUniqueRetargetChainName(const FName& NameToMakeUnique) const
 {
-	if (!Asset->GetRetargetChainByName(NameToMakeUnique))
-	{
-		return NameToMakeUnique; // name is already unique
-	}
-	
 	auto IsNameBeingUsed = [this](const FName& NameToTry)->bool
 	{
 		for (const FBoneChain& Chain : Asset->RetargetDefinition.BoneChains)
@@ -352,17 +345,20 @@ FName UIKRigController::GetUniqueRetargetChainName(const FName& NameToMakeUnique
 		return false;
 	};
 
-	FString ChainName = NameToMakeUnique.ToString();
-	int8 SuffixInt = 0;
-	while (true)
+	// check if name is already unique
+	if (!IsNameBeingUsed(NameToMakeUnique))
 	{
-		const FName ChainNameToTry = FName(*FString::Format(TEXT("{0}_{1}"), {ChainName, FString::FromInt(SuffixInt)}));
-		if (!IsNameBeingUsed(ChainNameToTry))
-		{
-			return ChainNameToTry;
-		}
-		++SuffixInt;
+		return NameToMakeUnique; 
 	}
+	
+	// keep concatenating an incremented integer suffix until name is unique
+	int32 Number = NameToMakeUnique.GetNumber() + 1;
+	while(IsNameBeingUsed(FName(NameToMakeUnique, Number)))
+	{
+		Number++;
+	}
+
+	return FName(NameToMakeUnique, Number);
 }
 
 bool UIKRigController::ValidateChain(const FName& ChainName, TSet<int32>& OutChainIndices) const
@@ -734,6 +730,11 @@ bool UIKRigController::RemoveGoal(const FName& GoalName) const
 
 FName UIKRigController::RenameGoal(const FName& OldName, const FName& PotentialNewName) const
 {
+	if (OldName == PotentialNewName)
+	{
+		return OldName; // skipping renaming the same name
+	}
+	
 	const int32 GoalIndex = GetGoalIndex(OldName);
 	if (GoalIndex == INDEX_NONE)
 	{
@@ -743,22 +744,8 @@ FName UIKRigController::RenameGoal(const FName& OldName, const FName& PotentialN
 	// sanitize the potential new name
 	FString CleanName = PotentialNewName.ToString();
 	SanitizeGoalName(CleanName);
-	const FName NewName = FName(CleanName);
-	// check if this goal already exists (case sensitive)
-	int32 ExistingGoalIndex = GetGoalIndex(NewName, ENameCase::CaseSensitive);
-	if (ExistingGoalIndex != INDEX_NONE)
-	{
-		return NAME_None; // name already in use, can't use that
-	}
-	// check if this goal already exists (case insensitive, it might just be renamed with a different case)
-	ExistingGoalIndex = GetGoalIndex(NewName);
-	if (ExistingGoalIndex != INDEX_NONE)
-	{
-		if (!NewName.IsEqual(OldName, ENameCase::IgnoreCase))
-		{
-			return NAME_None; // name already in use, can't use that
-		}
-	}
+	// make the name unique
+	const FName NewName = GetUniqueGoalName(FName(CleanName));
 	
 	FScopedTransaction Transaction(LOCTEXT("RenameGoal_Label", "Rename Goal"));
 	Asset->Modify();
@@ -789,6 +776,31 @@ FName UIKRigController::RenameGoal(const FName& OldName, const FName& PotentialN
 	return NewName;
 }
 
+FName UIKRigController::GetUniqueGoalName(const FName& NameToMakeUnique) const
+{
+	auto IsNameBeingUsed = [this](const FName& NameToTry) -> bool
+	{
+		// check if this goal already exists (case sensitive)
+		int32 ExistingGoalIndex = GetGoalIndex(NameToTry, ENameCase::IgnoreCase);
+		return ExistingGoalIndex != INDEX_NONE;
+	};
+
+	// check if name is already unique
+	if (!IsNameBeingUsed(NameToMakeUnique))
+	{
+		return NameToMakeUnique; 
+	}
+	
+	// keep concatenating an incremented integer suffix until name is unique
+	int32 Number = NameToMakeUnique.GetNumber() + 1;
+	while(IsNameBeingUsed(FName(NameToMakeUnique, Number)))
+	{
+		Number++;
+	}
+
+	return FName(NameToMakeUnique, Number);
+}
+
 bool UIKRigController::ModifyGoal(const FName& GoalName) const
 {
 	if (UIKRigEffectorGoal* Goal = GetGoal(GoalName))
diff --git a/Engine/Plugins/Animation/IKRig/Source/IKRigEditor/Public/RigEditor/IKRigController.h b/Engine/Plugins/Animation/IKRig/Source/IKRigEditor/Public/RigEditor/IKRigController.h
--- a/Engine/Plugins/Animation/IKRig/Source/IKRigEditor/Public/RigEditor/IKRigController.h
+++ b/Engine/Plugins/Animation/IKRig/Source/IKRigEditor/Public/RigEditor/IKRigController.h
@@ -119,6 +119,9 @@ public:
 	void ResetGoalTransforms() const;
 	/** Ensure that the given name adheres to required standards for Goal names (no special characters etc..)*/
 	static void SanitizeGoalName(FString& InOutName);
+	/** Add a suffix as needed to ensure the Goal name is unique */
+	FName GetUniqueGoalName(const FName& NameToMakeUnique) const;
+	
 	/** END Goals */
 
 	/** Bone Settings
