Removing pluralization on missing goal warning in transform solver.

diff --git a/Engine/Plugins/Animation/IKRig/Source/IKRig/Private/Solvers/IKRig_SetTransform.cpp b/Engine/Plugins/Animation/IKRig/Source/IKRig/Private/Solvers/IKRig_SetTransform.cpp
--- a/Engine/Plugins/Animation/IKRig/Source/IKRig/Private/Solvers/IKRig_SetTransform.cpp
+++ b/Engine/Plugins/Animation/IKRig/Source/IKRig/Private/Solvers/IKRig_SetTransform.cpp
@@ -76,7 +76,7 @@ bool UIKRig_SetTransform::GetWarningMessage(FText& OutWarningMessage) const
 {
 	if (RootBone == NAME_None)
 	{
-		OutWarningMessage = LOCTEXT("MissingGoal", "Missing goals.");
+		OutWarningMessage = LOCTEXT("MissingGoal", "Missing goal.");
 		return true;
 	}
 	return false;
