Capped influence multiplier of Body Mover solver.

diff --git a/Engine/Plugins/Animation/IKRig/Source/IKRig/Public/Solvers/IKRig_BodyMover.h b/Engine/Plugins/Animation/IKRig/Source/IKRig/Public/Solvers/IKRig_BodyMover.h
--- a/Engine/Plugins/Animation/IKRig/Source/IKRig/Public/Solvers/IKRig_BodyMover.h
+++ b/Engine/Plugins/Animation/IKRig/Source/IKRig/Public/Solvers/IKRig_BodyMover.h
@@ -22,7 +22,7 @@ public:
 	FName BoneName;
 
 	/** Scale the influence this effector has on the body. Range is 0-10. Default is 1.0. */
-	UPROPERTY(EditAnywhere, Category = "Body Mover Effector", meta = (ClampMin = "0", ClampMa = "10", UIMin = "0.0", UIMax = "10.0"))
+	UPROPERTY(EditAnywhere, Category = "Body Mover Effector", meta = (ClampMin = "0", ClampMax = "10", UIMin = "0.0", UIMax = "10.0"))
 	float InfluenceMultiplier = 1.0f;
 };
 
