Populate Preview Mesh properties in IK Retargeter by default post assignment of IKRigs

diff --git a/Engine/Plugins/Animation/IKRig/Source/IKRigEditor/Private/RetargetEditor/IKRetargetEditor.cpp b/Engine/Plugins/Animation/IKRig/Source/IKRigEditor/Private/RetargetEditor/IKRetargetEditor.cpp
--- a/Engine/Plugins/Animation/IKRig/Source/IKRigEditor/Private/RetargetEditor/IKRetargetEditor.cpp
+++ b/Engine/Plugins/Animation/IKRig/Source/IKRigEditor/Private/RetargetEditor/IKRetargetEditor.cpp
@@ -368,20 +368,37 @@ void FIKRetargetEditor::HandleDetailsCreated(const TSharedRef<class IDetailsView
 
 void FIKRetargetEditor::OnFinishedChangingDetails(const FPropertyChangedEvent& PropertyChangedEvent)
 {
+	const UIKRetargeterController* AssetController = EditorController->AssetController;
+
+	// determine which properties were modified
 	const bool bSourceIKRigChanged = PropertyChangedEvent.GetPropertyName() == UIKRetargeter::GetSourceIKRigPropertyName();
 	const bool bTargetIKRigChanged = PropertyChangedEvent.GetPropertyName() == UIKRetargeter::GetTargetIKRigPropertyName();
 	const bool bSourcePreviewChanged = PropertyChangedEvent.GetPropertyName() == UIKRetargeter::GetSourcePreviewMeshPropertyName();
 	const bool bTargetPreviewChanged = PropertyChangedEvent.GetPropertyName() == UIKRetargeter::GetTargetPreviewMeshPropertyName();
 
+	// if no override target mesh has been specified, update the override to reflect the mesh in the ik rig asset
+	if (bTargetIKRigChanged)
+	{
+		AssetController->OnTargetIKRigChanged();
+	}
+
+	// if no override source mesh has been specified, update the override to reflect the mesh in the ik rig asset
+	if (bSourceIKRigChanged)
+	{
+		AssetController->OnSourceIKRigChanged();
+	}
+
+	// if either IK Rig asset has been modified, rebind and refresh UI
 	if (bTargetIKRigChanged || bSourceIKRigChanged)
 	{
 		EditorController->ClearOutputLog();
-		EditorController->BindToIKRigAsset(EditorController->AssetController->GetAsset()->GetTargetIKRigWriteable());
-		EditorController->BindToIKRigAsset(EditorController->AssetController->GetAsset()->GetSourceIKRigWriteable());
+		EditorController->BindToIKRigAsset(AssetController->GetAsset()->GetTargetIKRigWriteable());
+		EditorController->BindToIKRigAsset(AssetController->GetAsset()->GetSourceIKRigWriteable());
 		EditorController->AssetController->CleanChainMapping();
 		EditorController->AssetController->AutoMapChains();
 	}
-	
+
+	// if either the source or target meshes are possibly modified, update scene components, anim instance and UI
 	if (bTargetIKRigChanged || bSourceIKRigChanged || bTargetPreviewChanged || bSourcePreviewChanged)
 	{
 		EditorController->ClearOutputLog();
diff --git a/Engine/Plugins/Animation/IKRig/Source/IKRigEditor/Private/RetargetEditor/IKRetargeterController.cpp b/Engine/Plugins/Animation/IKRig/Source/IKRigEditor/Private/RetargetEditor/IKRetargeterController.cpp
--- a/Engine/Plugins/Animation/IKRig/Source/IKRigEditor/Private/RetargetEditor/IKRetargeterController.cpp
+++ b/Engine/Plugins/Animation/IKRig/Source/IKRigEditor/Private/RetargetEditor/IKRetargeterController.cpp
@@ -101,6 +101,22 @@ const UIKRigDefinition* UIKRetargeterController::GetTargetIKRig() const
 	return Asset->GetTargetIKRig();
 }
 
+void UIKRetargeterController::OnTargetIKRigChanged() const
+{
+	if (const UIKRigDefinition* IKRig = GetTargetIKRig())
+	{
+		Asset->TargetPreviewMesh = IKRig->PreviewSkeletalMesh;
+	}
+}
+
+void UIKRetargeterController::OnSourceIKRigChanged() const
+{
+	if (const UIKRigDefinition* IKRig = GetSourceIKRig())
+	{
+		Asset->SourcePreviewMesh = IKRig->PreviewSkeletalMesh;
+	}
+}
+
 FName UIKRetargeterController::GetSourceRootBone() const
 {
 	const UIKRigDefinition* SourceIKRig = Asset->GetSourceIKRig();
diff --git a/Engine/Plugins/Animation/IKRig/Source/IKRigEditor/Public/RetargetEditor/IKRetargeterController.h b/Engine/Plugins/Animation/IKRig/Source/IKRigEditor/Public/RetargetEditor/IKRetargeterController.h
--- a/Engine/Plugins/Animation/IKRig/Source/IKRigEditor/Public/RetargetEditor/IKRetargeterController.h
+++ b/Engine/Plugins/Animation/IKRig/Source/IKRigEditor/Public/RetargetEditor/IKRetargeterController.h
@@ -48,6 +48,10 @@ public:
 	const UIKRigDefinition* GetSourceIKRig() const;
 	/** Get source IK Rig asset */
 	const UIKRigDefinition* GetTargetIKRig() const;
+	/** Set the target preview mesh based on the mesh in the target IK Rig asset */
+	void OnTargetIKRigChanged() const;
+	/** Set the source preview mesh based on the mesh in the source IK Rig asset */
+	void OnSourceIKRigChanged() const;
 
 	/** Get name of the Root bone used for retargeting the Source skeleton. */
 	FName GetSourceRootBone() const;
