Fixing IKRig plugin logging to be processor-centric to get unique warnings per instance.

diff --git a/Engine/Plugins/Animation/IKRig/Source/IKRig/Private/AnimNodes/AnimNode_RetargetPoseFromMesh.cpp b/Engine/Plugins/Animation/IKRig/Source/IKRig/Private/AnimNodes/AnimNode_RetargetPoseFromMesh.cpp
--- a/Engine/Plugins/Animation/IKRig/Source/IKRig/Private/AnimNodes/AnimNode_RetargetPoseFromMesh.cpp
+++ b/Engine/Plugins/Animation/IKRig/Source/IKRig/Private/AnimNodes/AnimNode_RetargetPoseFromMesh.cpp
@@ -11,6 +11,11 @@ void FAnimNode_RetargetPoseFromMesh::Initialize_AnyThread(const FAnimationInitia
 
 	// Initial update of the node, so we dont have a frame-delay on setup
 	GetEvaluateGraphExposedInputs().Execute(Context);
+
+	if (Processor.IsNull() && IsInGameThread())
+	{
+		Processor = NewObject<UIKRetargetProcessor>(Context.AnimInstanceProxy->GetSkelMeshComponent());	
+	}
 }
 
 void FAnimNode_RetargetPoseFromMesh::CacheBones_AnyThread(const FAnimationCacheBonesContext& Context)
diff --git a/Engine/Plugins/Animation/IKRig/Source/IKRig/Private/IKRigDefinition.cpp b/Engine/Plugins/Animation/IKRig/Source/IKRig/Private/IKRigDefinition.cpp
--- a/Engine/Plugins/Animation/IKRig/Source/IKRig/Private/IKRigDefinition.cpp
+++ b/Engine/Plugins/Animation/IKRig/Source/IKRig/Private/IKRigDefinition.cpp
@@ -272,20 +272,6 @@ void UIKRigDefinition::Serialize(FArchive& Ar)
 	Ar.UsingCustomVersion(FIKRigObjectVersion::GUID);
 }
 
-void UIKRigDefinition::PostLoad()
-{
-	Super::PostLoad();
-
-	Log.SetLogTarget(GetUniqueIDAsName(), false);
-}
-
-void UIKRigDefinition::PostInitProperties()
-{
-	UObject::PostInitProperties();
-
-	Log.SetLogTarget(GetUniqueIDAsName(), false);
-}
-
 const FBoneChain* UIKRigDefinition::GetRetargetChainByName(FName ChainName) const
 {
 	for (const FBoneChain& Chain : RetargetDefinition.BoneChains)
@@ -309,12 +295,4 @@ USkeletalMesh* UIKRigDefinition::GetPreviewMesh() const
 	return PreviewSkeletalMesh.LoadSynchronous();
 }
 
-FName UIKRigDefinition::GetUniqueIDAsName() const
-{
-	LLM_SCOPE_BYNAME(TEXT("Animation/IKRig"));
-	
-	static const FString IKRigIDPrefix("IKRig_");
-	return FName(IKRigIDPrefix + FString::FromInt(GetUniqueID()));
-}
-
 #undef LOCTEXT_NAMESPACE
diff --git a/Engine/Plugins/Animation/IKRig/Source/IKRig/Private/IKRigLogger.cpp b/Engine/Plugins/Animation/IKRig/Source/IKRig/Private/IKRigLogger.cpp
--- a/Engine/Plugins/Animation/IKRig/Source/IKRig/Private/IKRigLogger.cpp
+++ b/Engine/Plugins/Animation/IKRig/Source/IKRig/Private/IKRigLogger.cpp
@@ -6,10 +6,9 @@
 
 DEFINE_LOG_CATEGORY(LogIKRig);
 
-void FIKRigLogger::SetLogTarget(const FName InLogName, bool bInSuppressWarnings)
+void FIKRigLogger::SetLogTarget(const FName InLogName)
 {
 	LogName = InLogName;
-	bWarningsSuppressed = bInSuppressWarnings;
 }
 
 FName FIKRigLogger::GetLogTarget() const
@@ -19,37 +18,25 @@ FName FIKRigLogger::GetLogTarget() const
 	
 void FIKRigLogger::LogError(const FText& Message) const
 {
-	// print to the global output log
-	UE_LOG(LogIKRig, Error, TEXT("%s"), *Message.ToString());
-
-	// print to the output log in the asset editor
-	FMessageLog MessageLog(LogName);
-	MessageLog.Error(Message);
+	FMessageLog(LogName).SuppressLoggingToOutputLog(true).Error(Message);
+	Errors.Add(Message);
 }
 
 void FIKRigLogger::LogWarning(const FText& Message) const
 {
-	if (bWarningsSuppressed)
-	{
-		return;
-	}
-
-	// print to the global output log
-	UE_LOG(LogIKRig, Warning, TEXT("%s"), *Message.ToString());
+	FMessageLog(LogName).SuppressLoggingToOutputLog(true).Warning(Message);
+	Warnings.Add(Message);
+}
 
-	// print to the output log in the asset editor
-	FMessageLog MessageLog(LogName);
-	MessageLog.Warning(Message);
+void FIKRigLogger::LogInfo(const FText& Message) const
+{
+	FMessageLog(LogName).SuppressLoggingToOutputLog(true).Info(Message);
+	Messages.Add(Message);
 }
 
-void FIKRigLogger::LogEditorMessage(const FText& Message) const
+void FIKRigLogger::Clear() const
 {
-	if (bWarningsSuppressed)
-	{
-		return;
-	}
-
-	// print to the output log in the asset editor
-	FMessageLog MessageLog(LogName);
-	MessageLog.Info(Message);
+	Errors.Empty();
+	Warnings.Empty();
+	Messages.Empty();
 }
diff --git a/Engine/Plugins/Animation/IKRig/Source/IKRig/Private/IKRigProcessor.cpp b/Engine/Plugins/Animation/IKRig/Source/IKRig/Private/IKRigProcessor.cpp
--- a/Engine/Plugins/Animation/IKRig/Source/IKRig/Private/IKRigProcessor.cpp
+++ b/Engine/Plugins/Animation/IKRig/Source/IKRig/Private/IKRigProcessor.cpp
@@ -8,6 +8,12 @@
 
 #define LOCTEXT_NAMESPACE "IKRigProcessor"
 
+UIKRigProcessor::UIKRigProcessor()
+{
+	const FName LogName = FName("IKRig_",GetUniqueID());
+	Log.SetLogTarget(LogName);
+}
+
 void UIKRigProcessor::Initialize(const UIKRigDefinition* InRigAsset, const USkeletalMesh* SkeletalMesh)
 {
 	// we instantiate UObjects here which MUST be done on game thread...
@@ -30,15 +36,15 @@ void UIKRigProcessor::Initialize(const UIKRigDefinition* InRigAsset, const USkel
 	
 	if (InRigAsset->Skeleton.BoneNames.IsEmpty())
 	{
-		InRigAsset->Log.LogError(FText::Format(
+		Log.LogError(FText::Format(
 			LOCTEXT("NoSkeleton", "Trying to initialize IK Rig, '{0}' that has no skeleton."),
 			FText::FromString(InRigAsset->GetName())));
 		return;
 	}
 
-	if (!UIKRigProcessor::IsIKRigCompatibleWithSkeleton(InRigAsset, SkeletalMesh))
+	if (!UIKRigProcessor::IsIKRigCompatibleWithSkeleton(InRigAsset, SkeletalMesh, &Log))
 	{
-		InRigAsset->Log.LogError(FText::Format(
+		Log.LogError(FText::Format(
 			LOCTEXT("SkeletonMissingRequiredBones", "Trying to initialize IKRig, '{0}' with a Skeleton that is missing required bones. See prior warnings."),
 			FText::FromString(InRigAsset->GetName())));
 		return;
@@ -64,7 +70,7 @@ void UIKRigProcessor::Initialize(const UIKRigDefinition* InRigAsset, const USkel
 		// validate that the skeleton we are trying to solve this goal on contains the bone the goal expects
 		if (NewGoalBone.BoneIndex == INDEX_NONE)
 		{
-			InRigAsset->Log.LogError(FText::Format(
+			Log.LogError(FText::Format(
 				LOCTEXT("MissingGoalBone", "IK Rig, '{0}' has a Goal, '{1}' that references an unknown bone, '{2}'. Cannot evaluate."),
 				FText::FromString(InRigAsset->GetName()), FText::FromName(EffectorGoal->GoalName), FText::FromName(EffectorGoal->BoneName) ));
 			return;
@@ -76,7 +82,7 @@ void UIKRigProcessor::Initialize(const UIKRigDefinition* InRigAsset, const USkel
 		{
 			if (Bone->BoneName != NewGoalBone.BoneName)
 			{
-				InRigAsset->Log.LogError(FText::Format(
+				Log.LogError(FText::Format(
 				LOCTEXT("DuplicateGoal", "IK Rig, '{0}' has a Goal, '{1}' that references different bones in different solvers, '{2}' and '{3}'. Cannot evaluate."),
                 FText::FromString(InRigAsset->GetName()),
                 FText::FromName(EffectorGoal->GoalName),
@@ -98,7 +104,7 @@ void UIKRigProcessor::Initialize(const UIKRigDefinition* InRigAsset, const USkel
 		if (!IKRigSolver)
 		{
 			// this can happen if asset references deleted IK Solver type which should only happen during development (if at all)
-			InRigAsset->Log.LogWarning(FText::Format(
+			Log.LogWarning(FText::Format(
 				LOCTEXT("UnknownSolver", "IK Rig, '{0}' has null/unknown solver in it. Please remove it."),
 				FText::FromString(InRigAsset->GetName())));
 			continue;
@@ -112,14 +118,17 @@ void UIKRigProcessor::Initialize(const UIKRigDefinition* InRigAsset, const USkel
 		Solvers.Add(Solver);
 	}
 
-	InRigAsset->Log.LogEditorMessage(FText::Format(
+	Log.LogInfo(FText::Format(
 				LOCTEXT("SuccessfulInit", "IK Rig, '{0}' ready to run on {1}."),
 				FText::FromString(InRigAsset->GetName()), FText::FromString(SkeletalMesh->GetName())));
 	
 	bInitialized = true;
 }
 
-bool UIKRigProcessor::IsIKRigCompatibleWithSkeleton(const UIKRigDefinition* InRigAsset,	const FIKRigInputSkeleton& InputSkeleton)
+bool UIKRigProcessor::IsIKRigCompatibleWithSkeleton(
+	const UIKRigDefinition* InRigAsset,
+	const FIKRigInputSkeleton& InputSkeleton,
+	const FIKRigLogger* Log)
 {
 	// first we validate that all the required bones are in the input skeleton...
 	
@@ -147,10 +156,13 @@ bool UIKRigProcessor::IsIKRigCompatibleWithSkeleton(const UIKRigDefinition* InRi
 	{
 		if (!InputSkeleton.BoneNames.Contains(RequiredBone))
 		{
-			InRigAsset->Log.LogError(FText::Format(
+			if (Log)
+			{
+				Log->LogError(FText::Format(
 			LOCTEXT("MissingBone", "IK Rig, '{0}' is missing a required bone, '{1}' in the Skeletal Mesh."),
-			FText::FromString(InRigAsset->GetName()),
-			FText::FromName(RequiredBone)));
+				FText::FromString(InRigAsset->GetName()),
+				FText::FromName(RequiredBone)));
+			}
 			
 			bAllRequiredBonesFound = false;
 		}
@@ -182,10 +194,13 @@ bool UIKRigProcessor::IsIKRigCompatibleWithSkeleton(const UIKRigDefinition* InRi
 			{
 				bAllParentsValid = false;
 
-				InRigAsset->Log.LogError(FText::Format(
-			LOCTEXT("InvalidParent", "IK Rig is running on a skeleton with a required bone, '{0}', that expected to have a valid parent. The expected parent was, '{1}'."),
-				FText::FromName(RequiredBone),
-				FText::FromName(AssetParentName)));
+				if (Log)
+				{
+					Log->LogError(FText::Format(
+					LOCTEXT("InvalidParent", "IK Rig is running on a skeleton with a required bone, '{0}', that expected to have a valid parent. The expected parent was, '{1}'."),
+					FText::FromName(RequiredBone),
+					FText::FromName(AssetParentName)));
+				}
 				
 				continue;
 			}
@@ -193,12 +208,16 @@ bool UIKRigProcessor::IsIKRigCompatibleWithSkeleton(const UIKRigDefinition* InRi
 			const FName& InputParentName = InputSkeleton.BoneNames[InputParentIndex];
 			if (AssetParentName != InputParentName)
 			{
-				// we only warn about this, because it may be nice not to have the exact same hierarchy
-				InRigAsset->Log.LogWarning(FText::Format(
-				LOCTEXT("UnexpectedParent", "IK Rig is running on a skeleton with a required bone, '{0}', that has a different parent '{1}'. The expected parent was, '{2}'."),
-				FText::FromName(RequiredBone),
-				FText::FromName(InputParentName),
-				FText::FromName(AssetParentName)));
+				if (Log)
+				{
+					// we only warn about this, because it may be nice not to have the exact same hierarchy
+					Log->LogWarning(FText::Format(
+					LOCTEXT("UnexpectedParent", "IK Rig is running on a skeleton with a required bone, '{0}', that has a different parent '{1}'. The expected parent was, '{2}'."),
+					FText::FromName(RequiredBone),
+					FText::FromName(InputParentName),
+					FText::FromName(AssetParentName)));
+				}
+				
 				continue;
 			}
 		}
diff --git a/Engine/Plugins/Animation/IKRig/Source/IKRig/Private/Retargeter/IKRetargetProcessor.cpp b/Engine/Plugins/Animation/IKRig/Source/IKRig/Private/Retargeter/IKRetargetProcessor.cpp
--- a/Engine/Plugins/Animation/IKRig/Source/IKRig/Private/Retargeter/IKRetargetProcessor.cpp
+++ b/Engine/Plugins/Animation/IKRig/Source/IKRig/Private/Retargeter/IKRetargetProcessor.cpp
@@ -1256,6 +1256,12 @@ void FRootRetargeter::DecodePose(TArray<FTransform>& OutTargetGlobalPose) const
 	TargetRootTransform.SetRotation(Rotation);
 }
 
+UIKRetargetProcessor::UIKRetargetProcessor()
+{
+	const FName LogName = FName("IKRetarget_",GetUniqueID());
+	Log.SetLogTarget(LogName);
+}
+
 void UIKRetargetProcessor::Initialize(
 		USkeletalMesh* SourceSkeletalMesh,
 		USkeletalMesh* TargetSkeletalMesh,
@@ -1278,27 +1284,27 @@ void UIKRetargetProcessor::Initialize(
 	// check prerequisite assets
 	if (!SourceSkeletalMesh)
 	{
-		RetargeterAsset->Log.LogError(LOCTEXT("MissingSourceMesh", "IK Retargeter unable to initialize. Missing source Skeletal Mesh asset."));
+		Log.LogError(LOCTEXT("MissingSourceMesh", "IK Retargeter unable to initialize. Missing source Skeletal Mesh asset."));
 		return;
 	}
 	if (!TargetSkeletalMesh)
 	{
-		RetargeterAsset->Log.LogError(LOCTEXT("MissingTargetMesh", "IK Retargeter unable to initialize. Missing target Skeletal Mesh asset."));
+		Log.LogError(LOCTEXT("MissingTargetMesh", "IK Retargeter unable to initialize. Missing target Skeletal Mesh asset."));
 		return;
 	}
 	if (!RetargeterAsset->GetSourceIKRig())
 	{
-		RetargeterAsset->Log.LogError(LOCTEXT("MissingSourceIKRig", "IK Retargeter unable to initialize. Missing source IK Rig asset."));
+		Log.LogError(LOCTEXT("MissingSourceIKRig", "IK Retargeter unable to initialize. Missing source IK Rig asset."));
 		return;
 	}
 	if (!RetargeterAsset->GetTargetIKRig())
 	{
-		RetargeterAsset->Log.LogError(LOCTEXT("MissingTargetIKRig", "IK Retargeter unable to initialize. Missing target IK Rig asset."));
+		Log.LogError(LOCTEXT("MissingTargetIKRig", "IK Retargeter unable to initialize. Missing target IK Rig asset."));
 		return;
 	}
 	if (!RetargeterAsset->GetCurrentRetargetPose())
 	{
-		RetargeterAsset->Log.LogError(LOCTEXT("MissingRetargetPose", "IK Retargeter unable to initialize. Missing retarget pose."));
+		Log.LogError(LOCTEXT("MissingRetargetPose", "IK Retargeter unable to initialize. Missing retarget pose."));
 		return;
 	}
 	
@@ -1318,7 +1324,7 @@ void UIKRetargetProcessor::Initialize(
 	if (!bAtLeastOneValidBoneChainPair)
 	{
 		// couldn't match up any BoneChain pairs, no limb retargeting possible
-		RetargeterAsset->Log.LogWarning( FText::Format(
+		Log.LogWarning( FText::Format(
 			LOCTEXT("NoMappedChains", "IK Retargeter unable to map any bone chains between source, {0} and target, {1}"),
 			FText::FromString(SourceSkeleton.SkeletalMesh->GetName()), FText::FromString(TargetSkeleton.SkeletalMesh->GetName())));
 	}
@@ -1328,7 +1334,7 @@ void UIKRetargetProcessor::Initialize(
 	if (!bIKRigInitialized)
 	{
 		// couldn't initialize the IK Rig, we don't disable the retargeter in this case, just warn the user
-		RetargeterAsset->Log.LogWarning( FText::Format(
+		Log.LogWarning( FText::Format(
 			LOCTEXT("CouldNotInitializeIKRig", "IK Retargeter was unable to initialize the IK Rig, {0} for the Skeletal Mesh {1}. See previous warnings."),
 			FText::FromString(RetargeterAsset->GetTargetIKRig()->GetName()), FText::FromString(TargetSkeleton.SkeletalMesh->GetName())));
 	}
@@ -1337,7 +1343,7 @@ void UIKRetargetProcessor::Initialize(
 	if (bRootsInitialized && bAtLeastOneValidBoneChainPair)
 	{
 		// confirm for the user that the IK Rig was initialized successfully
-		RetargeterAsset->Log.LogEditorMessage(FText::Format(
+		Log.LogInfo(FText::Format(
 				LOCTEXT("SuccessfulInit", "Success! The IK Retargeter is ready to transfer animation from the source, {0} to the target, {1}"),
 				FText::FromString(SourceSkeleton.SkeletalMesh->GetName()), FText::FromString(TargetSkeleton.SkeletalMesh->GetName())));
 	}
@@ -1349,20 +1355,20 @@ bool UIKRetargetProcessor::InitializeRoots()
 {
 	// initialize root encoder
 	const FName SourceRootBoneName = RetargeterAsset->GetSourceIKRig()->GetRetargetRoot();
-	const bool bRootEncoderInit = RootRetargeter.InitializeSource(SourceRootBoneName, SourceSkeleton, RetargeterAsset->Log);
+	const bool bRootEncoderInit = RootRetargeter.InitializeSource(SourceRootBoneName, SourceSkeleton, Log);
 	if (!bRootEncoderInit)
 	{
-		RetargeterAsset->Log.LogWarning( FText::Format(
+		Log.LogWarning( FText::Format(
 			LOCTEXT("NoSourceRoot", "IK Retargeter unable to initialize source root, '{0}' on skeletal mesh: '{1}'"),
 			FText::FromName(SourceRootBoneName), FText::FromString(SourceSkeleton.SkeletalMesh->GetName())));
 	}
 
 	// initialize root decoder
 	const FName TargetRootBoneName = RetargeterAsset->GetTargetIKRig()->GetRetargetRoot();
-	const bool bRootDecoderInit = RootRetargeter.InitializeTarget(TargetRootBoneName, TargetSkeleton, RetargeterAsset->Log);
+	const bool bRootDecoderInit = RootRetargeter.InitializeTarget(TargetRootBoneName, TargetSkeleton, Log);
 	if (!bRootDecoderInit)
 	{
-		RetargeterAsset->Log.LogWarning( FText::Format(
+		Log.LogWarning( FText::Format(
 			LOCTEXT("NoTargetRoot", "IK Retargeter unable to initialize target root, '{0}' on skeletal mesh: '{1}'"),
 			FText::FromName(TargetRootBoneName), FText::FromString(TargetSkeleton.SkeletalMesh->GetName())));
 	}
@@ -1385,7 +1391,7 @@ bool UIKRetargetProcessor::InitializeBoneChainPairs()
 		const FBoneChain* TargetBoneChain = RetargeterAsset->GetTargetIKRig()->GetRetargetChainByName(ChainMap->TargetChain);
 		if (!TargetBoneChain)
 		{
-			RetargeterAsset->Log.LogWarning( FText::Format(
+			Log.LogWarning( FText::Format(
 			LOCTEXT("MissingTargetChain", "IK Retargeter missing target bone chain: {0}. Please update the mapping."),
 			FText::FromString(ChainMap->TargetChain.ToString())));
 			continue;
@@ -1401,7 +1407,7 @@ bool UIKRetargetProcessor::InitializeBoneChainPairs()
 		const FBoneChain* SourceBoneChain = RetargeterAsset->GetSourceIKRig()->GetRetargetChainByName(ChainMap->SourceChain);
 		if (!SourceBoneChain)
 		{
-			RetargeterAsset->Log.LogWarning( FText::Format(
+			Log.LogWarning( FText::Format(
 			LOCTEXT("MissingSourceChain", "IK Retargeter missing source bone chain: {0}"),
 			FText::FromString(ChainMap->SourceChain.ToString())));
 			continue;
@@ -1409,14 +1415,14 @@ bool UIKRetargetProcessor::InitializeBoneChainPairs()
 
 		// all chains are loaded as FK (giving IK better starting pose)
 		FRetargetChainPairFK ChainPair;
-		if (ChainPair.Initialize(ChainMap, *SourceBoneChain, *TargetBoneChain, SourceSkeleton, TargetSkeleton, RetargeterAsset->Log))
+		if (ChainPair.Initialize(ChainMap, *SourceBoneChain, *TargetBoneChain, SourceSkeleton, TargetSkeleton, Log))
 		{
 			ChainPairsFK.Add(ChainPair);
 		}
 		
 		// load IK chain
 		FRetargetChainPairIK ChainPairIK;
-		if (ChainPairIK.Initialize(ChainMap, *SourceBoneChain, *TargetBoneChain, SourceSkeleton, TargetSkeleton, RetargeterAsset->Log))
+		if (ChainPairIK.Initialize(ChainMap, *SourceBoneChain, *TargetBoneChain, SourceSkeleton, TargetSkeleton, Log))
 		{
 			ChainPairsIK.Add(ChainPairIK);
 		}
@@ -1485,7 +1491,7 @@ bool UIKRetargetProcessor::InitializeIKRig(UObject* Outer, const USkeletalMesh*
 		// does the IK rig have the IK goal this bone chain requires?
 		if (!IKRigProcessor->GetGoalContainer().FindGoalByName(ChainPair.IKGoalName))
 		{
-			RetargeterAsset->Log.LogError( FText::Format(
+			Log.LogError( FText::Format(
 			LOCTEXT("TargetIKBoneNotInSolver", "IK Retargeter has target bone chain, {0} that references an IK Goal, {1} that is not present in any of the solvers in the IK Rig asset."),
 			FText::FromName(ChainPair.TargetBoneChainName), FText::FromName(ChainPair.IKGoalName)));
 			return false;
diff --git a/Engine/Plugins/Animation/IKRig/Source/IKRig/Private/Retargeter/IKRetargeter.cpp b/Engine/Plugins/Animation/IKRig/Source/IKRig/Private/Retargeter/IKRetargeter.cpp
--- a/Engine/Plugins/Animation/IKRig/Source/IKRig/Private/Retargeter/IKRetargeter.cpp
+++ b/Engine/Plugins/Animation/IKRig/Source/IKRig/Private/Retargeter/IKRetargeter.cpp
@@ -85,17 +85,6 @@ void UIKRetargeter::PostLoad()
 	{
 		CurrentRetargetPose = UIKRetargeter::GetDefaultPoseName();
 	}
-
-	// set the log target
-	Log.SetLogTarget(GetUniqueIDAsName(), false);
-}
-
-void UIKRetargeter::PostInitProperties()
-{
-	UObject::PostInitProperties();
-
-	// set the log target
-	Log.SetLogTarget(GetUniqueIDAsName(), false);
 }
 
 void FIKRetargetPose::SetBoneRotationOffset(FName BoneName, FQuat RotationDelta, const FIKRigSkeleton& Skeleton)
@@ -127,13 +116,6 @@ void FIKRetargetPose::SortHierarchically(const FIKRigSkeleton& Skeleton)
 	});
 }
 
-FName UIKRetargeter::GetUniqueIDAsName() const
-{
-	LLM_SCOPE_BYNAME(TEXT("Animation/IKRig"));
-	static const FString IKRetargetIDPrefix("IKRetargeter_");
-	return FName(IKRetargetIDPrefix + FString::FromInt(GetUniqueID()));
-}
-
 const FName UIKRetargeter::GetDefaultPoseName()
 {
 	static const FName DefaultPoseName = "Default Pose";
diff --git a/Engine/Plugins/Animation/IKRig/Source/IKRig/Public/IKRigDefinition.h b/Engine/Plugins/Animation/IKRig/Source/IKRig/Public/IKRigDefinition.h
--- a/Engine/Plugins/Animation/IKRig/Source/IKRig/Public/IKRigDefinition.h
+++ b/Engine/Plugins/Animation/IKRig/Source/IKRig/Public/IKRigDefinition.h
@@ -214,8 +214,6 @@ public:
 
 	/** UObject */
 	virtual void Serialize(FArchive& Ar) override;
-	virtual void PostLoad() override;
-	virtual void PostInitProperties() override;
 	/** END UObject */
 	
 	/** hierarchy and bone-pose transforms */
@@ -242,12 +240,6 @@ public:
 	virtual USkeletalMesh* GetPreviewMesh() const override;
 	/** END IInterface_PreviewMeshProvider interface */
 
-	/** Get the unique ID of this UObject as a name. */
-	FName GetUniqueIDAsName() const;
-
-	/** logging system */
-	FIKRigLogger Log;
-
 #if WITH_EDITOR
 	/* Get name of Preview Mesh property */
 	static const FName GetPreviewMeshPropertyName();
diff --git a/Engine/Plugins/Animation/IKRig/Source/IKRig/Public/IKRigLogger.h b/Engine/Plugins/Animation/IKRig/Source/IKRig/Public/IKRigLogger.h
--- a/Engine/Plugins/Animation/IKRig/Source/IKRig/Public/IKRigLogger.h
+++ b/Engine/Plugins/Animation/IKRig/Source/IKRig/Public/IKRigLogger.h
@@ -9,8 +9,8 @@ DECLARE_LOG_CATEGORY_EXTERN(LogIKRig, Warning, All);
 
 struct IKRIG_API FIKRigLogger
 {
-	/** Set the name of the log to output messages to, and whether to suppress warnings */
-	void SetLogTarget(const FName InLogName, bool bInSuppressWarnings);
+	/** Set the name of the log to output messages to */
+	void SetLogTarget(const FName InLogName);
 	/** Get the name this log is currently outputting to */
 	FName GetLogTarget() const;
 	/** Log a warning message to display to user. */
@@ -18,7 +18,13 @@ struct IKRIG_API FIKRigLogger
 	/** Log a warning message to display to user. */
 	void LogWarning(const FText& Message) const;
 	/** Log a message to display to editor output log. */
-	void LogEditorMessage(const FText& Message) const;
+	void LogInfo(const FText& Message) const;
+	/** clear all the stored messages */
+	void Clear() const;
+	/** get a list of messages that have been logged since last Clear() */
+	const TArray<FText>& GetErrors() const { return Errors; };
+	const TArray<FText>& GetWarnings() const { return Warnings; };
+	const TArray<FText>& GetMessages() const { return Messages; };
 
 private:
 	/** the name of the output log this logger will send messages to
@@ -27,7 +33,9 @@ private:
 	 * that is being edited. Therefore we name the log using the unique ID of the UObject itself (valid for lifetime of UObject between loads)
 	 */
 	FName LogName;
-	
-	/** when true, warnings will not be printed out. */
-	bool bWarningsSuppressed = false;
-};
\ No newline at end of file
+
+	/** store messages here so we can show them during anim graph compilation */
+	mutable TArray<FText> Errors;
+	mutable TArray<FText> Warnings;
+	mutable TArray<FText> Messages;
+};
diff --git a/Engine/Plugins/Animation/IKRig/Source/IKRig/Public/IKRigProcessor.h b/Engine/Plugins/Animation/IKRig/Source/IKRig/Public/IKRigProcessor.h
--- a/Engine/Plugins/Animation/IKRig/Source/IKRig/Public/IKRigProcessor.h
+++ b/Engine/Plugins/Animation/IKRig/Source/IKRig/Public/IKRigProcessor.h
@@ -39,7 +39,9 @@ public:
 	* 5. Copy output transforms with CopyOutputGlobalPoseToArray()
 	* 
 	*/
-	 
+
+	UIKRigProcessor();
+	
 	/** setup a new processor to run the given IKRig asset
 	 *  NOTE!! this function creates new UObjects and consequently MUST be called from the main thread!!
 	 *  @param InRigAsset - the IK Rig defining the collection of solvers to execute and all the rig settings
@@ -86,12 +88,18 @@ public:
 	FIKRigSkeleton& GetSkeleton();
 
 	/** Used to determine if the IK Rig asset is compatible with a given skeleton. */
-	static bool IsIKRigCompatibleWithSkeleton(const UIKRigDefinition* InRigAsset, const FIKRigInputSkeleton& InputSkeleton);
+	static bool IsIKRigCompatibleWithSkeleton(
+		const UIKRigDefinition* InRigAsset,
+		const FIKRigInputSkeleton& InputSkeleton,
+		const FIKRigLogger* Log);
 
 	bool IsInitialized() const { return bInitialized; };
 
 	void SetNeedsInitialized();
 
+	/** logging system */
+	FIKRigLogger Log;
+
 #if WITH_EDITOR
 	/** Used to propagate setting values from the source asset at runtime (settings that do not require re-initialization) */
 	void CopyAllInputsFromSourceAssetAtRuntime(const UIKRigDefinition* SourceAsset);
diff --git a/Engine/Plugins/Animation/IKRig/Source/IKRig/Public/Retargeter/IKRetargetProcessor.h b/Engine/Plugins/Animation/IKRig/Source/IKRig/Public/Retargeter/IKRetargetProcessor.h
--- a/Engine/Plugins/Animation/IKRig/Source/IKRig/Public/Retargeter/IKRetargetProcessor.h
+++ b/Engine/Plugins/Animation/IKRig/Source/IKRig/Public/Retargeter/IKRetargetProcessor.h
@@ -473,6 +473,8 @@ class IKRIG_API UIKRetargetProcessor : public UObject
 	GENERATED_BODY()
 
 public:
+
+	UIKRetargetProcessor();
 	
 	/**
 	* Initialize the retargeter to enable running it.
@@ -525,6 +527,9 @@ public:
 	/** Reset the IK planting state. */
 	void ResetPlanting();
 
+	/** logging system */
+	FIKRigLogger Log;
+
 #if WITH_EDITOR
 	/** Set that this processor needs to be reinitialized. */
 	void SetNeedsInitialized();
diff --git a/Engine/Plugins/Animation/IKRig/Source/IKRig/Public/Retargeter/IKRetargeter.h b/Engine/Plugins/Animation/IKRig/Source/IKRig/Public/Retargeter/IKRetargeter.h
--- a/Engine/Plugins/Animation/IKRig/Source/IKRig/Public/Retargeter/IKRetargeter.h
+++ b/Engine/Plugins/Animation/IKRig/Source/IKRig/Public/Retargeter/IKRetargeter.h
@@ -264,14 +264,11 @@ public:
 	URetargetRootSettings* GetRetargetRootSettings() const { return RootSettings; };
 	/** Get read-only access to a retarget pose */
 	const FIKRetargetPose* GetCurrentRetargetPose() const { return &RetargetPoses[CurrentRetargetPose]; };
-	/** Get the unique ID of this UObject as a name. */
-	FName GetUniqueIDAsName() const;
 	/* Get name of default pose */
 	static const FName GetDefaultPoseName();
 
 	/** UObject */
 	virtual void PostLoad() override;
-	virtual void PostInitProperties() override;
 	/** END UObject */
 
 #if WITH_EDITOR
@@ -332,9 +329,6 @@ public:
 	UPROPERTY(EditAnywhere, Category = RetargetPhases)
 	bool bRetargetIK = true;
 
-	/** log warnings and errors */
-	FIKRigLogger Log;
-
 #if WITH_EDITORONLY_DATA
 	UPROPERTY()
 	float TargetActorOffset_DEPRECATED = 0.0f;
diff --git a/Engine/Plugins/Animation/IKRig/Source/IKRigDeveloper/Private/AnimGraphNode_RetargetPoseFromMesh.cpp b/Engine/Plugins/Animation/IKRig/Source/IKRigDeveloper/Private/AnimGraphNode_RetargetPoseFromMesh.cpp
--- a/Engine/Plugins/Animation/IKRig/Source/IKRigDeveloper/Private/AnimGraphNode_RetargetPoseFromMesh.cpp
+++ b/Engine/Plugins/Animation/IKRig/Source/IKRigDeveloper/Private/AnimGraphNode_RetargetPoseFromMesh.cpp
@@ -97,6 +97,25 @@ void UAnimGraphNode_RetargetPoseFromMesh::ValidateAnimNodeDuringCompilation(USke
 		return;
 	}
 
+	// pull messages out of the processor's log
+	if (!Node.bSuppressWarnings)
+	{
+		if (const UIKRetargetProcessor* Processor = Node.GetRetargetProcessor())
+		{
+			const TArray<FText>& Warnings = Processor->Log.GetWarnings();
+			for (const FText& Warning : Warnings)
+			{
+				MessageLog.Warning(*Warning.ToString());
+			}
+		
+			const TArray<FText>& Errors = Processor->Log.GetErrors();
+			for (const FText& Error : Errors)
+			{
+				MessageLog.Error(*Error.ToString());
+			}
+		}
+	}
+
 	if (ForSkeleton && !Node.bSuppressWarnings)
 	{
 		// validate that target bone chains exist on this skeleton
diff --git a/Engine/Plugins/Animation/IKRig/Source/IKRigEditor/Private/RetargetEditor/IKRetargetEditorController.cpp b/Engine/Plugins/Animation/IKRig/Source/IKRigEditor/Private/RetargetEditor/IKRetargetEditorController.cpp
--- a/Engine/Plugins/Animation/IKRig/Source/IKRigEditor/Private/RetargetEditor/IKRetargetEditorController.cpp
+++ b/Engine/Plugins/Animation/IKRig/Source/IKRigEditor/Private/RetargetEditor/IKRetargetEditorController.cpp
@@ -243,6 +243,7 @@ void FIKRetargetEditorController::ClearOutputLog() const
 	if (OutputLogView.IsValid())
 	{
 		OutputLogView.Get()->ClearLog();
+		GetRetargetProcessor()->Log.Clear();
 	}
 }
 
diff --git a/Engine/Plugins/Animation/IKRig/Source/IKRigEditor/Private/RetargetEditor/IKRetargetOutputLogTabSummoner.cpp b/Engine/Plugins/Animation/IKRig/Source/IKRigEditor/Private/RetargetEditor/IKRetargetOutputLogTabSummoner.cpp
--- a/Engine/Plugins/Animation/IKRig/Source/IKRigEditor/Private/RetargetEditor/IKRetargetOutputLogTabSummoner.cpp
+++ b/Engine/Plugins/Animation/IKRig/Source/IKRigEditor/Private/RetargetEditor/IKRetargetOutputLogTabSummoner.cpp
@@ -32,10 +32,10 @@ TSharedPtr<SToolTip> FIKRetargetOutputLogTabSummoner::CreateTabToolTipWidget(con
 
 TSharedRef<SWidget> FIKRetargetOutputLogTabSummoner::CreateTabBody(const FWorkflowTabSpawnInfo& Info) const
 {
-	TSharedPtr<SIKRigOutputLog>& OutputLogView = IKRetargetEditor.Pin()->GetController()->OutputLogView;
-	FName LogName = IKRetargetEditor.Pin()->GetController()->AssetController->GetAsset()->Log.GetLogTarget();
+	const TSharedRef<FIKRetargetEditorController>& Controller = IKRetargetEditor.Pin()->GetController();
+	const FName LogName = Controller->GetRetargetProcessor()->Log.GetLogTarget();
 	
-	return SNew(SIKRigOutputLog, LogName, OutputLogView);
+	return SNew(SIKRigOutputLog, LogName, Controller->OutputLogView);
 }
 
 #undef LOCTEXT_NAMESPACE 
diff --git a/Engine/Plugins/Animation/IKRig/Source/IKRigEditor/Private/RetargetEditor/IKRetargeterController.cpp b/Engine/Plugins/Animation/IKRig/Source/IKRigEditor/Private/RetargetEditor/IKRetargeterController.cpp
--- a/Engine/Plugins/Animation/IKRig/Source/IKRigEditor/Private/RetargetEditor/IKRetargeterController.cpp
+++ b/Engine/Plugins/Animation/IKRig/Source/IKRigEditor/Private/RetargetEditor/IKRetargeterController.cpp
@@ -40,16 +40,6 @@ UIKRetargeter* UIKRetargeterController::GetAsset() const
 	return Asset;
 }
 
-FName UIKRetargeterController::GetAssetIDAsName() const 
-{
-	if (!Asset)
-	{
-		return NAME_None;
-	}
-
-	return Asset->GetUniqueIDAsName();
-}
-
 void UIKRetargeterController::SetSourceIKRig(UIKRigDefinition* SourceIKRig)
 {
 	Asset->SourceIKRigAsset = SourceIKRig;
diff --git a/Engine/Plugins/Animation/IKRig/Source/IKRigEditor/Private/RigEditor/IKRigController.cpp b/Engine/Plugins/Animation/IKRig/Source/IKRigEditor/Private/RigEditor/IKRigController.cpp
--- a/Engine/Plugins/Animation/IKRig/Source/IKRigEditor/Private/RigEditor/IKRigController.cpp
+++ b/Engine/Plugins/Animation/IKRig/Source/IKRigEditor/Private/RigEditor/IKRigController.cpp
@@ -416,11 +416,9 @@ bool UIKRigController::SetSkeletalMesh(USkeletalMesh* SkeletalMesh, bool bTransa
 	
 	// first determine runtime compatibility between the IK Rig asset and the skeleton we're trying to run it on
 	const FIKRigInputSkeleton InputSkeleton = FIKRigInputSkeleton(SkeletalMesh);
-	if (!UIKRigProcessor::IsIKRigCompatibleWithSkeleton(Asset, InputSkeleton))
+	if (!UIKRigProcessor::IsIKRigCompatibleWithSkeleton(Asset, InputSkeleton, nullptr))
 	{
-		Asset->Log.LogError( FText::Format(
-			LOCTEXT("IncompatibleSkeleton", "Trying to initialize IKRig with a Skeleton that is missing required bones. See output log. {0}"),
-			FText::FromString(Asset->GetName())));
+		UE_LOG(LogTemp, Warning, TEXT("Trying to initialize IKRig with a Skeleton that is missing required bones. See output log. {0}"), *Asset->GetName());
 		return false;
 	}
 
diff --git a/Engine/Plugins/Animation/IKRig/Source/IKRigEditor/Private/RigEditor/IKRigEditMode.cpp b/Engine/Plugins/Animation/IKRig/Source/IKRigEditor/Private/RigEditor/IKRigEditMode.cpp
--- a/Engine/Plugins/Animation/IKRig/Source/IKRigEditor/Private/RigEditor/IKRigEditMode.cpp
+++ b/Engine/Plugins/Animation/IKRig/Source/IKRigEditor/Private/RigEditor/IKRigEditMode.cpp
@@ -138,7 +138,7 @@ void FIKRigEditMode::RenderBones(FPrimitiveDrawInterface* PDI)
 	}
 
 	// IKRig processor initialized and running?
-	UIKRigProcessor* CurrentProcessor = Controller->AnimInstance->GetCurrentlyRunningProcessor();
+	UIKRigProcessor* CurrentProcessor = Controller->GetIKRigProcessor();
 	if (!IsValid(CurrentProcessor))
 	{
 		return;
diff --git a/Engine/Plugins/Animation/IKRig/Source/IKRigEditor/Private/RigEditor/IKRigEditorController.cpp b/Engine/Plugins/Animation/IKRig/Source/IKRigEditor/Private/RigEditor/IKRigEditorController.cpp
--- a/Engine/Plugins/Animation/IKRig/Source/IKRigEditor/Private/RigEditor/IKRigEditorController.cpp
+++ b/Engine/Plugins/Animation/IKRig/Source/IKRigEditor/Private/RigEditor/IKRigEditorController.cpp
@@ -206,6 +206,16 @@ void FIKRigEditorController::Initialize(TSharedPtr<FIKRigEditorToolkit> Toolkit,
 	}
 }
 
+UIKRigProcessor* FIKRigEditorController::GetIKRigProcessor() const
+{
+	if (AnimInstance.IsValid())
+	{
+		return AnimInstance->GetCurrentlyRunningProcessor();
+	}
+
+	return nullptr;
+}
+
 void FIKRigEditorController::OnIKRigNeedsInitialized(UIKRigDefinition* ModifiedIKRig)
 {
 	if (ModifiedIKRig != AssetController->GetAsset())
@@ -271,6 +281,7 @@ void FIKRigEditorController::ClearOutputLog() const
 	if (OutputLogView.IsValid())
 	{
 		OutputLogView->ClearLog();
+		GetIKRigProcessor()->Log.Clear();
 	}
 }
 
diff --git a/Engine/Plugins/Animation/IKRig/Source/IKRigEditor/Private/RigEditor/IKRigOutputLogTabSummoner.cpp b/Engine/Plugins/Animation/IKRig/Source/IKRigEditor/Private/RigEditor/IKRigOutputLogTabSummoner.cpp
--- a/Engine/Plugins/Animation/IKRig/Source/IKRigEditor/Private/RigEditor/IKRigOutputLogTabSummoner.cpp
+++ b/Engine/Plugins/Animation/IKRig/Source/IKRigEditor/Private/RigEditor/IKRigOutputLogTabSummoner.cpp
@@ -2,6 +2,7 @@
 #include "RigEditor/IKRigOutputLogTabSummoner.h"
 
 #include "IDocumentation.h"
+#include "IKRigProcessor.h"
 #include "RigEditor/IKRigToolkit.h"
 #include "RigEditor/IKRigEditorStyle.h"
 #include "RigEditor/SIKRigOutputLog.h"
@@ -31,11 +32,10 @@ TSharedPtr<SToolTip> FIKRigOutputLogTabSummoner::CreateTabToolTipWidget(const FW
 
 TSharedRef<SWidget> FIKRigOutputLogTabSummoner::CreateTabBody(const FWorkflowTabSpawnInfo& Info) const
 {
-	const TSharedRef<FIKRigEditorController> Controller = IKRigEditor.Pin()->GetController();
-	TSharedPtr<SIKRigOutputLog>& OutputLogView = Controller->OutputLogView;
-	FName LogName = Controller->AssetController->GetAsset()->Log.GetLogTarget();
+	const TSharedRef<FIKRigEditorController>& Controller = IKRigEditor.Pin()->GetController();
+	const FName LogName = Controller->GetIKRigProcessor()->Log.GetLogTarget();
 	
-	return SNew(SIKRigOutputLog, LogName, OutputLogView);
+	return SNew(SIKRigOutputLog, LogName, Controller->OutputLogView);
 }
 
 #undef LOCTEXT_NAMESPACE 
diff --git a/Engine/Plugins/Animation/IKRig/Source/IKRigEditor/Private/RigEditor/SIKRigOutputLog.cpp b/Engine/Plugins/Animation/IKRig/Source/IKRigEditor/Private/RigEditor/SIKRigOutputLog.cpp
--- a/Engine/Plugins/Animation/IKRig/Source/IKRigEditor/Private/RigEditor/SIKRigOutputLog.cpp
+++ b/Engine/Plugins/Animation/IKRig/Source/IKRigEditor/Private/RigEditor/SIKRigOutputLog.cpp
@@ -6,7 +6,10 @@
 #include "MessageLogModule.h"
 #include "RetargetEditor/IKRetargetEditorController.h"
 
-void SIKRigOutputLog::Construct(const FArguments& InArgs, const FName InLogName, TSharedPtr<SIKRigOutputLog>& OutLog)
+void SIKRigOutputLog::Construct(
+	const FArguments& InArgs,
+	const FName& InLogName,
+	TSharedPtr<SIKRigOutputLog>& OutLog)
 {
 	OutLog = SharedThis(this);
 
@@ -22,6 +25,7 @@ void SIKRigOutputLog::Construct(const FArguments& InArgs, const FName InLogName,
 		InitOptions.bShowFilters = false;
 		InitOptions.bShowPages = false;
 		InitOptions.bAllowClear = false;
+		InitOptions.bShowInLogWindow = false;
 		MessageLogModule.RegisterLogListing(InLogName, FText::FromString("IK Rig Log"), InitOptions);
 		MessageLogListing = MessageLogModule.GetLogListing(InLogName);
 	}
@@ -40,4 +44,4 @@ void SIKRigOutputLog::Construct(const FArguments& InArgs, const FName InLogName,
 void SIKRigOutputLog::ClearLog() const
 {
 	MessageLogListing.Get()->ClearMessages();
-}
+}
\ No newline at end of file
diff --git a/Engine/Plugins/Animation/IKRig/Source/IKRigEditor/Public/RetargetEditor/IKRetargeterController.h b/Engine/Plugins/Animation/IKRig/Source/IKRigEditor/Public/RetargetEditor/IKRetargeterController.h
--- a/Engine/Plugins/Animation/IKRig/Source/IKRigEditor/Public/RetargetEditor/IKRetargeterController.h
+++ b/Engine/Plugins/Animation/IKRig/Source/IKRigEditor/Public/RetargetEditor/IKRetargeterController.h
@@ -28,8 +28,6 @@ public:
 	/** Get access to the retargeter asset.
 	 *@warning Do not make modifications to the asset directly. Use this API instead. */
 	UIKRetargeter* GetAsset() const;
-	/** Get unique asset integer ID as a name */
-	FName GetAssetIDAsName() const;
 	/** Get access to the editor controller.*/
 	FIKRetargetEditorController* GetEditorController() const { return EditorController; };
 	/** Set the currently used editor controller.*/
diff --git a/Engine/Plugins/Animation/IKRig/Source/IKRigEditor/Public/RigEditor/IKRigEditorController.h b/Engine/Plugins/Animation/IKRig/Source/IKRigEditor/Public/RigEditor/IKRigEditorController.h
--- a/Engine/Plugins/Animation/IKRig/Source/IKRigEditor/Public/RigEditor/IKRigEditorController.h
+++ b/Engine/Plugins/Animation/IKRig/Source/IKRigEditor/Public/RigEditor/IKRigEditorController.h
@@ -80,6 +80,9 @@ public:
 	/** initialize the editor controller to an instance of the IK Rig editor */
 	void Initialize(TSharedPtr<FIKRigEditorToolkit> Toolkit, UIKRigDefinition* Asset);
 
+	/** get the currently active processor running the IK Rig in the editor */
+	UIKRigProcessor* GetIKRigProcessor() const;
+
 	/** callback when IK Rig requires re-initialization */
 	void OnIKRigNeedsInitialized(UIKRigDefinition* ModifiedIKRig);
 
diff --git a/Engine/Plugins/Animation/IKRig/Source/IKRigEditor/Public/RigEditor/SIKRigOutputLog.h b/Engine/Plugins/Animation/IKRig/Source/IKRigEditor/Public/RigEditor/SIKRigOutputLog.h
--- a/Engine/Plugins/Animation/IKRig/Source/IKRigEditor/Public/RigEditor/SIKRigOutputLog.h
+++ b/Engine/Plugins/Animation/IKRig/Source/IKRigEditor/Public/RigEditor/SIKRigOutputLog.h
@@ -4,6 +4,7 @@
 
 #include "Widgets/Layout/SBox.h"
 
+struct FIKRigLogger;
 class IMessageLogListing;
 class FIKRetargetEditorController;
 
@@ -15,8 +16,11 @@ public:
 	SLATE_BEGIN_ARGS(SIKRigOutputLog) {}
 	SLATE_END_ARGS()
 
-	void Construct(const FArguments& InArgs, const FName InLogName, TSharedPtr<SIKRigOutputLog>& OutLog);
-	
+	void Construct(
+		const FArguments& InArgs,
+		const FName& InLogName,
+		TSharedPtr<SIKRigOutputLog>& OutLog);
+
 	void ClearLog() const;
 	
 private:
