Adding ability to double-click IK Rig node to go to IK Rig definition.

diff --git a/Engine/Plugins/Animation/IKRig/Source/IKRigDeveloper/Private/AnimGraphNode_IKRig.cpp b/Engine/Plugins/Animation/IKRig/Source/IKRigDeveloper/Private/AnimGraphNode_IKRig.cpp
--- a/Engine/Plugins/Animation/IKRig/Source/IKRigDeveloper/Private/AnimGraphNode_IKRig.cpp
+++ b/Engine/Plugins/Animation/IKRig/Source/IKRigDeveloper/Private/AnimGraphNode_IKRig.cpp
@@ -325,6 +325,11 @@ void UAnimGraphNode_IKRig::ValidateAnimNodeDuringCompilation(USkeleton* ForSkele
 	}
 }
 
+UObject* UAnimGraphNode_IKRig::GetJumpTargetForDoubleClick() const
+{
+	return Node.RigDefinitionAsset;
+}
+
 void UAnimGraphNode_IKRig::PostEditChangeProperty(struct FPropertyChangedEvent& PropertyChangedEvent)
 {	
 	const FName PropertyName = PropertyChangedEvent.GetPropertyName();
diff --git a/Engine/Plugins/Animation/IKRig/Source/IKRigDeveloper/Public/AnimGraphNode_IKRig.h b/Engine/Plugins/Animation/IKRig/Source/IKRigDeveloper/Public/AnimGraphNode_IKRig.h
--- a/Engine/Plugins/Animation/IKRig/Source/IKRigDeveloper/Public/AnimGraphNode_IKRig.h
+++ b/Engine/Plugins/Animation/IKRig/Source/IKRigDeveloper/Public/AnimGraphNode_IKRig.h
@@ -110,6 +110,10 @@ public:
 	virtual void ValidateAnimNodeDuringCompilation(USkeleton* ForSkeleton, FCompilerResultsLog& MessageLog) override;
 	// End of UAnimGraphNode_Base interface
 
+	// UK2Node interface
+	UObject* GetJumpTargetForDoubleClick() const;
+	// END UK2Node
+
 	// Begin UObject Interface.
 	virtual void PostLoad() override;
 	// End UObject Interface.
