Adding transform widget for retarget bones.

diff --git a/Engine/Plugins/Animation/IKRig/Source/IKRig/Private/Retargeter/IKRetargetProcessor.cpp b/Engine/Plugins/Animation/IKRig/Source/IKRig/Private/Retargeter/IKRetargetProcessor.cpp
--- a/Engine/Plugins/Animation/IKRig/Source/IKRig/Private/Retargeter/IKRetargetProcessor.cpp
+++ b/Engine/Plugins/Animation/IKRig/Source/IKRig/Private/Retargeter/IKRetargetProcessor.cpp
@@ -16,7 +16,9 @@
 
 void FRetargetSkeleton::Initialize(
 	USkeletalMesh* InSkeletalMesh,
-	const TArray<FBoneChain>& BoneChains)
+	const TArray<FBoneChain>& BoneChains,
+	const FIKRetargetPose* RetargetPose,
+	const FName& RetargetRootBone)
 {	
 	// record which skeletal mesh this is running on
 	SkeletalMesh = InSkeletalMesh;
@@ -31,7 +33,7 @@ void FRetargetSkeleton::Initialize(
 
 	// determine set of bones referenced by one of the retarget bone chains
 	// this is the set of bones that will be affected by the retarget pose
-	IsBoneInAnyChain.Init(false, BoneNames.Num());
+	ChainThatContainsBone.Init(NAME_None, BoneNames.Num());
 	for (const FBoneChain& BoneChain : BoneChains)
 	{
 		TArray<int32> BonesInChain;
@@ -39,16 +41,17 @@ void FRetargetSkeleton::Initialize(
 		{
 			for (const int32 BoneInChain : BonesInChain)
 			{
-				IsBoneInAnyChain[BoneInChain] = true;
+				ChainThatContainsBone[BoneInChain] = BoneChain.ChainName;
 			}
 		}
 	}
 
-	// update retarget pose to reflect custom offsets
-	GenerateRetargetPose();
-
 	// initialize branch caching
 	CachedEndOfBranchIndices.Init(RETARGETSKELETON_INVALID_BRANCH_INDEX, ParentIndices.Num());
+
+	// update retarget pose to reflect custom offsets (applies stored offsets)
+	// NOTE: this must be done AFTER generating IsBoneInAnyTargetChain array above
+	GenerateRetargetPose(RetargetPose, RetargetRootBone);
 }
 
 void FRetargetSkeleton::Reset()
@@ -60,21 +63,72 @@ void FRetargetSkeleton::Reset()
 	SkeletalMesh = nullptr;
 }
 
-void FRetargetSkeleton::GenerateRetargetPose()
+void FRetargetSkeleton::GenerateRetargetPose(const FIKRetargetPose* InRetargetPose, const FName& RetargetRootBone)
 {
 	// initialize retarget pose to the skeletal mesh reference pose
 	RetargetLocalPose = SkeletalMesh->GetRefSkeleton().GetRefBonePose();
 	// copy local pose to global
 	RetargetGlobalPose = RetargetLocalPose;
 	// convert to global space
-	UpdateGlobalTransformsBelowBone(0, RetargetLocalPose, RetargetGlobalPose);
+	UpdateGlobalTransformsBelowBone(-1, RetargetLocalPose, RetargetGlobalPose);
+
+	// no retarget pose specified (will use default pose from skeletal mesh with no offsets)
+	if (InRetargetPose==nullptr  || RetargetRootBone == NAME_None)
+	{
+		return;
+	}
+
+	// apply retarget pose offsets (retarget pose is stored as offset relative to reference pose)
+	const TArray<FTransform>& RefPoseLocal = SkeletalMesh->GetRefSkeleton().GetRefBonePose();
+	
+	// apply root translation offset
+	const int32 RootBoneIndex = FindBoneIndexByName(RetargetRootBone);
+	if (RootBoneIndex != INDEX_NONE)
+	{
+		FTransform& RootTransform = RetargetGlobalPose[RootBoneIndex];
+		RootTransform.AddToTranslation(InRetargetPose->RootTranslationOffset);
+		UpdateLocalTransformOfSingleBone(RootBoneIndex, RetargetLocalPose, RetargetGlobalPose);
+	}
+
+	// apply bone rotation offsets
+	for (const TTuple<FName, FQuat>& BoneRotationOffset : InRetargetPose->BoneRotationOffsets)
+	{
+		const int32 BoneIndex = FindBoneIndexByName(BoneRotationOffset.Key);
+		if (BoneIndex == INDEX_NONE)
+		{
+			// this can happen if a retarget pose recorded a bone offset for a bone that is not present in the
+			// target skeleton; ie, the retarget pose was generated from a different Skeletal Mesh with extra bones
+			continue;
+		}
+
+		const FQuat DeltaRotation = GetRetargetPoseDeltaRotation(BoneRotationOffset.Key, InRetargetPose);
+		const FQuat LocalBoneRotation = RefPoseLocal[BoneIndex].GetRotation() * DeltaRotation;
+		RetargetLocalPose[BoneIndex].SetRotation(LocalBoneRotation);
+	}
+
+	UpdateGlobalTransformsBelowBone(-1, RetargetLocalPose, RetargetGlobalPose);
+}
+
+FQuat FRetargetSkeleton::GetRetargetPoseDeltaRotation(const FName BoneName, const FIKRetargetPose* InRetargetPose) const
+{
+	const int32 BoneIndex = FindBoneIndexByName(BoneName);
+	check(BoneIndex!=INDEX_NONE)
+	check(InRetargetPose)
+	
+	if (const FQuat* BoneRotationOffset = InRetargetPose->BoneRotationOffsets.Find(BoneNames[BoneIndex]))
+	{
+		return *BoneRotationOffset;
+
+	}
+
+	return FQuat::Identity;
 }
 
 int32 FRetargetSkeleton::FindBoneIndexByName(const FName InName) const
 {
 	return BoneNames.IndexOfByPredicate([&InName](const FName& BoneName)
 	{
-			return BoneName == InName;
+		return BoneName == InName;
 	});
 }
 
@@ -83,7 +137,7 @@ void FRetargetSkeleton::UpdateGlobalTransformsBelowBone(
 	const TArray<FTransform>& InLocalPose,
 	TArray<FTransform>& OutGlobalPose) const
 {
-	check(BoneNames.IsValidIndex(StartBoneIndex));
+	check(BoneNames.IsValidIndex(StartBoneIndex+1));
 	check(BoneNames.Num() == InLocalPose.Num());
 	check(BoneNames.Num() == OutGlobalPose.Num());
 	
@@ -116,7 +170,9 @@ void FRetargetSkeleton::UpdateGlobalTransformOfSingleBone(
 	const int32 ParentIndex = ParentIndices[BoneIndex];
 	if (ParentIndex == INDEX_NONE)
 	{
-		return; // root always in global space
+		// root always in global space already, no conversion required
+		OutGlobalPose[BoneIndex] = InLocalPose[BoneIndex];
+		return; 
 	}
 	const FTransform& ChildLocalTransform = InLocalPose[BoneIndex];
 	const FTransform& ParentGlobalTransform = OutGlobalPose[ParentIndex];
@@ -252,12 +308,12 @@ int32 FRetargetSkeleton::GetParentIndex(const int32 BoneIndex) const
 }
 
 void FTargetSkeleton::Initialize(
-	USkeletalMesh* InSkeletalMesh, 
+	USkeletalMesh* InSkeletalMesh,
+	const TArray<FBoneChain>& BoneChains,
 	const FIKRetargetPose* RetargetPose,
-	const FName& RetargetRootBone,
-	const TArray<FBoneChain>& TargetChains)
+	const FName& RetargetRootBone)
 {
-	FRetargetSkeleton::Initialize(InSkeletalMesh, TargetChains);
+	FRetargetSkeleton::Initialize(InSkeletalMesh, BoneChains, RetargetPose, RetargetRootBone);
 
 	// make storage for per-bone "Is Retargeted" flag (used for hierarchy updates)
 	// these are bones that are in a target chain that is mapped to a source chain (ie, will actually be retargeted)
@@ -266,58 +322,6 @@ void FTargetSkeleton::Initialize(
 
 	// initialize storage for output pose (the result of the retargeting)
 	OutputGlobalPose = RetargetGlobalPose;
-
-	// generate the retarget pose (applies stored offsets)
-	// NOTE: this must be done AFTER generating IsBoneInAnyTargetChain array above (FRetargetSkeleton::Initialize)
-	GenerateRetargetPose(RetargetPose, RetargetRootBone);
-}
-
-void FTargetSkeleton::GenerateRetargetPose(const FIKRetargetPose* InRetargetPose, const FName& RetargetRootBone)
-{
-	// create a retarget pose by copying the ref pose
-	FRetargetSkeleton::GenerateRetargetPose();
-	
-	// no retarget pose specified (will use default pose from skeletal mesh with no offsets)
-	if (InRetargetPose==nullptr  || RetargetRootBone == NAME_None)
-	{
-		return;
-	}
-
-	// apply retarget pose offsets (retarget pose is stored as offset relative to reference pose)
-	const TArray<FTransform>& RefPoseLocal = SkeletalMesh->GetRefSkeleton().GetRefBonePose();
-	
-	// apply root translation offset
-	const int32 RootBoneIndex = FindBoneIndexByName(RetargetRootBone);
-	if (RootBoneIndex != INDEX_NONE)
-	{
-		FTransform& RootTransform = RetargetGlobalPose[RootBoneIndex];
-		RootTransform.AddToTranslation(InRetargetPose->RootTranslationOffset);
-		UpdateLocalTransformOfSingleBone(RootBoneIndex, RetargetLocalPose, RetargetGlobalPose);
-	}
-
-	// apply bone rotation offsets
-	for (const TTuple<FName, FQuat>& BoneRotationOffset : InRetargetPose->BoneRotationOffsets)
-	{
-		const int32 BoneIndex = FindBoneIndexByName(BoneRotationOffset.Key);
-		if (BoneIndex == INDEX_NONE)
-		{
-			// this can happen if a retarget pose recorded a bone offset for a bone that is not present in the
-			// target skeleton; ie, the retarget pose was generated from a different Skeletal Mesh with extra bones
-			continue;
-		}
-
-		if (!IsBoneInAnyChain[BoneIndex] && BoneIndex!=RootBoneIndex)
-		{
-			// this can happen if a retarget pose includes bone edits from a bone chain that was subsequently removed,
-			// and the asset has not run through the "CleanChainMapping" operation yet (happens on load)
-			continue;
-		}
-
-		const FQuat LocalBoneRotation = BoneRotationOffset.Value * RefPoseLocal[BoneIndex].GetRotation();
-		RetargetLocalPose[BoneIndex].SetRotation(LocalBoneRotation);
-	}
-
-	UpdateGlobalTransformsBelowBone(0, RetargetLocalPose, RetargetGlobalPose);
 }
 
 void FTargetSkeleton::Reset()
@@ -340,8 +344,10 @@ void FTargetSkeleton::UpdateGlobalTransformsAllNonRetargetedBones(TArray<FTransf
 	}
 }
 
-FResolvedBoneChain::FResolvedBoneChain( const FBoneChain& BoneChain, const FRetargetSkeleton& Skeleton,
-										TArray<int32>& OutBoneIndices)
+FResolvedBoneChain::FResolvedBoneChain(
+	const FBoneChain& BoneChain,
+	const FRetargetSkeleton& Skeleton,
+	TArray<int32>& OutBoneIndices)
 {
 	// validate start and end bones exist and are not the root
 	const int32 StartIndex = Skeleton.FindBoneIndexByName(BoneChain.StartBone.BoneName);
@@ -1309,12 +1315,16 @@ void UIKRetargetProcessor::Initialize(
 	}
 	
 	// initialize skeleton data for source and target
-	SourceSkeleton.Initialize(SourceSkeletalMesh, RetargeterAsset->GetSourceIKRig()->GetRetargetChains());
+	SourceSkeleton.Initialize(
+		SourceSkeletalMesh,
+		RetargeterAsset->GetSourceIKRig()->GetRetargetChains(),
+		nullptr,
+		NAME_None);
 	TargetSkeleton.Initialize(
 		TargetSkeletalMesh,
+		RetargeterAsset->GetTargetIKRig()->GetRetargetChains(),
 		RetargeterAsset->GetCurrentRetargetPose(),
-		RetargeterAsset->GetTargetIKRig()->GetRetargetRoot(),
-		RetargeterAsset->GetTargetIKRig()->GetRetargetChains());
+		RetargeterAsset->GetTargetIKRig()->GetRetargetRoot());
 
 	// initialize roots
 	bRootsInitialized = InitializeRoots();
@@ -1713,6 +1723,30 @@ void UIKRetargetProcessor::CopyAllSettingsFromAsset()
 	RootRetargeter.StaticOffset = RootSettings.StaticOffset;
 	RootRetargeter.StaticRotationOffset = RootSettings.StaticRotationOffset;
 }
+
+bool UIKRetargetProcessor::IsBoneRetargeted(
+	const int32& BoneIndex,
+	const int8& SkeletonToCheck) const
+{
+	const bool bUseSourceSkeleton = SkeletonToCheck == 0;
+	const FRetargetSkeleton& Skeleton = bUseSourceSkeleton ? SourceSkeleton : TargetSkeleton;
+	const int32 RootBoneIndex = bUseSourceSkeleton ? GetSourceRetargetRoot() : GetTargetRetargetRoot();
+	return (Skeleton.ChainThatContainsBone[BoneIndex] != NAME_None) || BoneIndex==RootBoneIndex;
+}
+
+FName UIKRetargetProcessor::GetChainNameForBone(const int32& BoneIndex, const int8& SkeletonToCheck) const
+{
+	const bool bUseSourceSkeleton = SkeletonToCheck == 0;
+	const FRetargetSkeleton& Skeleton = bUseSourceSkeleton ? SourceSkeleton : TargetSkeleton;
+	const int32 RootBoneIndex = bUseSourceSkeleton ? GetSourceRetargetRoot() : GetTargetRetargetRoot();
+	if (RootBoneIndex == BoneIndex)
+	{
+		return FName("Retarget Root");
+	}
+
+	return Skeleton.ChainThatContainsBone[BoneIndex];
+}
+
 #endif
 
 #undef LOCTEXT_NAMESPACE
diff --git a/Engine/Plugins/Animation/IKRig/Source/IKRig/Private/Retargeter/IKRetargeter.cpp b/Engine/Plugins/Animation/IKRig/Source/IKRig/Private/Retargeter/IKRetargeter.cpp
--- a/Engine/Plugins/Animation/IKRig/Source/IKRig/Private/Retargeter/IKRetargeter.cpp
+++ b/Engine/Plugins/Animation/IKRig/Source/IKRig/Private/Retargeter/IKRetargeter.cpp
@@ -2,6 +2,8 @@
 
 #include "Retargeter/IKRetargeter.h"
 
+#include "IKRigObjectVersion.h"
+
 #if WITH_EDITOR
 const FName UIKRetargeter::GetSourceIKRigPropertyName() { return GET_MEMBER_NAME_STRING_CHECKED(UIKRetargeter, SourceIKRigAsset); };
 const FName UIKRetargeter::GetTargetIKRigPropertyName() { return GET_MEMBER_NAME_STRING_CHECKED(UIKRetargeter, TargetIKRigAsset); };
@@ -27,6 +29,12 @@ UIKRetargeter::UIKRetargeter(const FObjectInitializer& ObjectInitializer)
 	RootSettings->SetFlags(RF_Transactional);
 }
 
+void UIKRetargeter::Serialize(FArchive& Ar)
+{
+	Super::Serialize(Ar);
+	Ar.UsingCustomVersion(FIKRigObjectVersion::GUID);
+};
+
 void UIKRetargeter::PostLoad()
 {
 	Super::PostLoad();
@@ -101,7 +109,12 @@ void FIKRetargetPose::SetBoneRotationOffset(FName BoneName, FQuat RotationDelta,
 	*RotOffset = RotationDelta;
 }
 
-void FIKRetargetPose::AddTranslationDeltaToRoot(FVector TranslateDelta)
+void FIKRetargetPose::SetRootTranslationDelta(FVector TranslationDelta)
+{
+	RootTranslationOffset = TranslationDelta;
+}
+
+void FIKRetargetPose::AddToRootTranslationDelta(FVector TranslateDelta)
 {
 	RootTranslationOffset += TranslateDelta;
 }
@@ -120,4 +133,4 @@ const FName UIKRetargeter::GetDefaultPoseName()
 {
 	static const FName DefaultPoseName = "Default Pose";
 	return DefaultPoseName;
-};
+}
diff --git a/Engine/Plugins/Animation/IKRig/Source/IKRig/Public/IKRigObjectVersion.h b/Engine/Plugins/Animation/IKRig/Source/IKRig/Public/IKRigObjectVersion.h
--- a/Engine/Plugins/Animation/IKRig/Source/IKRig/Public/IKRigObjectVersion.h
+++ b/Engine/Plugins/Animation/IKRig/Source/IKRig/Public/IKRigObjectVersion.h
@@ -14,6 +14,9 @@ struct IKRIG_API FIKRigObjectVersion
 		// Before any version changes were made
 		BeforeCustomVersionWasAdded,
 
+		// Retarget pose quaternions changed from pre to post multiplied
+		RetargetPoseQuatPostMultiplied,
+
 		// -----<new versions can be added above this line>-------------------------------------------------
 		VersionPlusOne,
 		LatestVersion = VersionPlusOne - 1
diff --git a/Engine/Plugins/Animation/IKRig/Source/IKRig/Public/Retargeter/IKRetargetProcessor.h b/Engine/Plugins/Animation/IKRig/Source/IKRig/Public/Retargeter/IKRetargetProcessor.h
--- a/Engine/Plugins/Animation/IKRig/Source/IKRig/Public/Retargeter/IKRetargetProcessor.h
+++ b/Engine/Plugins/Animation/IKRig/Source/IKRig/Public/Retargeter/IKRetargetProcessor.h
@@ -8,6 +8,7 @@
 #include "Kismet/KismetMathLibrary.h"
 #include "IKRetargetProcessor.generated.h"
 
+enum class EIKRetargetSkeletonMode : uint8;
 class URetargetChainSettings;
 class UIKRigDefinition;
 class UIKRigProcessor;
@@ -30,18 +31,18 @@ struct IKRIG_API FRetargetSkeleton
 	
 	USkeletalMesh* SkeletalMesh;
 
-	// true for bones that are in a retarget bone chain and thus are posed by a retarget pose
-	// the application of a retarget pose to a bone happens regardless of whether the chain is mapped or not
-	// (this must happen or else child chains may not end up in the correct retarget pose)
-	TArray<bool> IsBoneInAnyChain;
+	// record which chain is actually controlling each bone
+	TArray<FName> ChainThatContainsBone;
 
 	void Initialize(
 		USkeletalMesh* InSkeletalMesh,
-		const TArray<FBoneChain>& BoneChains);
+		const TArray<FBoneChain>& BoneChains,
+		const FIKRetargetPose* RetargetPose,
+		const FName& RetargetRootBone);
 
 	void Reset();
 
-	void GenerateRetargetPose();
+	void GenerateRetargetPose(const FIKRetargetPose* InRetargetPose, const FName& RetargetRootBone);
 
 	int32 FindBoneIndexByName(const FName InName) const;
 
@@ -79,6 +80,8 @@ struct IKRIG_API FRetargetSkeleton
 	
 	bool IsParentOfChild(const int32 PotentialParentIndex, const int32 ChildBoneIndex) const;
 
+	FQuat GetRetargetPoseDeltaRotation(const FName BoneName, const FIKRetargetPose* InRetargetPose) const;
+
 private:
 	
 	/** One index per-bone. Lazy-filled on request. Stores the last element of the branch below the bone.
@@ -94,12 +97,10 @@ struct FTargetSkeleton : public FRetargetSkeleton
 	TArray<bool> IsBoneRetargeted;
 
 	void Initialize(
-		USkeletalMesh* InSkeletalMesh, 
+		USkeletalMesh* InSkeletalMesh,
+		const TArray<FBoneChain>& BoneChains,
 		const FIKRetargetPose* RetargetPose,
-		const FName& RetargetRootBone,
-		const TArray<FBoneChain>& TargetChains);
-
-	void GenerateRetargetPose(const FIKRetargetPose* InRetargetPose, const FName& RetargetRootBone);
+		const FName& RetargetRootBone);
 
 	void Reset();
 
@@ -535,6 +536,10 @@ public:
 	void SetNeedsInitialized();
 	/** During editor preview, drive the target IK Rig with the settings from it's source asset */
 	void CopyAllSettingsFromAsset();
+	/** Returns true if the bone is part of a retarget chain or root bone, false otherwise. */
+	bool IsBoneRetargeted(const int32& BoneIndex, const int8& SkeletonToCheck) const;
+	/** Returns name of the chain associated with this bone. Returns NAME_None if bone is not in a chain. */
+	FName GetChainNameForBone(const int32& BoneIndex, const int8& SkeletonToCheck) const;
 #endif
 
 private:
diff --git a/Engine/Plugins/Animation/IKRig/Source/IKRig/Public/Retargeter/IKRetargeter.h b/Engine/Plugins/Animation/IKRig/Source/IKRig/Public/Retargeter/IKRetargeter.h
--- a/Engine/Plugins/Animation/IKRig/Source/IKRig/Public/Retargeter/IKRetargeter.h
+++ b/Engine/Plugins/Animation/IKRig/Source/IKRig/Public/Retargeter/IKRetargeter.h
@@ -195,18 +195,18 @@ struct IKRIG_API FIKRetargetPose
 public:
 	
 	FIKRetargetPose() = default;
-	
+
+	// a translational delta in GLOBAL space, applied only to the retarget root bone
 	UPROPERTY(EditAnywhere, Category = RetargetPose)
 	FVector RootTranslationOffset = FVector::ZeroVector;
-	
+
+	// these are LOCAL-space rotation deltas to be applied to a bone to modify it's retarget pose
 	UPROPERTY(EditAnywhere, Category = RetargetPose)
 	TMap<FName, FQuat> BoneRotationOffsets;
-
-	// returns true if this is a new offset (not previously recorded)
+	
 	void SetBoneRotationOffset(FName BoneName, FQuat RotationOffset, const FIKRigSkeleton& Skeleton);
-
-	void AddTranslationDeltaToRoot(FVector TranslateDelta);
-
+	void SetRootTranslationDelta(FVector TranslationDelta);
+	void AddToRootTranslationDelta(FVector TranslationDelta);
 	void SortHierarchically(const FIKRigSkeleton& Skeleton);
 };
 
@@ -268,6 +268,7 @@ public:
 	static const FName GetDefaultPoseName();
 
 	/** UObject */
+	virtual void Serialize(FArchive& Ar) override;
 	virtual void PostLoad() override;
 	/** END UObject */
 
diff --git a/Engine/Plugins/Animation/IKRig/Source/IKRigEditor/Private/IKRigEditor.cpp b/Engine/Plugins/Animation/IKRig/Source/IKRigEditor/Private/IKRigEditor.cpp
--- a/Engine/Plugins/Animation/IKRig/Source/IKRigEditor/Private/IKRigEditor.cpp
+++ b/Engine/Plugins/Animation/IKRig/Source/IKRigEditor/Private/IKRigEditor.cpp
@@ -56,12 +56,15 @@ void FIKRigEditor::StartupModule()
 
 	// register detail customizations
 	FPropertyEditorModule& PropertyEditorModule = FModuleManager::LoadModuleChecked<FPropertyEditorModule>("PropertyEditor");
+	// custom IK rig bone widget
+	PropertyEditorModule.RegisterCustomClassLayout(UIKRigBoneDetails::StaticClass()->GetFName(), FOnGetDetailCustomizationInstance::CreateStatic(&FIKRigGenericDetailCustomization::MakeInstance));
 	ClassesToUnregisterOnShutdown.Add(UIKRigBoneDetails::StaticClass()->GetFName());
-	PropertyEditorModule.RegisterCustomClassLayout(ClassesToUnregisterOnShutdown.Last(), FOnGetDetailCustomizationInstance::CreateStatic(&FIKRigGenericDetailCustomization::MakeInstance));
+	// custom IK goal widget
+	PropertyEditorModule.RegisterCustomClassLayout(UIKRigEffectorGoal::StaticClass()->GetFName(), FOnGetDetailCustomizationInstance::CreateStatic(&FIKRigGenericDetailCustomization::MakeInstance));
 	ClassesToUnregisterOnShutdown.Add(UIKRigEffectorGoal::StaticClass()->GetFName());
-	PropertyEditorModule.RegisterCustomClassLayout(ClassesToUnregisterOnShutdown.Last(), FOnGetDetailCustomizationInstance::CreateStatic(&FIKRigGenericDetailCustomization::MakeInstance));
-	ClassesToUnregisterOnShutdown.Add(UIKRetargeter::StaticClass()->GetFName());
-	PropertyEditorModule.RegisterCustomClassLayout(ClassesToUnregisterOnShutdown.Last(), FOnGetDetailCustomizationInstance::CreateStatic(&FIKRetargeterDetails::MakeInstance));
+	// custom retargeter bone widget
+	PropertyEditorModule.RegisterCustomClassLayout(UIKRetargetBoneDetails::StaticClass()->GetFName(), FOnGetDetailCustomizationInstance::CreateStatic(&FIKRetargetBoneDetailCustomization::MakeInstance));
+	ClassesToUnregisterOnShutdown.Add(UIKRetargetBoneDetails::StaticClass()->GetFName());
 }
 
 void FIKRigEditor::ShutdownModule()
diff --git a/Engine/Plugins/Animation/IKRig/Source/IKRigEditor/Private/RetargetEditor/IKRetargetApplicationMode.cpp b/Engine/Plugins/Animation/IKRig/Source/IKRigEditor/Private/RetargetEditor/IKRetargetApplicationMode.cpp
--- a/Engine/Plugins/Animation/IKRig/Source/IKRigEditor/Private/RetargetEditor/IKRetargetApplicationMode.cpp
+++ b/Engine/Plugins/Animation/IKRig/Source/IKRigEditor/Private/RetargetEditor/IKRetargetApplicationMode.cpp
@@ -43,7 +43,7 @@ FIKRetargetApplicationMode::FIKRetargetApplicationMode(
 	TabFactories.RegisterFactory(MakeShared<FIKRetargetHierarchyTabSummoner>(IKRetargetEditor));
 
 	// create tab layout
-	TabLayout = FTabManager::NewLayout("Standalone_IKRetargetEditor_Layout_v1.014")
+	TabLayout = FTabManager::NewLayout("Standalone_IKRetargetEditor_Layout_v1.017")
 		->AddArea
 		(
 			FTabManager::NewPrimaryArea()
@@ -55,34 +55,34 @@ FIKRetargetApplicationMode::FIKRetargetApplicationMode(
 				->SetOrientation(Orient_Horizontal)
 				->Split
 				(
-					FTabManager::NewSplitter()
-					->SetSizeCoefficient(0.8f)
-					->SetOrientation(Orient_Vertical)
+				FTabManager::NewSplitter()
+					//->SetSizeCoefficient(0.8f)
+					->SetOrientation(Orient_Horizontal)
 					->Split
 					(
-					FTabManager::NewSplitter()
-						//->SetSizeCoefficient(0.8f)
-						->SetOrientation(Orient_Horizontal)
-						->Split
-						(
-							FTabManager::NewStack()
-							->SetSizeCoefficient(0.2f)
-							->SetHideTabWell(true)
-							->AddTab(FIKRetargetHierarchyTabSummoner::TabID, ETabState::OpenedTab)
-						)
-						->Split
-						(
 						FTabManager::NewStack()
-							->SetSizeCoefficient(0.8f)
-							->SetHideTabWell(true)
-							->AddTab(FPersonaTabs::PreviewViewportID, ETabState::OpenedTab)
-						)
+						->SetSizeCoefficient(0.2f)
+						->SetHideTabWell(true)
+						->AddTab(FIKRetargetHierarchyTabSummoner::TabID, ETabState::OpenedTab)
 					)
 					->Split
 					(
-					FTabManager::NewStack()
-						->SetSizeCoefficient(0.2f)
-						->AddTab(FIKRetargetOutputLogTabSummoner::TabID, ETabState::OpenedTab)
+					FTabManager::NewSplitter()
+							->SetSizeCoefficient(0.8f)
+							->SetOrientation(Orient_Vertical)
+							->Split
+							(
+							FTabManager::NewStack()
+								->SetSizeCoefficient(0.9f)
+								->SetHideTabWell(true)
+								->AddTab(FPersonaTabs::PreviewViewportID, ETabState::OpenedTab)
+							)
+							->Split
+							(
+							FTabManager::NewStack()
+								->SetSizeCoefficient(0.1f)
+								->AddTab(FIKRetargetOutputLogTabSummoner::TabID, ETabState::OpenedTab)
+							)
 					)
 				)
 				->Split
diff --git a/Engine/Plugins/Animation/IKRig/Source/IKRigEditor/Private/RetargetEditor/IKRetargetDefaultMode.cpp b/Engine/Plugins/Animation/IKRig/Source/IKRigEditor/Private/RetargetEditor/IKRetargetDefaultMode.cpp
--- a/Engine/Plugins/Animation/IKRig/Source/IKRigEditor/Private/RetargetEditor/IKRetargetDefaultMode.cpp
+++ b/Engine/Plugins/Animation/IKRig/Source/IKRigEditor/Private/RetargetEditor/IKRetargetDefaultMode.cpp
@@ -260,9 +260,6 @@ void FIKRetargetDefaultMode::Enter()
 	Controller->Editor.Pin()->GetPersonaToolkit()->GetPreviewScene()->SetAllowMeshHitProxies(true);
 	Controller->SourceSkelMeshComponent->bSelectable = true;
 	Controller->TargetSkelMeshComponent->bSelectable = true;
-
-	// clear all selection when entering mode
-	Controller->ClearSelection();
 }
 
 void FIKRetargetDefaultMode::Exit()
@@ -273,13 +270,14 @@ void FIKRetargetDefaultMode::Exit()
 		return; 
 	}
 
-	// disable selection in other modes
-	Controller->Editor.Pin()->GetPersonaToolkit()->GetPreviewScene()->SetAllowMeshHitProxies(false);
-	Controller->SourceSkelMeshComponent->bSelectable = false;
-	Controller->TargetSkelMeshComponent->bSelectable = false;
-	
-	// deselect everything
-	Controller->ClearSelection();
+	// editor can be closed while in edit mode
+	if (Controller->Editor.IsValid())
+	{
+		// disable selection in other modes
+		Controller->Editor.Pin()->GetPersonaToolkit()->GetPreviewScene()->SetAllowMeshHitProxies(false);
+		Controller->SourceSkelMeshComponent->bSelectable = false;
+		Controller->TargetSkelMeshComponent->bSelectable = false;
+	}
 
 	IPersonaEditMode::Exit();
 }
diff --git a/Engine/Plugins/Animation/IKRig/Source/IKRigEditor/Private/RetargetEditor/IKRetargetDetails.cpp b/Engine/Plugins/Animation/IKRig/Source/IKRigEditor/Private/RetargetEditor/IKRetargetDetails.cpp
--- a/Engine/Plugins/Animation/IKRig/Source/IKRigEditor/Private/RetargetEditor/IKRetargetDetails.cpp
+++ b/Engine/Plugins/Animation/IKRig/Source/IKRigEditor/Private/RetargetEditor/IKRetargetDetails.cpp
@@ -5,218 +5,750 @@
 #include "DetailCategoryBuilder.h"
 #include "DetailLayoutBuilder.h"
 #include "DetailWidgetRow.h"
+#include "Animation/AnimInstance.h"
+#include "Animation/DebugSkelMeshComponent.h"
 #include "Framework/MultiBox/SToolBarButtonBlock.h"
-#include "RetargetEditor/IKRetargetCommands.h"
 #include "RetargetEditor/IKRetargetEditor.h"
 #include "RetargetEditor/IKRetargetEditorController.h"
 #include "RetargetEditor/IKRetargeterController.h"
 #include "Retargeter/IKRetargeter.h"
+#include "Widgets/Input/SSegmentedControl.h"
+#include "ScopedTransaction.h"
+
+#if WITH_EDITOR
+#include "HAL/PlatformApplicationMisc.h"
+#endif
 
 #define LOCTEXT_NAMESPACE "IKRetargeterDetails"
 
-TSharedRef<IDetailCustomization> FIKRetargeterDetails::MakeInstance()
+
+FEulerTransform UIKRetargetBoneDetails::GetTransform(EIKRetargetTransformType TransformType, const bool bLocalSpace) const
 {
-	return MakeShareable(new FIKRetargeterDetails);
+	// editor setup?
+	const FIKRetargetEditorController* Controller = EditorController.Get();
+	if (!Controller)
+	{
+		return FEulerTransform::Identity;
+	}
+
+	// ensure we have a valid skeletal mesh
+	const bool bEditingSource = EditorController->GetSkeletonMode() == EIKRetargetSkeletonMode::Source;
+	UDebugSkelMeshComponent* Mesh = bEditingSource ? Controller->SourceSkelMeshComponent : Controller->TargetSkelMeshComponent;
+	if (!(Mesh && Mesh->SkeletalMesh))
+	{
+		return FEulerTransform::Identity;
+	}
+
+	// ensure bone is valid
+	const FReferenceSkeleton& RefSkeleton = Mesh->SkeletalMesh->GetRefSkeleton();
+	const int32 BoneIndex = RefSkeleton.FindBoneIndex(SelectedBone);
+	if (BoneIndex == INDEX_NONE)
+	{
+		return FEulerTransform::Identity;
+	}
+	
+	switch(TransformType)
+	{
+	case EIKRetargetTransformType::Current:
+		{
+			if (bLocalSpace)
+			{
+				const TArray<FTransform>& LocalTransforms = Mesh->GetBoneSpaceTransforms();
+				return LocalTransforms.IsValidIndex(BoneIndex) ? FEulerTransform(LocalTransforms[BoneIndex]) : FEulerTransform::Identity;
+			}
+			else
+			{
+				return FEulerTransform(Mesh->GetBoneTransform(BoneIndex, FTransform::Identity));
+			}
+		}
+	
+	case EIKRetargetTransformType::Reference:
+		{
+			if (bLocalSpace)
+			{
+				return FEulerTransform(RefSkeleton.GetRefBonePose()[BoneIndex]);
+			}
+			else
+			{
+				return FEulerTransform(FAnimationRuntime::GetComponentSpaceTransform(RefSkeleton, RefSkeleton.GetRefBonePose(), BoneIndex));
+			}
+		}
+
+	case EIKRetargetTransformType::RelativeOffset:
+		{
+			// this is the only stored data we have for bone pose offsets
+			const FRotator LocalRotationDelta = EditorController->AssetController->GetRotationOffsetForRetargetPoseBone(SelectedBone).Rotator();
+			const FVector GlobalTranslationDelta = IsRootBone() ? EditorController->AssetController->GetTranslationOffsetOnRetargetRootBone() : FVector::Zero();
+			const int32 ParentIndex = RefSkeleton.GetParentIndex(BoneIndex);
+
+			if (bLocalSpace)
+			{
+				// create partial local delta transform
+				FEulerTransform LocalDeltaTransform = FEulerTransform::Identity;
+				LocalDeltaTransform.Rotation = LocalRotationDelta;
+				// get parent global transform to calculate local translation delta
+				FTransform ParentRefGlobalTransform = FTransform::Identity;
+				if (ParentIndex != INDEX_NONE)
+				{
+					ParentRefGlobalTransform = FAnimationRuntime::GetComponentSpaceTransform(RefSkeleton, RefSkeleton.GetRefBonePose(), ParentIndex);
+				}
+				// now calculate local translational delta from global
+				LocalDeltaTransform.SetLocation(ParentRefGlobalTransform.InverseTransformVector(GlobalTranslationDelta));
+				return LocalDeltaTransform;
+			}
+			else
+			{
+				// get the CURRENT parent global transform and reference LOCAL transform to calculate the
+				// current global transform of the bone without any offsets applied
+				FTransform ParentGlobalTransform = FTransform::Identity;
+				if (ParentIndex != INDEX_NONE)
+				{
+					ParentGlobalTransform = Mesh->GetBoneTransform(ParentIndex, FTransform::Identity);
+				}
+				FTransform LocalRefTransform = RefSkeleton.GetRefBonePose()[BoneIndex];
+				FTransform GlobalTransform = LocalRefTransform * ParentGlobalTransform;
+				// get global rotation plus delta
+				FQuat GlobalRotationPlusDelta = GlobalTransform.GetRotation() * LocalRotationDelta.Quaternion();
+				// get global delta rotation
+				FQuat GlobalDeltaRotation = GlobalRotationPlusDelta * GlobalTransform.GetRotation().Inverse();
+				// combine with translation delta
+				return FEulerTransform(GlobalTranslationDelta, GlobalDeltaRotation.Rotator(), FVector::OneVector);
+			}
+		}
+	default:
+		checkNoEntry();
+	}
+
+	return FEulerTransform::Identity;
 }
 
-void FIKRetargeterDetails::CustomizeDetails(IDetailLayoutBuilder& DetailBuilder)
+bool UIKRetargetBoneDetails::IsComponentRelative(
+	ESlateTransformComponent::Type Component,
+	EIKRetargetTransformType TransformType) const
 {
-	const TObjectPtr<UIKRetargeterController> Controller = GetAssetControllerFromSelectedObjects(DetailBuilder);
-	if (!Controller)
+	switch(TransformType)
+	{
+	case EIKRetargetTransformType::Current:
+		{
+			return CurrentTransformRelative[(int32)Component]; 
+		}
+	case EIKRetargetTransformType::Reference:
+		{
+			return ReferenceTransformRelative[(int32)Component]; 
+		}
+	case EIKRetargetTransformType::RelativeOffset:
+		{
+			return RelativeOffsetTransformRelative[(int32)Component];
+		}
+	}
+	return true;
+}
+
+void UIKRetargetBoneDetails::OnComponentRelativeChanged(
+	ESlateTransformComponent::Type Component,
+	bool bIsRelative,
+	EIKRetargetTransformType TransformType)
+{
+	switch(TransformType)
+	{
+	case EIKRetargetTransformType::Current:
+		{
+			CurrentTransformRelative[(int32)Component] = bIsRelative;
+			break; 
+		}
+	case EIKRetargetTransformType::Reference:
+		{
+			ReferenceTransformRelative[(int32)Component] = bIsRelative;
+			break; 
+		}
+	case EIKRetargetTransformType::RelativeOffset:
+		{
+			RelativeOffsetTransformRelative[(int32)Component] = bIsRelative;
+			break; 
+		}
+	}
+}
+
+void UIKRetargetBoneDetails::OnCopyToClipboard(
+	ESlateTransformComponent::Type Component,
+	EIKRetargetTransformType TransformType) const
+{
+	// get is local or global space
+	bool bIsRelative = false;
+	switch(TransformType)
+	{
+	case EIKRetargetTransformType::Current:
+		{
+			bIsRelative = CurrentTransformRelative[(int32)Component];
+			break;
+		}
+	case EIKRetargetTransformType::Reference:
+		{
+			bIsRelative = ReferenceTransformRelative[(int32)Component];
+			break;
+		}
+	case EIKRetargetTransformType::RelativeOffset:
+		{
+			bIsRelative = RelativeOffsetTransformRelative[(int32)Component];
+			break;
+		}
+	default:
+		checkNoEntry();
+	}
+
+	// get the transform of correct type and space
+	const FEulerTransform Transform = GetTransform(TransformType, bIsRelative);
+	
+	FString Content;
+	switch(Component)
+	{
+	case ESlateTransformComponent::Location:
+		{
+			GetContentFromData(Transform.GetLocation(), Content);
+			break;
+		}
+	case ESlateTransformComponent::Rotation:
+		{
+			GetContentFromData(Transform.Rotator(), Content);
+			break;
+		}
+	case ESlateTransformComponent::Scale:
+		{
+			GetContentFromData(Transform.GetScale3D(), Content);
+			break;
+		}
+	case ESlateTransformComponent::Max:
+	default:
+		{
+			GetContentFromData(Transform, Content);
+			TBaseStructure<FTransform>::Get()->ExportText(Content, &Transform, &Transform, nullptr, PPF_None, nullptr);
+			break;
+		}
+	}
+
+	if(!Content.IsEmpty())
+	{
+		FPlatformApplicationMisc::ClipboardCopy(*Content);
+	}
+}
+
+void UIKRetargetBoneDetails::OnPasteFromClipboard(
+	ESlateTransformComponent::Type Component,
+	EIKRetargetTransformType TransformType)
+{
+	// only allow editing of relative offsets in retarget poses
+	if(TransformType != EIKRetargetTransformType::RelativeOffset)
 	{
 		return;
 	}
 
-	FIKRetargetEditorController* EditorController = Controller->GetEditorController();
-	if (!EditorController)
+	// must have valid controller
+	UIKRetargeterController* AssetController = EditorController->AssetController;
+	if(!AssetController)
+	{
+		return;
+	}
+	
+	// get content of clipboard to paste
+	FString Content;
+	FPlatformApplicationMisc::ClipboardPaste(Content);
+	if(Content.IsEmpty())
 	{
 		return;
 	}
+	
+	class FRetargetPasteTransformWidgetErrorPipe : public FOutputDevice
+	{
+	public:
 
-	// the commands for the menus
-	TSharedPtr<FUICommandList> Commands = Controller->GetEditorController()->Editor.Pin()->GetToolkitCommands();
+		int32 NumErrors;
 
-	// add a new category at the top to edit the retarget pose
-	IDetailCategoryBuilder& EditPoseCategoryBuilder = DetailBuilder.EditCategory( "Edit Retarget Pose", LOCTEXT("EditRetargetPoseLabel", "Edit Retarget Pose"), ECategoryPriority::Default );
+		FRetargetPasteTransformWidgetErrorPipe(FIKRigLogger* InLog)	: FOutputDevice(), NumErrors(0), Log(InLog) {}
 
-	// fill list of pose names
-	PoseNames.Reset();
-	for (const TTuple<FName, FIKRetargetPose>& Pose : Controller->GetRetargetPoses())
+		virtual void Serialize(const TCHAR* V, ELogVerbosity::Type Verbosity, const class FName& Category) override
+		{
+			if (Log)
+			{
+				Log->LogError(LOCTEXT("RetargetPasteError", "Error pasting transform data to bone."));
+			}
+			
+			NumErrors++;
+		}
+
+		FIKRigLogger* Log;
+	};
+
+	FIKRigLogger* Log = nullptr;
+	if (const auto Processor = EditorController->GetRetargetProcessor())
 	{
-		PoseNames.Add(MakeShareable(new FName(Pose.Key)));
+		Log = &Processor->Log;
 	}
+	FRetargetPasteTransformWidgetErrorPipe ErrorPipe(Log);
+	
+	// get the transform of correct type and space
+	const bool bIsRelative = IsComponentRelative(Component, TransformType);
+	FEulerTransform Transform = GetTransform(TransformType, bIsRelative);
 
-	// add pose selection combo box
-	FDetailWidgetRow& CurrentPoseRow = EditPoseCategoryBuilder.AddCustomRow(LOCTEXT("CurrentPoseLabel", "Current Pose"))
-	.NameContent()
-	[
-		SNew(STextBlock)
-		.Text(LOCTEXT("CurrentPose", "Current Retarget Pose"))
-		.Font(DetailBuilder.GetDetailFont())
-	]
-	.ValueContent()
-	[
-		SNew(SComboBox<TSharedPtr<FName>>)
-		.OptionsSource(&PoseNames)
-		.OnGenerateWidget_Lambda([](TSharedPtr<FName> InItem)
+	// create a transaction on the asset
+	FScopedTransaction Transaction(LOCTEXT("PasteTransform", "Paste Transform"));
+	EditorController.Get()->AssetController->GetAsset()->Modify();
+	
+	switch(Component)
+	{
+	case ESlateTransformComponent::Location:
 		{
-			return SNew(STextBlock).Text(FText::FromName(*InItem.Get()));
-		})
-		.OnSelectionChanged(EditorController, &FIKRetargetEditorController::OnPoseSelected)
-		[
-			SNew(STextBlock).Text(EditorController, &FIKRetargetEditorController::GetCurrentPoseName)
-		]
-	];
-
-	// add pose editing toolbar
-	FDetailWidgetRow& ToolbarRow = EditPoseCategoryBuilder.AddCustomRow(LOCTEXT("EditPoseToolbarLabel", "Edit Pose"))
-	.WholeRowWidget
-	[
-		SNew(SHorizontalBox)
-		+SHorizontalBox::Slot()
-		.FillWidth(1.0f)
-		.HAlign(HAlign_Center)
-		[
-			MakeToolbar(Commands)
-		]
-	];
+			FVector Data = Transform.GetLocation();
+			const TCHAR* Result = TBaseStructure<FVector>::Get()->ImportText(*Content, &Data, nullptr, PPF_None, &ErrorPipe, TBaseStructure<FVector>::Get()->GetName(), true);
+			if (Result && ErrorPipe.NumErrors == 0)
+			{
+				Transform.SetLocation(Data);
+				EditorController->AssetController->SetTranslationOffsetOnRetargetRootBone(Transform.GetLocation());
+			}
+			break;
+		}
+	case ESlateTransformComponent::Rotation:
+		{
+			FRotator Data = Transform.Rotator();
+			const TCHAR* Result = TBaseStructure<FRotator>::Get()->ImportText(*Content, &Data, nullptr, PPF_None, &ErrorPipe, TBaseStructure<FRotator>::Get()->GetName(), true);
+			if (Result && ErrorPipe.NumErrors == 0)
+			{
+				Transform.SetRotator(Data);
+				EditorController->AssetController->SetRotationOffsetForRetargetPoseBone(SelectedBone, Transform.GetRotation());
+			}
+			break;
+		}
+	default:
+		checkNoEntry();
+	}
 }
 
-TSharedRef<SWidget>  FIKRetargeterDetails::MakeToolbar(TSharedPtr<FUICommandList> Commands)
+TOptional<FVector::FReal> UIKRetargetBoneDetails::GetNumericValue(
+	EIKRetargetTransformType TransformType,
+	ESlateTransformComponent::Type Component,
+	ESlateRotationRepresentation::Type Representation,
+	ESlateTransformSubComponent::Type SubComponent)
 {
-	FToolBarBuilder ToolbarBuilder(Commands, FMultiBoxCustomization::None);
+	FEulerTransform Transform = FEulerTransform::Identity;
+	
+	switch(TransformType)
+	{
+	case EIKRetargetTransformType::Current:
+		{
+			Transform = GetTransform(TransformType, CurrentTransformRelative[(int32)Component]);
+			break;
+		}
 	
-	ToolbarBuilder.BeginSection("Edit Current Pose");
-
-	ToolbarBuilder.AddToolBarButton(
-		FIKRetargetCommands::Get().EditRetargetPose,
-		NAME_None,
-		TAttribute<FText>(),
-		TAttribute<FText>(),
-		FSlateIcon(FAppStyle::Get().GetStyleSetName(),"Icons.Edit"));
-
-	ToolbarBuilder.AddComboButton(
-		FUIAction(),
-		FOnGetContent::CreateSP(this, &FIKRetargeterDetails::GenerateResetMenuContent, Commands),
-		LOCTEXT("ResetPose_Label", "Reset"),
-		LOCTEXT("ResetPoseToolTip_Label", "Reset bones to reference pose."),
-		FSlateIcon(FAppStyle::GetAppStyleSetName(), "Icons.Refresh"));
-
-	ToolbarBuilder.EndSection();
-
-	ToolbarBuilder.BeginSection("Create Poses");
-
-	ToolbarBuilder.AddComboButton(
-		FUIAction(),
-		FOnGetContent::CreateSP(this, &FIKRetargeterDetails::GenerateNewMenuContent, Commands),
-		LOCTEXT("CreatePose_Label", "Create"),
-		TAttribute<FText>(),
-		FSlateIcon(FAppStyle::GetAppStyleSetName(), "Icons.Plus"));
-
-	ToolbarBuilder.AddToolBarButton(
-		FIKRetargetCommands::Get().DeleteRetargetPose,
-		NAME_None,
-		TAttribute<FText>(),
-		TAttribute<FText>(),
-		FSlateIcon(FAppStyle::Get().GetStyleSetName(),"Icons.Delete"));
-
-	ToolbarBuilder.AddToolBarButton(
-		FIKRetargetCommands::Get().RenameRetargetPose,
-		NAME_None,
-		TAttribute<FText>(),
-		TAttribute<FText>(),
-		FSlateIcon(FAppStyle::Get().GetStyleSetName(),"Icons.Settings"));
-
-	ToolbarBuilder.EndSection();
-
-	return ToolbarBuilder.MakeWidget();
+	case EIKRetargetTransformType::Reference:
+		{
+			Transform = GetTransform(TransformType, ReferenceTransformRelative[(int32)Component]);
+			break;
+		}
+
+	case EIKRetargetTransformType::RelativeOffset:
+		{
+			Transform = GetTransform(TransformType, RelativeOffsetTransformRelative[(int32)Component]);
+			break;
+		}
+	default:
+		checkNoEntry();
+	}
+
+	return CleanRealValue(SAdvancedTransformInputBox<FEulerTransform>::GetNumericValueFromTransform(Transform, Component, Representation, SubComponent));
 }
 
-TSharedRef<SWidget> FIKRetargeterDetails::GenerateResetMenuContent(TSharedPtr<FUICommandList> Commands)
+void UIKRetargetBoneDetails::OnNumericValueCommitted(
+	ESlateTransformComponent::Type Component,
+	ESlateRotationRepresentation::Type Representation,
+	ESlateTransformSubComponent::Type SubComponent,
+	FVector::FReal Value,
+	ETextCommit::Type CommitType,
+	EIKRetargetTransformType TransformType,
+	bool bIsCommit)
 {
-	FMenuBuilder MenuBuilder(true, Commands);
-
-	MenuBuilder.AddMenuEntry(
-		FIKRetargetCommands::Get().ResetSelectedBones,
-		TEXT("Reset Selected"),
-		TAttribute<FText>(),
-		TAttribute<FText>());
-
-	MenuBuilder.AddMenuEntry(
-		FIKRetargetCommands::Get().ResetSelectedAndChildrenBones,
-		TEXT("Reset Selected And Children"),
-		TAttribute<FText>(),
-		TAttribute<FText>());
-
-	MenuBuilder.AddMenuEntry(
-		FIKRetargetCommands::Get().ResetAllBones,
-		TEXT("Reset All"),
-		TAttribute<FText>(),
-		TAttribute<FText>());
-
-	return MenuBuilder.MakeWidget();
+	if(TransformType != EIKRetargetTransformType::RelativeOffset)
+	{
+		return;
+	}
+
+	UIKRetargeterController* AssetController = EditorController->AssetController;
+	if(!AssetController)
+	{
+		return;
+	}
+
+	// ensure we have a valid skeletal mesh
+	const bool bEditingSource = EditorController->GetSkeletonMode() == EIKRetargetSkeletonMode::Source;
+	UDebugSkelMeshComponent* Mesh = bEditingSource ? EditorController->SourceSkelMeshComponent : EditorController->TargetSkelMeshComponent;
+	if (!(Mesh && Mesh->SkeletalMesh))
+	{
+		return;
+	}
+
+	// ensure bone is valid
+	const FReferenceSkeleton& RefSkeleton = Mesh->SkeletalMesh->GetRefSkeleton();
+	const int32 BoneIndex = RefSkeleton.FindBoneIndex(SelectedBone);
+	if (BoneIndex == INDEX_NONE)
+	{
+		return;
+	}
+
+	switch(Component)
+	{
+	case ESlateTransformComponent::Location:
+		{
+			const bool bIsTranslationLocal = RelativeOffsetTransformRelative[0];
+			FTransform CurrentGlobalOffset = FTransform::Identity;
+			CurrentGlobalOffset.SetTranslation(AssetController->GetTranslationOffsetOnRetargetRootBone());
+			
+			if (bIsTranslationLocal)
+			{	
+				// get the current LOCAL offset
+				FTransform CurrentLocalOffset = CurrentGlobalOffset;
+				FTransform ParentGlobalRefTransform = FTransform::Identity;
+				const int32 ParentIndex = RefSkeleton.GetParentIndex(BoneIndex);
+				if (ParentIndex != INDEX_NONE)
+				{
+					ParentGlobalRefTransform = FAnimationRuntime::GetComponentSpaceTransform(RefSkeleton, RefSkeleton.GetRefBonePose(), ParentIndex);
+				}
+				CurrentLocalOffset = CurrentLocalOffset.GetRelativeTransform(ParentGlobalRefTransform);
+
+				// apply the numerical value to the local space values
+				SAdvancedTransformInputBox<FTransform>::ApplyNumericValueChange(CurrentLocalOffset, Value, Component, Representation, SubComponent);
+
+				// convert back to global space for storage in the pose
+				CurrentGlobalOffset = CurrentLocalOffset * ParentGlobalRefTransform;
+			}
+			else
+			{
+				// apply the edit
+				SAdvancedTransformInputBox<FTransform>::ApplyNumericValueChange(CurrentGlobalOffset, Value, Component, Representation, SubComponent);
+			}
+			
+			// store the new transform in the retarget pose
+			FScopedTransaction Transaction(LOCTEXT("EditRootTranslation", "Edit Retarget Root Pose Translation"));
+			EditorController->AssetController->GetAsset()->Modify();
+			EditorController->AssetController->SetTranslationOffsetOnRetargetRootBone(CurrentGlobalOffset.GetTranslation());
+			
+			break;
+		}
+	case ESlateTransformComponent::Rotation:
+		{
+			const bool bIsRotationLocal = RelativeOffsetTransformRelative[1];
+			const FQuat LocalRotationDelta = AssetController->GetRotationOffsetForRetargetPoseBone(SelectedBone);
+			FEulerTransform LocalDeltaTransform = FEulerTransform(FVector::ZeroVector, LocalRotationDelta.Rotator(), FVector::OneVector);
+			FQuat NewLocalRotationDelta;
+			
+			if (bIsRotationLocal)
+			{
+				// rotations are stored in local space, so just apply the edit
+				SAdvancedTransformInputBox<FEulerTransform>::ApplyNumericValueChange(LocalDeltaTransform, Value, Component, Representation, SubComponent);
+				NewLocalRotationDelta = LocalDeltaTransform.GetRotation();
+			}
+			else
+			{
+				FTransform ParentGlobalTransform = FTransform::Identity;
+				const int32 ParentIndex = RefSkeleton.GetParentIndex(BoneIndex);
+				if (ParentIndex != INDEX_NONE)
+				{
+					ParentGlobalTransform = Mesh->GetBoneTransform(ParentIndex, FTransform::Identity);
+				}
+				FTransform LocalRefTransform = RefSkeleton.GetRefBonePose()[BoneIndex];
+				FTransform GlobalTransform = LocalRefTransform * ParentGlobalTransform;
+				
+				// get reference global transform
+				//FTransform GlobalTransform = Mesh->GetBoneTransform(BoneIndex, FTransform::Identity);
+				// get offset global transform
+				FQuat GlobalRefRotationPlusDelta = GlobalTransform.GetRotation() * LocalRotationDelta;
+				// get global delta
+				FQuat GlobalDeltaRotation = GlobalRefRotationPlusDelta * GlobalTransform.GetRotation().Inverse();
+				// apply edit to global delta
+				FEulerTransform GlobalDeltaTransform = FEulerTransform(FVector::ZeroVector, GlobalDeltaRotation.Rotator(), FVector::OneVector);
+				SAdvancedTransformInputBox<FEulerTransform>::ApplyNumericValueChange(GlobalDeltaTransform, Value, Component, Representation, SubComponent);
+
+				// convert world space delta quaternion to bone-space
+				GlobalDeltaRotation = GlobalDeltaTransform.GetRotation();
+				const FVector RotationAxis = GlobalDeltaRotation.GetRotationAxis();
+				const FVector UnRotatedAxis = GlobalTransform.InverseTransformVector(RotationAxis);
+				NewLocalRotationDelta = FQuat(UnRotatedAxis, GlobalDeltaRotation.GetAngle());
+			}
+			
+			// store the new rotation in the retarget pose
+			FScopedTransaction Transaction(LOCTEXT("EditRootTranslation", "Edit Retarget Pose Rotation"));
+			EditorController->AssetController->GetAsset()->Modify();
+			EditorController->AssetController->SetRotationOffsetForRetargetPoseBone(SelectedBone, NewLocalRotationDelta);
+			break;
+		}
+	default:
+		checkNoEntry();
+	}
 }
 
-TSharedRef<SWidget> FIKRetargeterDetails::GenerateNewMenuContent(TSharedPtr<FUICommandList> Commands)
+bool UIKRetargetBoneDetails::IsRootBone() const
 {
-	FMenuBuilder MenuBuilder(true, Commands);
+	const bool bIsSource = EditorController->GetSkeletonMode() == EIKRetargetSkeletonMode::Source;
+	const FName RootBone = bIsSource ? EditorController->AssetController->GetSourceRootBone() :  EditorController->AssetController->GetTargetRootBone();
+	return SelectedBone == RootBone;
+}
 
-	MenuBuilder.BeginSection("Create", LOCTEXT("CreatePoseOperations", "Create New Retarget Pose"));
-	{
-		MenuBuilder.AddMenuEntry(
-		FIKRetargetCommands::Get().NewRetargetPose,
-		TEXT("Create"),
-		TAttribute<FText>(),
-		TAttribute<FText>());
+void UIKRetargetBoneDetails::OnMultiNumericValueCommitted(
+	ESlateTransformComponent::Type Component,
+	ESlateRotationRepresentation::Type Representation,
+	ESlateTransformSubComponent::Type SubComponent,
+	FVector::FReal Value,
+	ETextCommit::Type CommitType,
+	EIKRetargetTransformType TransformType,
+	TArrayView<UIKRetargetBoneDetails*> Bones,
+	bool bIsCommit)
+{
+	for(UIKRetargetBoneDetails* Bone : Bones)
+	{	
+		Bone->OnNumericValueCommitted(Component, Representation, SubComponent, Value, CommitType, TransformType, bIsCommit);
+	}
+}
 
-		MenuBuilder.AddMenuEntry(
-			FIKRetargetCommands::Get().DuplicateRetargetPose,
-			TEXT("Create"),
-			TAttribute<FText>(),
-			TAttribute<FText>());
+template <typename DataType>
+void UIKRetargetBoneDetails::GetContentFromData(const DataType& InData, FString& Content) const
+{
+	TBaseStructure<DataType>::Get()->ExportText(Content, &InData, &InData, nullptr, PPF_None, nullptr);
+}
+
+TOptional<FVector::FReal> UIKRetargetBoneDetails::CleanRealValue(TOptional<FVector::FReal> InValue)
+{
+	// remove insignificant decimal noise and sign bit if value near zero
+	if (InValue.IsSet() && FMath::IsNearlyZero(InValue.GetValue(), UE_KINDA_SMALL_NUMBER))
+	{
+		InValue = 0.0f;
 	}
-	MenuBuilder.EndSection();
+	
+	return InValue;
+}
 
-	MenuBuilder.BeginSection("Import",LOCTEXT("ImportPoseOperations", "Import Retarget Pose"));
+// ------------------------------------------- BEGIN FIKRetargetBoneDetailCustomization -------------------------------
+
+void FIKRetargetBoneDetailCustomization::CustomizeDetails(IDetailLayoutBuilder& DetailBuilder)
+{
+	TArray<TWeakObjectPtr<UObject>> ObjectsBeingCustomized = DetailBuilder.GetSelectedObjects();
+	Bones.Reset();
+	for (const TWeakObjectPtr<UObject>& Object : ObjectsBeingCustomized)
 	{
-		MenuBuilder.AddMenuEntry(
-		FIKRetargetCommands::Get().ImportRetargetPose,
-		TEXT("Import"),
-		TAttribute<FText>(),
-		TAttribute<FText>());
+		if (UIKRetargetBoneDetails* Bone = Cast<UIKRetargetBoneDetails>(Object.Get()))
+		{
+			Bones.Add(Bone);
+		}
+	}
 	
-		MenuBuilder.AddMenuEntry(
-			FIKRetargetCommands::Get().ImportRetargetPoseFromAnim,
-			TEXT("ImportFromSequence"),
-			TAttribute<FText>(),
-			TAttribute<FText>());
+	if (Bones.IsEmpty())
+	{
+		return;
 	}
-	MenuBuilder.EndSection();
 
-	MenuBuilder.BeginSection("Export",LOCTEXT("EmportPoseOperations", "Export Retarget Pose"));
+	if (!Bones[0]->EditorController.IsValid())
 	{
-		MenuBuilder.AddMenuEntry(
-			FIKRetargetCommands::Get().ExportRetargetPose,
-			TEXT("Export"),
-			TAttribute<FText>(),
-			TAttribute<FText>());
+		return;
 	}
-	MenuBuilder.EndSection();
+	
+	const FIKRetargetEditorController& Controller = *Bones[0]->EditorController.Get();
+	const UIKRetargeterController* AssetController = Controller.AssetController;
+	
+	const bool bIsSourceMode = Controller.GetSkeletonMode() == EIKRetargetSkeletonMode::Source;
+	const bool bIsEditingPose = Controller.IsEditingPose();
+
+	const FName CurrentRootName = bIsSourceMode ?AssetController->GetSourceRootBone() : AssetController->GetTargetRootBone();
+	const bool bIsRootSelected =  Bones[0]->SelectedBone == CurrentRootName;
+
+	FIKRetargetTransformUIData UIData;
+	GetTransformUIData( bIsEditingPose, DetailBuilder, UIData);
+
+	TSharedPtr<SSegmentedControl<EIKRetargetTransformType>> TransformChoiceWidget =
+		SSegmentedControl<EIKRetargetTransformType>::Create(
+			UIData.TransformTypes,
+			UIData.ButtonLabels,
+			UIData.ButtonTooltips,
+			UIData.VisibleTransforms
+		);
+
+	DetailBuilder.EditCategory(TEXT("Selection")).SetSortOrder(1);
 
-	return MenuBuilder.MakeWidget();
+	IDetailCategoryBuilder& CategoryBuilder = DetailBuilder.EditCategory(TEXT("Transforms"));
+	CategoryBuilder.SetSortOrder(2);
+	CategoryBuilder.AddCustomRow(FText::FromString(TEXT("TransformType")))
+	.ValueContent()
+	.MinDesiredWidth(375.f)
+	.MaxDesiredWidth(375.f)
+	.HAlign(HAlign_Left)
+	[
+		SNew(SHorizontalBox)
+		+SHorizontalBox::Slot()
+		.AutoWidth()
+		.HAlign(HAlign_Left)
+		.VAlign(VAlign_Center)
+		[
+			TransformChoiceWidget.ToSharedRef()
+		]
+	];
+
+	SAdvancedTransformInputBox<FTransform>::FArguments TransformWidgetArgs = SAdvancedTransformInputBox<FTransform>::FArguments()
+	.ConstructLocation(!bIsEditingPose || (bIsEditingPose && bIsRootSelected))
+	.ConstructRotation(true)
+	.ConstructScale(!bIsEditingPose)
+	.DisplayRelativeWorld(true)
+	.DisplayScaleLock(false)
+	.AllowEditRotationRepresentation(true)
+	.Font(IDetailLayoutBuilder::GetDetailFont())
+	.UseQuaternionForRotation(true);
+
+	TArrayView<UIKRetargetBoneDetails*> BonesView = TArrayView<UIKRetargetBoneDetails*>(Bones);
+	
+	for(int32 PropertyIndex=0;PropertyIndex<UIData.Properties.Num();PropertyIndex++)
+	{
+		const EIKRetargetTransformType TransformType = UIData.TransformTypes[PropertyIndex];
+
+		// only enable editing of the relative offset transform type while in edit mode
+		const bool bIsEditable = bIsEditingPose && TransformType == EIKRetargetTransformType::RelativeOffset;
+		
+		TransformWidgetArgs.IsEnabled(bIsEditable);
+		// edit transform
+		if(bIsEditable)
+		{
+			TransformWidgetArgs.OnNumericValueChanged_Static(
+				&UIKRetargetBoneDetails::OnMultiNumericValueCommitted,
+				ETextCommit::Default,
+				TransformType,
+				BonesView,
+				false);
+		}
+
+		// get/set relative
+		TransformWidgetArgs.OnGetIsComponentRelative_Lambda( [bIsEditable, BonesView, TransformType](ESlateTransformComponent::Type InComponent)
+		{
+			return BonesView.ContainsByPredicate( [&](const UIKRetargetBoneDetails* Bone)
+			{
+				return Bone->IsComponentRelative(InComponent, TransformType);
+			} );
+		})
+		.OnIsComponentRelativeChanged_Lambda( [bIsEditable, BonesView, TransformType](ESlateTransformComponent::Type InComponent, bool bIsRelative)
+		{
+			for (UIKRetargetBoneDetails* Bone: BonesView)
+			{
+				Bone->OnComponentRelativeChanged(InComponent, bIsRelative, TransformType);
+			}
+		} );
+
+		TransformWidgetArgs.OnGetNumericValue_Lambda([BonesView, TransformType](
+		ESlateTransformComponent::Type Component,
+		ESlateRotationRepresentation::Type Representation,
+		ESlateTransformSubComponent::Type SubComponent) -> TOptional<FVector::FReal>
+		{
+			TOptional<FVector::FReal> FirstValue = BonesView[0]->GetNumericValue(TransformType, Component, Representation, SubComponent);
+			
+			if (FirstValue)
+			{
+				for (int32 Index = 1; Index < BonesView.Num(); Index++)
+				{
+					const TOptional<FVector::FReal> CurrentValue = BonesView[Index]->GetNumericValue(TransformType, Component, Representation, SubComponent);
+					if (CurrentValue.IsSet())
+					{
+						// using a more permissive precision to avoid "Multiple Values" in details panel caused
+						// by floating point noise introduced by normal rotation calculations
+						constexpr double EDITING_PRECISION = 1.e-2;
+						if(! FMath::IsNearlyEqual(FirstValue.GetValue(), CurrentValue.GetValue(), EDITING_PRECISION))
+						{
+							return TOptional<FVector::FReal>();
+						}
+					}
+				}
+			}
+			return FirstValue;
+		});
+
+		// copy/paste bones transforms
+		TransformWidgetArgs.OnCopyToClipboard_UObject(Bones[0], &UIKRetargetBoneDetails::OnCopyToClipboard, TransformType);
+		TransformWidgetArgs.OnPasteFromClipboard_UObject(Bones[0], &UIKRetargetBoneDetails::OnPasteFromClipboard, TransformType);
+
+		TransformWidgetArgs.Visibility_Lambda([TransformChoiceWidget, TransformType]() -> EVisibility
+		{
+			return TransformChoiceWidget->HasValue(TransformType) ? EVisibility::Visible : EVisibility::Collapsed;
+		});
+
+		SAdvancedTransformInputBox<FTransform>::ConstructGroupedTransformRows(
+			CategoryBuilder, 
+			UIData.ButtonLabels[PropertyIndex], 
+			UIData.ButtonTooltips[PropertyIndex], 
+			TransformWidgetArgs);
+	}
 }
 
-TObjectPtr<UIKRetargeterController> FIKRetargeterDetails::GetAssetControllerFromSelectedObjects(IDetailLayoutBuilder& DetailBuilder) const
+void FIKRetargetBoneDetailCustomization::GetTransformUIData(
+	const bool bIsEditingPose,
+	const IDetailLayoutBuilder& DetailBuilder,
+	FIKRetargetTransformUIData& OutData) const
 {
-	TArray< TWeakObjectPtr<UObject> > OutObjects;
-	DetailBuilder.GetObjectsBeingCustomized(OutObjects);
-	if (OutObjects.IsEmpty())
+	// read-only
+	static const TArray<EIKRetargetTransformType> ReadOnlyTransformTypes
 	{
-		return nullptr;
-	}
+		EIKRetargetTransformType::Current,
+		EIKRetargetTransformType::Reference
+	};
+	static const TArray<FText> ReadOnlyButtonLabels
+	{
+		LOCTEXT("CurrentTransform", "Current"),
+		LOCTEXT("ReferenceTransform", "Reference")
+	};
+	static const TArray<FText> ReadOnlyButtonTooltips
+	{
+		LOCTEXT("CurrentBoneTransformTooltip", "The current transform of the bone"),
+		LOCTEXT("ReferenceBoneTransformTooltip", "The reference transform of the bone")
+	};
+	static TAttribute<TArray<EIKRetargetTransformType>> ReadOnlyVisibleTransforms =
+		TArray<EIKRetargetTransformType>({EIKRetargetTransformType::Current});
 
-	const TObjectPtr<UIKRetargeter> Asset = CastChecked<UIKRetargeter>(OutObjects[0].Get());
-	return UIKRetargeterController::GetController(Asset);
+
+	// editable
+	static const TArray<EIKRetargetTransformType> EditableTransformTypes
+	{
+		EIKRetargetTransformType::RelativeOffset,
+		EIKRetargetTransformType::Reference
+	};
+	static const TArray<FText> EditableButtonLabels
+	{
+		LOCTEXT("RelativeOffsetTransform", "Relative Offset"),
+		LOCTEXT("ReferenceTransform", "Reference")
+	};
+	static const TArray<FText> EditableButtonTooltips
+	{
+		LOCTEXT("RelativeOffsetBoneTransformTooltip", "The offset transform in the current retarget pose, relative to the reference pose."),
+		LOCTEXT("ReferenceBoneTransformTooltip", "The reference transform of the bone.")
+	};
+	static TAttribute<TArray<EIKRetargetTransformType>> EditableVisibleTransforms =
+		TArray<EIKRetargetTransformType>({EIKRetargetTransformType::RelativeOffset});
+
+
+	if (bIsEditingPose)
+	{
+		OutData.TransformTypes = EditableTransformTypes;
+		OutData.ButtonLabels = EditableButtonLabels;
+		OutData.ButtonTooltips = EditableButtonTooltips;
+		OutData.VisibleTransforms = EditableVisibleTransforms;
+
+		OutData.Properties.Append(
+			{
+			DetailBuilder.GetProperty(GET_MEMBER_NAME_CHECKED(UIKRetargetBoneDetails, OffsetTransform)),
+			DetailBuilder.GetProperty(GET_MEMBER_NAME_CHECKED(UIKRetargetBoneDetails, ReferenceTransform))
+			});
+	}
+	else
+	{
+		OutData.TransformTypes = ReadOnlyTransformTypes;
+		OutData.ButtonLabels = ReadOnlyButtonLabels;
+		OutData.ButtonTooltips = ReadOnlyButtonTooltips;
+		OutData.VisibleTransforms = ReadOnlyVisibleTransforms;
+
+		OutData.Properties.Append(
+			{
+			DetailBuilder.GetProperty(GET_MEMBER_NAME_CHECKED(UIKRetargetBoneDetails, CurrentTransform)),
+			DetailBuilder.GetProperty(GET_MEMBER_NAME_CHECKED(UIKRetargetBoneDetails, ReferenceTransform))
+			});
+	}
 }
 
 #undef LOCTEXT_NAMESPACE
diff --git a/Engine/Plugins/Animation/IKRig/Source/IKRigEditor/Private/RetargetEditor/IKRetargetEditPoseMode.cpp b/Engine/Plugins/Animation/IKRig/Source/IKRigEditor/Private/RetargetEditor/IKRetargetEditPoseMode.cpp
--- a/Engine/Plugins/Animation/IKRig/Source/IKRigEditor/Private/RetargetEditor/IKRetargetEditPoseMode.cpp
+++ b/Engine/Plugins/Animation/IKRig/Source/IKRigEditor/Private/RetargetEditor/IKRetargetEditPoseMode.cpp
@@ -4,12 +4,10 @@
 
 #include "EditorViewportClient.h"
 #include "AssetEditorModeManager.h"
-#include "EditorScriptingHelpers.h"
 #include "IKRigDebugRendering.h"
 #include "IKRigProcessor.h"
 #include "IPersonaPreviewScene.h"
 #include "Animation/DebugSkelMeshComponent.h"
-#include "Materials/MaterialInstanceDynamic.h"
 #include "Preferences/PersonaOptions.h"
 #include "RetargetEditor/IKRetargetAnimInstance.h"
 #include "RetargetEditor/IKRetargetHitProxies.h"
@@ -86,7 +84,6 @@ void FIKRetargetEditPoseMode::RenderSkeleton(
 	FVector Offset;
 	GetEditedComponentScaleAndOffset(Scale, Offset);
 	const float BoneDrawSize = Asset->BoneDrawSize;
-	const int32 RootBoneIndex = GetEditedRetargetRootBoneIndex(*Processor);
 	const float MaxDrawRadius = Controller->TargetSkelMeshComponent->Bounds.SphereRadius * 0.01f;
 
 	// use colors from user preferences
@@ -98,10 +95,10 @@ void FIKRetargetEditPoseMode::RenderSkeleton(
 		// selected bones are drawn with different color, affected bones are highlighted
 		const bool bIsSelected = OutSelectedBones.Contains(BoneIndex);
 		const bool bIsAffected = OutAffectedBones.Contains(BoneIndex);
-		const bool bIsDisabled = !(Skeleton.IsBoneInAnyChain[BoneIndex] || BoneIndex==RootBoneIndex);
+		const bool bIsRetargeted = Processor->IsBoneRetargeted(BoneIndex, (int8)SkeletonMode);
 		FLinearColor LineColor = bIsAffected ? PersonaOptions->AffectedBoneColor : PersonaOptions->DefaultBoneColor;
 		LineColor = bIsSelected ? PersonaOptions->SelectedBoneColor : LineColor;
-		LineColor = bIsDisabled ? PersonaOptions->DisabledBoneColor : LineColor;
+		LineColor = bIsRetargeted ? LineColor : FLinearColor::LerpUsingHSV(LineColor, PersonaOptions->DisabledBoneColor, 0.35f);
 		
 		const float BoneRadius = FMath::Min(1.0f, MaxDrawRadius) * BoneDrawSize;
 		
@@ -131,10 +128,7 @@ void FIKRetargetEditPoseMode::RenderSkeleton(
 		}
 		
 		// draw the bone proxy
-		if (!bIsDisabled)
-		{
-			PDI->SetHitProxy(new HIKRetargetEditorBoneProxy(Skeleton.BoneNames[BoneIndex]));
-		}
+		PDI->SetHitProxy(new HIKRetargetEditorBoneProxy(Skeleton.BoneNames[BoneIndex]));
 		SkeletalDebugRendering::DrawWireBoneAdvanced(
 			PDI,
 			BoneTransform,
@@ -144,10 +138,7 @@ void FIKRetargetEditPoseMode::RenderSkeleton(
 			SDPG_Foreground,
 			BoneRadius,
 			bIsSelected || bIsAffected);
-		if (!bIsDisabled)
-		{
-			PDI->SetHitProxy(NULL);
-		}
+		PDI->SetHitProxy(NULL);
 	}
 }
 
@@ -372,6 +363,13 @@ bool FIKRetargetEditPoseMode::InputDelta(
 		return false; 
 	}
 
+	const UIKRetargetProcessor* Processor = Controller->GetRetargetProcessor();
+	if (!(Processor && Processor->IsInitialized()))
+	{
+		return false;
+	}
+	const FRetargetSkeleton& Skeleton = GetCurrentlyEditedSkeleton(*Processor);
+	
 	// rotating any bone
 	if (TrackingState == FIKRetargetTrackingState::RotatingBone)
 	{
@@ -381,23 +379,22 @@ bool FIKRetargetEditPoseMode::InputDelta(
 		}
 
 		// accumulate the rotation from the viewport gizmo (amount of rotation since tracking started)
-		BoneEdit.AccumulatedGlobalOffset = InRot.Quaternion().Inverse() * BoneEdit.AccumulatedGlobalOffset;
-
-		// convert to a delta rotation in the local space of the last selected bone
-		const FQuat BoneGlobalOrig = BoneEdit.GlobalTransform.GetRotation();
-		const FQuat BoneGlobalPlusOffset = BoneEdit.AccumulatedGlobalOffset * BoneGlobalOrig;
-		const FQuat ParentInv = BoneEdit.ParentGlobalTransform.GetRotation().Inverse();
-		const FQuat BoneLocal = ParentInv * BoneGlobalOrig;
-		const FQuat BoneLocalPlusOffset = ParentInv * BoneGlobalPlusOffset;
-		const FQuat BoneLocalOffset = BoneLocal * BoneLocalPlusOffset.Inverse();
+		BoneEdit.AccumulatedGlobalOffset = InRot.Quaternion() * BoneEdit.AccumulatedGlobalOffset;
+
+		// convert world space delta quaternion to bone-space
+		const FVector RotationAxis = BoneEdit.AccumulatedGlobalOffset.GetRotationAxis();
+		const FVector UnRotatedAxis = BoneEdit.GlobalTransform.InverseTransformVector(RotationAxis);
+		const FQuat BoneLocalDelta = FQuat(UnRotatedAxis, BoneEdit.AccumulatedGlobalOffset.GetAngle());
 		
 		// apply rotation delta to all selected bones
 		const TArray<FName>& SelectedBones = Controller->GetSelectedBones();
 		for (int32 SelectionIndex=0; SelectionIndex<SelectedBones.Num(); ++SelectionIndex)
 		{
-			const FQuat TotalOffsetLocal = BoneLocalOffset * BoneEdit.PrevLocalOffsets[SelectionIndex];
 			const FName& BoneName = SelectedBones[SelectionIndex];
-			Controller->AssetController->SetRotationOffsetForRetargetPoseBone(BoneName, TotalOffsetLocal);
+
+			// apply the new delta to the retarget pose
+			const FQuat TotalDeltaRotation = BoneEdit.PreviousDeltaRotation[SelectionIndex] * BoneLocalDelta;
+			Controller->AssetController->SetRotationOffsetForRetargetPoseBone(BoneName, TotalDeltaRotation);
 		}
 
 		return true;
@@ -469,8 +466,12 @@ void FIKRetargetEditPoseMode::Enter()
 	// put asset in mode where target mesh will output retarget pose (for preview purposes)
 	Controller->AssetController->GetAsset()->SetOutputMode(ERetargeterOutputMode::EditRetargetPose);
 
-	// clear selection
+	// clear bone edit
 	BoneEdit.Reset();
+
+	// deselect everything except bones
+	constexpr bool bKeepBoneSelection = true;
+	Controller->ClearSelection(bKeepBoneSelection);
 	
 	// record skeleton to edit (must be constant between enter/exit)
 	SkeletonMode = Controller->GetSkeletonMode();
@@ -512,12 +513,6 @@ void FIKRetargetEditPoseMode::GetEditedComponentScaleAndOffset(float& OutScale,
 	OutOffset = bIsSource ? Asset->SourceMeshOffset : Asset->TargetMeshOffset;
 }
 
-int32 FIKRetargetEditPoseMode::GetEditedRetargetRootBoneIndex(const UIKRetargetProcessor& Processor) const
-{
-	const bool bIsSource = EditorController.Pin()->GetSkeletonMode() == EIKRetargetSkeletonMode::Source;
-	return bIsSource ? Processor.GetSourceRetargetRoot() : Processor.GetTargetRetargetRoot();
-}
-
 void FIKRetargetEditPoseMode::UpdateWidgetTransform()
 {
 	const TSharedPtr<FIKRetargetEditorController> Controller = EditorController.Pin();
@@ -546,16 +541,18 @@ void FIKRetargetEditPoseMode::UpdateWidgetTransform()
 	FVector Offset;
 	GetEditedComponentScaleAndOffset(Scale,Offset);
 
+	const FIKRetargetPose& RetargetPose = AssetController->GetCurrentRetargetPose();
+
 	BoneEdit.Name = SelectedBones.Last();
 	BoneEdit.Index = Skeleton.FindBoneIndexByName(BoneEdit.Name);
 	BoneEdit.GlobalTransform = Controller->GetGlobalRetargetPoseOfBone(Skeleton, BoneEdit.Index, Scale, Offset);
 	BoneEdit.AccumulatedGlobalOffset = FQuat::Identity;
 
-	BoneEdit.PrevLocalOffsets.Reset();
+	BoneEdit.PreviousDeltaRotation.Reset();
 	for (int32 SelectionIndex=0; SelectionIndex<SelectedBones.Num(); ++SelectionIndex)
 	{
-		FQuat PrevLocalOffset = AssetController->GetRotationOffsetForRetargetPoseBone(SelectedBones[SelectionIndex]);
-		BoneEdit.PrevLocalOffsets.Add(PrevLocalOffset);
+		FQuat PrevDeltaRotation = RetargetProcessor->GetTargetSkeleton().GetRetargetPoseDeltaRotation(SelectedBones[SelectionIndex], &RetargetPose);
+		BoneEdit.PreviousDeltaRotation.Add(PrevDeltaRotation);
 	}
 	
 	const int32 ParentIndex = Skeleton.GetParentIndex(BoneEdit.Index);
diff --git a/Engine/Plugins/Animation/IKRig/Source/IKRigEditor/Private/RetargetEditor/IKRetargetEditorController.cpp b/Engine/Plugins/Animation/IKRig/Source/IKRigEditor/Private/RetargetEditor/IKRetargetEditorController.cpp
--- a/Engine/Plugins/Animation/IKRig/Source/IKRigEditor/Private/RetargetEditor/IKRetargetEditorController.cpp
+++ b/Engine/Plugins/Animation/IKRig/Source/IKRigEditor/Private/RetargetEditor/IKRetargetEditorController.cpp
@@ -32,6 +32,7 @@ void FIKRetargetEditorController::Initialize(TSharedPtr<FIKRetargetEditor> InEdi
 	AssetController = UIKRetargeterController::GetController(InAsset);
 	AssetController->SetEditorController(this);
 	SkeletonMode = EIKRetargetSkeletonMode::Target;
+	InAsset->SetOutputMode(ERetargeterOutputMode::ShowRetargetPose);
 
 	// bind callbacks when SOURCE or TARGET IK Rigs are modified
 	BindToIKRigAsset(AssetController->GetAsset()->GetSourceIKRigWriteable());
@@ -92,12 +93,14 @@ void FIKRetargetEditorController::OnRetargeterNeedsInitialized(UIKRetargeter* Re
 {
 	// clear the output log
 	ClearOutputLog();
+	
 	// force reinit the retarget processor (also inits the target IK Rig processor)
 	if (UIKRetargetProcessor* Processor = GetRetargetProcessor())
 	{
 		constexpr bool bSuppressWarnings = false;
 		Processor->Initialize(GetSourceSkeletalMesh(), GetTargetSkeletalMesh(), Retargeter, bSuppressWarnings);
 	}
+	
 	// refresh all the UI views
 	RefreshAllViews();
 }
@@ -127,6 +130,70 @@ void FIKRetargetEditorController::AddOffsetToMeshComponent(const FVector& Offset
 	MeshComponent->SetWorldScale3D(FVector(Scale,Scale,Scale));
 }
 
+bool FIKRetargetEditorController::IsBoneRetargeted(const FName& BoneName, EIKRetargetSkeletonMode WhichSkeleton) const
+{
+	// get an initialized processor
+	const UIKRetargetProcessor* Processor = GetRetargetProcessor();
+	if (!(Processor && Processor->IsInitialized()))
+	{
+		return false;
+	}
+
+	// get the bone index
+	const FRetargetSkeleton& Skeleton = WhichSkeleton == EIKRetargetSkeletonMode::Source ? Processor->GetSourceSkeleton() : Processor->GetTargetSkeleton();
+	const int32 BoneIndex = Skeleton.FindBoneIndexByName(BoneName);
+	if (BoneIndex == INDEX_NONE)
+	{
+		return false;
+	}
+
+	// return if it's a retargeted bone
+	return Processor->IsBoneRetargeted(BoneIndex, (int8)WhichSkeleton);
+}
+
+FName FIKRetargetEditorController::GetChainNameFromBone(const FName& BoneName, EIKRetargetSkeletonMode WhichSkeleton) const
+{
+	// get an initialized processor
+	const UIKRetargetProcessor* Processor = GetRetargetProcessor();
+	if (!(Processor && Processor->IsInitialized()))
+	{
+		return NAME_None;
+	}
+
+	// get the bone index
+	const FRetargetSkeleton& Skeleton = WhichSkeleton == EIKRetargetSkeletonMode::Source ? Processor->GetSourceSkeleton() : Processor->GetTargetSkeleton();
+	const int32 BoneIndex = Skeleton.FindBoneIndexByName(BoneName);
+	if (BoneIndex == INDEX_NONE)
+	{
+		return NAME_None;
+	}
+
+	return Processor->GetChainNameForBone(BoneIndex, (int8)WhichSkeleton);
+}
+
+TObjectPtr<UIKRetargetBoneDetails> FIKRetargetEditorController::GetDetailsObjectForBone(const FName& BoneName)
+{
+	if (AllBoneDetails.Contains(BoneName))
+	{
+		return AllBoneDetails[BoneName];
+	}
+
+	return CreateBoneDetails(BoneName);
+}
+
+TObjectPtr<UIKRetargetBoneDetails> FIKRetargetEditorController::CreateBoneDetails(const FName& BoneName)
+{
+	// create and store a new one
+	UIKRetargetBoneDetails* NewBoneDetails = NewObject<UIKRetargetBoneDetails>(AssetController->GetAsset(), FName(BoneName), RF_Standalone | RF_Transient );
+	NewBoneDetails->SelectedBone = BoneName;
+	NewBoneDetails->EditorController = SharedThis(this);
+
+	// store it in the map
+	AllBoneDetails.Add(BoneName, NewBoneDetails);
+	
+	return NewBoneDetails;
+}
+
 USkeletalMesh* FIKRetargetEditorController::GetSourceSkeletalMesh() const
 {
 	return AssetController ? AssetController->GetSourcePreviewMesh() : nullptr;
@@ -162,6 +229,11 @@ UDebugSkelMeshComponent* FIKRetargetEditorController::GetEditedSkeletalMesh() co
 	return SkeletonMode == EIKRetargetSkeletonMode::Source ? SourceSkelMeshComponent : TargetSkelMeshComponent;
 }
 
+const FRetargetSkeleton& FIKRetargetEditorController::GetCurrentlyEditedSkeleton(const UIKRetargetProcessor& Processor) const
+{
+	return SkeletonMode == EIKRetargetSkeletonMode::Source ? Processor.GetSourceSkeleton() : Processor.GetTargetSkeleton();
+}
+
 FTransform FIKRetargetEditorController::GetGlobalRetargetPoseOfBone(
 	const FRetargetSkeleton& Skeleton,
 	const int32& BoneIndex,
@@ -250,7 +322,12 @@ void FIKRetargetEditorController::ClearOutputLog() const
 void FIKRetargetEditorController::RefreshAllViews() const
 {
 	Editor.Pin()->RegenerateMenusAndToolbars();
-	DetailsView->ForceRefresh();
+
+	// refresh the details panel
+	if (DetailsView.IsValid())
+	{
+		DetailsView->ForceRefresh();
+	}
 
 	// cannot assume chains view is always available
 	if (ChainsView.IsValid())
@@ -400,28 +477,49 @@ void FIKRetargetEditorController::EditBoneSelection(
 	// apply selection to debug mesh component so rendering knows
 	DebugComponent->BonesOfInterest = SelectedBoneIndices;
 
-	// update views
+	// update hierarchy view
 	if (!bFromHierarchyView && HierarchyView.IsValid())
 	{
 		HierarchyView->RefreshTreeView();
 	}
+
+	// update details
+	if (SelectedBones.IsEmpty())
+	{
+		DetailsView->SetObject(AssetController->GetAsset());
+	}
+	else
+	{
+		SelectedBoneDetails.Reset();
+		for (const FName& SelectedBone : SelectedBones)
+		{
+			TObjectPtr<UIKRetargetBoneDetails> BoneDetails = GetDetailsObjectForBone(SelectedBone);
+			SelectedBoneDetails.Add(BoneDetails);
+		}
+		DetailsView->SetObjects(SelectedBoneDetails);
+	}
 }
 
-void FIKRetargetEditorController::ClearSelection()
+void FIKRetargetEditorController::ClearSelection(const bool bKeepBoneSelection)
 {
 	// clear mesh selection
 	SetSelectedMesh(nullptr);
 	
-	// clear bone selection
-	const TArray<FName> Empty;
-	constexpr bool bFromHierarchy = false;
-	EditBoneSelection(Empty, EBoneSelectionEdit::Replace, bFromHierarchy);
-	
 	// deselect all chains
 	ChainsView->ClearSelection();
 
-	// show global details
-	DetailsView->SetObject(AssetController->GetAsset());
+	// clear bone selection
+	if (!bKeepBoneSelection)
+	{
+		const TArray<FName> Empty;
+		constexpr bool bFromHierarchy = false;
+		EditBoneSelection(Empty, EBoneSelectionEdit::Replace, bFromHierarchy);
+
+		// show global details
+		DetailsView->SetObject(AssetController->GetAsset());
+	}
+
+	DetailsView->ForceRefresh();
 }
 
 void FIKRetargetEditorController::SetRetargeterMode(ERetargeterOutputMode Mode)
@@ -1046,7 +1144,7 @@ FReply FIKRetargetEditorController::OnImportPoseFromSequence()
 			const FTransform LocalTransformImported = UAnimPoseExtensions::GetBonePose(ImportedPose, BoneName, EAnimPoseSpaces::Local);
 			const FTransform LocalTransformReference = RefPose[BoneIndex];
 			const FQuat DeltaRotation = LocalTransformImported.GetRotation() * LocalTransformReference.GetRotation().Inverse();
-			// only if it's difference than the ref pose
+			// only if it's different than the ref pose
 			if (DeltaRotation.GetAngle() > KINDA_SMALL_NUMBER)
 			{
 				ImportedRetargetPose.BoneRotationOffsets.Add(BoneName, DeltaRotation);
diff --git a/Engine/Plugins/Animation/IKRig/Source/IKRigEditor/Private/RetargetEditor/IKRetargeterController.cpp b/Engine/Plugins/Animation/IKRig/Source/IKRigEditor/Private/RetargetEditor/IKRetargeterController.cpp
--- a/Engine/Plugins/Animation/IKRig/Source/IKRigEditor/Private/RetargetEditor/IKRetargeterController.cpp
+++ b/Engine/Plugins/Animation/IKRig/Source/IKRigEditor/Private/RetargetEditor/IKRetargeterController.cpp
@@ -445,7 +445,7 @@ void UIKRetargeterController::ResetRetargetPose(FName PoseToReset, const TArray<
 	{
 		FScopedTransaction Transaction(LOCTEXT("ResetRetargetPose", "Reset Retarget Pose"));
 		Asset->Modify();
-		
+
 		PoseToEdit.BoneRotationOffsets.Reset();
 		PoseToEdit.RootTranslationOffset = FVector::ZeroVector;
 	}
@@ -498,7 +498,7 @@ const FIKRetargetPose& UIKRetargeterController::GetCurrentRetargetPose() const
 	return GetAsset()->RetargetPoses[GetCurrentRetargetPoseName()];
 }
 
-void UIKRetargeterController::SetRotationOffsetForRetargetPoseBone(FName BoneName, FQuat RotationOffset) const
+void UIKRetargeterController::SetRotationOffsetForRetargetPoseBone(FName BoneName, const FQuat& RotationOffset) const
 {
 	const FIKRigSkeleton& Skeleton = Asset->GetTargetIKRig()->Skeleton;
 	Asset->RetargetPoses[Asset->CurrentRetargetPose].SetBoneRotationOffset(BoneName, RotationOffset, Skeleton);
@@ -515,9 +515,19 @@ FQuat UIKRetargeterController::GetRotationOffsetForRetargetPoseBone(FName BoneNa
 	return BoneOffsets[BoneName];
 }
 
+void UIKRetargeterController::SetTranslationOffsetOnRetargetRootBone(FVector TranslationOffset) const
+{
+	Asset->RetargetPoses[Asset->CurrentRetargetPose].SetRootTranslationDelta(TranslationOffset);
+}
+
+const FVector& UIKRetargeterController::GetTranslationOffsetOnRetargetRootBone() const
+{
+	return Asset->RetargetPoses[Asset->CurrentRetargetPose].RootTranslationOffset;
+}
+
 void UIKRetargeterController::AddTranslationOffsetToRetargetRootBone(FVector TranslationOffset) const
 {
-	Asset->RetargetPoses[Asset->CurrentRetargetPose].AddTranslationDeltaToRoot(TranslationOffset);
+	Asset->RetargetPoses[Asset->CurrentRetargetPose].AddToRootTranslationDelta(TranslationOffset);
 }
 
 FName UIKRetargeterController::MakePoseNameUnique(FString PoseName) const
diff --git a/Engine/Plugins/Animation/IKRig/Source/IKRigEditor/Private/RetargetEditor/SIKRetargetHierarchy.cpp b/Engine/Plugins/Animation/IKRig/Source/IKRigEditor/Private/RetargetEditor/SIKRetargetHierarchy.cpp
--- a/Engine/Plugins/Animation/IKRig/Source/IKRigEditor/Private/RetargetEditor/SIKRetargetHierarchy.cpp
+++ b/Engine/Plugins/Animation/IKRig/Source/IKRigEditor/Private/RetargetEditor/SIKRetargetHierarchy.cpp
@@ -3,6 +3,7 @@
 #include "RetargetEditor/SIKRetargetHierarchy.h"
 
 #include "Preferences/PersonaOptions.h"
+#include "RetargetEditor/SIKRetargetPoseEditor.h"
 #include "Widgets/Input/SSegmentedControl.h"
 #include "Retargeter/IKRetargetProcessor.h"
 #include "RigEditor/IKRigEditorStyle.h"
@@ -10,74 +11,103 @@
 
 #define LOCTEXT_NAMESPACE "SIKRetargetHierarchy"
 
-FIKRetargetTreeElement::FIKRetargetTreeElement(
-	const FText& InKey,
+static const FName BoneColumnName(TEXT("BoneName"));
+static const FName ChainColumnName(TEXT("RetargetChainName"));
+
+FIKRetargetHierarchyElement::FIKRetargetHierarchyElement(
+	const FName& InName,
 	const TSharedRef<FIKRetargetEditorController>& InEditorController)
-	: Key(InKey)
-	, EditorController(InEditorController)
+	: Key(FText::FromName(InName)),
+	Name(InName),
+	EditorController(InEditorController)
 {}
 
-TSharedRef<ITableRow> FIKRetargetTreeElement::MakeTreeRowWidget(
-	TSharedRef<FIKRetargetEditorController> InEditorController,
-	const TSharedRef<STableViewBase>& InOwnerTable,
-	TSharedRef<FIKRetargetTreeElement> InTreeElement,
-	TSharedRef<FUICommandList> InCommandList,
-	TSharedPtr<SIKRetargetHierarchy> InHierarchy)
+TSharedRef<SWidget> SIKRetargetHierarchyRow::GenerateWidgetForColumn(const FName& ColumnName)
 {
-	return SNew(SIKRetargetHierarchyItem, InEditorController, InOwnerTable, InTreeElement, InCommandList, InHierarchy);
-}
-
-void SIKRetargetHierarchyItem::Construct(
-	const FArguments& InArgs,
-	TSharedRef<FIKRetargetEditorController> InEditorController,
-	const TSharedRef<STableViewBase>& OwnerTable,
-	TSharedRef<FIKRetargetTreeElement> InTreeElement,
-	TSharedRef<FUICommandList> InCommandList,
-	TSharedPtr<SIKRetargetHierarchy> InHierarchy)
-{
-	WeakTreeElement = InTreeElement;
-	EditorController = InEditorController;
-	HierarchyView = InHierarchy;
+	FName BoneName = WeakTreeElement.Pin()->Name;
+	EIKRetargetSkeletonMode CurrentSkeleton = EditorController.Pin()->GetSkeletonMode();
+	const bool bIsBoneRetargeted = EditorController.Pin()->IsBoneRetargeted(BoneName, CurrentSkeleton);
+	const FText ChainName = FText::FromName(EditorController.Pin()->GetChainNameFromBone(BoneName, CurrentSkeleton));
+
+	// determine icon based on if bone is retargeted
+	const FSlateBrush* Brush;
+	if (bIsBoneRetargeted)
+	{
+		Brush = FAppStyle::Get().GetBrush("SkeletonTree.Bone");
+	}
+	else
+	{
+		Brush = FAppStyle::Get().GetBrush("SkeletonTree.BoneNonWeighted");
+	}
 
-	const FSlateBrush* Brush = FAppStyle::Get().GetBrush("SkeletonTree.Bone");
-	
-	const FTextBlockStyle NormalText = FIKRigEditorStyle::Get().GetWidgetStyle<FTextBlockStyle>("IKRig.Tree.NormalText");
-	const FTextBlockStyle ItalicText = FIKRigEditorStyle::Get().GetWidgetStyle<FTextBlockStyle>("IKRig.Tree.ItalicText");
-	FSlateFontInfo TextFont = NormalText.Font;
-	FSlateColor TextColor = FLinearColor(0.2f, 0.2f, 0.2f, 0.5f);
+	// determine text based on if bone is retargeted
+	FTextBlockStyle NormalText = FIKRigEditorStyle::Get().GetWidgetStyle<FTextBlockStyle>("IKRig.Tree.NormalText");
+	FTextBlockStyle ItalicText = FIKRigEditorStyle::Get().GetWidgetStyle<FTextBlockStyle>("IKRig.Tree.ItalicText");
+	FSlateFontInfo TextFont;
+	FSlateColor TextColor;
+	if (bIsBoneRetargeted)
+	{
+		// elements connected to the selected solver are green
+		TextFont = ItalicText.Font;
+		TextColor = NormalText.ColorAndOpacity;
+	}
+	else
+	{
+		TextFont = NormalText.Font;
+		TextColor = FLinearColor(0.2f, 0.2f, 0.2f, 0.5f);
+	}
 	
-	TSharedPtr<SHorizontalBox> HorizontalBox;
-	STableRow<TSharedPtr<FIKRetargetTreeElement>>::Construct(
-		STableRow<TSharedPtr<FIKRetargetTreeElement>>::FArguments()
-		.ShowWires(true)
-		.Content()
+	if (ColumnName == BoneColumnName)
+	{
+		TSharedPtr< SHorizontalBox > RowBox;
+		SAssignNew(RowBox, SHorizontalBox)
+		+ SHorizontalBox::Slot()
+		.AutoWidth()
 		[
-			SAssignNew(HorizontalBox, SHorizontalBox)
-			+ SHorizontalBox::Slot()
-			.MaxWidth(18)
-			.FillWidth(1.0)
-			.HAlign(HAlign_Left)
-			.VAlign(VAlign_Center)
-			[
-				SNew(SImage)
-				.Image(Brush)
-			]
-		], OwnerTable);
+			SNew( SExpanderArrow, SharedThis(this) )
+			.ShouldDrawWires(true)
+		];
+
+		RowBox->AddSlot()
+		.MaxWidth(18)
+		.FillWidth(1.0)
+		.HAlign(HAlign_Left)
+		.VAlign(VAlign_Center)
+		[
+			SNew(SImage)
+			.Image(Brush)
+		];
 
-	HorizontalBox->AddSlot()
+		RowBox->AddSlot()
 		.AutoWidth()
 		.VAlign(VAlign_Center)
 		[
 			SNew(STextBlock)
-			.Text(this, &SIKRetargetHierarchyItem::GetName)
+			.Text(FText::FromName(BoneName))
 			.Font(TextFont)
 			.ColorAndOpacity(TextColor)
 		];
-}
 
-FText SIKRetargetHierarchyItem::GetName() const
-{
-	return WeakTreeElement.Pin()->Key;
+		return RowBox.ToSharedRef();
+	}
+
+	if (ColumnName == ChainColumnName)
+	{
+		TSharedPtr< SHorizontalBox > RowBox;
+		SAssignNew(RowBox, SHorizontalBox)
+		+ SHorizontalBox::Slot()
+		.AutoWidth()
+		[
+			SNew(STextBlock)
+			.Text(ChainName)
+			.Font(TextFont)
+			.ColorAndOpacity(TextColor)
+		];
+
+		return RowBox.ToSharedRef();
+	}
+	
+	return SNullWidget::NullWidget;
 }
 
 void SIKRetargetHierarchy::Construct(
@@ -91,6 +121,7 @@ void SIKRetargetHierarchy::Construct(
 	ChildSlot
 	[
 		SNew(SVerticalBox)
+
 		+SVerticalBox::Slot()
 		.Padding(2.0f)
 		.AutoHeight()
@@ -114,6 +145,13 @@ void SIKRetargetHierarchy::Construct(
 				.Text(LOCTEXT("TargetSkeleton", "Target"))
 			]
 		]
+		
+		+SVerticalBox::Slot()
+		.Padding(2.0f)
+		.AutoHeight()
+		[
+			SNew(SIKRetargetPoseEditor, InEditorController)
+		]
 
 		+SVerticalBox::Slot()
 		.Padding(2.0f)
@@ -125,7 +163,14 @@ void SIKRetargetHierarchy::Construct(
 				SAssignNew(TreeView, SIKRetargetHierarchyTreeView)
 				.TreeItemsSource(&RootElements)
 				.SelectionMode(ESelectionMode::Multi)
-				.OnGenerateRow(this, &SIKRetargetHierarchy::MakeTableRowWidget)
+				.OnGenerateRow_Lambda( [this](
+					TSharedPtr<FIKRetargetHierarchyElement> InItem,
+					const TSharedRef<STableViewBase>& OwnerTable) -> TSharedRef<ITableRow>
+				{
+					return SNew(SIKRetargetHierarchyRow, OwnerTable)
+					.EditorController(EditorController.Pin())
+					.TreeElement(InItem);
+				})
 				.OnGetChildren(this, &SIKRetargetHierarchy::HandleGetChildrenForTree)
 				.OnSelectionChanged(this, &SIKRetargetHierarchy::OnSelectionChanged)
 				.OnMouseButtonClick(this, &SIKRetargetHierarchy::OnItemClicked)
@@ -133,6 +178,18 @@ void SIKRetargetHierarchy::Construct(
 				.OnSetExpansionRecursive(this, &SIKRetargetHierarchy::OnSetExpansionRecursive)
 				.HighlightParentNodesForSelection(false)
 				.ItemHeight(24)
+				.HeaderRow
+				(
+					SNew(SHeaderRow)
+					
+					+ SHeaderRow::Column(BoneColumnName)
+					.DefaultLabel(LOCTEXT("RetargetBoneNameLabel", "Bone Name"))
+					.FillWidth(0.7f)
+						
+					+ SHeaderRow::Column(ChainColumnName)
+					.DefaultLabel(LOCTEXT("RetargetChainNameLabel", "Retarget Chain"))
+					.FillWidth(0.3f)
+				)
 			]
 		]
 	];
@@ -143,8 +200,8 @@ void SIKRetargetHierarchy::Construct(
 
 void SIKRetargetHierarchy::ShowItemAfterSelection(FName ItemName)
 {
-	TSharedPtr<FIKRetargetTreeElement> ItemToShow;
-	for (const TSharedPtr<FIKRetargetTreeElement>& Element : AllElements)
+	TSharedPtr<FIKRetargetHierarchyElement> ItemToShow;
+	for (const TSharedPtr<FIKRetargetHierarchyElement>& Element : AllElements)
 	{
 		if (Element.Get()->Name == ItemName)
 		{
@@ -159,7 +216,7 @@ void SIKRetargetHierarchy::ShowItemAfterSelection(FName ItemName)
 	
 	if(GetDefault<UPersonaOptions>()->bExpandTreeOnSelection)
 	{
-		TSharedPtr<FIKRetargetTreeElement> ItemToExpand = ItemToShow->Parent;
+		TSharedPtr<FIKRetargetHierarchyElement> ItemToExpand = ItemToShow->Parent;
 		while(ItemToExpand.IsValid())
 		{
 			TreeView->SetItemExpansion(ItemToExpand, true);
@@ -214,24 +271,21 @@ void SIKRetargetHierarchy::RefreshTreeView(bool IsInitialSetup)
 	TMap<FName, int32> BoneTreeElementIndices;
 
 	// create all bone elements
+	// TODO, add filtering of retargeted bones, by name, etc..
 	for (const FName BoneName : BoneNames)
 	{
 		// create "Bone" tree element for this bone
-		const FText BoneDisplayName = FText::FromName(BoneName);
-		TSharedPtr<FIKRetargetTreeElement> BoneElement = MakeShared<FIKRetargetTreeElement>(BoneDisplayName, ControllerRef);
+		TSharedPtr<FIKRetargetHierarchyElement> BoneElement = MakeShared<FIKRetargetHierarchyElement>(BoneName, ControllerRef);
 		BoneElement.Get()->Name = BoneName;
 		const int32 BoneElementIndex = AllElements.Add(BoneElement);
 		BoneTreeElementIndices.Add(BoneName, BoneElementIndex);
-
-		// TODO, this will show ALL bones in the skeletal mesh. represent bone differently if it's in the retargeted set
-		// TODO, add filtering of retargeted bones, by name, etc..
 	}
 
 	// store children/parent pointers on all bone elements
 	for (int32 BoneIndex=0; BoneIndex<BoneNames.Num(); ++BoneIndex)
 	{
 		const FName BoneName = BoneNames[BoneIndex];
-		const TSharedPtr<FIKRetargetTreeElement> BoneTreeElement = AllElements[BoneTreeElementIndices[BoneName]];
+		const TSharedPtr<FIKRetargetHierarchyElement> BoneTreeElement = AllElements[BoneTreeElementIndices[BoneName]];
 		const int32 ParentIndex = ParentIndices[BoneIndex];
 		if (ParentIndex < 0)
 		{
@@ -243,7 +297,7 @@ void SIKRetargetHierarchy::RefreshTreeView(bool IsInitialSetup)
 
 		// get parent tree element
 		const FName ParentBoneName = BoneNames[ParentIndex];
-		const TSharedPtr<FIKRetargetTreeElement> ParentBoneTreeElement = AllElements[BoneTreeElementIndices[ParentBoneName]];
+		const TSharedPtr<FIKRetargetHierarchyElement> ParentBoneTreeElement = AllElements[BoneTreeElementIndices[ParentBoneName]];
 		// store pointer to child on parent
 		ParentBoneTreeElement->Children.Add(BoneTreeElement);
 		// store pointer to parent on child
@@ -253,7 +307,7 @@ void SIKRetargetHierarchy::RefreshTreeView(bool IsInitialSetup)
 	// expand all elements upon the initial construction of the tree
 	if (IsInitialSetup)
 	{
-		for (TSharedPtr<FIKRetargetTreeElement> RootElement : RootElements)
+		for (TSharedPtr<FIKRetargetHierarchyElement> RootElement : RootElements)
 		{
 			SetExpansionRecursive(RootElement, false, true);
 		}
@@ -261,7 +315,7 @@ void SIKRetargetHierarchy::RefreshTreeView(bool IsInitialSetup)
 	else
 	{
 		// restore expansion and selection state
-		for (const TSharedPtr<FIKRetargetTreeElement>& Element : AllElements)
+		for (const TSharedPtr<FIKRetargetHierarchyElement>& Element : AllElements)
 		{
 			TreeView->RestoreState(Element);
 		}
@@ -270,22 +324,15 @@ void SIKRetargetHierarchy::RefreshTreeView(bool IsInitialSetup)
 	TreeView->RequestTreeRefresh();
 }
 
-TSharedRef<ITableRow> SIKRetargetHierarchy::MakeTableRowWidget(
-	TSharedPtr<FIKRetargetTreeElement> InItem,
-	const TSharedRef<STableViewBase>& OwnerTable)
-{
-	return InItem->MakeTreeRowWidget(EditorController.Pin().ToSharedRef(), OwnerTable, InItem.ToSharedRef(), CommandList.ToSharedRef(), SharedThis(this));
-}
-
 void SIKRetargetHierarchy::HandleGetChildrenForTree(
-	TSharedPtr<FIKRetargetTreeElement> InItem,
-	TArray<TSharedPtr<FIKRetargetTreeElement>>& OutChildren)
+	TSharedPtr<FIKRetargetHierarchyElement> InItem,
+	TArray<TSharedPtr<FIKRetargetHierarchyElement>>& OutChildren)
 {
 	OutChildren = InItem.Get()->Children;
 }
 
 void SIKRetargetHierarchy::OnSelectionChanged(
-	TSharedPtr<FIKRetargetTreeElement> Selection,
+	TSharedPtr<FIKRetargetHierarchyElement> Selection,
 	ESelectInfo::Type SelectInfo)
 {
 	if (SelectInfo == ESelectInfo::Direct)
@@ -300,7 +347,7 @@ void SIKRetargetHierarchy::OnSelectionChanged(
 	}
 
 	TArray<FName> SelectedBoneNames;
-	TArray<TSharedPtr<FIKRetargetTreeElement>> SelectedItems = TreeView->GetSelectedItems();
+	TArray<TSharedPtr<FIKRetargetHierarchyElement>> SelectedItems = TreeView->GetSelectedItems();
 	for (auto SelectedItem : SelectedItems)
 	{
 		SelectedBoneNames.Add(SelectedItem.Get()->Name);
@@ -310,12 +357,12 @@ void SIKRetargetHierarchy::OnSelectionChanged(
 	Controller->EditBoneSelection(SelectedBoneNames, EBoneSelectionEdit::Replace, bFromHierarchy);
 }
 
-void SIKRetargetHierarchy::OnItemClicked(TSharedPtr<FIKRetargetTreeElement> InItem)
+void SIKRetargetHierarchy::OnItemClicked(TSharedPtr<FIKRetargetHierarchyElement> InItem)
 {
 	// TODO show bone details
 }
 
-void SIKRetargetHierarchy::OnItemDoubleClicked(TSharedPtr<FIKRetargetTreeElement> InItem)
+void SIKRetargetHierarchy::OnItemDoubleClicked(TSharedPtr<FIKRetargetHierarchyElement> InItem)
 {
 	if (TreeView->IsItemExpanded(InItem))
 	{
@@ -327,13 +374,13 @@ void SIKRetargetHierarchy::OnItemDoubleClicked(TSharedPtr<FIKRetargetTreeElement
 	}
 }
 
-void SIKRetargetHierarchy::OnSetExpansionRecursive(TSharedPtr<FIKRetargetTreeElement> InItem, bool bShouldBeExpanded)
+void SIKRetargetHierarchy::OnSetExpansionRecursive(TSharedPtr<FIKRetargetHierarchyElement> InItem, bool bShouldBeExpanded)
 {
 	SetExpansionRecursive(InItem, false, bShouldBeExpanded);
 }
 
 void SIKRetargetHierarchy::SetExpansionRecursive(
-	TSharedPtr<FIKRetargetTreeElement> InElement,
+	TSharedPtr<FIKRetargetHierarchyElement> InElement,
 	bool bTowardsParent,
 	bool bShouldBeExpanded)
 {
diff --git a/Engine/Plugins/Animation/IKRig/Source/IKRigEditor/Private/RetargetEditor/SIKRetargetPoseEditor.cpp b/Engine/Plugins/Animation/IKRig/Source/IKRigEditor/Private/RetargetEditor/SIKRetargetPoseEditor.cpp
new file mode 100644
--- /dev/null
+++ b/Engine/Plugins/Animation/IKRig/Source/IKRigEditor/Private/RetargetEditor/SIKRetargetPoseEditor.cpp
@@ -0,0 +1,203 @@
+// Copyright Epic Games, Inc. All Rights Reserved.
+
+#include "RetargetEditor/SIKRetargetPoseEditor.h"
+
+#include "RetargetEditor/IKRetargetCommands.h"
+#include "RetargetEditor/IKRetargetEditor.h"
+#include "RetargetEditor/IKRetargetEditorController.h"
+
+#define LOCTEXT_NAMESPACE "SIKRetargetPoseEditor"
+
+void SIKRetargetPoseEditor::Construct(
+	const FArguments& InArgs,
+	TSharedRef<FIKRetargetEditorController> InEditorController)
+{
+	EditorController = InEditorController;
+	
+	// fill list of pose names
+	PoseNames.Reset();
+	for (const TTuple<FName, FIKRetargetPose>& Pose : EditorController.Pin()->AssetController->GetRetargetPoses())
+	{
+		PoseNames.Add(MakeShareable(new FName(Pose.Key)));
+	}
+
+	// the commands for the menus
+	TSharedPtr<FUICommandList> Commands = EditorController.Pin()->Editor.Pin()->GetToolkitCommands();
+	
+	ChildSlot
+	[
+		SNew(SVerticalBox)
+
+		+SVerticalBox::Slot()
+		.Padding(2.0f)
+		.AutoHeight()
+		[
+			SNew(SHorizontalBox)
+
+			// add pose selection label
+			+SHorizontalBox::Slot()
+			.Padding(2.0f)
+			.HAlign(HAlign_Right)
+			.VAlign(VAlign_Center)
+			[
+				SNew(STextBlock)
+				.Text(LOCTEXT("CurrentPose", "Current Retarget Pose:"))
+			]
+
+			// add pose selection combo box
+			+SHorizontalBox::Slot()
+			.Padding(2.0f)
+			[
+				SNew(SComboBox<TSharedPtr<FName>>)
+				.OptionsSource(&PoseNames)
+				.OnGenerateWidget_Lambda([](TSharedPtr<FName> InItem)
+				{
+					return SNew(STextBlock).Text(FText::FromName(*InItem.Get()));
+				})
+				.OnSelectionChanged(EditorController.Pin().Get(), &FIKRetargetEditorController::OnPoseSelected)
+				[
+					SNew(STextBlock).Text(EditorController.Pin().Get(), &FIKRetargetEditorController::GetCurrentPoseName)
+				]
+			]
+		]
+
+		+SVerticalBox::Slot()
+		.Padding(2.0f)
+		.AutoHeight()
+		[
+			SNew(SHorizontalBox)
+			+SHorizontalBox::Slot()
+			.FillWidth(1.0f)
+			.HAlign(HAlign_Center)
+			[
+				MakeToolbar(Commands)
+			]
+		]
+	];
+}
+
+TSharedRef<SWidget>  SIKRetargetPoseEditor::MakeToolbar(TSharedPtr<FUICommandList> Commands)
+{
+	FToolBarBuilder ToolbarBuilder(Commands, FMultiBoxCustomization::None);
+	
+	ToolbarBuilder.BeginSection("Edit Current Pose");
+
+	ToolbarBuilder.AddToolBarButton(
+		FIKRetargetCommands::Get().EditRetargetPose,
+		NAME_None,
+		TAttribute<FText>(),
+		TAttribute<FText>(),
+		FSlateIcon(FAppStyle::Get().GetStyleSetName(),"Icons.Edit"));
+
+	ToolbarBuilder.AddComboButton(
+		FUIAction(),
+		FOnGetContent::CreateSP(this, &SIKRetargetPoseEditor::GenerateResetMenuContent, Commands),
+		LOCTEXT("ResetPose_Label", "Reset"),
+		LOCTEXT("ResetPoseToolTip_Label", "Reset bones to reference pose."),
+		FSlateIcon(FAppStyle::GetAppStyleSetName(), "Icons.Refresh"));
+
+	ToolbarBuilder.EndSection();
+
+	ToolbarBuilder.BeginSection("Create Poses");
+
+	ToolbarBuilder.AddComboButton(
+		FUIAction(),
+		FOnGetContent::CreateSP(this, &SIKRetargetPoseEditor::GenerateNewMenuContent, Commands),
+		LOCTEXT("CreatePose_Label", "Create"),
+		TAttribute<FText>(),
+		FSlateIcon(FAppStyle::GetAppStyleSetName(), "Icons.Plus"));
+
+	ToolbarBuilder.AddToolBarButton(
+		FIKRetargetCommands::Get().DeleteRetargetPose,
+		NAME_None,
+		TAttribute<FText>(),
+		TAttribute<FText>(),
+		FSlateIcon(FAppStyle::Get().GetStyleSetName(),"Icons.Delete"));
+
+	ToolbarBuilder.AddToolBarButton(
+		FIKRetargetCommands::Get().RenameRetargetPose,
+		NAME_None,
+		TAttribute<FText>(),
+		TAttribute<FText>(),
+		FSlateIcon(FAppStyle::Get().GetStyleSetName(),"Icons.Settings"));
+
+	ToolbarBuilder.EndSection();
+
+	return ToolbarBuilder.MakeWidget();
+}
+
+TSharedRef<SWidget> SIKRetargetPoseEditor::GenerateResetMenuContent(TSharedPtr<FUICommandList> Commands)
+{
+	FMenuBuilder MenuBuilder(true, Commands);
+
+	MenuBuilder.AddMenuEntry(
+		FIKRetargetCommands::Get().ResetSelectedBones,
+		TEXT("Reset Selected"),
+		TAttribute<FText>(),
+		TAttribute<FText>());
+
+	MenuBuilder.AddMenuEntry(
+		FIKRetargetCommands::Get().ResetSelectedAndChildrenBones,
+		TEXT("Reset Selected And Children"),
+		TAttribute<FText>(),
+		TAttribute<FText>());
+
+	MenuBuilder.AddMenuEntry(
+		FIKRetargetCommands::Get().ResetAllBones,
+		TEXT("Reset All"),
+		TAttribute<FText>(),
+		TAttribute<FText>());
+
+	return MenuBuilder.MakeWidget();
+}
+
+TSharedRef<SWidget> SIKRetargetPoseEditor::GenerateNewMenuContent(TSharedPtr<FUICommandList> Commands)
+{
+	FMenuBuilder MenuBuilder(true, Commands);
+
+	MenuBuilder.BeginSection("Create", LOCTEXT("CreatePoseOperations", "Create New Retarget Pose"));
+	{
+		MenuBuilder.AddMenuEntry(
+		FIKRetargetCommands::Get().NewRetargetPose,
+		TEXT("Create"),
+		TAttribute<FText>(),
+		TAttribute<FText>());
+
+		MenuBuilder.AddMenuEntry(
+			FIKRetargetCommands::Get().DuplicateRetargetPose,
+			TEXT("Create"),
+			TAttribute<FText>(),
+			TAttribute<FText>());
+	}
+	MenuBuilder.EndSection();
+
+	MenuBuilder.BeginSection("Import",LOCTEXT("ImportPoseOperations", "Import Retarget Pose"));
+	{
+		MenuBuilder.AddMenuEntry(
+		FIKRetargetCommands::Get().ImportRetargetPose,
+		TEXT("Import"),
+		TAttribute<FText>(),
+		TAttribute<FText>());
+	
+		MenuBuilder.AddMenuEntry(
+			FIKRetargetCommands::Get().ImportRetargetPoseFromAnim,
+			TEXT("ImportFromSequence"),
+			TAttribute<FText>(),
+			TAttribute<FText>());
+	}
+	MenuBuilder.EndSection();
+
+	MenuBuilder.BeginSection("Export",LOCTEXT("EmportPoseOperations", "Export Retarget Pose"));
+	{
+		MenuBuilder.AddMenuEntry(
+			FIKRetargetCommands::Get().ExportRetargetPose,
+			TEXT("Export"),
+			TAttribute<FText>(),
+			TAttribute<FText>());
+	}
+	MenuBuilder.EndSection();
+
+	return MenuBuilder.MakeWidget();
+}
+
+#undef LOCTEXT_NAMESPACE
diff --git a/Engine/Plugins/Animation/IKRig/Source/IKRigEditor/Private/RigEditor/IKRigDetailCustomizations.cpp b/Engine/Plugins/Animation/IKRig/Source/IKRigEditor/Private/RigEditor/IKRigDetailCustomizations.cpp
--- a/Engine/Plugins/Animation/IKRig/Source/IKRigEditor/Private/RigEditor/IKRigDetailCustomizations.cpp
+++ b/Engine/Plugins/Animation/IKRig/Source/IKRigEditor/Private/RigEditor/IKRigDetailCustomizations.cpp
@@ -1,6 +1,6 @@
 // Copyright Epic Games, Inc. All Rights Reserved.
 
-#include "IKRigDetailCustomizations.h"
+#include "RigEditor/IKRigDetailCustomizations.h"
 
 #include "DetailLayoutBuilder.h"
 #include "DetailCategoryBuilder.h"
diff --git a/Engine/Plugins/Animation/IKRig/Source/IKRigEditor/Private/RigEditor/IKRigDetailCustomizations.h b/Engine/Plugins/Animation/IKRig/Source/IKRigEditor/Private/RigEditor/IKRigDetailCustomizations.h
deleted file mode 100644
diff --git a/Engine/Plugins/Animation/IKRig/Source/IKRigEditor/Private/RigEditor/IKRigEditorController.cpp b/Engine/Plugins/Animation/IKRig/Source/IKRigEditor/Private/RigEditor/IKRigEditorController.cpp
--- a/Engine/Plugins/Animation/IKRig/Source/IKRigEditor/Private/RigEditor/IKRigEditorController.cpp
+++ b/Engine/Plugins/Animation/IKRig/Source/IKRigEditor/Private/RigEditor/IKRigEditorController.cpp
@@ -129,17 +129,6 @@ void UIKRigBoneDetails::OnComponentRelativeChanged(
 
 #if WITH_EDITOR
 
-namespace
-{
-
-template<typename DataType>
-void GetContentFromData(const DataType& InData, FString& Content)
-{
-	TBaseStructure<DataType>::Get()->ExportText(Content, &InData, &InData, nullptr, PPF_None, nullptr);
-}
-	
-}
-
 void UIKRigBoneDetails::OnCopyToClipboard(ESlateTransformComponent::Type Component, EIKRigTransformType::Type TransformType) const
 {
 	TOptional<FTransform> Optional = GetTransform(TransformType);
diff --git a/Engine/Plugins/Animation/IKRig/Source/IKRigEditor/Private/RigEditor/SIKRigOutputLog.cpp b/Engine/Plugins/Animation/IKRig/Source/IKRigEditor/Private/RigEditor/SIKRigOutputLog.cpp
--- a/Engine/Plugins/Animation/IKRig/Source/IKRigEditor/Private/RigEditor/SIKRigOutputLog.cpp
+++ b/Engine/Plugins/Animation/IKRig/Source/IKRigEditor/Private/RigEditor/SIKRigOutputLog.cpp
@@ -26,6 +26,7 @@ void SIKRigOutputLog::Construct(
 		InitOptions.bShowPages = false;
 		InitOptions.bAllowClear = false;
 		InitOptions.bShowInLogWindow = false;
+		InitOptions.bDiscardDuplicates = true;
 		MessageLogModule.RegisterLogListing(InLogName, FText::FromString("IK Rig Log"), InitOptions);
 		MessageLogListing = MessageLogModule.GetLogListing(InLogName);
 	}
diff --git a/Engine/Plugins/Animation/IKRig/Source/IKRigEditor/Public/RetargetEditor/IKRetargetDetails.h b/Engine/Plugins/Animation/IKRig/Source/IKRigEditor/Public/RetargetEditor/IKRetargetDetails.h
--- a/Engine/Plugins/Animation/IKRig/Source/IKRigEditor/Public/RetargetEditor/IKRetargetDetails.h
+++ b/Engine/Plugins/Animation/IKRig/Source/IKRigEditor/Public/RetargetEditor/IKRetargetDetails.h
@@ -5,29 +5,113 @@
 #include "CoreMinimal.h"
 #include "IDetailCustomization.h"
 #include "IKRetargeterController.h"
+#include "Retargeter/IKRetargetProcessor.h"
 
-// details customization for IKRetargeter asset
-class FIKRetargeterDetails : public IDetailCustomization
+#include "IKRetargetDetails.generated.h"
+
+enum class EIKRetargetTransformType : int8
 {
-	
+	Current,
+	Reference,
+	RelativeOffset,
+};
+
+UCLASS(config = Engine, hidecategories = UObject)
+class IKRIGEDITOR_API UIKRetargetBoneDetails : public UObject
+{
+	GENERATED_BODY()
+
 public:
+
+	UPROPERTY(VisibleAnywhere, Category = "Selection")
+	FName SelectedBone;
+
+	UPROPERTY()
+	FTransform OffsetTransform;
 	
-	// makes a new instance of this detail layout class for a specific detail view requesting it
-	static TSharedRef<IDetailCustomization> MakeInstance();
+	UPROPERTY()
+	FTransform CurrentTransform;
 
-	// Begin IDetailCustomization interface
-	virtual void CustomizeDetails(IDetailLayoutBuilder& DetailBuilder) override;
-	// End IDetailCustomization interface
+	UPROPERTY()
+	FTransform ReferenceTransform;
+	
+	TSharedPtr<FIKRetargetEditorController> EditorController;
 
-private:
+#if WITH_EDITOR
 
-	TSharedRef<SWidget> MakeToolbar(TSharedPtr<FUICommandList> Commands);
+	FEulerTransform GetTransform(EIKRetargetTransformType TransformType, const bool bLocalSpace) const;
+	bool IsComponentRelative(ESlateTransformComponent::Type Component, EIKRetargetTransformType TransformType) const;
+	void OnComponentRelativeChanged(ESlateTransformComponent::Type Component, bool bIsRelative, EIKRetargetTransformType TransformType);
+	void OnCopyToClipboard(ESlateTransformComponent::Type Component, EIKRetargetTransformType TransformType) const;
+	void OnPasteFromClipboard(ESlateTransformComponent::Type Component, EIKRetargetTransformType TransformType);
+	void OnNumericValueCommitted(
+		ESlateTransformComponent::Type Component,
+		ESlateRotationRepresentation::Type Representation,
+		ESlateTransformSubComponent::Type SubComponent,
+		FVector::FReal Value,
+		ETextCommit::Type CommitType,
+		EIKRetargetTransformType TransformType,
+		bool bIsCommit);
 
-	TSharedRef<SWidget> GenerateResetMenuContent(TSharedPtr<FUICommandList> Commands);
-	
-	TSharedRef<SWidget> GenerateNewMenuContent(TSharedPtr<FUICommandList> Commands);
+	TOptional<FVector::FReal> GetNumericValue(
+		EIKRetargetTransformType TransformType,
+		ESlateTransformComponent::Type Component,
+		ESlateRotationRepresentation::Type Representation,
+		ESlateTransformSubComponent::Type SubComponent);
+	/** Method to react to changes of numeric values in the widget */
+	static void OnMultiNumericValueCommitted(
+		ESlateTransformComponent::Type Component,
+		ESlateRotationRepresentation::Type Representation,
+		ESlateTransformSubComponent::Type SubComponent,
+		FVector::FReal Value,
+		ETextCommit::Type CommitType,
+		EIKRetargetTransformType TransformType,
+		TArrayView<UIKRetargetBoneDetails*> Bones,
+		bool bIsCommit);
+
+	template<typename DataType>
+	void GetContentFromData(const DataType& InData, FString& Content) const;
+
+#endif
+
+private:
+
+	static TOptional<FVector::FReal> CleanRealValue(TOptional<FVector::FReal> InValue);
 	
-	TObjectPtr<UIKRetargeterController> GetAssetControllerFromSelectedObjects(IDetailLayoutBuilder& DetailBuilder) const;
+	bool IsRootBone() const;
 	
-	TArray<TSharedPtr<FName>> PoseNames;
+	bool RelativeOffsetTransformRelative[3] = { false, true, false };
+	bool CurrentTransformRelative[3];
+	bool ReferenceTransformRelative[3];
+};
+
+struct FIKRetargetTransformUIData
+{
+	TArray<EIKRetargetTransformType> TransformTypes;
+	TArray<FText> ButtonLabels;
+	TArray<FText> ButtonTooltips;
+	TAttribute<TArray<EIKRetargetTransformType>> VisibleTransforms;
+	TArray<TSharedRef<IPropertyHandle>> Properties;
+};
+
+class FIKRetargetBoneDetailCustomization : public IDetailCustomization
+{
+public:
+
+	static TSharedRef<IDetailCustomization> MakeInstance()
+	{
+		return MakeShareable(new FIKRetargetBoneDetailCustomization);
+	}
+
+	/** IDetailCustomization interface */
+	virtual void CustomizeDetails(IDetailLayoutBuilder& DetailBuilder) override;
+
+private:
+
+	void GetTransformUIData(
+		const bool bIsEditingPose,
+		const IDetailLayoutBuilder& DetailBuilder,
+		FIKRetargetTransformUIData& OutData) const;
+
+	TArray<UIKRetargetBoneDetails*> Bones;
 };
diff --git a/Engine/Plugins/Animation/IKRig/Source/IKRigEditor/Public/RetargetEditor/IKRetargetEditPoseMode.h b/Engine/Plugins/Animation/IKRig/Source/IKRigEditor/Public/RetargetEditor/IKRetargetEditPoseMode.h
--- a/Engine/Plugins/Animation/IKRig/Source/IKRigEditor/Public/RetargetEditor/IKRetargetEditPoseMode.h
+++ b/Engine/Plugins/Animation/IKRig/Source/IKRigEditor/Public/RetargetEditor/IKRetargetEditPoseMode.h
@@ -30,7 +30,7 @@ struct BoneEdit
 	FTransform LocalTransform;			// local transform of last selected bone
 	FQuat AccumulatedGlobalOffset;		// the accumulated offset from rotation gizmo
 	
-	TArray<FQuat> PrevLocalOffsets;		// the prev stored local offsets of all selected bones
+	TArray<FQuat> PreviousDeltaRotation;		// the prev stored local offsets of all selected bones
 
 	void Reset()
 	{
@@ -79,9 +79,6 @@ private:
 	// get the scale and offset associated with the currently edited skeletal mesh component
 	void GetEditedComponentScaleAndOffset(float& OutScale, FVector& OutOffset) const;
 
-	// get the scale and offset associated with the currently edited skeletal mesh component
-	int32 GetEditedRetargetRootBoneIndex(const UIKRetargetProcessor& Processor) const;
-
 	// render skeleton in viewport
 	void RenderSkeleton(
 		FPrimitiveDrawInterface* PDI,
diff --git a/Engine/Plugins/Animation/IKRig/Source/IKRigEditor/Public/RetargetEditor/IKRetargetEditorController.h b/Engine/Plugins/Animation/IKRig/Source/IKRigEditor/Public/RetargetEditor/IKRetargetEditorController.h
--- a/Engine/Plugins/Animation/IKRig/Source/IKRigEditor/Public/RetargetEditor/IKRetargetEditorController.h
+++ b/Engine/Plugins/Animation/IKRig/Source/IKRigEditor/Public/RetargetEditor/IKRetargetEditorController.h
@@ -1,6 +1,7 @@
 // Copyright Epic Games, Inc. All Rights Reserved.
 #pragma once
 
+#include "IKRetargetDetails.h"
 #include "IKRetargeterController.h"
 #include "IPersonaToolkit.h"
 #include "SIKRetargetAssetBrowser.h"
@@ -43,7 +44,7 @@ enum class EBoneSelectionEdit : uint8
 };
 
 /** a home for cross-widget communication to synchronize state across all tabs and viewport */
-class FIKRetargetEditorController : public TSharedFromThis<FIKRetargetEditorController>
+class FIKRetargetEditorController : public TSharedFromThis<FIKRetargetEditorController>, FGCObject
 {
 public:
 
@@ -104,6 +105,8 @@ public:
 	const USkeleton* GetTargetSkeleton() const;
 	/** get currently edited debug skeletal mesh */
 	UDebugSkelMeshComponent* GetEditedSkeletalMesh() const;
+	/** get the currently edited retarget skeleton */
+	const FRetargetSkeleton& GetCurrentlyEditedSkeleton(const UIKRetargetProcessor& Processor) const;
 	
 	/** get world space pose of a bone (with component scale / offset applied) */
 	FTransform GetGlobalRetargetPoseOfBone(
@@ -147,13 +150,21 @@ public:
 		const TArray<FName>& InBoneNames,
 		EBoneSelectionEdit EditMode,
 		const bool bFromHierarchyView);
-	void ClearSelection();
+	void ClearSelection(const bool bKeepBoneSelection=false);
 	const TArray<FName>& GetSelectedBones() const {return SelectedBones; };
 	/** mesh selection management (viewport view) */
 	void SetSelectedMesh(UPrimitiveComponent* InComponent);
 	UPrimitiveComponent* GetSelectedMesh();
 	void AddOffsetToMeshComponent(const FVector& Offset, USceneComponent* MeshComponent) const;
 
+	/** determine if bone in the specified skeleton is part of the retarget (in a mapped chain) */
+	bool IsBoneRetargeted(const FName& BoneName, EIKRetargetSkeletonMode WhichSkeleton) const;
+	/** get the name of the chain that contains this bone */
+	FName GetChainNameFromBone(const FName& BoneName, EIKRetargetSkeletonMode WhichSkeleton) const;
+
+	/** factory to get/create bone details object */
+	TObjectPtr<UIKRetargetBoneDetails> GetDetailsObjectForBone(const FName& BoneName);
+
 	/* START RETARGET POSES */
 	
 	/** go to retarget pose */
@@ -220,6 +231,17 @@ public:
 	
 	/* END RETARGET POSES */
 
+	/** FGCObject interface */
+	virtual void AddReferencedObjects(FReferenceCollector& Collector) override
+	{
+		for (TTuple<FName, TObjectPtr<UIKRetargetBoneDetails>> Pair : AllBoneDetails)
+		{
+			Collector.AddReferencedObject(Pair.Value);
+		}
+	};
+	virtual FString GetReferencerName() const override { return TEXT("Retarget Editor"); };
+	/** END FGCObject interface */
+
 private:
 	
 	/** which skeleton are we editing / viewing? */
@@ -227,6 +249,10 @@ private:
 
 	/** current selection set */
 	TArray<FName> SelectedBones;
+	UPROPERTY()
+	TMap<FName,TObjectPtr<UIKRetargetBoneDetails>> AllBoneDetails;
+	TArray<UObject*> SelectedBoneDetails;
+	TObjectPtr<UIKRetargetBoneDetails> CreateBoneDetails(const FName& InBoneName);
 
 	/** currently selected mesh */
 	UPrimitiveComponent* SelectedMesh;
diff --git a/Engine/Plugins/Animation/IKRig/Source/IKRigEditor/Public/RetargetEditor/IKRetargeterController.h b/Engine/Plugins/Animation/IKRig/Source/IKRigEditor/Public/RetargetEditor/IKRetargeterController.h
--- a/Engine/Plugins/Animation/IKRig/Source/IKRigEditor/Public/RetargetEditor/IKRetargeterController.h
+++ b/Engine/Plugins/Animation/IKRig/Source/IKRigEditor/Public/RetargetEditor/IKRetargeterController.h
@@ -100,9 +100,13 @@ public:
 	/** Get the current retarget pose */
 	const FIKRetargetPose& GetCurrentRetargetPose() const;
 	/** Set a delta rotation for a given bone for the current retarget pose (used in Edit Mode in the retarget editor) */
-	void SetRotationOffsetForRetargetPoseBone(FName BoneName, FQuat RotationOffset) const;
+	void SetRotationOffsetForRetargetPoseBone(FName BoneName, const FQuat& RotationOffset) const;
 	/** Get a delta rotation for a given bone for the current retarget pose (used in Edit Mode in the retarget editor) */
 	FQuat GetRotationOffsetForRetargetPoseBone(FName BoneName) const;
+	/** Set the delta translation of the root bone (used in Edit Mode in the retarget editor) */
+	void SetTranslationOffsetOnRetargetRootBone(FVector TranslationOffset) const;
+	/** Get the delta translation of the root bone (used in Edit Mode in the retarget editor) */
+	const FVector& GetTranslationOffsetOnRetargetRootBone() const;
 	/** Add a delta translation to the root bone (used in Edit Mode in the retarget editor) */
 	void AddTranslationOffsetToRetargetRootBone(FVector TranslationOffset) const;
 	/** Add a numbered suffix to the given pose name to make it unique. */
diff --git a/Engine/Plugins/Animation/IKRig/Source/IKRigEditor/Public/RetargetEditor/SIKRetargetHierarchy.h b/Engine/Plugins/Animation/IKRig/Source/IKRigEditor/Public/RetargetEditor/SIKRetargetHierarchy.h
--- a/Engine/Plugins/Animation/IKRig/Source/IKRigEditor/Public/RetargetEditor/SIKRetargetHierarchy.h
+++ b/Engine/Plugins/Animation/IKRig/Source/IKRigEditor/Public/RetargetEditor/SIKRetargetHierarchy.h
@@ -11,22 +11,15 @@
 
 class SIKRetargetHierarchy;
 
-class FIKRetargetTreeElement : public TSharedFromThis<FIKRetargetTreeElement>
+class FIKRetargetHierarchyElement : public TSharedFromThis<FIKRetargetHierarchyElement>
 {
 public:
 	
-	FIKRetargetTreeElement(const FText& InKey, const TSharedRef<FIKRetargetEditorController>& InEditorController);
-
-	TSharedRef<ITableRow> MakeTreeRowWidget(
-		TSharedRef<FIKRetargetEditorController> InEditorController,
-		const TSharedRef<STableViewBase>& InOwnerTable,
-		TSharedRef<FIKRetargetTreeElement> InRigTreeElement,
-		TSharedRef<FUICommandList> InCommandList,
-		TSharedPtr<SIKRetargetHierarchy> InHierarchy);
+	FIKRetargetHierarchyElement(const FName& InName, const TSharedRef<FIKRetargetEditorController>& InEditorController);
 
 	FText Key;
-	TSharedPtr<FIKRetargetTreeElement> Parent;
-	TArray<TSharedPtr<FIKRetargetTreeElement>> Children;
+	TSharedPtr<FIKRetargetHierarchyElement> Parent;
+	TArray<TSharedPtr<FIKRetargetHierarchyElement>> Children;
 	FName Name;
 
 private:
@@ -34,28 +27,38 @@ private:
 	TWeakPtr<FIKRetargetEditorController> EditorController;
 };
 
-class SIKRetargetHierarchyItem : public STableRow<TSharedPtr<FIKRetargetTreeElement>>
+class SIKRetargetHierarchyRow : public SMultiColumnTableRow<TSharedPtr<FIKRetargetHierarchyElement>>
 {
+	
 public:
+
+	SLATE_BEGIN_ARGS(SIKRetargetHierarchyRow) {}
+	SLATE_ARGUMENT(TSharedPtr<FIKRetargetEditorController>, EditorController)
+	SLATE_ARGUMENT(TSharedPtr<FIKRetargetHierarchyElement>, TreeElement)
+	SLATE_END_ARGS()
 	
 	void Construct(
 		const FArguments& InArgs,
-		TSharedRef<FIKRetargetEditorController> InEditorController,
-		const TSharedRef<STableViewBase>& OwnerTable,
-		TSharedRef<FIKRetargetTreeElement> InTreeElement,
-		TSharedRef<FUICommandList> InCommandList,
-		TSharedPtr<SIKRetargetHierarchy> InHierarchy);
+		const TSharedRef<STableViewBase>& InOwnerTable)
+	{
+		OwnerTable = InOwnerTable;
+		WeakTreeElement = InArgs._TreeElement;
+		EditorController = InArgs._EditorController;
+		
+		SMultiColumnTableRow<TSharedPtr<FIKRetargetHierarchyElement>>::Construct(FSuperRowType::FArguments(), InOwnerTable);
+	}
+
+	/** Overridden from SMultiColumnTableRow.  Generates a widget for this column of the list view. */
+	virtual TSharedRef<SWidget> GenerateWidgetForColumn(const FName& ColumnName) override;
 
 private:
-	
-	FText GetName() const;
 
-	TWeakPtr<FIKRetargetTreeElement> WeakTreeElement;
+	TSharedPtr<STableViewBase> OwnerTable;
+	TWeakPtr<FIKRetargetHierarchyElement> WeakTreeElement;
 	TWeakPtr<FIKRetargetEditorController> EditorController;
-	TWeakPtr<SIKRetargetHierarchy> HierarchyView;
 };
 
-typedef SBaseHierarchyTreeView<FIKRetargetTreeElement> SIKRetargetHierarchyTreeView;
+typedef SBaseHierarchyTreeView<FIKRetargetHierarchyElement> SIKRetargetHierarchyTreeView;
 
 class SIKRetargetHierarchy : public SCompoundWidget, public FEditorUndoClient
 {
@@ -77,20 +80,19 @@ private:
 	
 	/** tree view widget */
 	TSharedPtr<SIKRetargetHierarchyTreeView> TreeView;
-	TArray<TSharedPtr<FIKRetargetTreeElement>> RootElements;
-	TArray<TSharedPtr<FIKRetargetTreeElement>> AllElements;
+	TArray<TSharedPtr<FIKRetargetHierarchyElement>> RootElements;
+	TArray<TSharedPtr<FIKRetargetHierarchyElement>> AllElements;
 	
 	/** tree view callbacks */
 	void RefreshTreeView(bool IsInitialSetup=false);
-	TSharedRef<ITableRow> MakeTableRowWidget(TSharedPtr<FIKRetargetTreeElement> InItem, const TSharedRef<STableViewBase>& OwnerTable);
-	void HandleGetChildrenForTree(TSharedPtr<FIKRetargetTreeElement> InItem, TArray<TSharedPtr<FIKRetargetTreeElement>>& OutChildren);
-	void OnSelectionChanged(TSharedPtr<FIKRetargetTreeElement> Selection, ESelectInfo::Type SelectInfo);
-	void OnItemClicked(TSharedPtr<FIKRetargetTreeElement> InItem);
-	void OnItemDoubleClicked(TSharedPtr<FIKRetargetTreeElement> InItem);
-	void OnSetExpansionRecursive(TSharedPtr<FIKRetargetTreeElement> InItem, bool bShouldBeExpanded);
-	void SetExpansionRecursive(TSharedPtr<FIKRetargetTreeElement> InElement, bool bTowardsParent, bool bShouldBeExpanded);
+	void HandleGetChildrenForTree(TSharedPtr<FIKRetargetHierarchyElement> InElement, TArray<TSharedPtr<FIKRetargetHierarchyElement>>& OutChildren);
+	void OnSelectionChanged(TSharedPtr<FIKRetargetHierarchyElement> Selection, ESelectInfo::Type SelectInfo);
+	void OnItemClicked(TSharedPtr<FIKRetargetHierarchyElement> InElement);
+	void OnItemDoubleClicked(TSharedPtr<FIKRetargetHierarchyElement> InElement);
+	void OnSetExpansionRecursive(TSharedPtr<FIKRetargetHierarchyElement> InElement, bool bShouldBeExpanded);
+	void SetExpansionRecursive(TSharedPtr<FIKRetargetHierarchyElement> InElement, bool bTowardsParent, bool bShouldBeExpanded);
 	/** END tree view callbacks */
 
-	friend SIKRetargetHierarchyItem;
+	friend SIKRetargetHierarchyRow;
 	friend FIKRetargetEditorController;
 };
diff --git a/Engine/Plugins/Animation/IKRig/Source/IKRigEditor/Public/RetargetEditor/SIKRetargetPoseEditor.h b/Engine/Plugins/Animation/IKRig/Source/IKRigEditor/Public/RetargetEditor/SIKRetargetPoseEditor.h
new file mode 100644
--- /dev/null
+++ b/Engine/Plugins/Animation/IKRig/Source/IKRigEditor/Public/RetargetEditor/SIKRetargetPoseEditor.h
@@ -0,0 +1,30 @@
+﻿// Copyright Epic Games, Inc. All Rights Reserved.
+
+#pragma once
+
+#include "CoreMinimal.h"
+#include "DetailLayoutBuilder.h"
+
+class UIKRetargeterController;
+class FIKRetargetEditorController;
+
+class SIKRetargetPoseEditor : public SCompoundWidget
+{
+	
+public:
+	
+	SLATE_BEGIN_ARGS(SIKRetargetPoseEditor) {}
+	SLATE_END_ARGS()
+
+	void Construct(const FArguments& InArgs, TSharedRef<FIKRetargetEditorController> InEditorController);
+
+private:
+	
+	TSharedRef<SWidget> MakeToolbar(TSharedPtr<FUICommandList> Commands);
+	TSharedRef<SWidget> GenerateResetMenuContent(TSharedPtr<FUICommandList> Commands);
+	TSharedRef<SWidget> GenerateNewMenuContent(TSharedPtr<FUICommandList> Commands);
+	TObjectPtr<UIKRetargeterController> GetAssetControllerFromSelectedObjects(IDetailLayoutBuilder& DetailBuilder) const;
+	
+	TWeakPtr<FIKRetargetEditorController> EditorController;
+	TArray<TSharedPtr<FName>> PoseNames;
+};
diff --git a/Engine/Plugins/Animation/IKRig/Source/IKRigEditor/Public/RigEditor/IKRigDetailCustomizations.h b/Engine/Plugins/Animation/IKRig/Source/IKRigEditor/Public/RigEditor/IKRigDetailCustomizations.h
new file mode 100644
--- /dev/null
+++ b/Engine/Plugins/Animation/IKRig/Source/IKRigEditor/Public/RigEditor/IKRigDetailCustomizations.h
@@ -0,0 +1,40 @@
+// Copyright Epic Games, Inc. All Rights Reserved.
+
+#pragma once
+
+#include "CoreMinimal.h"
+#include "IDetailCustomization.h"
+#include "ScopedTransaction.h"
+#include "RigEditor/IKRigEditorController.h"
+
+class FIKRigGenericDetailCustomization : public IDetailCustomization
+{
+public:
+
+	static TSharedRef<IDetailCustomization> MakeInstance()
+	{
+		return MakeShareable(new FIKRigGenericDetailCustomization);
+	}
+
+	/** IDetailCustomization interface */
+	virtual void CustomizeDetails(IDetailLayoutBuilder& DetailBuilder) override;
+
+private:
+
+	/** Labels and transform types used by the transform widget */
+	static const TArray<FText>& GetButtonLabels();
+	static const TArray<EIKRigTransformType::Type>& GetTransformTypes();
+
+	/** Per class customization */
+	template<typename ClassToCustomize>
+	void CustomizeDetailsForClass(IDetailLayoutBuilder& DetailBuilder, const TArray<TWeakObjectPtr<UObject>>& ObjectsBeingCustomized) {}
+
+	/** Template specializations */
+	template<>
+	void CustomizeDetailsForClass<UIKRigBoneDetails>(IDetailLayoutBuilder& DetailBuilder, const TArray<TWeakObjectPtr<UObject>>& ObjectsBeingCustomized);
+	template<>
+	void CustomizeDetailsForClass<UIKRigEffectorGoal>(IDetailLayoutBuilder& DetailBuilder, const TArray<TWeakObjectPtr<UObject>>& ObjectsBeingCustomized);
+
+	/** Transaction to handle value changed. It's being effective once the value committed */
+	TSharedPtr<FScopedTransaction> ValueChangedTransaction;
+};
\ No newline at end of file
diff --git a/Engine/Plugins/Animation/IKRig/Source/IKRigEditor/Public/RigEditor/IKRigEditorController.h b/Engine/Plugins/Animation/IKRig/Source/IKRigEditor/Public/RigEditor/IKRigEditorController.h
--- a/Engine/Plugins/Animation/IKRig/Source/IKRigEditor/Public/RigEditor/IKRigEditorController.h
+++ b/Engine/Plugins/Animation/IKRig/Source/IKRigEditor/Public/RigEditor/IKRigEditorController.h
@@ -64,8 +64,14 @@ public:
 	void OnCopyToClipboard(ESlateTransformComponent::Type Component, EIKRigTransformType::Type TransformType) const;
 	void OnPasteFromClipboard(ESlateTransformComponent::Type Component, EIKRigTransformType::Type TransformType);
 
-#endif
+	template <typename DataType>
+	void GetContentFromData(const DataType& InData, FString& Content) const
+	{
+		TBaseStructure<DataType>::Get()->ExportText(Content, &InData, &InData, nullptr, PPF_None, nullptr);
+	}
 
+#endif
+	
 private:
 	
 	static bool CurrentTransformRelative[3];
