Fixing duplicate localization keys in IK Rig plugin

diff --git a/Engine/Plugins/Animation/IKRig/Source/IKRigEditor/Private/RetargetEditor/IKRetargetDetails.cpp b/Engine/Plugins/Animation/IKRig/Source/IKRigEditor/Private/RetargetEditor/IKRetargetDetails.cpp
--- a/Engine/Plugins/Animation/IKRig/Source/IKRigEditor/Private/RetargetEditor/IKRetargetDetails.cpp
+++ b/Engine/Plugins/Animation/IKRig/Source/IKRigEditor/Private/RetargetEditor/IKRetargetDetails.cpp
@@ -473,7 +473,7 @@ void UIKRetargetBoneDetails::OnNumericValueCommitted(
 			}
 			
 			// store the new rotation in the retarget pose
-			FScopedTransaction Transaction(LOCTEXT("EditRootTranslation", "Edit Retarget Pose Rotation"));
+			FScopedTransaction Transaction(LOCTEXT("EditRootRotation", "Edit Retarget Pose Rotation"));
 			EditorController->AssetController->GetAsset()->Modify();
 			EditorController->AssetController->SetRotationOffsetForRetargetPoseBone(SelectedBone, NewLocalRotationDelta);
 			break;
@@ -696,8 +696,8 @@ void FIKRetargetBoneDetailCustomization::GetTransformUIData(
 	};
 	static const TArray<FText> ReadOnlyButtonTooltips
 	{
-		LOCTEXT("CurrentBoneTransformTooltip", "The current transform of the bone"),
-		LOCTEXT("ReferenceBoneTransformTooltip", "The reference transform of the bone")
+		LOCTEXT("CurrentBoneTransformTooltip", "The current transform of the bone."),
+		LOCTEXT("ReferenceBoneTransformTooltip", "The reference transform of the bone.")
 	};
 	static TAttribute<TArray<EIKRetargetTransformType>> ReadOnlyVisibleTransforms =
 		TArray<EIKRetargetTransformType>({EIKRetargetTransformType::Current});
@@ -711,13 +711,13 @@ void FIKRetargetBoneDetailCustomization::GetTransformUIData(
 	};
 	static const TArray<FText> EditableButtonLabels
 	{
-		LOCTEXT("RelativeOffsetTransform", "Relative Offset"),
-		LOCTEXT("ReferenceTransform", "Reference")
+		LOCTEXT("EditableRelativeOffsetTransform", "Relative Offset"),
+		LOCTEXT("EditableReferenceTransform", "Reference")
 	};
 	static const TArray<FText> EditableButtonTooltips
 	{
 		LOCTEXT("RelativeOffsetBoneTransformTooltip", "The offset transform in the current retarget pose, relative to the reference pose."),
-		LOCTEXT("ReferenceBoneTransformTooltip", "The reference transform of the bone.")
+		LOCTEXT("EditableReferenceBoneTransformTooltip", "The reference transform of the bone.")
 	};
 	static TAttribute<TArray<EIKRetargetTransformType>> EditableVisibleTransforms =
 		TArray<EIKRetargetTransformType>({EIKRetargetTransformType::RelativeOffset});
