Fix issue renaming retarget chains adds "_0" suffix

diff --git a/Engine/Plugins/Animation/IKRig/Source/IKRigEditor/Private/RigEditor/SIKRigRetargetChainList.cpp b/Engine/Plugins/Animation/IKRig/Source/IKRigEditor/Private/RigEditor/SIKRigRetargetChainList.cpp
--- a/Engine/Plugins/Animation/IKRig/Source/IKRigEditor/Private/RigEditor/SIKRigRetargetChainList.cpp
+++ b/Engine/Plugins/Animation/IKRig/Source/IKRigEditor/Private/RigEditor/SIKRigRetargetChainList.cpp
@@ -248,7 +248,7 @@ void SIKRigRetargetChainRow::OnGoalComboSelectionChanged(TSharedPtr<FString> InG
 	ChainList.Pin()->RefreshView();
 }
 
-void SIKRigRetargetChainRow::OnRenameChain(const FText& InText, ETextCommit::Type) const
+void SIKRigRetargetChainRow::OnRenameChain(const FText& InText, ETextCommit::Type CommitType) const
 {
 	const TSharedPtr<FIKRigEditorController> Controller = ChainList.Pin()->EditorController.Pin();
 	if (!Controller.IsValid())
@@ -258,6 +258,12 @@ void SIKRigRetargetChainRow::OnRenameChain(const FText& InText, ETextCommit::Typ
 
 	const FName OldName = ChainElement.Pin()->ChainName;
 	const FName NewName = FName(*InText.ToString());
+	if (OldName == NewName)
+	{
+		// most reliable way to catch multiple commits
+		return;
+	}
+	
 	ChainElement.Pin()->ChainName = Controller->AssetController->RenameRetargetChain(OldName, NewName);
 	ChainList.Pin()->RefreshView();
 }
