Add ability to edit retarget pose on source (not just target).
Retarget pose I/O converted to use Pose Assets
Various UI/UX changes.

diff --git a/Engine/Plugins/Animation/IKRig/Source/IKRig/Private/Retargeter/IKRetargetProcessor.cpp b/Engine/Plugins/Animation/IKRig/Source/IKRig/Private/Retargeter/IKRetargetProcessor.cpp
--- a/Engine/Plugins/Animation/IKRig/Source/IKRig/Private/Retargeter/IKRetargetProcessor.cpp
+++ b/Engine/Plugins/Animation/IKRig/Source/IKRig/Private/Retargeter/IKRetargetProcessor.cpp
@@ -86,14 +86,14 @@ void FRetargetSkeleton::GenerateRetargetPose(const FIKRetargetPose* InRetargetPo
 	if (RootBoneIndex != INDEX_NONE)
 	{
 		FTransform& RootTransform = RetargetGlobalPose[RootBoneIndex];
-		RootTransform.AddToTranslation(InRetargetPose->RootTranslationOffset);
+		RootTransform.AddToTranslation(InRetargetPose->GetRootTranslationDelta());
 		UpdateLocalTransformOfSingleBone(RootBoneIndex, RetargetLocalPose, RetargetGlobalPose);
 	}
 
 	// apply bone rotation offsets
-	for (const TTuple<FName, FQuat>& BoneRotationOffset : InRetargetPose->BoneRotationOffsets)
+	for (const TTuple<FName, FQuat>& BoneDelta : InRetargetPose->GetAllDeltaRotations())
 	{
-		const int32 BoneIndex = FindBoneIndexByName(BoneRotationOffset.Key);
+		const int32 BoneIndex = FindBoneIndexByName(BoneDelta.Key);
 		if (BoneIndex == INDEX_NONE)
 		{
 			// this can happen if a retarget pose recorded a bone offset for a bone that is not present in the
@@ -101,29 +101,13 @@ void FRetargetSkeleton::GenerateRetargetPose(const FIKRetargetPose* InRetargetPo
 			continue;
 		}
 
-		const FQuat DeltaRotation = GetRetargetPoseDeltaRotation(BoneRotationOffset.Key, InRetargetPose);
-		const FQuat LocalBoneRotation = RefPoseLocal[BoneIndex].GetRotation() * DeltaRotation;
+		const FQuat LocalBoneRotation = RefPoseLocal[BoneIndex].GetRotation() * BoneDelta.Value;
 		RetargetLocalPose[BoneIndex].SetRotation(LocalBoneRotation);
 	}
 
 	UpdateGlobalTransformsBelowBone(-1, RetargetLocalPose, RetargetGlobalPose);
 }
 
-FQuat FRetargetSkeleton::GetRetargetPoseDeltaRotation(const FName BoneName, const FIKRetargetPose* InRetargetPose) const
-{
-	const int32 BoneIndex = FindBoneIndexByName(BoneName);
-	check(BoneIndex!=INDEX_NONE)
-	check(InRetargetPose)
-	
-	if (const FQuat* BoneRotationOffset = InRetargetPose->BoneRotationOffsets.Find(BoneNames[BoneIndex]))
-	{
-		return *BoneRotationOffset;
-
-	}
-
-	return FQuat::Identity;
-}
-
 int32 FRetargetSkeleton::FindBoneIndexByName(const FName InName) const
 {
 	return BoneNames.IndexOfByPredicate([&InName](const FName& BoneName)
@@ -399,6 +383,31 @@ void FTargetSkeleton::SetBoneIsRetargeted(const int32 BoneIndex, const bool IsRe
 	IsBoneRetargeted[BoneIndex] = IsRetargeted;
 }
 
+void FRetargetChainSettings::CopySettingsFromAsset(const URetargetChainSettings* AssetChainSettings)
+{
+	TargetChainName = AssetChainSettings->TargetChain;
+
+	CopyPoseUsingFK = AssetChainSettings->CopyPoseUsingFK;
+	RotationMode = AssetChainSettings->RotationMode;
+	RotationAlpha = AssetChainSettings->RotationAlpha;
+	TranslationMode = AssetChainSettings->TranslationMode;
+	TranslationAlpha = AssetChainSettings->TranslationAlpha;
+		
+	DriveIKGoal = AssetChainSettings->DriveIKGoal;
+	Extension = AssetChainSettings->Extension;
+	StaticOffset = AssetChainSettings->StaticOffset;
+	StaticLocalOffset = AssetChainSettings->StaticLocalOffset;
+	StaticRotationOffset = AssetChainSettings->StaticRotationOffset;
+	BlendToSource = AssetChainSettings->BlendToSource;
+	BlendToSourceWeights = AssetChainSettings->BlendToSourceWeights;
+
+	UseSpeedCurveToPlantIK = AssetChainSettings->UseSpeedCurveToPlantIK;
+	UnplantStiffness = AssetChainSettings->UnplantStiffness;
+	UnplantCriticalDamping = AssetChainSettings->UnplantCriticalDamping;
+	SpeedThreshold = AssetChainSettings->VelocityThreshold;
+	SpeedCurveName = AssetChainSettings->SpeedCurveName;
+};
+
 bool FChainFK::Initialize(
 	const FRetargetSkeleton& Skeleton,
 	const TArray<int32>& BoneIndices,
@@ -1222,6 +1231,16 @@ bool FRootRetargeter::InitializeTarget(
 	return true;
 }
 
+void FRetargetRootSettings::CopySettingsFromAsset(const URetargetRootSettings* AssetRootSettings)
+{
+	GlobalScaleHorizontal = AssetRootSettings->GlobalScaleHorizontal;
+	GlobalScaleVertical = AssetRootSettings->GlobalScaleVertical;
+	BlendToSource = AssetRootSettings->BlendToSource;
+	BlendToSourceWeights = AssetRootSettings->BlendToSourceWeights;
+	StaticOffset = AssetRootSettings->StaticOffset;
+	StaticRotationOffset = AssetRootSettings->StaticRotationOffset;
+}
+
 void FRootRetargeter::Reset()
 {
 	Source = FRootSource();
@@ -1240,21 +1259,21 @@ void FRootRetargeter::DecodePose(TArray<FTransform>& OutTargetGlobalPose) const
 {
 	// scale normalized position by root height
 	FVector RetargetedPosition = Source.CurrentPositionNormalized * Target.InitialHeight;// * GlobalScaleVertical;
-	RetargetedPosition.Z *= GlobalScaleVertical;
+	RetargetedPosition.Z *= Settings.GlobalScaleVertical;
 	// globally scale offset of root
-	const FVector RootOffset = (RetargetedPosition - Target.InitialPosition) * FVector(GlobalScaleHorizontal, GlobalScaleHorizontal, 1.0f);
+	const FVector RootOffset = (RetargetedPosition - Target.InitialPosition) * FVector(Settings.GlobalScaleHorizontal, Settings.GlobalScaleHorizontal, 1.0f);
 	RetargetedPosition = Target.InitialPosition + RootOffset;
 	// blend the retarget root position towards the source retarget root position
-	FVector Position = FMath::Lerp(RetargetedPosition, Source.CurrentPosition, BlendToSource);
+	FVector Position = FMath::Lerp(RetargetedPosition, Source.CurrentPosition, Settings.BlendToSource*Settings.BlendToSourceWeights);
 	// apply a static offset
-	Position += StaticOffset;
+	Position += Settings.StaticOffset;
 
 	// calc offset between initial source/target root rotations
 	const FQuat RotationDelta = Source.CurrentRotation * Source.InitialRotation.Inverse();
 	// add retarget pose delta to the current source rotation
 	FQuat Rotation = RotationDelta * Target.InitialRotation;
 	// add static rotation offset
-	Rotation = StaticRotationOffset.Quaternion() * Rotation;
+	Rotation = Settings.StaticRotationOffset.Quaternion() * Rotation;
 
 	// apply to target
 	FTransform& TargetRootTransform = OutTargetGlobalPose[Target.BoneIndex];
@@ -1287,6 +1306,9 @@ void UIKRetargetProcessor::Initialize(
 	ChainPairsIK.Reset();
 	RootRetargeter.Reset();
 
+	const UIKRigDefinition* SourceIKRig = RetargeterAsset->GetSourceIKRig();
+	const UIKRigDefinition* TargetIKRig = RetargeterAsset->GetTargetIKRig();
+
 	// check prerequisite assets
 	if (!SourceSkeletalMesh)
 	{
@@ -1298,17 +1320,17 @@ void UIKRetargetProcessor::Initialize(
 		Log.LogError(LOCTEXT("MissingTargetMesh", "IK Retargeter unable to initialize. Missing target Skeletal Mesh asset."));
 		return;
 	}
-	if (!RetargeterAsset->GetSourceIKRig())
+	if (!SourceIKRig)
 	{
 		Log.LogError(LOCTEXT("MissingSourceIKRig", "IK Retargeter unable to initialize. Missing source IK Rig asset."));
 		return;
 	}
-	if (!RetargeterAsset->GetTargetIKRig())
+	if (!TargetIKRig)
 	{
 		Log.LogError(LOCTEXT("MissingTargetIKRig", "IK Retargeter unable to initialize. Missing target IK Rig asset."));
 		return;
 	}
-	if (!RetargeterAsset->GetCurrentRetargetPose())
+	if (!RetargeterAsset->GetCurrentRetargetPose(ERetargetSourceOrTarget::Target))
 	{
 		Log.LogError(LOCTEXT("MissingRetargetPose", "IK Retargeter unable to initialize. Missing retarget pose."));
 		return;
@@ -1317,14 +1339,14 @@ void UIKRetargetProcessor::Initialize(
 	// initialize skeleton data for source and target
 	SourceSkeleton.Initialize(
 		SourceSkeletalMesh,
-		RetargeterAsset->GetSourceIKRig()->GetRetargetChains(),
-		nullptr,
-		NAME_None);
+		SourceIKRig->GetRetargetChains(),
+		RetargeterAsset->GetCurrentRetargetPose(ERetargetSourceOrTarget::Source),
+		SourceIKRig->GetRetargetRoot());
 	TargetSkeleton.Initialize(
 		TargetSkeletalMesh,
-		RetargeterAsset->GetTargetIKRig()->GetRetargetChains(),
-		RetargeterAsset->GetCurrentRetargetPose(),
-		RetargeterAsset->GetTargetIKRig()->GetRetargetRoot());
+		TargetIKRig->GetRetargetChains(),
+		RetargeterAsset->GetCurrentRetargetPose(ERetargetSourceOrTarget::Target),
+		TargetIKRig->GetRetargetRoot());
 
 	// initialize roots
 	bRootsInitialized = InitializeRoots();
@@ -1336,7 +1358,7 @@ void UIKRetargetProcessor::Initialize(
 		// couldn't match up any BoneChain pairs, no limb retargeting possible
 		Log.LogWarning( FText::Format(
 			LOCTEXT("NoMappedChains", "IK Retargeter unable to map any bone chains between source, {0} and target, {1}"),
-			FText::FromString(SourceSkeleton.SkeletalMesh->GetName()), FText::FromString(TargetSkeleton.SkeletalMesh->GetName())));
+			FText::FromString(SourceSkeletalMesh->GetName()), FText::FromString(TargetSkeletalMesh->GetName())));
 	}
 
 	// initialize the IKRigProcessor for doing IK decoding
@@ -1346,7 +1368,7 @@ void UIKRetargetProcessor::Initialize(
 		// couldn't initialize the IK Rig, we don't disable the retargeter in this case, just warn the user
 		Log.LogWarning( FText::Format(
 			LOCTEXT("CouldNotInitializeIKRig", "IK Retargeter was unable to initialize the IK Rig, {0} for the Skeletal Mesh {1}. See previous warnings."),
-			FText::FromString(RetargeterAsset->GetTargetIKRig()->GetName()), FText::FromString(TargetSkeleton.SkeletalMesh->GetName())));
+			FText::FromString(TargetIKRig->GetName()), FText::FromString(TargetSkeletalMesh->GetName())));
 	}
 
 	// must have a mapped root bone OR at least a single mapped chain to be able to do any retargeting at all
@@ -1355,7 +1377,7 @@ void UIKRetargetProcessor::Initialize(
 		// confirm for the user that the IK Rig was initialized successfully
 		Log.LogInfo(FText::Format(
 				LOCTEXT("SuccessfulInit", "Success! The IK Retargeter is ready to transfer animation from the source, {0} to the target, {1}"),
-				FText::FromString(SourceSkeleton.SkeletalMesh->GetName()), FText::FromString(TargetSkeleton.SkeletalMesh->GetName())));
+				FText::FromString(SourceSkeletalMesh->GetName()), FText::FromString(TargetSkeletalMesh->GetName())));
 	}
 	
 	bIsInitialized = true;
@@ -1363,6 +1385,9 @@ void UIKRetargetProcessor::Initialize(
 
 bool UIKRetargetProcessor::InitializeRoots()
 {
+	// copy the retarget root settings
+	RootRetargeter.Settings.CopySettingsFromAsset(RetargeterAsset->GetRetargetRootSettings());
+	
 	// initialize root encoder
 	const FName SourceRootBoneName = RetargeterAsset->GetSourceIKRig()->GetRetargetRoot();
 	const bool bRootEncoderInit = RootRetargeter.InitializeSource(SourceRootBoneName, SourceSkeleton, Log);
@@ -1517,21 +1542,7 @@ TArray<FTransform>&  UIKRetargetProcessor::RunRetargeter(
 	const float DeltaTime)
 {
 	check(bIsInitialized);
-
-#if WITH_EDITOR
-	// in edit mode we just want to see the edited reference pose, not actually run the retargeting
-	// as long as the retargeter is reinitialized after every modification to the limb rotation offsets,
-	// then the TargetSkeleton.RetargetGlobalPose will contain the updated retarget pose.
-	const ERetargeterOutputMode CurrentMode = RetargeterAsset->GetOutputMode();
-	const bool bOutputRetargetPose = CurrentMode == ERetargeterOutputMode::EditRetargetPose || CurrentMode == ERetargeterOutputMode::ShowRetargetPose;
-	if (bOutputRetargetPose && RetargeterAsset->GetTargetIKRig())
-	{
-		const FName RootBoneName = RetargeterAsset->GetTargetIKRig()->GetRetargetRoot();
-		TargetSkeleton.GenerateRetargetPose(RetargeterAsset->GetCurrentRetargetPose(), RootBoneName);
-		return TargetSkeleton.RetargetGlobalPose; 
-	}
-#endif
-
+	
 	// start from retarget pose
 	TargetSkeleton.OutputGlobalPose = TargetSkeleton.RetargetGlobalPose;
 
@@ -1686,13 +1697,10 @@ void UIKRetargetProcessor::SetNeedsInitialized()
 
 void UIKRetargetProcessor::CopyAllSettingsFromAsset()
 {
-	const UIKRigDefinition* TargetIKRig = RetargeterAsset->GetTargetIKRig();
-	if (!TargetIKRig)
+	if (const UIKRigDefinition* TargetIKRig = RetargeterAsset->GetTargetIKRig())
 	{
-		return;
+		IKRigProcessor->CopyAllInputsFromSourceAssetAtRuntime(TargetIKRig);
 	}
-	
-	IKRigProcessor->CopyAllInputsFromSourceAssetAtRuntime(TargetIKRig);
 
 	// copy most recent settings from asset for each chain
 	const TArray<TObjectPtr<URetargetChainSettings>>& AllChainSettings = RetargeterAsset->GetAllChainSettings();
@@ -1716,12 +1724,7 @@ void UIKRetargetProcessor::CopyAllSettingsFromAsset()
 	}
 
 	// copy root settings
-	const URetargetRootSettings& RootSettings = *RetargeterAsset->GetRetargetRootSettings();
-	RootRetargeter.GlobalScaleHorizontal = RootSettings.GlobalScaleHorizontal;
-	RootRetargeter.GlobalScaleVertical = RootSettings.GlobalScaleVertical;
-	RootRetargeter.BlendToSource = RootSettings.BlendToSource;
-	RootRetargeter.StaticOffset = RootSettings.StaticOffset;
-	RootRetargeter.StaticRotationOffset = RootSettings.StaticRotationOffset;
+	RootRetargeter.Settings.CopySettingsFromAsset(RetargeterAsset->GetRetargetRootSettings());
 }
 
 bool UIKRetargetProcessor::IsBoneRetargeted(
diff --git a/Engine/Plugins/Animation/IKRig/Source/IKRig/Private/Retargeter/IKRetargeter.cpp b/Engine/Plugins/Animation/IKRig/Source/IKRig/Private/Retargeter/IKRetargeter.cpp
--- a/Engine/Plugins/Animation/IKRig/Source/IKRig/Private/Retargeter/IKRetargeter.cpp
+++ b/Engine/Plugins/Animation/IKRig/Source/IKRig/Private/Retargeter/IKRetargeter.cpp
@@ -19,7 +19,8 @@ void UIKRetargeter::GetSpeedCurveNames(TArray<FName>& OutSpeedCurveNames) const
 			OutSpeedCurveNames.Add(ChainSetting->SpeedCurveName);
 		}
 	}
-};
+}
+
 #endif
 
 UIKRetargeter::UIKRetargeter(const FObjectInitializer& ObjectInitializer)
@@ -27,6 +28,8 @@ UIKRetargeter::UIKRetargeter(const FObjectInitializer& ObjectInitializer)
 {
 	RootSettings = CreateDefaultSubobject<URetargetRootSettings>(TEXT("RootSettings"));
 	RootSettings->SetFlags(RF_Transactional);
+
+	CleanAndInitialize();
 }
 
 void UIKRetargeter::Serialize(FArchive& Ar)
@@ -82,39 +85,93 @@ void UIKRetargeter::PostLoad()
 			TargetMeshScale = TargetActorScale_DEPRECATED;
 		}
 	#endif
+
+
+	// load deprecated retarget poses (pre adding retarget poses for source)
+	if (!RetargetPoses_DEPRECATED.IsEmpty())
+	{
+		TargetRetargetPoses = RetargetPoses_DEPRECATED;
+		RetargetPoses_DEPRECATED.Empty();
+	}
+
+	// load deprecated current retarget pose (pre adding retarget poses for source)
+	if (CurrentRetargetPose_DEPRECATED != NAME_None)
+	{
+		CurrentTargetRetargetPose = CurrentRetargetPose_DEPRECATED;
+		CurrentRetargetPose_DEPRECATED = NAME_None;
+	}
 	
 	PRAGMA_ENABLE_DEPRECATION_WARNINGS
 
+	CleanAndInitialize();
+}
+
+void UIKRetargeter::CleanAndInitialize()
+{
 	// remove null settings
 	ChainSettings.Remove(nullptr);
 
-	// ensure current pose is valid
-	if (!RetargetPoses.Contains(CurrentRetargetPose))
+	// use default pose as current pose unless set to something else
+	if (CurrentSourceRetargetPose == NAME_None)
 	{
-		CurrentRetargetPose = UIKRetargeter::GetDefaultPoseName();
+		CurrentSourceRetargetPose = GetDefaultPoseName();
 	}
+	if (CurrentTargetRetargetPose == NAME_None)
+	{
+		CurrentTargetRetargetPose = GetDefaultPoseName();
+	}
+
+	// enforce the existence of a default pose
+	if (!SourceRetargetPoses.Contains(GetDefaultPoseName()))
+	{
+		SourceRetargetPoses.Emplace(GetDefaultPoseName());
+	}
+	if (!TargetRetargetPoses.Contains(GetDefaultPoseName()))
+	{
+		TargetRetargetPoses.Emplace(GetDefaultPoseName());
+	}
+
+	// ensure current pose exists, otherwise set it to the default pose
+	if (!SourceRetargetPoses.Contains(CurrentSourceRetargetPose))
+	{
+		CurrentSourceRetargetPose = GetDefaultPoseName();
+	}
+	if (!TargetRetargetPoses.Contains(CurrentTargetRetargetPose))
+	{
+		CurrentTargetRetargetPose = GetDefaultPoseName();
+	}
+};
+
+FQuat FIKRetargetPose::GetDeltaRotationForBone(const FName BoneName) const
+{
+	const FQuat* BoneRotationOffset = BoneRotationOffsets.Find(BoneName);
+	return BoneRotationOffset != nullptr ? *BoneRotationOffset : FQuat::Identity;
 }
 
-void FIKRetargetPose::SetBoneRotationOffset(FName BoneName, FQuat RotationDelta, const FIKRigSkeleton& Skeleton)
+void FIKRetargetPose::SetDeltaRotationForBone(FName BoneName, FQuat RotationDelta)
 {
 	FQuat* RotOffset = BoneRotationOffsets.Find(BoneName);
 	if (RotOffset == nullptr)
 	{
 		// first time this bone has been modified in this pose
 		BoneRotationOffsets.Emplace(BoneName, RotationDelta);
-		SortHierarchically(Skeleton);
 		return;
 	}
 
 	*RotOffset = RotationDelta;
 }
 
-void FIKRetargetPose::SetRootTranslationDelta(FVector TranslationDelta)
+FVector FIKRetargetPose::GetRootTranslationDelta() const
+{
+	return RootTranslationOffset;
+}
+
+void FIKRetargetPose::SetRootTranslationDelta(const FVector& TranslationDelta)
 {
 	RootTranslationOffset = TranslationDelta;
 }
 
-void FIKRetargetPose::AddToRootTranslationDelta(FVector TranslateDelta)
+void FIKRetargetPose::AddToRootTranslationDelta(const FVector& TranslateDelta)
 {
 	RootTranslationOffset += TranslateDelta;
 }
@@ -129,6 +186,11 @@ void FIKRetargetPose::SortHierarchically(const FIKRigSkeleton& Skeleton)
 	});
 }
 
+const FIKRetargetPose* UIKRetargeter::GetCurrentRetargetPose(const ERetargetSourceOrTarget& SourceOrTarget) const
+{
+	return SourceOrTarget == ERetargetSourceOrTarget::Source ? &SourceRetargetPoses[CurrentSourceRetargetPose] : &TargetRetargetPoses[CurrentTargetRetargetPose];
+}
+
 const FName UIKRetargeter::GetDefaultPoseName()
 {
 	static const FName DefaultPoseName = "Default Pose";
diff --git a/Engine/Plugins/Animation/IKRig/Source/IKRig/Public/Retargeter/IKRetargetProcessor.h b/Engine/Plugins/Animation/IKRig/Source/IKRig/Public/Retargeter/IKRetargetProcessor.h
--- a/Engine/Plugins/Animation/IKRig/Source/IKRig/Public/Retargeter/IKRetargetProcessor.h
+++ b/Engine/Plugins/Animation/IKRig/Source/IKRig/Public/Retargeter/IKRetargetProcessor.h
@@ -8,7 +8,8 @@
 #include "Kismet/KismetMathLibrary.h"
 #include "IKRetargetProcessor.generated.h"
 
-enum class EIKRetargetSkeletonMode : uint8;
+enum class ERetargetTranslationMode : uint8;
+enum class ERetargetRotationMode : uint8;
 class URetargetChainSettings;
 class UIKRigDefinition;
 class UIKRigProcessor;
@@ -80,8 +81,6 @@ struct IKRIG_API FRetargetSkeleton
 	
 	bool IsParentOfChild(const int32 PotentialParentIndex, const int32 ChildBoneIndex) const;
 
-	FQuat GetRetargetPoseDeltaRotation(const FName BoneName, const FIKRetargetPose* InRetargetPose) const;
-
 private:
 	
 	/** One index per-bone. Lazy-filled on request. Stores the last element of the branch below the bone.
@@ -148,21 +147,23 @@ struct FRootTarget
 	float InitialHeight;
 };
 
-struct FRootRetargeter
-{	
-	FRootSource Source;
-	
-	FRootTarget Target;
-
+struct FRetargetRootSettings
+{
 	float GlobalScaleHorizontal = 1.0f;
-
 	float GlobalScaleVertical = 1.0f;
-
-	FVector BlendToSource = FVector::ZeroVector;
-
+	float BlendToSource = 0.0f;
+	FVector BlendToSourceWeights = FVector::OneVector;
 	FVector StaticOffset = FVector::ZeroVector;
-
 	FRotator StaticRotationOffset = FRotator::ZeroRotator;
+	
+	void CopySettingsFromAsset(const URetargetRootSettings* AssetRootSettings);
+};
+
+struct FRootRetargeter
+{	
+	FRootSource Source;
+	FRootTarget Target;
+	FRetargetRootSettings Settings;
 
 	void Reset();
 	
@@ -182,7 +183,7 @@ struct FRootRetargeter
 
 	FVector GetGlobalScaleVector() const
 	{
-		return FVector(GlobalScaleHorizontal, GlobalScaleHorizontal, GlobalScaleVertical);
+		return FVector(Settings.GlobalScaleHorizontal, Settings.GlobalScaleHorizontal, Settings.GlobalScaleVertical);
 	}
 };
 
@@ -212,30 +213,7 @@ struct FRetargetChainSettings
 
 public:
 	
-	void CopySettingsFromAsset(const URetargetChainSettings* AssetChainSettings)
-	{
-		TargetChainName = AssetChainSettings->TargetChain;
-
-		CopyPoseUsingFK = AssetChainSettings->CopyPoseUsingFK;
-		RotationMode = AssetChainSettings->RotationMode;
-		RotationAlpha = AssetChainSettings->RotationAlpha;
-		TranslationMode = AssetChainSettings->TranslationMode;
-		TranslationAlpha = AssetChainSettings->TranslationAlpha;
-		
-		DriveIKGoal = AssetChainSettings->DriveIKGoal;
-		Extension = AssetChainSettings->Extension;
-		StaticOffset = AssetChainSettings->StaticOffset;
-		StaticLocalOffset = AssetChainSettings->StaticLocalOffset;
-		StaticRotationOffset = AssetChainSettings->StaticRotationOffset;
-		BlendToSource = AssetChainSettings->BlendToSource;
-		BlendToSourceWeights = AssetChainSettings->BlendToSourceWeights;
-
-		UseSpeedCurveToPlantIK = AssetChainSettings->UseSpeedCurveToPlantIK;
-		UnplantStiffness = AssetChainSettings->UnplantStiffness;
-		UnplantCriticalDamping = AssetChainSettings->UnplantCriticalDamping;
-		SpeedThreshold = AssetChainSettings->VelocityThreshold;
-		SpeedCurveName = AssetChainSettings->SpeedCurveName;
-	};
+	void CopySettingsFromAsset(const URetargetChainSettings* AssetChainSettings);
 };
 
 
diff --git a/Engine/Plugins/Animation/IKRig/Source/IKRig/Public/Retargeter/IKRetargeter.h b/Engine/Plugins/Animation/IKRig/Source/IKRig/Public/Retargeter/IKRetargeter.h
--- a/Engine/Plugins/Animation/IKRig/Source/IKRig/Public/Retargeter/IKRetargeter.h
+++ b/Engine/Plugins/Animation/IKRig/Source/IKRig/Public/Retargeter/IKRetargeter.h
@@ -160,9 +160,15 @@ class IKRIG_API URetargetRootSettings: public UObject
 	
 	URetargetRootSettings() = default;
 
-	/** Whether to modify the location of the retarget root bone. Default is true. */
-	UPROPERTY(EditAnywhere, Category = "Root Retarget Settings")
-	bool RetargetRootTranslation = true;
+	/** Range 0 to 1. Default 0. Blends the retarget root's translation to the exact source location.
+	*  At 0 the root is placed at the retargeted location.
+	*  At 1 the root is placed at the location of the source's retarget root bone. */
+	UPROPERTY(EditAnywhere, Category = "Root Retarget Settings", meta = (ClampMin = "0.0", ClampMax = "1.0", UIMin = "0.0", UIMax = "1.0"))
+	float BlendToSource = 0.0f;
+
+	/** Per-axis weights for the Blend to Source. */
+	UPROPERTY(EditAnywhere, Category = "Root Retarget Settings", meta = (ClampMin = "0.0", ClampMax = "1.0", UIMin = "0.0", UIMax = "1.0"))
+	FVector BlendToSourceWeights = FVector::OneVector;
 
 	/** Default 1. Scales the motion of the root position in the horizontal plane (X,Y). */
 	UPROPERTY(EditAnywhere, Category = "Root Retarget Settings", meta = (UIMin = "0.0", UIMax = "3.0"))
@@ -171,19 +177,13 @@ class IKRIG_API URetargetRootSettings: public UObject
 	/** Default 1. Scales the motion of the root position in the vertical direction (Z). */
 	UPROPERTY(EditAnywhere, Category = "Root Retarget Settings", meta = (UIMin = "0.0", UIMax = "3.0"))
 	float GlobalScaleVertical = 1.0f;
-	
-	/** Range 0 to 1. Default 0. Blends the retarget root's translation to the exact source location (per axis).
-	*  At 0 the root is placed at the retargeted location.
-	*  At 1 the root is placed at the location of the source's retarget root bone. */
-	UPROPERTY(EditAnywhere, Category = "Root Retarget Settings")
-	FVector BlendToSource = FVector::ZeroVector;
 
 	/** Applies a static component-space translation offset to the retarget root.*/
 	UPROPERTY(EditAnywhere, Category = "Root Retarget Settings")
 	FVector StaticOffset = FVector::ZeroVector;
 
 	/** Applies a static component-space rotation offset to the retarget root.*/
-	UPROPERTY(EditAnywhere, Category = "Root Retarget Settings")
+	UPROPERTY(EditAnywhere, Category = "Root Retarget Settings", meta = (ClampMin = "-180.0", ClampMax = "180.0", UIMin = "-180.0", UIMax = "180.0"))
 	FRotator StaticRotationOffset = FRotator::ZeroRotator;
 };
 
@@ -196,48 +196,34 @@ public:
 	
 	FIKRetargetPose() = default;
 
-	// a translational delta in GLOBAL space, applied only to the retarget root bone
-	UPROPERTY(EditAnywhere, Category = RetargetPose)
-	FVector RootTranslationOffset = FVector::ZeroVector;
+	FQuat GetDeltaRotationForBone(const FName BoneName) const;
+	void SetDeltaRotationForBone(FName BoneName, FQuat RotationOffset);
+	const TMap<FName, FQuat>& GetAllDeltaRotations() const { return BoneRotationOffsets; };
 
-	// these are LOCAL-space rotation deltas to be applied to a bone to modify it's retarget pose
-	UPROPERTY(EditAnywhere, Category = RetargetPose)
-	TMap<FName, FQuat> BoneRotationOffsets;
+	FVector GetRootTranslationDelta() const;
+	void SetRootTranslationDelta(const FVector& TranslationDelta);
+	void AddToRootTranslationDelta(const FVector& TranslationDelta);
 	
-	void SetBoneRotationOffset(FName BoneName, FQuat RotationOffset, const FIKRigSkeleton& Skeleton);
-	void SetRootTranslationDelta(FVector TranslationDelta);
-	void AddToRootTranslationDelta(FVector TranslationDelta);
 	void SortHierarchically(const FIKRigSkeleton& Skeleton);
-};
-
-UCLASS(Blueprintable)
-class IKRIG_API URetargetPose: public UObject
-{
-	GENERATED_BODY()
-
-public:
 
+private:
+	// a translational delta in GLOBAL space, applied only to the retarget root bone
 	UPROPERTY(EditAnywhere, Category = RetargetPose)
 	FVector RootTranslationOffset = FVector::ZeroVector;
-	
+
+	// these are LOCAL-space rotation deltas to be applied to a bone to modify it's retarget pose
 	UPROPERTY(EditAnywhere, Category = RetargetPose)
 	TMap<FName, FQuat> BoneRotationOffsets;
 
-	void GetAsRetargetPose(FIKRetargetPose& OutPose) const
-	{
-		OutPose.RootTranslationOffset = RootTranslationOffset;
-		OutPose.BoneRotationOffsets = BoneRotationOffsets;
-	}
+	friend class UIKRetargeterController;
 };
 
-#if WITH_EDITOR
-enum class ERetargeterOutputMode : uint8
+// which skeleton are we referring to?
+enum class ERetargetSourceOrTarget : uint8
 {
-	RunRetarget,
-	ShowRetargetPose,
-	EditRetargetPose
+	Source,	// the SOURCE skeleton (to copy FROM)
+	Target, // the TARGET skeleton (to copy TO)
 };
-#endif
 
 UCLASS(Blueprintable)
 class IKRIG_API UIKRetargeter : public UObject
@@ -254,7 +240,7 @@ public:
 	const UIKRigDefinition* GetTargetIKRig() const { return TargetIKRigAsset.LoadSynchronous(); };
 	/** Get read-write access to the source IK Rig asset.
 	 * WARNING: do not use for editing the data model. Use Controller class instead. */
-	 UIKRigDefinition* GetSourceIKRigWriteable() const { return SourceIKRigAsset.LoadSynchronous(); };
+	UIKRigDefinition* GetSourceIKRigWriteable() const { return SourceIKRigAsset.LoadSynchronous(); };
 	/** Get read-write access to the target IK Rig asset.
 	 * WARNING: do not use for editing the data model. Use Controller class instead. */
 	UIKRigDefinition* GetTargetIKRigWriteable() const { return TargetIKRigAsset.LoadSynchronous(); };
@@ -263,7 +249,7 @@ public:
 	/** Get access to the root settings */
 	URetargetRootSettings* GetRetargetRootSettings() const { return RootSettings; };
 	/** Get read-only access to a retarget pose */
-	const FIKRetargetPose* GetCurrentRetargetPose() const { return &RetargetPoses[CurrentRetargetPose]; };
+	const FIKRetargetPose* GetCurrentRetargetPose(const ERetargetSourceOrTarget& SourceOrTarget) const;
 	/* Get name of default pose */
 	static const FName GetDefaultPoseName();
 
@@ -285,13 +271,10 @@ public:
 	void GetSpeedCurveNames(TArray<FName>& OutSpeedCurveNames) const;
 #endif
 
-#if WITH_EDITOR
-	ERetargeterOutputMode GetOutputMode() const { return bOutputMode; }
-	void SetOutputMode(const ERetargeterOutputMode InMode) { bOutputMode = InMode; }
-#endif
-
 private:
 
+	void CleanAndInitialize();
+
 	/** The rig to copy animation FROM.*/
 	UPROPERTY(EditAnywhere, Category = Source)
 	TSoftObjectPtr<UIKRigDefinition> SourceIKRigAsset = nullptr;
@@ -353,9 +336,6 @@ public:
 	float BoneDrawSize = 8.0f;
 	
 private:
-	/** A special editor-only mode which forces the retargeter to output the current retarget reference pose,
-	* rather than actually running the retarget and outputting the retargeted pose. Used in Edit-Pose mode.*/
-	ERetargeterOutputMode bOutputMode;
 
 	/** The controller responsible for managing this asset's data (all editor mutation goes through this) */
 	UPROPERTY(Transient)
@@ -363,9 +343,6 @@ private:
 #endif
 	
 private:
-	/** The set of retarget poses available as options for retargeting.*/
-	UPROPERTY()
-	TMap<FName, FIKRetargetPose> RetargetPoses;
 
 	/** (OLD VERSION) Mapping of chains to copy animation between source and target rigs.*/
 	PRAGMA_DISABLE_DEPRECATION_WARNINGS
@@ -377,9 +354,27 @@ private:
 	UPROPERTY()
 	TArray<TObjectPtr<URetargetChainSettings>> ChainSettings;
 	
-	/** The set of retarget poses available as options for retargeting.*/
+	/** The set of retarget poses for the SOURCE skeleton.*/
+	UPROPERTY()
+	TMap<FName, FIKRetargetPose> SourceRetargetPoses;
+	/** The set of retarget poses for the TARGET skeleton.*/
+	UPROPERTY()
+	TMap<FName, FIKRetargetPose> TargetRetargetPoses;
+	
+	/** The current retarget pose to use for the SOURCE.*/
+	UPROPERTY()
+	FName CurrentSourceRetargetPose;
+	/** The current retarget pose to use for the TARGET.*/
 	UPROPERTY()
-	FName CurrentRetargetPose;
+	FName CurrentTargetRetargetPose;
+
+	/** (OLD VERSION) Before retarget poses were stored for target AND source.*/
+	PRAGMA_DISABLE_DEPRECATION_WARNINGS
+	UPROPERTY()
+	TMap<FName, FIKRetargetPose> RetargetPoses_DEPRECATED;
+	UPROPERTY()
+	FName CurrentRetargetPose_DEPRECATED;
+	PRAGMA_ENABLE_DEPRECATION_WARNINGS
 
 	friend class UIKRetargeterController;
 };
diff --git a/Engine/Plugins/Animation/IKRig/Source/IKRigEditor/Private/IKRigEditor.cpp b/Engine/Plugins/Animation/IKRig/Source/IKRigEditor/Private/IKRigEditor.cpp
--- a/Engine/Plugins/Animation/IKRig/Source/IKRigEditor/Private/IKRigEditor.cpp
+++ b/Engine/Plugins/Animation/IKRig/Source/IKRigEditor/Private/IKRigEditor.cpp
@@ -10,7 +10,6 @@
 #include "IKRigDefinition.h"
 #include "RigEditor/AssetTypeActions_IKRigDefinition.h"
 #include "RetargetEditor/AssetTypeActions_IKRetargeter.h"
-#include "RetargetEditor/AssetTypeActions_RetargetPose.h"
 #include "RetargetEditor/IKRetargetCommands.h"
 #include "RetargetEditor/IKRetargetDefaultMode.h"
 #include "RetargetEditor/IKRetargetDetails.h"
@@ -42,9 +41,6 @@ void FIKRigEditor::StartupModule()
 	//
 	IKRetargeterAssetAction = MakeShareable(new FAssetTypeActions_IKRetargeter);
 	ToolsModule.RegisterAssetTypeActions(IKRetargeterAssetAction.ToSharedRef());
-	//
-	RetargetPoseAssetAction = MakeShareable(new FAssetTypeActions_RetargetPose);
-	ToolsModule.RegisterAssetTypeActions(RetargetPoseAssetAction.ToSharedRef());
 
 	// extend the content browser menu
 	FAssetTypeActions_IKRetargeter::ExtendAnimSequenceToolMenu();
@@ -91,11 +87,6 @@ void FIKRigEditor::ShutdownModule()
 		{
 			ToolsModule.UnregisterAssetTypeActions(IKRetargeterAssetAction.ToSharedRef());
 		}
-		
-		if (RetargetPoseAssetAction.IsValid())
-		{
-			ToolsModule.UnregisterAssetTypeActions(RetargetPoseAssetAction.ToSharedRef());
-		}
 	}
 }
 
diff --git a/Engine/Plugins/Animation/IKRig/Source/IKRigEditor/Private/RetargetEditor/AssetTypeActions_RetargetPose.cpp b/Engine/Plugins/Animation/IKRig/Source/IKRigEditor/Private/RetargetEditor/AssetTypeActions_RetargetPose.cpp
deleted file mode 100644
diff --git a/Engine/Plugins/Animation/IKRig/Source/IKRigEditor/Private/RetargetEditor/IKRetargetAnimInstance.cpp b/Engine/Plugins/Animation/IKRig/Source/IKRigEditor/Private/RetargetEditor/IKRetargetAnimInstance.cpp
--- a/Engine/Plugins/Animation/IKRig/Source/IKRigEditor/Private/RetargetEditor/IKRetargetAnimInstance.cpp
+++ b/Engine/Plugins/Animation/IKRig/Source/IKRigEditor/Private/RetargetEditor/IKRetargetAnimInstance.cpp
@@ -2,6 +2,114 @@
 
 #include "RetargetEditor/IKRetargetAnimInstance.h"
 #include "RetargetEditor/IKRetargetAnimInstanceProxy.h"
+#include "RetargetEditor/IKRetargeterController.h"
+
+
+void FAnimNode_PreviewRetargetPose::CacheBones_AnyThread(const FAnimationCacheBonesContext& Context)
+{
+	DECLARE_SCOPE_HIERARCHICAL_COUNTER_ANIMNODE(CacheBones_AnyThread)
+
+	const FBoneContainer& RequiredBones = Context.AnimInstanceProxy->GetRequiredBones();
+	if (!RequiredBones.IsValid())
+	{
+		return;
+	}
+	
+	if (!Context.AnimInstanceProxy->GetSkelMeshComponent()->GetSkeletalMesh())
+	{
+		return;
+	}
+
+	// rebuild mapping
+	RequiredBoneToMeshBoneMap.Reset();
+
+	const FReferenceSkeleton& RefSkeleton = RequiredBones.GetReferenceSkeleton();
+	const FReferenceSkeleton& TargetSkeleton = Context.AnimInstanceProxy->GetSkelMeshComponent()->GetSkeletalMesh()->GetRefSkeleton();
+	const TArray<FBoneIndexType>& RequiredBonesArray = RequiredBones.GetBoneIndicesArray();
+	for (int32 RequiredBoneIndex = 0; RequiredBoneIndex < RequiredBonesArray.Num(); ++RequiredBoneIndex)
+	{
+		const FBoneIndexType ReqBoneIndex = RequiredBonesArray[RequiredBoneIndex]; 
+		if (ReqBoneIndex == INDEX_NONE)
+		{
+			continue;
+		}
+		
+		const FName BoneName = RefSkeleton.GetBoneName(ReqBoneIndex);
+		const int32 TargetBoneIndex = TargetSkeleton.FindBoneIndex(BoneName);
+		if (TargetBoneIndex == INDEX_NONE)
+		{
+			continue;
+		}
+
+		// store require bone to target bone indices
+		RequiredBoneToMeshBoneMap.Emplace(TargetBoneIndex);
+	}
+}
+
+void FAnimNode_PreviewRetargetPose::Evaluate_AnyThread(FPoseContext& Output)
+{
+	DECLARE_SCOPE_HIERARCHICAL_COUNTER_ANIMNODE(Evaluate_AnyThread)
+
+	Output.Pose.ResetToRefPose();
+	
+	if (!IKRetargeterAsset)
+	{
+		return;
+	}
+	
+	const FIKRetargetPose* RetargetPose = IKRetargeterAsset->GetCurrentRetargetPose(SourceOrTarget);
+	if (!RetargetPose)
+	{
+		return;	
+	}
+	
+	// generate full (no LOD) LOCAL retarget pose by applying the local rotation offsets from the stored retarget pose
+	// these poses are read back by the editor and need to be complete (not culled)
+	const FReferenceSkeleton& RefSkeleton = Output.AnimInstanceProxy->GetSkelMeshComponent()->GetSkeletalMesh()->GetRefSkeleton();
+	const TArray<FTransform>& RefPose = RefSkeleton.GetRefBonePose(); 
+	RetargetLocalPose = RefPose;
+	for (const TTuple<FName, FQuat>& BoneDelta : RetargetPose->GetAllDeltaRotations())
+	{
+		const int32 BoneIndex = RefSkeleton.FindBoneIndex(BoneDelta.Key);
+		if (BoneIndex != INDEX_NONE)
+		{
+			FTransform& LocalBoneTransform = RetargetLocalPose[BoneIndex];
+			LocalBoneTransform.SetRotation(LocalBoneTransform.GetRotation() * BoneDelta.Value);
+		}
+	}
+
+	// generate GLOBAL space pose (for editor to query)
+	FAnimationRuntime::FillUpComponentSpaceTransforms(RefSkeleton, RetargetLocalPose, RetargetGlobalPose);
+
+	// apply root translation offset from the retarget pose (done in global space)
+	const UIKRetargeterController* Controller = UIKRetargeterController::GetController(IKRetargeterAsset);
+	const FName RetargetRootBoneName = Controller->GetRetargetRootBone(SourceOrTarget);
+	const int32 RootBoneIndex = RefSkeleton.FindBoneIndex(RetargetRootBoneName);
+	if (RootBoneIndex != INDEX_NONE)
+	{
+		FTransform& RootTransform = RetargetGlobalPose[RootBoneIndex];
+		RootTransform.AddToTranslation(RetargetPose->GetRootTranslationDelta());
+		// update local transform of retarget root
+		const int32 ParentIndex = RefSkeleton.GetParentIndex(RootBoneIndex);
+		if (ParentIndex != INDEX_NONE)
+		{
+			const FTransform& ChildGlobalTransform = RetargetGlobalPose[RootBoneIndex];
+			const FTransform& ParentGlobalTransform = RetargetGlobalPose[ParentIndex];
+			RetargetLocalPose[RootBoneIndex] = ChildGlobalTransform.GetRelativeTransform(ParentGlobalTransform);
+		}
+	}
+	
+	// copy to the compact output pose
+	const int32 NumBones = Output.Pose.GetNumBones();
+	for (int32 Index = 0; Index<NumBones; ++Index)
+	{
+		FCompactPoseBoneIndex BoneIndex(Index);
+		const int32 MeshBoneIndex = RequiredBoneToMeshBoneMap[Index];
+
+		BlendTransform<ETransformBlendMode::Overwrite>(RefPose[MeshBoneIndex], Output.Pose[BoneIndex], 1.f - RetargetPoseBlend);
+		BlendTransform<ETransformBlendMode::Accumulate>(RetargetLocalPose[MeshBoneIndex], Output.Pose[BoneIndex], RetargetPoseBlend);
+	}
+}
 
 UIKRetargetAnimInstance::UIKRetargetAnimInstance(const FObjectInitializer& ObjectInitializer)
 	: Super(ObjectInitializer)
@@ -9,21 +117,34 @@ UIKRetargetAnimInstance::UIKRetargetAnimInstance(const FObjectInitializer& Objec
 	bUseMultiThreadedAnimationUpdate = false;
 }
 
-void UIKRetargetAnimInstance::SetRetargetAssetAndSourceComponent(
+void UIKRetargetAnimInstance::SetRetargetMode(const ERetargeterOutputMode& InOutputMode)
+{
+	FIKRetargetAnimInstanceProxy& Proxy = GetProxyOnGameThread<FIKRetargetAnimInstanceProxy>();
+	Proxy.SetRetargetMode(InOutputMode);
+}
+
+void UIKRetargetAnimInstance::SetRetargetPoseBlend(const float& RetargetPoseBlend)
+{
+	FIKRetargetAnimInstanceProxy& Proxy = GetProxyOnGameThread<FIKRetargetAnimInstanceProxy>();
+	Proxy.SetRetargetPoseBlend(RetargetPoseBlend);
+}
+
+void UIKRetargetAnimInstance::ConfigureAnimInstance(
+	const ERetargetSourceOrTarget& InSourceOrTarget,
 	UIKRetargeter* InAsset,
 	TWeakObjectPtr<USkeletalMeshComponent> InSourceMeshComponent)
 {
 	FIKRetargetAnimInstanceProxy& Proxy = GetProxyOnGameThread<FIKRetargetAnimInstanceProxy>();
-	Proxy.SetRetargetAssetAndSourceComponent(InAsset, InSourceMeshComponent);
+	Proxy.ConfigureAnimInstance(InSourceOrTarget, InAsset, InSourceMeshComponent);
 }
 
 UIKRetargetProcessor* UIKRetargetAnimInstance::GetRetargetProcessor() const
 {
-	return IKRetargeterNode.GetRetargetProcessor();
+	return RetargetNode.GetRetargetProcessor();
 }
 
 FAnimInstanceProxy* UIKRetargetAnimInstance::CreateAnimInstanceProxy()
 {
 	LLM_SCOPE_BYNAME(TEXT("Animation/IKRig"));
-	return new FIKRetargetAnimInstanceProxy(this, &IKRetargeterNode);
+	return new FIKRetargetAnimInstanceProxy(this, &PreviewPoseNode, &RetargetNode);
 }
diff --git a/Engine/Plugins/Animation/IKRig/Source/IKRigEditor/Private/RetargetEditor/IKRetargetAnimInstanceProxy.cpp b/Engine/Plugins/Animation/IKRig/Source/IKRigEditor/Private/RetargetEditor/IKRetargetAnimInstanceProxy.cpp
--- a/Engine/Plugins/Animation/IKRig/Source/IKRigEditor/Private/RetargetEditor/IKRetargetAnimInstanceProxy.cpp
+++ b/Engine/Plugins/Animation/IKRig/Source/IKRigEditor/Private/RetargetEditor/IKRetargetAnimInstanceProxy.cpp
@@ -2,33 +2,72 @@
 
 #include "RetargetEditor/IKRetargetAnimInstanceProxy.h"
 #include "RetargetEditor/IKRetargetAnimInstance.h"
+#include "AnimNodes/AnimNode_RetargetPoseFromMesh.h"
+#include "RetargetEditor/IKRetargetEditorController.h"
 
-
-FIKRetargetAnimInstanceProxy::FIKRetargetAnimInstanceProxy(UAnimInstance* InAnimInstance, FAnimNode_RetargetPoseFromMesh* InNode)
+FIKRetargetAnimInstanceProxy::FIKRetargetAnimInstanceProxy(
+	UAnimInstance* InAnimInstance,
+	FAnimNode_PreviewRetargetPose* InPreviewPoseNode,
+	FAnimNode_RetargetPoseFromMesh* InRetargetNode)
 	: FAnimPreviewInstanceProxy(InAnimInstance),
-	IKRetargetNode(InNode)
+	PreviewPoseNode(InPreviewPoseNode),
+	RetargetNode(InRetargetNode),
+	OutputMode(ERetargeterOutputMode::ShowRetargetPose)
 {
 }
 
 void FIKRetargetAnimInstanceProxy::Initialize(UAnimInstance* InAnimInstance)
 {
 	FAnimPreviewInstanceProxy::Initialize(InAnimInstance);
+	PreviewPoseNode->InputPose.SetLinkNode(RetargetNode);
+
+	FAnimationInitializeContext InitContext(this);
+	PreviewPoseNode->Initialize_AnyThread(InitContext);
+	RetargetNode->Initialize_AnyThread(InitContext);
 }
 
-bool FIKRetargetAnimInstanceProxy::Evaluate(FPoseContext& Output)
+void FIKRetargetAnimInstanceProxy::CacheBones()
 {
-	IKRetargetNode->Evaluate_AnyThread(Output);
+	if (bBoneCachesInvalidated)
+	{
+		FAnimationCacheBonesContext Context(this);
+		SingleNode.CacheBones_AnyThread(Context);
+		RetargetNode->CacheBones_AnyThread(Context);
+		PreviewPoseNode->CacheBones_AnyThread(Context);
+		bBoneCachesInvalidated = false;
+	}
+}
+
+bool FIKRetargetAnimInstanceProxy::Evaluate(FPoseContext& Output)
+{	
+	if (OutputMode == ERetargeterOutputMode::RunRetarget)
+	{
+		if (SourceOrTarget == ERetargetSourceOrTarget::Source)
+		{
+			FAnimPreviewInstanceProxy::Evaluate(Output);
+		}
+		else
+		{
+			RetargetNode->Evaluate_AnyThread(Output);
+		}
+	}
+	else
+	{
+		PreviewPoseNode->Evaluate_AnyThread(Output);
+	}
+	
 	return true;
 }
 
 FAnimNode_Base* FIKRetargetAnimInstanceProxy::GetCustomRootNode()
 {
-	return IKRetargetNode;
+	return PreviewPoseNode;
 }
 
 void FIKRetargetAnimInstanceProxy::GetCustomNodes(TArray<FAnimNode_Base*>& OutNodes)
 {
-	OutNodes.Add(IKRetargetNode);
+	OutNodes.Add(RetargetNode);
+	OutNodes.Add(PreviewPoseNode);
 }
 
 void FIKRetargetAnimInstanceProxy::UpdateAnimationNode(const FAnimationUpdateContext& InContext)
@@ -39,21 +78,41 @@ void FIKRetargetAnimInstanceProxy::UpdateAnimationNode(const FAnimationUpdateCon
 	}
 	else
 	{
-		IKRetargetNode->Update_AnyThread(InContext);
+		PreviewPoseNode->Update_AnyThread(InContext);
+		RetargetNode->Update_AnyThread(InContext);
 	}
 }
 
-void FIKRetargetAnimInstanceProxy::SetRetargetAssetAndSourceComponent(
+void FIKRetargetAnimInstanceProxy::ConfigureAnimInstance(
+	const ERetargetSourceOrTarget& InSourceOrTarget,
 	UIKRetargeter* InIKRetargetAsset,
-	TWeakObjectPtr<USkeletalMeshComponent> InSourceMeshComponent) const
+	TWeakObjectPtr<USkeletalMeshComponent> InSourceMeshComponent)
 {
-	IKRetargetNode->IKRetargeterAsset = InIKRetargetAsset;
-	IKRetargetNode->bUseAttachedParent = false;
-	IKRetargetNode->SourceMeshComponent = InSourceMeshComponent;
-	IKRetargetNode->bDriveWithAsset = true;
-	if (UIKRetargetProcessor* Processor = IKRetargetNode->GetRetargetProcessor())
+	SourceOrTarget = InSourceOrTarget;
+
+	PreviewPoseNode->SourceOrTarget = InSourceOrTarget;
+	PreviewPoseNode->IKRetargeterAsset = InIKRetargetAsset;
+
+	if (SourceOrTarget == ERetargetSourceOrTarget::Target)
 	{
-		Processor->SetNeedsInitialized();
+		RetargetNode->IKRetargeterAsset = InIKRetargetAsset;
+		RetargetNode->bUseAttachedParent = false;
+		RetargetNode->SourceMeshComponent = InSourceMeshComponent;
+		RetargetNode->bDriveWithAsset = true;
+		if (UIKRetargetProcessor* Processor = RetargetNode->GetRetargetProcessor())
+		{
+			Processor->SetNeedsInitialized();
+		}
 	}
 }
 
+void FIKRetargetAnimInstanceProxy::SetRetargetMode(const ERetargeterOutputMode& InOutputMode)
+{
+	OutputMode = InOutputMode;
+}
+
+void FIKRetargetAnimInstanceProxy::SetRetargetPoseBlend(const float& InRetargetPoseBlend) const
+{
+	PreviewPoseNode->RetargetPoseBlend = InRetargetPoseBlend;
+}
+
diff --git a/Engine/Plugins/Animation/IKRig/Source/IKRigEditor/Private/RetargetEditor/IKRetargetCommands.cpp b/Engine/Plugins/Animation/IKRig/Source/IKRigEditor/Private/RetargetEditor/IKRetargetCommands.cpp
--- a/Engine/Plugins/Animation/IKRig/Source/IKRigEditor/Private/RetargetEditor/IKRetargetCommands.cpp
+++ b/Engine/Plugins/Animation/IKRig/Source/IKRigEditor/Private/RetargetEditor/IKRetargetCommands.cpp
@@ -6,7 +6,7 @@
 
 void FIKRetargetCommands::RegisterCommands()
 {
-	UI_COMMAND(GoToRetargetPose, "Go to Retarget Pose", "Stop playback and set Source and Target meshes to the Retarget Pose.", EUserInterfaceActionType::Button, FInputChord());
+	UI_COMMAND(ShowRetargetPose, "Show Retarget Pose", "Pause playback and set Source and Target meshes to the Retarget Pose.", EUserInterfaceActionType::ToggleButton, FInputChord());
 	
 	UI_COMMAND(EditRetargetPose, "Edit Mode", "Enter into mode allowing manual editing of the target skeleton pose in the viewport.", EUserInterfaceActionType::ToggleButton, FInputChord());
 	UI_COMMAND(ResetAllBones, "Reset All", "Sets the selected bones to the mesh reference pose.", EUserInterfaceActionType::Button, FInputChord());
@@ -17,9 +17,9 @@ void FIKRetargetCommands::RegisterCommands()
 	UI_COMMAND(DuplicateRetargetPose, "Duplicate Current", "Duplicate the current retarget pose.", EUserInterfaceActionType::Button, FInputChord());
 	UI_COMMAND(DeleteRetargetPose, "Delete", "Delete current retarget pose.", EUserInterfaceActionType::Button, FInputChord());
 	UI_COMMAND(RenameRetargetPose, "Rename", "Rename current retarget pose.", EUserInterfaceActionType::Button, FInputChord());
-	UI_COMMAND(ImportRetargetPose, "Import from Retarget Pose Asset", "Import a retarget pose asset.", EUserInterfaceActionType::Button, FInputChord());
+	UI_COMMAND(ImportRetargetPose, "Import from Pose Asset", "Import a retarget pose asset.", EUserInterfaceActionType::Button, FInputChord());
 	UI_COMMAND(ImportRetargetPoseFromAnim, "Import from Animation Sequence", "Import a retarget pose from an animation sequence.", EUserInterfaceActionType::Button, FInputChord());
-	UI_COMMAND(ExportRetargetPose, "Export Retarget Pose Asset", "Export current retarget pose as an asset.", EUserInterfaceActionType::Button, FInputChord());
+	UI_COMMAND(ExportRetargetPose, "Export Pose Asset", "Export current retarget pose as a Pose Asset.", EUserInterfaceActionType::Button, FInputChord());
 }
 
 #undef LOCTEXT_NAMESPACE
diff --git a/Engine/Plugins/Animation/IKRig/Source/IKRigEditor/Private/RetargetEditor/IKRetargetDefaultMode.cpp b/Engine/Plugins/Animation/IKRig/Source/IKRigEditor/Private/RetargetEditor/IKRetargetDefaultMode.cpp
--- a/Engine/Plugins/Animation/IKRig/Source/IKRigEditor/Private/RetargetEditor/IKRetargetDefaultMode.cpp
+++ b/Engine/Plugins/Animation/IKRig/Source/IKRigEditor/Private/RetargetEditor/IKRetargetDefaultMode.cpp
@@ -254,7 +254,7 @@ void FIKRetargetDefaultMode::Enter()
 	}
 
 	// record which skeleton is being viewed/edited
-	SkeletonMode = Controller->GetSkeletonMode();
+	SkeletonMode = Controller->GetSourceOrTarget();
 
 	// allow selection of meshes in this mode
 	Controller->Editor.Pin()->GetPersonaToolkit()->GetPreviewScene()->SetAllowMeshHitProxies(true);
@@ -301,7 +301,7 @@ UDebugSkelMeshComponent* FIKRetargetDefaultMode::GetCurrentlyEditedMesh() const
 		return nullptr; 
 	}
 	
-	return SkeletonMode == EIKRetargetSkeletonMode::Source ? Controller->SourceSkelMeshComponent : Controller->TargetSkelMeshComponent;
+	return SkeletonMode == ERetargetSourceOrTarget::Source ? Controller->SourceSkelMeshComponent : Controller->TargetSkelMeshComponent;
 }
 
 void FIKRetargetDefaultMode::ApplyOffsetToMeshTransform(const FVector& Offset, USceneComponent* Component)
@@ -328,7 +328,7 @@ void FIKRetargetDefaultMode::Tick(FEditorViewportClient* ViewportClient, float D
 	const TSharedPtr<FIKRetargetEditorController> Controller = EditorController.Pin();
 	if (Controller.IsValid())
 	{
-		const bool bEditingSource = Controller->GetSkeletonMode() == EIKRetargetSkeletonMode::Source;
+		const bool bEditingSource = Controller->GetSourceOrTarget() == ERetargetSourceOrTarget::Source;
 		Controller->SourceSkelMeshComponent->SkeletonDrawMode = bEditingSource ? ESkeletonDrawMode::Default : ESkeletonDrawMode::GreyedOut;
 		Controller->TargetSkelMeshComponent->SkeletonDrawMode = !bEditingSource ? ESkeletonDrawMode::Default : ESkeletonDrawMode::GreyedOut;
 	}
diff --git a/Engine/Plugins/Animation/IKRig/Source/IKRigEditor/Private/RetargetEditor/IKRetargetDetails.cpp b/Engine/Plugins/Animation/IKRig/Source/IKRigEditor/Private/RetargetEditor/IKRetargetDetails.cpp
--- a/Engine/Plugins/Animation/IKRig/Source/IKRigEditor/Private/RetargetEditor/IKRetargetDetails.cpp
+++ b/Engine/Plugins/Animation/IKRig/Source/IKRigEditor/Private/RetargetEditor/IKRetargetDetails.cpp
@@ -33,7 +33,7 @@ FEulerTransform UIKRetargetBoneDetails::GetTransform(EIKRetargetTransformType Tr
 	}
 
 	// ensure we have a valid skeletal mesh
-	const bool bEditingSource = EditorController->GetSkeletonMode() == EIKRetargetSkeletonMode::Source;
+	const bool bEditingSource = EditorController->GetSourceOrTarget() == ERetargetSourceOrTarget::Source;
 	UDebugSkelMeshComponent* Mesh = bEditingSource ? Controller->SourceSkelMeshComponent : Controller->TargetSkelMeshComponent;
 	if (!(Mesh && Mesh->GetSkeletalMesh()))
 	{
@@ -78,8 +78,9 @@ FEulerTransform UIKRetargetBoneDetails::GetTransform(EIKRetargetTransformType Tr
 	case EIKRetargetTransformType::RelativeOffset:
 		{
 			// this is the only stored data we have for bone pose offsets
-			const FRotator LocalRotationDelta = EditorController->AssetController->GetRotationOffsetForRetargetPoseBone(SelectedBone).Rotator();
-			const FVector GlobalTranslationDelta = IsRootBone() ? EditorController->AssetController->GetTranslationOffsetOnRetargetRootBone() : FVector::Zero();
+			const ERetargetSourceOrTarget SourceOrTarget = EditorController->GetSourceOrTarget();
+			const FRotator LocalRotationDelta = EditorController->AssetController->GetRotationOffsetForRetargetPoseBone(SelectedBone, SourceOrTarget).Rotator();
+			const FVector GlobalTranslationDelta = IsRootBone() ? EditorController->AssetController->GetTranslationOffsetOnRetargetRootBone(SourceOrTarget) : FVector::Zero();
 			const int32 ParentIndex = RefSkeleton.GetParentIndex(BoneIndex);
 
 			if (bLocalSpace)
@@ -280,7 +281,7 @@ void UIKRetargetBoneDetails::OnPasteFromClipboard(
 	};
 
 	FIKRigLogger* Log = nullptr;
-	if (const auto Processor = EditorController->GetRetargetProcessor())
+	if (UIKRetargetProcessor* Processor = EditorController->GetRetargetProcessor())
 	{
 		Log = &Processor->Log;
 	}
@@ -290,6 +291,8 @@ void UIKRetargetBoneDetails::OnPasteFromClipboard(
 	const bool bIsRelative = IsComponentRelative(Component, TransformType);
 	FEulerTransform Transform = GetTransform(TransformType, bIsRelative);
 
+	const ERetargetSourceOrTarget SourceOrTarget = EditorController->GetSourceOrTarget();
+
 	// create a transaction on the asset
 	FScopedTransaction Transaction(LOCTEXT("PasteTransform", "Paste Transform"));
 	EditorController.Get()->AssetController->GetAsset()->Modify();
@@ -303,7 +306,7 @@ void UIKRetargetBoneDetails::OnPasteFromClipboard(
 			if (Result && ErrorPipe.NumErrors == 0)
 			{
 				Transform.SetLocation(Data);
-				EditorController->AssetController->SetTranslationOffsetOnRetargetRootBone(Transform.GetLocation());
+				EditorController->AssetController->SetTranslationOffsetOnRetargetRootBone(Transform.GetLocation(), SourceOrTarget);
 			}
 			break;
 		}
@@ -314,7 +317,7 @@ void UIKRetargetBoneDetails::OnPasteFromClipboard(
 			if (Result && ErrorPipe.NumErrors == 0)
 			{
 				Transform.SetRotator(Data);
-				EditorController->AssetController->SetRotationOffsetForRetargetPoseBone(SelectedBone, Transform.GetRotation());
+				EditorController->AssetController->SetRotationOffsetForRetargetPoseBone(SelectedBone, Transform.GetRotation(), SourceOrTarget);
 			}
 			break;
 		}
@@ -378,7 +381,7 @@ void UIKRetargetBoneDetails::OnNumericValueCommitted(
 	}
 
 	// ensure we have a valid skeletal mesh
-	const bool bEditingSource = EditorController->GetSkeletonMode() == EIKRetargetSkeletonMode::Source;
+	const bool bEditingSource = EditorController->GetSourceOrTarget() == ERetargetSourceOrTarget::Source;
 	UDebugSkelMeshComponent* Mesh = bEditingSource ? EditorController->SourceSkelMeshComponent : EditorController->TargetSkelMeshComponent;
 	if (!(Mesh && Mesh->GetSkeletalMesh()))
 	{
@@ -393,13 +396,15 @@ void UIKRetargetBoneDetails::OnNumericValueCommitted(
 		return;
 	}
 
+	const ERetargetSourceOrTarget SourceOrTarget = EditorController->GetSourceOrTarget();
+
 	switch(Component)
 	{
 	case ESlateTransformComponent::Location:
 		{
 			const bool bIsTranslationLocal = RelativeOffsetTransformRelative[0];
 			FTransform CurrentGlobalOffset = FTransform::Identity;
-			CurrentGlobalOffset.SetTranslation(AssetController->GetTranslationOffsetOnRetargetRootBone());
+			CurrentGlobalOffset.SetTranslation(AssetController->GetTranslationOffsetOnRetargetRootBone(SourceOrTarget));
 			
 			if (bIsTranslationLocal)
 			{	
@@ -428,14 +433,14 @@ void UIKRetargetBoneDetails::OnNumericValueCommitted(
 			// store the new transform in the retarget pose
 			FScopedTransaction Transaction(LOCTEXT("EditRootTranslation", "Edit Retarget Root Pose Translation"));
 			EditorController->AssetController->GetAsset()->Modify();
-			EditorController->AssetController->SetTranslationOffsetOnRetargetRootBone(CurrentGlobalOffset.GetTranslation());
+			EditorController->AssetController->SetTranslationOffsetOnRetargetRootBone(CurrentGlobalOffset.GetTranslation(), SourceOrTarget);
 			
 			break;
 		}
 	case ESlateTransformComponent::Rotation:
 		{
 			const bool bIsRotationLocal = RelativeOffsetTransformRelative[1];
-			const FQuat LocalRotationDelta = AssetController->GetRotationOffsetForRetargetPoseBone(SelectedBone);
+			const FQuat LocalRotationDelta = AssetController->GetRotationOffsetForRetargetPoseBone(SelectedBone, SourceOrTarget);
 			FEulerTransform LocalDeltaTransform = FEulerTransform(FVector::ZeroVector, LocalRotationDelta.Rotator(), FVector::OneVector);
 			FQuat NewLocalRotationDelta;
 			
@@ -476,7 +481,7 @@ void UIKRetargetBoneDetails::OnNumericValueCommitted(
 			// store the new rotation in the retarget pose
 			FScopedTransaction Transaction(LOCTEXT("EditRootRotation", "Edit Retarget Pose Rotation"));
 			EditorController->AssetController->GetAsset()->Modify();
-			EditorController->AssetController->SetRotationOffsetForRetargetPoseBone(SelectedBone, NewLocalRotationDelta);
+			EditorController->AssetController->SetRotationOffsetForRetargetPoseBone(SelectedBone, NewLocalRotationDelta, SourceOrTarget);
 			break;
 		}
 	default:
@@ -486,8 +491,7 @@ void UIKRetargetBoneDetails::OnNumericValueCommitted(
 
 bool UIKRetargetBoneDetails::IsRootBone() const
 {
-	const bool bIsSource = EditorController->GetSkeletonMode() == EIKRetargetSkeletonMode::Source;
-	const FName RootBone = bIsSource ? EditorController->AssetController->GetSourceRootBone() :  EditorController->AssetController->GetTargetRootBone();
+	const FName RootBone = EditorController->AssetController->GetRetargetRootBone(EditorController->GetSourceOrTarget());
 	return SelectedBone == RootBone;
 }
 
@@ -551,10 +555,9 @@ void FIKRetargetBoneDetailCustomization::CustomizeDetails(IDetailLayoutBuilder&
 	const FIKRetargetEditorController& Controller = *Bones[0]->EditorController.Get();
 	const UIKRetargeterController* AssetController = Controller.AssetController;
 	
-	const bool bIsSourceMode = Controller.GetSkeletonMode() == EIKRetargetSkeletonMode::Source;
 	const bool bIsEditingPose = Controller.IsEditingPose();
 
-	const FName CurrentRootName = bIsSourceMode ?AssetController->GetSourceRootBone() : AssetController->GetTargetRootBone();
+	const FName CurrentRootName = AssetController->GetRetargetRootBone(Controller.GetSourceOrTarget());
 	const bool bIsRootSelected =  Bones[0]->SelectedBone == CurrentRootName;
 
 	FIKRetargetTransformUIData UIData;
diff --git a/Engine/Plugins/Animation/IKRig/Source/IKRigEditor/Private/RetargetEditor/IKRetargetEditPoseMode.cpp b/Engine/Plugins/Animation/IKRig/Source/IKRigEditor/Private/RetargetEditor/IKRetargetEditPoseMode.cpp
--- a/Engine/Plugins/Animation/IKRig/Source/IKRigEditor/Private/RetargetEditor/IKRetargetEditPoseMode.cpp
+++ b/Engine/Plugins/Animation/IKRig/Source/IKRigEditor/Private/RetargetEditor/IKRetargetEditPoseMode.cpp
@@ -55,91 +55,74 @@ void FIKRetargetEditPoseMode::Render(const FSceneView* View, FViewport* Viewport
 		return;
 	}
 
-	const UIKRetargetProcessor* Processor = EditorController.Pin()->GetRetargetProcessor();
-	if (!(Processor && Processor->IsInitialized()))
-	{
-		return;
-	}
-
 	// render the skeleton
-	RenderSkeleton(PDI, Controller, Processor);
+	RenderSkeleton(PDI, Controller);
 }
 
 void FIKRetargetEditPoseMode::RenderSkeleton(
 	FPrimitiveDrawInterface* PDI,
-	const FIKRetargetEditorController* Controller,
-	const UIKRetargetProcessor* Processor)
+	const FIKRetargetEditorController* Controller)
 {
-	check(Processor && Controller && Processor->IsInitialized())
-	
-	const FRetargetSkeleton& Skeleton = GetCurrentlyEditedSkeleton(*Processor);
-	
-	// get selected and affected bones in this skeleton
-	TSet<int32> OutAffectedBones;
-	TSet<int32> OutSelectedBones;
-	GetSelectedAndAffectedBones(Controller, Skeleton, OutAffectedBones, OutSelectedBones);
+	const UDebugSkelMeshComponent* MeshComponent = Controller->GetSkeletalMeshComponent(SourceOrTarget);
+	FTransform ComponentTransform = MeshComponent->GetComponentTransform();
+	const FReferenceSkeleton& RefSkeleton = MeshComponent->GetReferenceSkeleton();
+	const int32 NumBones = RefSkeleton.GetNum();
 
+	// get world transforms of bones
+	TArray<FBoneIndexType> RequiredBones;
+	RequiredBones.AddUninitialized(NumBones);
+	TArray<FTransform> WorldTransforms;
+	WorldTransforms.AddUninitialized(NumBones);
+	for (int32 Index=0; Index<NumBones; ++Index)
+	{
+		RequiredBones[Index] = Index;
+		WorldTransforms[Index] = MeshComponent->GetBoneTransform(Index, ComponentTransform);
+	}
+	
 	const UIKRetargeter* Asset = Controller->AssetController->GetAsset();
-	float Scale;
-	FVector Offset;
-	GetEditedComponentScaleAndOffset(Scale, Offset);
 	const float BoneDrawSize = Asset->BoneDrawSize;
 	const float MaxDrawRadius = Controller->TargetSkelMeshComponent->Bounds.SphereRadius * 0.01f;
+	const float BoneRadius = FMath::Min(1.0f, MaxDrawRadius) * BoneDrawSize;
+	
+	UPersonaOptions* ConfigOption = UPersonaOptions::StaticClass()->GetDefaultObject<UPersonaOptions>();
+	
+	FSkelDebugDrawConfig DrawConfig;
+	DrawConfig.BoneDrawMode = (EBoneDrawMode::Type)ConfigOption->DefaultBoneDrawSelection;
+	DrawConfig.BoneDrawSize = BoneRadius;
+	DrawConfig.bAddHitProxy = true;
+	DrawConfig.bForceDraw = true;
+	DrawConfig.DefaultBoneColor = GetMutableDefault<UPersonaOptions>()->DefaultBoneColor;
+	DrawConfig.AffectedBoneColor = GetMutableDefault<UPersonaOptions>()->AffectedBoneColor;
+	DrawConfig.SelectedBoneColor = GetMutableDefault<UPersonaOptions>()->SelectedBoneColor;
+	DrawConfig.ParentOfSelectedBoneColor = GetMutableDefault<UPersonaOptions>()->ParentOfSelectedBoneColor;
 
-	// use colors from user preferences
-	const UPersonaOptions* PersonaOptions = GetDefault<UPersonaOptions>();
-
-	// loop over whole skeleton and draw each bone
-	for (int32 BoneIndex = 0; BoneIndex<Skeleton.BoneNames.Num(); ++BoneIndex)
+	TArray<HHitProxy*> HitProxies;
+	HitProxies.AddUninitialized(NumBones);
+	for (int32 Index = 0; Index < NumBones; ++Index)
 	{
-		// selected bones are drawn with different color, affected bones are highlighted
-		const bool bIsSelected = OutSelectedBones.Contains(BoneIndex);
-		const bool bIsAffected = OutAffectedBones.Contains(BoneIndex);
-		const bool bIsRetargeted = Processor->IsBoneRetargeted(BoneIndex, (int8)SkeletonMode);
-		FLinearColor LineColor = bIsAffected ? PersonaOptions->AffectedBoneColor : PersonaOptions->DefaultBoneColor;
-		LineColor = bIsSelected ? PersonaOptions->SelectedBoneColor : LineColor;
-		LineColor = bIsRetargeted ? LineColor : FLinearColor::LerpUsingHSV(LineColor, PersonaOptions->DisabledBoneColor, 0.35f);
-		
-		const float BoneRadius = FMath::Min(1.0f, MaxDrawRadius) * BoneDrawSize;
-		
-		// get the location of the bone on the currently initialized target skeletal mesh
-		// along with array of child positions
-		const FTransform BoneTransform = Controller->GetGlobalRetargetPoseOfBone(Skeleton, BoneIndex, Scale, Offset);
-		TArray<FVector> ChildrenPositions;
-		TArray<int32> ChildrenIndices;
-		Controller->GetGlobalRetargetPoseOfImmediateChildren(
-			Skeleton,
-			BoneIndex,
-			Scale,
-			Offset,
-			ChildrenIndices,
-			ChildrenPositions);
-
-		// draw cone to parent with different color if this bone is NOT selected, but the child IS selected
-		TArray<FLinearColor> ChildColors;
-		for (int32 ChildIndex : ChildrenIndices)
-		{
-			FLinearColor ChildLineColor = LineColor;
-			if (!bIsSelected && OutSelectedBones.Contains(ChildIndex))
-			{
-				ChildLineColor = PersonaOptions->ParentOfSelectedBoneColor;
-			}
-			ChildColors.Add(ChildLineColor);
-		}
-		
-		// draw the bone proxy
-		PDI->SetHitProxy(new HIKRetargetEditorBoneProxy(Skeleton.BoneNames[BoneIndex]));
-		SkeletalDebugRendering::DrawWireBoneAdvanced(
-			PDI,
-			BoneTransform,
-			ChildrenPositions,
-			ChildColors,
-			LineColor,
-			SDPG_Foreground,
-			BoneRadius,
-			bIsSelected || bIsAffected);
-		PDI->SetHitProxy(NULL);
+		HitProxies[Index] = new HIKRetargetEditorBoneProxy(RefSkeleton.GetBoneName(Index));
 	}
+
+	// record selected bone indices
+	const TArray<FName>& SelectedBoneNames = Controller->GetSelectedBones();
+	TArray<int32> SelectedBones;
+	for (const FName& SelectedBoneName : SelectedBoneNames)
+	{
+		int32 SelectedBoneIndex = RefSkeleton.FindBoneIndex(SelectedBoneName);
+		SelectedBones.Add(SelectedBoneIndex);
+	}
+
+	SkeletalDebugRendering::DrawBones(
+		PDI,
+		ComponentTransform.GetLocation(),
+		RequiredBones,
+		RefSkeleton,
+		WorldTransforms,
+		SelectedBones,
+		TArray<FLinearColor>(),
+		HitProxies,
+		DrawConfig
+	);
 }
 
 void FIKRetargetEditPoseMode::GetSelectedAndAffectedBones(
@@ -225,26 +208,24 @@ FVector FIKRetargetEditPoseMode::GetWidgetLocation() const
 		return FVector::ZeroVector; 
 	}
 
-	const UIKRetargetProcessor* RetargetProcessor = Controller->GetRetargetProcessor();
-	if (!(RetargetProcessor && RetargetProcessor->IsInitialized()))
+	const USkeletalMesh* SkeletalMesh = Controller->GetSkeletalMesh(SourceOrTarget);
+	if (!SkeletalMesh)
 	{
-		return FVector::ZeroVector;
+		return FVector::ZeroVector; 
 	}
-
-	const bool bIsSource = Controller->GetSkeletonMode() == EIKRetargetSkeletonMode::Source;
 	
-	const FRetargetSkeleton& Skeleton = bIsSource ? RetargetProcessor->GetSourceSkeleton() : RetargetProcessor->GetTargetSkeleton();
-	const int32 BoneIndex = Skeleton.FindBoneIndexByName(SelectedBones.Last());
+	const FReferenceSkeleton& RefSkeleton = SkeletalMesh->GetRefSkeleton();
+	const int32 BoneIndex = RefSkeleton.FindBoneIndex(SelectedBones.Last());
 	if (BoneIndex == INDEX_NONE)
 	{
 		return FVector::ZeroVector;
 	}
 	
-	const UIKRetargeter* Asset = Controller->AssetController->GetAsset();
-	const float Scale = bIsSource ? 1.0f : Asset->TargetMeshScale;
-	const FVector Offset = bIsSource ? Asset->SourceMeshOffset : Asset->TargetMeshOffset;
+	float Scale;
+	FVector Offset;
+	GetEditedComponentScaleAndOffset(Scale,Offset);
 	
-	return Controller->GetGlobalRetargetPoseOfBone(Skeleton, BoneIndex, Scale, Offset).GetTranslation();
+	return Controller->GetGlobalRetargetPoseOfBone(SourceOrTarget, BoneIndex, Scale, Offset).GetTranslation();
 }
 
 bool FIKRetargetEditPoseMode::HandleClick(FEditorViewportClient* InViewportClient, HHitProxy* HitProxy, const FViewportClick& Click)
@@ -362,13 +343,6 @@ bool FIKRetargetEditPoseMode::InputDelta(
 	{
 		return false; 
 	}
-
-	const UIKRetargetProcessor* Processor = Controller->GetRetargetProcessor();
-	if (!(Processor && Processor->IsInitialized()))
-	{
-		return false;
-	}
-	const FRetargetSkeleton& Skeleton = GetCurrentlyEditedSkeleton(*Processor);
 	
 	// rotating any bone
 	if (TrackingState == FIKRetargetTrackingState::RotatingBone)
@@ -394,7 +368,7 @@ bool FIKRetargetEditPoseMode::InputDelta(
 
 			// apply the new delta to the retarget pose
 			const FQuat TotalDeltaRotation = BoneEdit.PreviousDeltaRotation[SelectionIndex] * BoneLocalDelta;
-			Controller->AssetController->SetRotationOffsetForRetargetPoseBone(BoneName, TotalDeltaRotation);
+			Controller->AssetController->SetRotationOffsetForRetargetPoseBone(BoneName, TotalDeltaRotation, Controller->GetSourceOrTarget());
 		}
 
 		return true;
@@ -409,7 +383,7 @@ bool FIKRetargetEditPoseMode::InputDelta(
 		}
 
 		// apply translation delta to root
-		Controller->AssetController->AddTranslationOffsetToRetargetRootBone(InDrag);
+		Controller->AssetController->AddTranslationOffsetToRetargetRootBone(InDrag, Controller->GetSourceOrTarget());
 		return true;
 	}
 	
@@ -452,20 +426,7 @@ void FIKRetargetEditPoseMode::Enter()
 	{
 		return;
 	}
-
-	const UIKRetargetProcessor* Processor = Controller->GetRetargetProcessor();
-	if (!Processor)
-	{
-		return;
-	}
-
-	// put source mesh in reference pose
-	Controller->SourceSkelMeshComponent->ShowReferencePose(true);
-	// have to move component back to offset position because ShowReferencePose() sets it back to origin
-	Controller->AddOffsetToMeshComponent(FVector::ZeroVector, Controller->SourceSkelMeshComponent);
-	// put asset in mode where target mesh will output retarget pose (for preview purposes)
-	Controller->AssetController->GetAsset()->SetOutputMode(ERetargeterOutputMode::EditRetargetPose);
-
+	
 	// clear bone edit
 	BoneEdit.Reset();
 
@@ -474,9 +435,9 @@ void FIKRetargetEditPoseMode::Enter()
 	Controller->ClearSelection(bKeepBoneSelection);
 	
 	// record skeleton to edit (must be constant between enter/exit)
-	SkeletonMode = Controller->GetSkeletonMode();
+	SourceOrTarget = Controller->GetSourceOrTarget();
 	// hide skeleton we're editing so that we can draw our own editable version of it
-	const bool bEditingSource = SkeletonMode == EIKRetargetSkeletonMode::Source;
+	const bool bEditingSource = SourceOrTarget == ERetargetSourceOrTarget::Source;
 	Controller->SourceSkelMeshComponent->SkeletonDrawMode = bEditingSource ? ESkeletonDrawMode::Hidden : ESkeletonDrawMode::GreyedOut;
 	Controller->TargetSkelMeshComponent->SkeletonDrawMode = !bEditingSource ? ESkeletonDrawMode::Hidden : ESkeletonDrawMode::GreyedOut;
 }
@@ -492,17 +453,12 @@ void FIKRetargetEditPoseMode::Exit()
 	IPersonaEditMode::Exit();
 }
 
-const FRetargetSkeleton& FIKRetargetEditPoseMode::GetCurrentlyEditedSkeleton(const UIKRetargetProcessor& Processor) const
-{
-	return SkeletonMode == EIKRetargetSkeletonMode::Source ? Processor.GetSourceSkeleton() : Processor.GetTargetSkeleton();
-}
-
 void FIKRetargetEditPoseMode::GetEditedComponentScaleAndOffset(float& OutScale, FVector& OutOffset) const
 {
 	OutScale = 1.0f;
 	OutOffset = FVector::Zero();
 
-	const bool bIsSource = EditorController.Pin()->GetSkeletonMode() == EIKRetargetSkeletonMode::Source;
+	const bool bIsSource = EditorController.Pin()->GetSourceOrTarget() == ERetargetSourceOrTarget::Source;
 	const UIKRetargeter* Asset = EditorController.Pin()->AssetController->GetAsset();
 	if (!Asset)
 	{
@@ -529,40 +485,35 @@ void FIKRetargetEditPoseMode::UpdateWidgetTransform()
 		return;
 	}
 
-	const UIKRetargetProcessor* RetargetProcessor = Controller->GetRetargetProcessor();
-	if (!(RetargetProcessor && RetargetProcessor->IsInitialized()))
+	const USkeletalMesh* SkeletalMesh = Controller->GetSkeletalMesh(SourceOrTarget);
+	if (!SkeletalMesh)
 	{
+		BoneEdit.GlobalTransform = FTransform::Identity;
 		return;
 	}
 	
-	const UIKRetargeterController* AssetController = Controller->AssetController;
-	const FRetargetSkeleton& Skeleton = GetCurrentlyEditedSkeleton(*RetargetProcessor);
 	float Scale;
 	FVector Offset;
 	GetEditedComponentScaleAndOffset(Scale,Offset);
 
-	const FIKRetargetPose& RetargetPose = AssetController->GetCurrentRetargetPose();
+	const UIKRetargeterController* AssetController = Controller->AssetController;
+	const FIKRetargetPose& RetargetPose = AssetController->GetCurrentRetargetPose(SourceOrTarget);
+	const FReferenceSkeleton& RefSkeleton = SkeletalMesh->GetRefSkeleton();
 
 	BoneEdit.Name = SelectedBones.Last();
-	BoneEdit.Index = Skeleton.FindBoneIndexByName(BoneEdit.Name);
-	BoneEdit.GlobalTransform = Controller->GetGlobalRetargetPoseOfBone(Skeleton, BoneEdit.Index, Scale, Offset);
+	BoneEdit.Index = RefSkeleton.FindBoneIndex(BoneEdit.Name);
+	BoneEdit.GlobalTransform = Controller->GetGlobalRetargetPoseOfBone(SourceOrTarget, BoneEdit.Index, Scale, Offset);
 	BoneEdit.AccumulatedGlobalOffset = FQuat::Identity;
 
 	BoneEdit.PreviousDeltaRotation.Reset();
 	for (int32 SelectionIndex=0; SelectionIndex<SelectedBones.Num(); ++SelectionIndex)
 	{
-		FQuat PrevDeltaRotation = RetargetProcessor->GetTargetSkeleton().GetRetargetPoseDeltaRotation(SelectedBones[SelectionIndex], &RetargetPose);
+		FQuat PrevDeltaRotation = RetargetPose.GetDeltaRotationForBone(SelectedBones[SelectionIndex]);
 		BoneEdit.PreviousDeltaRotation.Add(PrevDeltaRotation);
 	}
 	
-	const int32 ParentIndex = Skeleton.GetParentIndex(BoneEdit.Index);
-	if (ParentIndex != INDEX_NONE)
-	{
-		BoneEdit.ParentGlobalTransform = Controller->GetGlobalRetargetPoseOfBone(Skeleton, ParentIndex, Scale, Offset);
-	}else
-	{
-		BoneEdit.ParentGlobalTransform = FTransform::Identity;
-	}
+	const int32 ParentIndex = RefSkeleton.GetParentIndex(BoneEdit.Index);
+	BoneEdit.ParentGlobalTransform = Controller->GetGlobalRetargetPoseOfBone(SourceOrTarget, ParentIndex, Scale, Offset);
 }
 
 bool FIKRetargetEditPoseMode::IsRootSelected() const
@@ -574,7 +525,7 @@ bool FIKRetargetEditPoseMode::IsRootSelected() const
 	}
 	
 	const TArray<FName>& SelectedBones = Controller->GetSelectedBones();
-	const FName& RootName = Controller->AssetController->GetTargetRootBone();
+	const FName& RootName = Controller->AssetController->GetRetargetRootBone(SourceOrTarget);
 	return SelectedBones.Contains(RootName);
 }
 
@@ -592,7 +543,7 @@ bool FIKRetargetEditPoseMode::IsOnlyRootSelected() const
 		return false;
 	}
 
-	return Controller->AssetController->GetTargetRootBone() == SelectedBones[0];
+	return Controller->AssetController->GetRetargetRootBone(SourceOrTarget) == SelectedBones[0];
 }
 
 void FIKRetargetEditPoseMode::Tick(FEditorViewportClient* ViewportClient, float DeltaTime)
diff --git a/Engine/Plugins/Animation/IKRig/Source/IKRigEditor/Private/RetargetEditor/IKRetargetEditor.cpp b/Engine/Plugins/Animation/IKRig/Source/IKRigEditor/Private/RetargetEditor/IKRetargetEditor.cpp
--- a/Engine/Plugins/Animation/IKRig/Source/IKRigEditor/Private/RetargetEditor/IKRetargetEditor.cpp
+++ b/Engine/Plugins/Animation/IKRig/Source/IKRigEditor/Private/RetargetEditor/IKRetargetEditor.cpp
@@ -38,10 +38,6 @@ FIKRetargetEditor::FIKRetargetEditor()
 {
 }
 
-FIKRetargetEditor::~FIKRetargetEditor()
-{
-}
-
 void FIKRetargetEditor::InitAssetEditor(
 	const EToolkitMode::Type Mode,
     const TSharedPtr<IToolkitHost>& InitToolkitHost, 
@@ -111,13 +107,6 @@ void FIKRetargetEditor::UnregisterTabSpawners(const TSharedRef<FTabManager>& InT
 void FIKRetargetEditor::BindCommands()
 {
 	const FIKRetargetCommands& Commands = FIKRetargetCommands::Get();
-
-	ToolkitCommands->MapAction(
-		Commands.GoToRetargetPose,
-		FExecuteAction::CreateSP(EditorController, &FIKRetargetEditorController::HandleGoToRetargetPose),
-		FCanExecuteAction(),
-		FCanExecuteAction(),
-		EUIActionRepeatMode::RepeatDisabled);
 	
 	ToolkitCommands->MapAction(
         Commands.EditRetargetPose,
@@ -129,7 +118,7 @@ void FIKRetargetEditor::BindCommands()
 	ToolkitCommands->MapAction(
 		Commands.ResetAllBones,
 		FExecuteAction::CreateSP(EditorController, &FIKRetargetEditorController::HandleResetAllBones),
-		FCanExecuteAction::CreateSP(EditorController, &FIKRetargetEditorController::CanResetPose),
+		FCanExecuteAction(),
 		EUIActionRepeatMode::RepeatDisabled);
 
 	ToolkitCommands->MapAction(
@@ -147,14 +136,14 @@ void FIKRetargetEditor::BindCommands()
 	ToolkitCommands->MapAction(
 		Commands.NewRetargetPose,
 		FExecuteAction::CreateSP(EditorController, &FIKRetargetEditorController::HandleNewPose),
-		FCanExecuteAction::CreateSP(EditorController, &FIKRetargetEditorController::CanCreatePose),
+		FCanExecuteAction(),
 		FCanExecuteAction(),
 		EUIActionRepeatMode::RepeatDisabled);
 
 	ToolkitCommands->MapAction(
 		Commands.DuplicateRetargetPose,
 		FExecuteAction::CreateSP(EditorController, &FIKRetargetEditorController::HandleDuplicatePose),
-		FCanExecuteAction::CreateSP(EditorController, &FIKRetargetEditorController::CanCreatePose),
+		FCanExecuteAction(),
 		FCanExecuteAction(),
 		EUIActionRepeatMode::RepeatDisabled);
 
@@ -172,24 +161,26 @@ void FIKRetargetEditor::BindCommands()
 		FCanExecuteAction(),
 		EUIActionRepeatMode::RepeatDisabled);
 
+	const TSharedRef<FIKRetargetPoseExporter> PoseExporterRef = EditorController->PoseExporter.ToSharedRef();
+	
 	ToolkitCommands->MapAction(
 		Commands.ImportRetargetPose,
-		FExecuteAction::CreateSP(EditorController, &FIKRetargetEditorController::HandleImportPose),
-		FCanExecuteAction::CreateSP(EditorController, &FIKRetargetEditorController::CanCreatePose),
+		FExecuteAction::CreateSP(PoseExporterRef, &FIKRetargetPoseExporter::HandleImportFromPoseAsset),
+		FCanExecuteAction(),
 		FCanExecuteAction(),
 		EUIActionRepeatMode::RepeatDisabled);
 
 	ToolkitCommands->MapAction(
 		Commands.ImportRetargetPoseFromAnim,
-		FExecuteAction::CreateSP(EditorController, &FIKRetargetEditorController::HandleImportPoseFromSequence),
-		FCanExecuteAction::CreateSP(EditorController, &FIKRetargetEditorController::CanCreatePose),
+		FExecuteAction::CreateSP(PoseExporterRef, &FIKRetargetPoseExporter::HandleImportFromSequenceAsset),
+		FCanExecuteAction(),
 		FCanExecuteAction(),
 		EUIActionRepeatMode::RepeatDisabled);
 
 	ToolkitCommands->MapAction(
 		Commands.ExportRetargetPose,
-		FExecuteAction::CreateSP(EditorController, &FIKRetargetEditorController::HandleExportPose),
-		FCanExecuteAction::CreateSP(EditorController, &FIKRetargetEditorController::CanCreatePose),
+		FExecuteAction::CreateSP(PoseExporterRef, &FIKRetargetPoseExporter::HandleExportPoseAsset),
+		FCanExecuteAction(),
 		FCanExecuteAction(),
 		EUIActionRepeatMode::RepeatDisabled);
 }
@@ -210,15 +201,11 @@ void FIKRetargetEditor::ExtendToolbar()
 
 void FIKRetargetEditor::FillToolbar(FToolBarBuilder& ToolbarBuilder)
 {
-	ToolbarBuilder.BeginSection("Go To Retarget Pose");
+	ToolbarBuilder.BeginSection("Show Retarget Pose");
 	{
-		ToolbarBuilder.AddToolBarButton(
-			FIKRetargetCommands::Get().GoToRetargetPose,
-			NAME_None,
-			TAttribute<FText>(),
-			TAttribute<FText>(),
-			FSlateIcon(FAppStyle::GetAppStyleSetName(),"GenericStop"));
+		// TODO debug section
 	}
+	
 	ToolbarBuilder.EndSection();
 }
 
@@ -310,6 +297,33 @@ void FIKRetargetEditor::HandleViewportCreated(const TSharedRef<class IPersonaVie
 			return 1.0f;
 		});
 	}
+
+	auto GetBorderColorAndOpacity = [this]()
+	{
+		// no processor or processor not initialized
+		const UIKRetargetProcessor* Processor = EditorController.Get().GetRetargetProcessor();
+		if (!Processor || !Processor->IsInitialized() )
+		{
+			return FLinearColor::Red;
+		}
+
+		const ERetargeterOutputMode OutputMode = EditorController.Get().GetRetargeterMode();
+		if (OutputMode == ERetargeterOutputMode::RunRetarget)
+		{
+			return FLinearColor::Transparent;
+		}
+
+		return FLinearColor::Blue;
+	};
+	
+	InViewport->AddOverlayWidget(
+		SNew(SBorder)
+		.BorderImage(FIKRetargetEditorStyle::Get().GetBrush( "IKRetarget.Viewport.Border"))
+		.BorderBackgroundColor_Lambda(GetBorderColorAndOpacity)
+		.Visibility(EVisibility::HitTestInvisible)
+		.Padding(0.0f)
+		.ShowEffectWhenDisabled(false)
+	);
 }
 
 void FIKRetargetEditor::HandlePreviewSceneCreated(const TSharedRef<IPersonaPreviewScene>& InPersonaPreviewScene)
@@ -323,14 +337,14 @@ void FIKRetargetEditor::HandlePreviewSceneCreated(const TSharedRef<IPersonaPrevi
 	EditorController->TargetSkelMeshComponent = NewObject<UDebugSkelMeshComponent>(Actor);
 
 	// setup an apply an anim instance to the skeletal mesh component
-	EditorController->SourceAnimInstance = NewObject<UAnimPreviewInstance>(EditorController->SourceSkelMeshComponent, TEXT("IKRetargetSourceAnimScriptInstance"));
+	EditorController->SourceAnimInstance = NewObject<UIKRetargetAnimInstance>(EditorController->SourceSkelMeshComponent, TEXT("IKRetargetSourceAnimScriptInstance"));
 	EditorController->TargetAnimInstance = NewObject<UIKRetargetAnimInstance>(EditorController->TargetSkelMeshComponent, TEXT("IKRetargetTargetAnimScriptInstance"));
 	SetupAnimInstance();
 	
 	// set the source and target skeletal meshes on the component
 	// NOTE: this must be done AFTER setting the AnimInstance so that the correct root anim node is loaded
-	USkeletalMesh* SourceMesh = EditorController->GetSourceSkeletalMesh();
-	USkeletalMesh* TargetMesh = EditorController->GetTargetSkeletalMesh();
+	USkeletalMesh* SourceMesh = EditorController->GetSkeletalMesh(ERetargetSourceOrTarget::Source);
+	USkeletalMesh* TargetMesh = EditorController->GetSkeletalMesh(ERetargetSourceOrTarget::Target);
 	EditorController->SourceSkelMeshComponent->SetSkeletalMesh(SourceMesh);
 	EditorController->TargetSkelMeshComponent->SetSkeletalMesh(TargetMesh);
 
@@ -349,8 +363,12 @@ void FIKRetargetEditor::HandlePreviewSceneCreated(const TSharedRef<IPersonaPrevi
 
 void FIKRetargetEditor::SetupAnimInstance()
 {
-	// connect the retarget asset and the source component ot the target anim instance
-	EditorController->TargetAnimInstance->SetRetargetAssetAndSourceComponent(EditorController->AssetController->GetAsset(), EditorController->SourceSkelMeshComponent);
+	UIKRetargeter* Asset = EditorController->AssetController->GetAsset();
+	
+	// configure SOURCE anim instance (will only output retarget pose)
+	EditorController->SourceAnimInstance->ConfigureAnimInstance(ERetargetSourceOrTarget::Source, Asset, nullptr);
+	// configure TARGET anim instance (will output retarget pose AND retarget pose from source skel mesh component)
+	EditorController->TargetAnimInstance->ConfigureAnimInstance(ERetargetSourceOrTarget::Target, Asset, EditorController->SourceSkelMeshComponent);
 
 	EditorController->SourceSkelMeshComponent->PreviewInstance = EditorController->SourceAnimInstance.Get();
 	EditorController->TargetSkelMeshComponent->PreviewInstance = EditorController->TargetAnimInstance.Get();
@@ -361,9 +379,9 @@ void FIKRetargetEditor::SetupAnimInstance()
 
 void FIKRetargetEditor::HandleDetailsCreated(const TSharedRef<class IDetailsView>& InDetailsView)
 {
-	EditorController->DetailsView = InDetailsView;
-	EditorController->DetailsView->OnFinishedChangingProperties().AddSP(this, &FIKRetargetEditor::OnFinishedChangingDetails);
-	EditorController->DetailsView->SetObject(EditorController->AssetController->GetAsset());
+	InDetailsView->OnFinishedChangingProperties().AddSP(this, &FIKRetargetEditor::OnFinishedChangingDetails);
+	InDetailsView->SetObject(EditorController->AssetController->GetAsset());
+	EditorController->SetDetailsView(InDetailsView);
 }
 
 void FIKRetargetEditor::OnFinishedChangingDetails(const FPropertyChangedEvent& PropertyChangedEvent)
@@ -379,13 +397,13 @@ void FIKRetargetEditor::OnFinishedChangingDetails(const FPropertyChangedEvent& P
 	// if no override target mesh has been specified, update the override to reflect the mesh in the ik rig asset
 	if (bTargetIKRigChanged)
 	{
-		AssetController->OnTargetIKRigChanged();
+		AssetController->OnIKRigChanged(ERetargetSourceOrTarget::Target);
 	}
 
 	// if no override source mesh has been specified, update the override to reflect the mesh in the ik rig asset
 	if (bSourceIKRigChanged)
 	{
-		AssetController->OnSourceIKRigChanged();
+		AssetController->OnIKRigChanged(ERetargetSourceOrTarget::Source);
 	}
 
 	// if either IK Rig asset has been modified, rebind and refresh UI
@@ -405,8 +423,8 @@ void FIKRetargetEditor::OnFinishedChangingDetails(const FPropertyChangedEvent& P
 		
 		// set the source and target skeletal meshes on the component
 		// NOTE: this must be done AFTER setting the AnimInstance so that the correct root anim node is loaded
-		USkeletalMesh* SourceMesh = EditorController->GetSourceSkeletalMesh();
-		USkeletalMesh* TargetMesh = EditorController->GetTargetSkeletalMesh();
+		USkeletalMesh* SourceMesh = EditorController->GetSkeletalMesh(ERetargetSourceOrTarget::Source);
+		USkeletalMesh* TargetMesh = EditorController->GetSkeletalMesh(ERetargetSourceOrTarget::Target);
 		EditorController->SourceSkelMeshComponent->SetSkeletalMesh(SourceMesh);
 		EditorController->TargetSkelMeshComponent->SetSkeletalMesh(TargetMesh);
 	
@@ -421,7 +439,7 @@ void FIKRetargetEditor::OnFinishedChangingDetails(const FPropertyChangedEvent& P
 	
 		SetupAnimInstance();
 
-		EditorController->RefreshAllViews();
+		AssetController->BroadcastNeedsReinitialized();
 	}
 }
 
diff --git a/Engine/Plugins/Animation/IKRig/Source/IKRigEditor/Private/RetargetEditor/IKRetargetEditorController.cpp b/Engine/Plugins/Animation/IKRig/Source/IKRigEditor/Private/RetargetEditor/IKRetargetEditorController.cpp
--- a/Engine/Plugins/Animation/IKRig/Source/IKRigEditor/Private/RetargetEditor/IKRetargetEditorController.cpp
+++ b/Engine/Plugins/Animation/IKRig/Source/IKRigEditor/Private/RetargetEditor/IKRetargetEditorController.cpp
@@ -12,11 +12,11 @@
 #include "Animation/AnimMontage.h"
 #include "Animation/AnimSequence.h"
 #include "Animation/PoseAsset.h"
+#include "Factories/PoseAssetFactory.h"
 #include "RetargetEditor/IKRetargetAnimInstance.h"
 #include "RetargetEditor/IKRetargetDefaultMode.h"
 #include "RetargetEditor/IKRetargetEditPoseMode.h"
 #include "RetargetEditor/IKRetargetEditor.h"
-#include "RetargetEditor/IKRetargetFactory.h"
 #include "RetargetEditor/SIKRetargetChainMapList.h"
 #include "RetargetEditor/SIKRetargetHierarchy.h"
 #include "Retargeter/IKRetargeter.h"
@@ -26,13 +26,21 @@
 
 #define LOCTEXT_NAMESPACE "IKRetargetEditorController"
 
+
 void FIKRetargetEditorController::Initialize(TSharedPtr<FIKRetargetEditor> InEditor, UIKRetargeter* InAsset)
 {
 	Editor = InEditor;
 	AssetController = UIKRetargeterController::GetController(InAsset);
-	AssetController->SetEditorController(this);
-	SkeletonMode = EIKRetargetSkeletonMode::Target;
-	InAsset->SetOutputMode(ERetargeterOutputMode::ShowRetargetPose);
+	CurrentlyEditingSourceOrTarget = ERetargetSourceOrTarget::Target;
+	OutputMode = ERetargeterOutputMode::ShowRetargetPose;
+	PreviousMode = OutputMode;
+	PoseExporter = MakeShared<FIKRetargetPoseExporter>();
+	PoseExporter->Initialize(SharedThis(this));
+
+	// clean the asset before editing
+	const bool bForceReinitialization = false;
+	AssetController->CleanChainMapping(bForceReinitialization);
+	AssetController->CleanPoseLists(bForceReinitialization);
 
 	// bind callbacks when SOURCE or TARGET IK Rigs are modified
 	BindToIKRigAsset(AssetController->GetAsset()->GetSourceIKRigWriteable());
@@ -98,13 +106,29 @@ void FIKRetargetEditorController::OnRetargeterNeedsInitialized(UIKRetargeter* Re
 	if (UIKRetargetProcessor* Processor = GetRetargetProcessor())
 	{
 		constexpr bool bSuppressWarnings = false;
-		Processor->Initialize(GetSourceSkeletalMesh(), GetTargetSkeletalMesh(), Retargeter, bSuppressWarnings);
+		Processor->Initialize(
+			GetSkeletalMesh(ERetargetSourceOrTarget::Source),
+			GetSkeletalMesh(ERetargetSourceOrTarget::Target),
+			Retargeter,
+			bSuppressWarnings);
 	}
 	
 	// refresh all the UI views
 	RefreshAllViews();
 }
 
+UDebugSkelMeshComponent* FIKRetargetEditorController::GetSkeletalMeshComponent(
+	const ERetargetSourceOrTarget& SourceOrTarget) const
+{
+	return SourceOrTarget == ERetargetSourceOrTarget::Source ? SourceSkelMeshComponent : TargetSkelMeshComponent;
+}
+
+UIKRetargetAnimInstance* FIKRetargetEditorController::GetAnimInstance(
+	const ERetargetSourceOrTarget& SourceOrTarget) const
+{
+	return SourceOrTarget == ERetargetSourceOrTarget::Source ? SourceAnimInstance.Get() : TargetAnimInstance.Get();
+}
+
 void FIKRetargetEditorController::AddOffsetToMeshComponent(const FVector& Offset, USceneComponent* MeshComponent) const
 {
 	UIKRetargeter* Asset = AssetController->GetAsset();
@@ -130,7 +154,7 @@ void FIKRetargetEditorController::AddOffsetToMeshComponent(const FVector& Offset
 	MeshComponent->SetWorldScale3D(FVector(Scale,Scale,Scale));
 }
 
-bool FIKRetargetEditorController::IsBoneRetargeted(const FName& BoneName, EIKRetargetSkeletonMode WhichSkeleton) const
+bool FIKRetargetEditorController::IsBoneRetargeted(const FName& BoneName, ERetargetSourceOrTarget SourceOrTarget) const
 {
 	// get an initialized processor
 	const UIKRetargetProcessor* Processor = GetRetargetProcessor();
@@ -140,7 +164,7 @@ bool FIKRetargetEditorController::IsBoneRetargeted(const FName& BoneName, EIKRet
 	}
 
 	// get the bone index
-	const FRetargetSkeleton& Skeleton = WhichSkeleton == EIKRetargetSkeletonMode::Source ? Processor->GetSourceSkeleton() : Processor->GetTargetSkeleton();
+	const FRetargetSkeleton& Skeleton = SourceOrTarget == ERetargetSourceOrTarget::Source ? Processor->GetSourceSkeleton() : Processor->GetTargetSkeleton();
 	const int32 BoneIndex = Skeleton.FindBoneIndexByName(BoneName);
 	if (BoneIndex == INDEX_NONE)
 	{
@@ -148,10 +172,10 @@ bool FIKRetargetEditorController::IsBoneRetargeted(const FName& BoneName, EIKRet
 	}
 
 	// return if it's a retargeted bone
-	return Processor->IsBoneRetargeted(BoneIndex, (int8)WhichSkeleton);
+	return Processor->IsBoneRetargeted(BoneIndex, (int8)SourceOrTarget);
 }
 
-FName FIKRetargetEditorController::GetChainNameFromBone(const FName& BoneName, EIKRetargetSkeletonMode WhichSkeleton) const
+FName FIKRetargetEditorController::GetChainNameFromBone(const FName& BoneName, ERetargetSourceOrTarget SourceOrTarget) const
 {
 	// get an initialized processor
 	const UIKRetargetProcessor* Processor = GetRetargetProcessor();
@@ -161,14 +185,14 @@ FName FIKRetargetEditorController::GetChainNameFromBone(const FName& BoneName, E
 	}
 
 	// get the bone index
-	const FRetargetSkeleton& Skeleton = WhichSkeleton == EIKRetargetSkeletonMode::Source ? Processor->GetSourceSkeleton() : Processor->GetTargetSkeleton();
+	const FRetargetSkeleton& Skeleton = SourceOrTarget == ERetargetSourceOrTarget::Source ? Processor->GetSourceSkeleton() : Processor->GetTargetSkeleton();
 	const int32 BoneIndex = Skeleton.FindBoneIndexByName(BoneName);
 	if (BoneIndex == INDEX_NONE)
 	{
 		return NAME_None;
 	}
 
-	return Processor->GetChainNameForBone(BoneIndex, (int8)WhichSkeleton);
+	return Processor->GetChainNameForBone(BoneIndex, (int8)SourceOrTarget);
 }
 
 TObjectPtr<UIKRetargetBoneDetails> FIKRetargetEditorController::GetDetailsObjectForBone(const FName& BoneName)
@@ -194,29 +218,14 @@ TObjectPtr<UIKRetargetBoneDetails> FIKRetargetEditorController::CreateBoneDetail
 	return NewBoneDetails;
 }
 
-USkeletalMesh* FIKRetargetEditorController::GetSourceSkeletalMesh() const
-{
-	return AssetController ? AssetController->GetSourcePreviewMesh() : nullptr;
-}
-
-USkeletalMesh* FIKRetargetEditorController::GetTargetSkeletalMesh() const
+USkeletalMesh* FIKRetargetEditorController::GetSkeletalMesh(const ERetargetSourceOrTarget SourceOrTarget) const
 {
-	return AssetController ? AssetController->GetTargetPreviewMesh() : nullptr;
+	return AssetController ? AssetController->GetPreviewMesh(SourceOrTarget) : nullptr;
 }
 
-const USkeleton* FIKRetargetEditorController::GetSourceSkeleton() const
+const USkeleton* FIKRetargetEditorController::GetSkeleton(const ERetargetSourceOrTarget SourceOrTarget) const
 {
-	if (const USkeletalMesh* Mesh = GetSourceSkeletalMesh())
-	{
-		return Mesh->GetSkeleton();
-	}
-	
-	return nullptr;
-}
-
-const USkeleton* FIKRetargetEditorController::GetTargetSkeleton() const
-{
-	if (const USkeletalMesh* Mesh = GetTargetSkeletalMesh())
+	if (const USkeletalMesh* Mesh = GetSkeletalMesh(SourceOrTarget))
 	{
 		return Mesh->GetSkeleton();
 	}
@@ -226,24 +235,30 @@ const USkeleton* FIKRetargetEditorController::GetTargetSkeleton() const
 
 UDebugSkelMeshComponent* FIKRetargetEditorController::GetEditedSkeletalMesh() const
 {
-	return SkeletonMode == EIKRetargetSkeletonMode::Source ? SourceSkelMeshComponent : TargetSkelMeshComponent;
+	return CurrentlyEditingSourceOrTarget == ERetargetSourceOrTarget::Source ? SourceSkelMeshComponent : TargetSkelMeshComponent;
 }
 
 const FRetargetSkeleton& FIKRetargetEditorController::GetCurrentlyEditedSkeleton(const UIKRetargetProcessor& Processor) const
 {
-	return SkeletonMode == EIKRetargetSkeletonMode::Source ? Processor.GetSourceSkeleton() : Processor.GetTargetSkeleton();
+	return CurrentlyEditingSourceOrTarget == ERetargetSourceOrTarget::Source ? Processor.GetSourceSkeleton() : Processor.GetTargetSkeleton();
 }
 
 FTransform FIKRetargetEditorController::GetGlobalRetargetPoseOfBone(
-	const FRetargetSkeleton& Skeleton,
+	const ERetargetSourceOrTarget SourceOrTarget,
 	const int32& BoneIndex,
 	const float& Scale,
 	const FVector& Offset) const
 {
-	check(Skeleton.BoneNames.IsValidIndex(BoneIndex))
+	if (BoneIndex == INDEX_NONE)
+	{
+		return FTransform::Identity;
+	}
+	
+	const UIKRetargetAnimInstance* AnimInstance = GetAnimInstance(SourceOrTarget);
+	const TArray<FTransform>& GlobalRetargetPose = AnimInstance->GetGlobalRetargetPose();
 	
 	// get transform of bone
-	FTransform BoneTransform = Skeleton.RetargetGlobalPose[BoneIndex];
+	FTransform BoneTransform = GlobalRetargetPose[BoneIndex];
 
 	// scale and offset
 	BoneTransform.ScaleTranslation(Scale);
@@ -315,38 +330,82 @@ void FIKRetargetEditorController::ClearOutputLog() const
 	if (OutputLogView.IsValid())
 	{
 		OutputLogView.Get()->ClearLog();
-		GetRetargetProcessor()->Log.Clear();
+		if (const UIKRetargetProcessor* Processor = GetRetargetProcessor())
+		{
+			Processor->Log.Clear();
+		}
 	}
 }
 
 void FIKRetargetEditorController::RefreshAllViews() const
 {
 	Editor.Pin()->RegenerateMenusAndToolbars();
+	RefreshDetailsView();
+	RefreshChainsView();
+	RefreshAssetBrowserView();
+	RefreshHierarchyView();
+}
 
-	// refresh the details panel
+void FIKRetargetEditorController::RefreshDetailsView() const
+{
+	// refresh the details panel, cannot assume tab is not closed
 	if (DetailsView.IsValid())
 	{
 		DetailsView->ForceRefresh();
 	}
+}
 
-	// cannot assume chains view is always available
+void FIKRetargetEditorController::RefreshChainsView() const
+{
+	// refesh chains view, cannot assume tab is not closed
 	if (ChainsView.IsValid())
 	{
 		ChainsView.Get()->RefreshView();
 	}
+}
 
+void FIKRetargetEditorController::RefreshAssetBrowserView() const
+{
 	// refresh the asset browser to ensure it shows compatible sequences
 	if (AssetBrowserView.IsValid())
 	{
 		AssetBrowserView.Get()->RefreshView();
 	}
+}
 
+void FIKRetargetEditorController::RefreshHierarchyView() const
+{
 	if (HierarchyView.IsValid())
 	{
+		HierarchyView.Get()->RefreshPoseList();
 		HierarchyView.Get()->RefreshTreeView();
 	}
 }
 
+void FIKRetargetEditorController::RefreshPoseList() const
+{
+	if (HierarchyView.IsValid())
+	{
+		HierarchyView.Get()->RefreshPoseList();
+	}
+}
+
+void FIKRetargetEditorController::SetDetailsObject(UObject* DetailsObject) const
+{
+	if (DetailsView.IsValid())
+	{
+		DetailsView->SetObject(DetailsObject);
+	}
+}
+
+void FIKRetargetEditorController::SetDetailsObjects(const TArray<UObject*>& DetailsObjects) const
+{
+	if (DetailsView.IsValid())
+	{
+		DetailsView->SetObjects(DetailsObjects);
+	}
+}
+
 void FIKRetargetEditorController::PlayAnimationAsset(UAnimationAsset* AssetToPlay)
 {
 	if (AssetToPlay && SourceAnimInstance.IsValid())
@@ -359,10 +418,46 @@ void FIKRetargetEditorController::PlayAnimationAsset(UAnimationAsset* AssetToPla
 	}
 }
 
-void FIKRetargetEditorController::SetSkeletonMode(EIKRetargetSkeletonMode NewMode)
+void FIKRetargetEditorController::PausePlayback()
+{
+	if (UAnimationAsset* CurrentAnim = SourceAnimInstance->GetAnimationAsset())
+	{
+		AnimThatWasPlaying = CurrentAnim;
+		TimeWhenPaused = SourceAnimInstance->GetCurrentTime();
+	}
+	
+	SourceAnimInstance->SetPlaying(false);
+	SourceAnimInstance->SetAnimationAsset(nullptr);
+}
+
+void FIKRetargetEditorController::ResumePlayback()
+{
+	SourceAnimInstance->SetAnimationAsset(AnimThatWasPlaying);
+	SourceAnimInstance->SetPlaying(true);
+	SourceAnimInstance->SetPosition(TimeWhenPaused);
+}
+
+float FIKRetargetEditorController::GetRetargetPoseAmount() const
+{
+	return RetargetPosePreviewBlend;
+}
+
+void FIKRetargetEditorController::SetRetargetPoseAmount(float InValue)
+{
+	if (OutputMode==ERetargeterOutputMode::RunRetarget)
+	{
+		SetRetargeterMode(ERetargeterOutputMode::ShowRetargetPose);
+	}
+	
+	RetargetPosePreviewBlend = InValue;
+	SourceAnimInstance->SetRetargetPoseBlend(RetargetPosePreviewBlend);
+	TargetAnimInstance->SetRetargetPoseBlend(RetargetPosePreviewBlend);
+}
+
+void FIKRetargetEditorController::SetSourceOrTargetMode(ERetargetSourceOrTarget NewMode)
 {
 	// already in this mode, so do nothing
-	if (NewMode == SkeletonMode)
+	if (NewMode == CurrentlyEditingSourceOrTarget)
 	{
 		return;
 	}
@@ -371,25 +466,27 @@ void FIKRetargetEditorController::SetSkeletonMode(EIKRetargetSkeletonMode NewMod
 	ClearSelection();
 	
 	// store the new skeleton mode
-	SkeletonMode = NewMode;
+	CurrentlyEditingSourceOrTarget = NewMode;
 
-	const ERetargeterOutputMode Mode = AssetController->GetAsset()->GetOutputMode();
-	switch (Mode)
+	switch (GetRetargeterMode())
 	{
 		case ERetargeterOutputMode::EditRetargetPose:
 		{
-			// TODO restart pose editing on other skeleton once we support source pose editing
-			//FEditorModeTools& EditorModeManager = Editor.Pin()->GetEditorModeManager();
-			//FIKRetargetEditPoseMode* EditMode = EditorModeManager.GetActiveModeTyped<FIKRetargetEditPoseMode>(FIKRetargetEditPoseMode::ModeName);
-			//EditMode->Enter();
-			SetRetargeterMode(ERetargeterOutputMode::ShowRetargetPose);
+			FEditorModeTools& EditorModeManager = Editor.Pin()->GetEditorModeManager();
+			FIKRetargetEditPoseMode* EditMode = EditorModeManager.GetActiveModeTyped<FIKRetargetEditPoseMode>(FIKRetargetEditPoseMode::ModeName);
+			if (EditMode)
+			{
+				// FIKRetargetEditPoseMode::Enter() is reentrant and written so we can switch between editing
+				// source / target skeleton without having to enter/exit the mode; just call Enter() again
+				EditMode->Enter();
+			}
 			break;
 		}
 		case ERetargeterOutputMode::ShowRetargetPose:
 		case ERetargeterOutputMode::RunRetarget:
 		{
 			// toggle visibility of currently active skeleton
-			const bool bEditingSource = NewMode == EIKRetargetSkeletonMode::Source;
+			const bool bEditingSource = NewMode == ERetargetSourceOrTarget::Source;
 			SourceSkelMeshComponent->SkeletonDrawMode = bEditingSource ? ESkeletonDrawMode::Default : ESkeletonDrawMode::GreyedOut;
 			TargetSkelMeshComponent->SkeletonDrawMode = bEditingSource ? ESkeletonDrawMode::GreyedOut : ESkeletonDrawMode::Default;
 			break;
@@ -478,15 +575,15 @@ void FIKRetargetEditorController::EditBoneSelection(
 	DebugComponent->BonesOfInterest = SelectedBoneIndices;
 
 	// update hierarchy view
-	if (!bFromHierarchyView && HierarchyView.IsValid())
+	if (!bFromHierarchyView)
 	{
-		HierarchyView->RefreshTreeView();
+		RefreshHierarchyView();
 	}
 
 	// update details
 	if (SelectedBones.IsEmpty())
 	{
-		DetailsView->SetObject(AssetController->GetAsset());
+		SetDetailsObject(AssetController->GetAsset());
 	}
 	else
 	{
@@ -496,7 +593,7 @@ void FIKRetargetEditorController::EditBoneSelection(
 			TObjectPtr<UIKRetargetBoneDetails> BoneDetails = GetDetailsObjectForBone(SelectedBone);
 			SelectedBoneDetails.Add(BoneDetails);
 		}
-		DetailsView->SetObjects(SelectedBoneDetails);
+		SetDetailsObjects(SelectedBoneDetails);
 	}
 }
 
@@ -506,7 +603,10 @@ void FIKRetargetEditorController::ClearSelection(const bool bKeepBoneSelection)
 	SetSelectedMesh(nullptr);
 	
 	// deselect all chains
-	ChainsView->ClearSelection();
+	if (ChainsView.IsValid())
+	{
+		ChainsView->ClearSelection();
+	}
 
 	// clear bone selection
 	if (!bKeepBoneSelection)
@@ -516,63 +616,92 @@ void FIKRetargetEditorController::ClearSelection(const bool bKeepBoneSelection)
 		EditBoneSelection(Empty, EBoneSelectionEdit::Replace, bFromHierarchy);
 
 		// show global details
-		DetailsView->SetObject(AssetController->GetAsset());
+		SetDetailsObject(AssetController->GetAsset());
 	}
 
-	DetailsView->ForceRefresh();
+	RefreshDetailsView();
 }
 
 void FIKRetargetEditorController::SetRetargeterMode(ERetargeterOutputMode Mode)
 {
-	const bool bWasEditingPose = IsEditingPose();
+	if (OutputMode == Mode)
+	{
+		return;
+	}
+		
+	PreviousMode = OutputMode;
+	
 	FEditorModeTools& EditorModeManager = Editor.Pin()->GetEditorModeManager();
-	UIKRetargeter* Asset = AssetController->GetAsset();
 	
 	switch (Mode)
 	{
 		case ERetargeterOutputMode::EditRetargetPose:
-			bWasPlayingAnim = SourceAnimInstance->IsPlaying();
+			// enter edit mode
 			EditorModeManager.DeactivateMode(FIKRetargetDefaultMode::ModeName);
 			EditorModeManager.ActivateMode(FIKRetargetEditPoseMode::ModeName);
-			Asset->SetOutputMode(ERetargeterOutputMode::EditRetargetPose);
+			OutputMode = ERetargeterOutputMode::EditRetargetPose;
+			SourceAnimInstance->SetRetargetMode(ERetargeterOutputMode::EditRetargetPose);
+			TargetAnimInstance->SetRetargetMode(ERetargeterOutputMode::EditRetargetPose);
+			PausePlayback();
+			SetRetargetPoseAmount(1.0f);
 			break;
 		
 		case ERetargeterOutputMode::RunRetarget:
 			EditorModeManager.DeactivateMode(FIKRetargetEditPoseMode::ModeName);
 			EditorModeManager.ActivateMode(FIKRetargetDefaultMode::ModeName);
-			Asset->SetOutputMode(ERetargeterOutputMode::RunRetarget);
-			if (bWasEditingPose)
+			OutputMode = ERetargeterOutputMode::RunRetarget;
+			SourceAnimInstance->SetRetargetMode(ERetargeterOutputMode::RunRetarget);
+			TargetAnimInstance->SetRetargetMode(ERetargeterOutputMode::RunRetarget);
+			if (PreviousMode == ERetargeterOutputMode::EditRetargetPose)
 			{
 				// must reinitialize after editing the retarget pose
 				AssetController->BroadcastNeedsReinitialized();
-				if (bWasPlayingAnim)
-				{
-					// continue playing whatever animation asset was last used
-					SourceAnimInstance->SetAnimationAsset(AnimThatWasPlaying);
-				}
 			}
+			ResumePlayback();
 			break;
 
 		case ERetargeterOutputMode::ShowRetargetPose:
 			EditorModeManager.DeactivateMode(FIKRetargetEditPoseMode::ModeName);
 			EditorModeManager.ActivateMode(FIKRetargetDefaultMode::ModeName);
-			Asset->SetOutputMode(ERetargeterOutputMode::ShowRetargetPose);
-			// stop playback of animations
-			SourceAnimInstance->SetPlaying(false);
-			// put source back in ref pose
-			SourceSkelMeshComponent->ShowReferencePose(true);
-			// have to move component back to offset position because ShowReferencePose() sets it back to origin
-			AddOffsetToMeshComponent(FVector::ZeroVector, SourceSkelMeshComponent);
+			OutputMode = ERetargeterOutputMode::ShowRetargetPose;
+			// show retarget pose
+			SourceAnimInstance->SetRetargetMode(ERetargeterOutputMode::ShowRetargetPose);
+			TargetAnimInstance->SetRetargetMode(ERetargeterOutputMode::ShowRetargetPose);
+			PausePlayback();
+			SetRetargetPoseAmount(1.0f);
 			break;
 
 		default:
 			checkNoEntry();
 	}
+
+	// details view displays differently depending on output mode
+	RefreshDetailsView();
 }
 
-void FIKRetargetEditorController::HandleGoToRetargetPose()
+FReply FIKRetargetEditorController::HandleShowRetargetPose()
 {
-	SetRetargeterMode(ERetargeterOutputMode::ShowRetargetPose);
+	const ERetargeterOutputMode CurrentMode = GetRetargeterMode();
+	if (CurrentMode == ERetargeterOutputMode::ShowRetargetPose || CurrentMode == ERetargeterOutputMode::EditRetargetPose)
+	{
+		SetRetargeterMode(ERetargeterOutputMode::RunRetarget);
+	}
+	else
+	{
+		SetRetargeterMode(ERetargeterOutputMode::ShowRetargetPose);
+	}
+	
+	return FReply::Handled();
+}
+
+bool FIKRetargetEditorController::CanShowRetargetPose() const
+{
+	return GetRetargeterMode() != ERetargeterOutputMode::ShowRetargetPose;
+}
+
+bool FIKRetargetEditorController::IsShowingRetargetPose() const
+{
+	return GetRetargeterMode() == ERetargeterOutputMode::ShowRetargetPose;
 }
 
 void FIKRetargetEditorController::HandleEditPose()
@@ -580,7 +709,7 @@ void FIKRetargetEditorController::HandleEditPose()
 	if (IsEditingPose())
 	{
 		// stop pose editing
-		SetRetargeterMode(ERetargeterOutputMode::RunRetarget);
+		SetRetargeterMode(PreviousMode);
 	}
 	else
 	{
@@ -591,28 +720,21 @@ void FIKRetargetEditorController::HandleEditPose()
 
 bool FIKRetargetEditorController::CanEditPose() const
 {
-	const UIKRetargetProcessor* Processor = GetRetargetProcessor();
-	if (!Processor)
-	{
-		return false;
-	}
-
-	const bool bRetargeterIsReady = Processor->IsInitialized();
-	// TODO in the future, add support for editing source pose
-	const bool bEditingTarget = GetSkeletonMode() == EIKRetargetSkeletonMode::Target;
-	return bRetargeterIsReady && bEditingTarget;
+	return GetSkeletalMesh(GetSourceOrTarget()) != nullptr;
 }
 
 bool FIKRetargetEditorController::IsEditingPose() const
 {
-	return AssetController->GetAsset()->GetOutputMode() == ERetargeterOutputMode::EditRetargetPose;
+	return GetRetargeterMode() == ERetargeterOutputMode::EditRetargetPose;
 }
 
 void FIKRetargetEditorController::HandleNewPose()
 {
+	SetRetargeterMode(ERetargeterOutputMode::ShowRetargetPose);
+	
 	// get a unique pose name to use as suggestion
 	const FString DefaultNewPoseName = LOCTEXT("NewRetargetPoseName", "CustomRetargetPose").ToString();
-	const FName UniqueNewPoseName = AssetController->MakePoseNameUnique(DefaultNewPoseName);
+	const FName UniqueNewPoseName = AssetController->MakePoseNameUnique(DefaultNewPoseName, GetSourceOrTarget());
 	
 	SAssignNew(NewPoseWindow, SWindow)
 	.Title(LOCTEXT("NewRetargetPoseOptions", "Create New Retarget Pose"))
@@ -682,19 +804,21 @@ bool FIKRetargetEditorController::CanCreatePose() const
 FReply FIKRetargetEditorController::CreateNewPose() const
 {
 	const FName NewPoseName = FName(NewPoseEditableText.Get()->GetText().ToString());
-	AssetController->AddRetargetPose(NewPoseName);
+	AssetController->AddRetargetPose(NewPoseName, nullptr, GetSourceOrTarget());
 	NewPoseWindow->RequestDestroyWindow();
-	DetailsView->ForceRefresh();
+	RefreshPoseList();
 	return FReply::Handled();
 }
 
 void FIKRetargetEditorController::HandleDuplicatePose()
 {
+	SetRetargeterMode(ERetargeterOutputMode::ShowRetargetPose);
+	
 	// get a unique pose name to use as suggestion for duplicate
 	const FString DuplicateSuffix = LOCTEXT("DuplicateSuffix", "_Copy").ToString();
 	FString CurrentPoseName = GetCurrentPoseName().ToString();
 	const FString DefaultDuplicatePoseName = CurrentPoseName.Append(*DuplicateSuffix);
-	const FName UniqueNewPoseName = AssetController->MakePoseNameUnique(DefaultDuplicatePoseName);
+	const FName UniqueNewPoseName = AssetController->MakePoseNameUnique(DefaultDuplicatePoseName, GetSourceOrTarget());
 	
 	SAssignNew(NewPoseWindow, SWindow)
 	.Title(LOCTEXT("DuplicateRetargetPoseOptions", "Duplicate Retarget Pose"))
@@ -759,503 +883,85 @@ void FIKRetargetEditorController::HandleDuplicatePose()
 
 FReply FIKRetargetEditorController::CreateDuplicatePose() const
 {
-	const FIKRetargetPose& PoseToDuplicate = AssetController->GetCurrentRetargetPose();
+	const FIKRetargetPose& PoseToDuplicate = AssetController->GetCurrentRetargetPose(CurrentlyEditingSourceOrTarget);
 	const FName NewPoseName = FName(NewPoseEditableText.Get()->GetText().ToString());
-	AssetController->AddRetargetPose(NewPoseName, &PoseToDuplicate);
+	AssetController->AddRetargetPose(NewPoseName, &PoseToDuplicate, GetSourceOrTarget());
 	NewPoseWindow->RequestDestroyWindow();
-	DetailsView->ForceRefresh();
+	RefreshPoseList();
 	return FReply::Handled();
 }
 
-void FIKRetargetEditorController::HandleImportPose()
-{
-	RetargetPoseToImport = nullptr;
-
-	// load the content browser module to display an asset picker
-	const FContentBrowserModule& ContentBrowserModule = FModuleManager::LoadModuleChecked<FContentBrowserModule>("ContentBrowser");
-
-	// the asset picker will only show animation sequences
-	FAssetPickerConfig AssetPickerConfig;
-	AssetPickerConfig.Filter.ClassPaths.Add(URetargetPose::StaticClass()->GetClassPathName());
-	AssetPickerConfig.Filter.bRecursiveClasses = true;
-	AssetPickerConfig.OnAssetSelected = FOnAssetSelected::CreateSP(this, &FIKRetargetEditorController::OnRetargetPoseSelected);
-	AssetPickerConfig.InitialAssetViewType = EAssetViewType::Tile;
-
-	ImportPoseWindow = SNew(SWindow)
-	.Title(LOCTEXT("ImportRetargetPose", "Import Retarget Pose"))
-	.ClientSize(FVector2D(500, 600))
-	.SupportsMinimize(false) .SupportsMaximize(false)
-	[
-		SNew(SBorder)
-		.BorderImage( FAppStyle::GetBrush("Menu.Background") )
-		[
-			SNew(SVerticalBox)
-			
-			+ SVerticalBox::Slot()
-			.Padding(4)
-			[
-				ContentBrowserModule.Get().CreateAssetPicker(AssetPickerConfig)
-			]
-			
-			+ SVerticalBox::Slot()
-			.Padding(4)
-			.AutoHeight()
-			.HAlign(HAlign_Right)
-			[
-				SNew(SHorizontalBox)
-				+ SHorizontalBox::Slot()
-				.AutoWidth()
-				.Padding(4)
-				[
-					SNew(SButton)
-					.ButtonStyle(FAppStyle::Get(), "Button")
-					.TextStyle( FAppStyle::Get(), "DialogButtonText" )
-					.HAlign(HAlign_Center)
-					.VAlign(VAlign_Center)
-					.Text(LOCTEXT("ImportNewRetargetPoseButtonLabel", "Import New Retarget Pose"))
-					.OnClicked(this, &FIKRetargetEditorController::ImportRetargetPose)
-				]
-				
-				+ SHorizontalBox::Slot()
-				.AutoWidth()
-				.Padding(4)
-				[
-					SNew(SButton)
-					.ButtonStyle(FAppStyle::Get(), "Button")
-					.TextStyle( FAppStyle::Get(), "DialogButtonText" )
-					.HAlign(HAlign_Center)
-					.VAlign(VAlign_Center)
-					.Text(LOCTEXT("CancelButtonLabel", "Cancel"))
-					.OnClicked_Lambda( [this]()
-					{
-						ImportPoseWindow->RequestDestroyWindow();
-						return FReply::Handled();
-					})
-				]
-			]	
-		]
-	];
-
-	GEditor->EditorAddModalWindow(ImportPoseWindow.ToSharedRef());
-	ImportPoseWindow.Reset();
-}
-
-FReply FIKRetargetEditorController::ImportRetargetPose() const
-{
-	ImportPoseWindow->RequestDestroyWindow();
-	
-	if (RetargetPoseToImport.IsNull())
-	{
-		return FReply::Handled();
-	}
-
-	const TObjectPtr<URetargetPose> RetargetPose = Cast<URetargetPose>(RetargetPoseToImport.TryLoad());
-	if (!RetargetPose)
-	{
-		return FReply::Handled();
-	}
-
-	// create a new pose with the data from the selected retarget pose asset
-	FIKRetargetPose Pose;
-	RetargetPose->GetAsRetargetPose(Pose);
-	AssetController->AddRetargetPose(FName(RetargetPose->GetName()), &Pose);
-	
-	RefreshAllViews();
-	
-	return FReply::Unhandled();
-}
-
-void FIKRetargetEditorController::OnRetargetPoseSelected(const FAssetData& SelectedAsset)
-{
-	RetargetPoseToImport = SelectedAsset.ToSoftObjectPath();
-}
-
-void FIKRetargetEditorController::HandleImportPoseFromSequence()
-{
-	SequenceToImportAsPose = nullptr;
-
-	// get a unique pose name to use as suggestion
-	const FString DefaultImportedPoseName = LOCTEXT("ImportedRetargetPoseName", "ImportedRetargetPose").ToString();
-	const FName UniqueNewPoseName = AssetController->MakePoseNameUnique(DefaultImportedPoseName);
-	ImportedPoseName = FText::FromName(UniqueNewPoseName);
-
-	// load the content browser module to display an asset picker
-	const FContentBrowserModule& ContentBrowserModule = FModuleManager::LoadModuleChecked<FContentBrowserModule>("ContentBrowser");
-
-	// the asset picker will only show animation sequences compatible with the preview mesh
-	FAssetPickerConfig AssetPickerConfig;
-	AssetPickerConfig.Filter.ClassPaths.Add(UAnimSequence::StaticClass()->GetClassPathName());
-	AssetPickerConfig.Filter.ClassPaths.Add(UAnimMontage::StaticClass()->GetClassPathName());
-	AssetPickerConfig.Filter.ClassPaths.Add(UPoseAsset::StaticClass()->GetClassPathName());
-	AssetPickerConfig.InitialAssetViewType = EAssetViewType::Column;
-	AssetPickerConfig.bAddFilterUI = true;
-	AssetPickerConfig.bShowPathInColumnView = true;
-	AssetPickerConfig.bShowTypeInColumnView = true;
-	AssetPickerConfig.DefaultFilterMenuExpansion = EAssetTypeCategories::Animation;
-	AssetPickerConfig.OnShouldFilterAsset = FOnShouldFilterAsset::CreateSP(this, &FIKRetargetEditorController::OnShouldFilterSequenceToImport);
-	AssetPickerConfig.OnAssetSelected = FOnAssetSelected::CreateSP(this, &FIKRetargetEditorController::OnSequenceSelectedForPose);
-	AssetPickerConfig.bAllowNullSelection = false;
-
-	// hide all asset registry columns by default (we only really want the name and path)
-	TArray<UObject::FAssetRegistryTag> AssetRegistryTags;
-	UAnimSequence::StaticClass()->GetDefaultObject()->GetAssetRegistryTags(AssetRegistryTags);
-	FName ColumnToKeep = FName("Number of Frames");
-	for(UObject::FAssetRegistryTag& AssetRegistryTag : AssetRegistryTags)
-	{
-		if (AssetRegistryTag.Name != ColumnToKeep)
-		{
-			AssetPickerConfig.HiddenColumnNames.Add(AssetRegistryTag.Name.ToString());
-		}
-	}
-
-	// Also hide the type column by default (but allow users to enable it, so don't use bShowTypeInColumnView)
-	AssetPickerConfig.HiddenColumnNames.Add(TEXT("Class"));
-	AssetPickerConfig.HiddenColumnNames.Add(TEXT("HasVirtualizedData"));
-	AssetPickerConfig.HiddenColumnNames.Add(TEXT("DiskSize"));
-
-	// create pop-up window for user to select animation sequence asset to import as a retarget pose]
-	ImportPoseFromSequenceWindow = SNew(SWindow)
-	.Title(LOCTEXT("ImportRetargetPoseFromSequenceAsset", "Import Retarget Pose from Sequence Asset"))
-	.ClientSize(FVector2D(500, 600))
-	.SupportsMinimize(false) .SupportsMaximize(false)
-	[
-		SNew(SBorder)
-		.BorderImage( FAppStyle::GetBrush("Menu.Background") )
-		[
-			SNew(SVerticalBox)
-			
-			+ SVerticalBox::Slot()
-			.Padding(4)
-			[
-				ContentBrowserModule.Get().CreateAssetPicker(AssetPickerConfig)
-			]
-
-			+ SVerticalBox::Slot()
-			.Padding(4)
-			.AutoHeight()
-			.HAlign(HAlign_Right)
-			[
-				SNew(SHorizontalBox)
-				+SHorizontalBox::Slot()
-				.AutoWidth()
-				.VAlign(VAlign_Center)
-				.Padding(4.0f, 0.0f)
-				[
-					SNew(STextBlock)
-					.Text(LOCTEXT("ImportFrame_Label", "Sequence Frame: "))
-				]
-				+SHorizontalBox::Slot()
-				.AutoWidth()
-				.HAlign(HAlign_Right)
-				.Padding(2.0f, 0.0f)
-				[
-					SNew(SNumericEntryBox<int32>)
-					.ToolTipText(LOCTEXT("ArrayIndex", "Frame of sequence to import pose from."))
-					.AllowSpin(true)
-					.Font(IDetailLayoutBuilder::GetDetailFont())
-					.MinValue(0)
-					.Value_Lambda([this]
-					{
-						return FrameOfSequenceToImport;
-					})
-					.OnValueChanged(SNumericEntryBox<int32>::FOnValueChanged::CreateLambda([this](int32 Value)
-					{
-						FrameOfSequenceToImport = Value;
-					}))
-				]
-			]
-
-			+ SVerticalBox::Slot()
-			.Padding(4)
-			.AutoHeight()
-			.HAlign(HAlign_Right)
-			[
-				SNew(SHorizontalBox)
-				+SHorizontalBox::Slot()
-				.AutoWidth()
-				.VAlign(VAlign_Center)
-				.Padding(4.0f, 0.0f)
-				[
-					SNew(STextBlock)
-					.Text(LOCTEXT("ImportName_Label", "Pose Name: "))
-				]
-				
-				+SHorizontalBox::Slot()
-				.AutoWidth()
-				[
-					SNew(SEditableTextBox)
-					.Text(FText::FromName(UniqueNewPoseName))
-					.OnTextChanged(FOnTextChanged::CreateLambda([this](const FText InText)
-					{
-						ImportedPoseName = InText;
-					}))
-				]
-			]
-			
-			+ SVerticalBox::Slot()
-			.Padding(4)
-			.AutoHeight()
-			.HAlign(HAlign_Right)
-			[
-				SNew(SHorizontalBox)
-				+ SHorizontalBox::Slot()
-				.AutoWidth()
-				.Padding(4)
-				[
-					SNew(SButton)
-					.ButtonStyle(FAppStyle::Get(), "Button")
-					.TextStyle( FAppStyle::Get(), "DialogButtonText" )
-					.HAlign(HAlign_Center)
-					.VAlign(VAlign_Center)
-					.Text(LOCTEXT("ImportAsRetargetPoseButtonLabel", "Import As Retarget Pose"))
-					.OnClicked(this, &FIKRetargetEditorController::OnImportPoseFromSequence)
-				]
-				
-				+ SHorizontalBox::Slot()
-				.AutoWidth()
-				.Padding(4)
-				[
-					SNew(SButton)
-					.ButtonStyle(FAppStyle::Get(), "Button")
-					.TextStyle( FAppStyle::Get(), "DialogButtonText" )
-					.HAlign(HAlign_Center)
-					.VAlign(VAlign_Center)
-					.Text(LOCTEXT("CancelButtonLabel", "Cancel"))
-					.OnClicked_Lambda( [this]()
-					{
-						ImportPoseFromSequenceWindow->RequestDestroyWindow();
-						return FReply::Handled();
-					})
-				]
-			]	
-		]
-	];
-
-	GEditor->EditorAddModalWindow(ImportPoseFromSequenceWindow.ToSharedRef());
-	ImportPoseFromSequenceWindow.Reset();
-}
-
-bool FIKRetargetEditorController::OnShouldFilterSequenceToImport(const struct FAssetData& AssetData) const
-{
-	// is this an animation asset?
-	if (!AssetData.IsInstanceOf(UAnimationAsset::StaticClass()))
-	{
-		return true;
-	}
-
-	// get target skeleton
-	constexpr bool bEditingTarget = true; // TODO in future allow editing source or target
-	const USkeleton* DesiredSkeleton = bEditingTarget ? GetTargetSkeleton() : GetSourceSkeleton();
-	if (!DesiredSkeleton)
-	{
-		return true;
-	}
-
-	return !DesiredSkeleton->IsCompatibleSkeletonByAssetData(AssetData);
-}
-
-FReply FIKRetargetEditorController::OnImportPoseFromSequence()
-{
-	ImportPoseFromSequenceWindow->RequestDestroyWindow();
-	
-	if (SequenceToImportAsPose.IsNull())
-	{
-		return FReply::Handled();
-	}
-
-	const TObjectPtr<UAnimSequence> AnimSequence = Cast<UAnimSequence>(SequenceToImportAsPose.TryLoad());
-	if (!AnimSequence)
-	{
-		return FReply::Handled();
-	}
-	
-	USkeletalMesh* Mesh = GetTargetSkeletalMesh();
-	if (!Mesh)
-	{
-		return FReply::Handled();
-	}
-
-	UIKRetargetProcessor* Retargeter = GetRetargetProcessor();
-	if (!Retargeter)
-	{
-		return FReply::Handled();
-	}
-	
-	// ensure we evaluate the source animation using the skeletal mesh proportions that were evaluated in the viewport
-	FAnimPoseEvaluationOptions EvaluationOptions = FAnimPoseEvaluationOptions();
-	EvaluationOptions.OptionalSkeletalMesh = Mesh;
-
-	FrameOfSequenceToImport = FMath::Clamp(FrameOfSequenceToImport, 0, AnimSequence->GetNumberOfSampledKeys());
-	
-	// evaluate the sequence at the desired frame
-	FAnimPose ImportedPose;
-	UAnimPoseExtensions::GetAnimPoseAtFrame(AnimSequence, FrameOfSequenceToImport, EvaluationOptions, ImportedPose);
-
-	// record delta pose for all bones being retargeted
-	FIKRetargetPose ImportedRetargetPose;
-	
-	// get all imported bone transforms and record them in the retarget pose
-	FReferenceSkeleton& RefSkeleton = Mesh->GetRefSkeleton();
-	const TArray<FTransform>& RefPose = RefSkeleton.GetRefBonePose();
-	const FTargetSkeleton& TargetSkeleton = Retargeter->GetTargetSkeleton();
-	int32 NumBones = RefSkeleton.GetNum();
-	const int32 RootBoneIndex = Retargeter->GetTargetRetargetRoot();
-	for (int32 BoneIndex = 0; BoneIndex < NumBones; BoneIndex++)
-	{
-		const FName& BoneName = RefSkeleton.GetBoneName(BoneIndex);
-		const int32 RetargetBoneIndex = TargetSkeleton.FindBoneIndexByName(BoneName);
-		const bool bIsRetargetRoot = RetargetBoneIndex == RootBoneIndex;
-		
-		// if this is the retarget root, we want to record the translation delta as well
-		if (bIsRetargetRoot)
-		{
-			const FTransform GlobalTransformImported = UAnimPoseExtensions::GetBonePose(ImportedPose, BoneName, EAnimPoseSpaces::World);
-			const FTransform GlobalTransformReference = UAnimPoseExtensions::GetRefBonePose(ImportedPose, BoneName, EAnimPoseSpaces::World);
-			const FVector TranslationDelta = GlobalTransformImported.GetLocation() - GlobalTransformReference.GetLocation();
-			ImportedRetargetPose.RootTranslationOffset = TranslationDelta;
-
-			// rotation offsets are interpreted as relative to the parent (local), but in the case of the retarget root bone,
-			// when we generate the retarget pose, it's parents will be left at ref pose, so we need to generate a local
-			// rotation offset relative to the ref pose parent, NOT the (potentially) posed parent transform from the animation.
-			FTransform GlobalParentTransformInRefPose = FTransform::Identity;
-			const int32 ParentIndex = RefSkeleton.GetParentIndex(BoneIndex);
-			if (ParentIndex != INDEX_NONE)
-			{
-				const FName& ParentBoneName = RefSkeleton.GetBoneName(ParentIndex);
-				GlobalParentTransformInRefPose = UAnimPoseExtensions::GetRefBonePose(ImportedPose, ParentBoneName, EAnimPoseSpaces::World);
-			}
-
-			// this is a bit crazy, but we have to generate a delta rotation in the local space of the retarget root
-			// bone while treating the root bone as being in global space since the retarget pose does not consider any
-			// bones above it.
-			const FQuat GlobalDeltaRotation = GlobalTransformImported.GetRotation() * GlobalTransformReference.GetRotation().Inverse();
-			const FQuat BoneGlobalOrig = GlobalTransformReference.GetRotation();
-			const FQuat BoneGlobalPlusOffset = GlobalDeltaRotation * BoneGlobalOrig;
-			const FQuat ParentInv = GlobalParentTransformInRefPose.GetRotation().Inverse();
-			const FQuat BoneLocal = ParentInv * BoneGlobalOrig;
-			const FQuat BoneLocalPlusOffset = ParentInv * BoneGlobalPlusOffset;
-			const FQuat BoneLocalOffset = BoneLocal * BoneLocalPlusOffset.Inverse();
-			
-			ImportedRetargetPose.BoneRotationOffsets.Add(BoneName, BoneLocalOffset.Inverse());
-		}
-		else
-		{
-			// record the delta rotation
-			const FTransform LocalTransformImported = UAnimPoseExtensions::GetBonePose(ImportedPose, BoneName, EAnimPoseSpaces::Local);
-			const FTransform LocalTransformReference = RefPose[BoneIndex];
-			const FQuat DeltaRotation = LocalTransformImported.GetRotation() * LocalTransformReference.GetRotation().Inverse();
-			// only if it's different than the ref pose
-			if (DeltaRotation.GetAngle() > KINDA_SMALL_NUMBER)
-			{
-				ImportedRetargetPose.BoneRotationOffsets.Add(BoneName, DeltaRotation);
-			}
-		}
-	}
-
-	// store the newly imported retarget pose in the asset
-	AssetController->AddRetargetPose(FName(ImportedPoseName.ToString()), &ImportedRetargetPose);
-	
-	return FReply::Unhandled();
-}
-
-void FIKRetargetEditorController::OnSequenceSelectedForPose(const FAssetData& SelectedAsset)
-{
-	SequenceToImportAsPose = SelectedAsset.ToSoftObjectPath();
-}
-
-void FIKRetargetEditorController::HandleExportPose()
+void FIKRetargetEditorController::HandleDeletePose()
 {
-	FSaveAssetDialogConfig SaveAssetDialogConfig;
-	SaveAssetDialogConfig.DefaultPath = AssetController->GetAsset()->GetPackage()->GetPathName();
-	SaveAssetDialogConfig.DefaultAssetName = GetCurrentPoseName().ToString();
-	SaveAssetDialogConfig.AssetClassNames.Add(URetargetPose::StaticClass()->GetClassPathName());
-	SaveAssetDialogConfig.ExistingAssetPolicy = ESaveAssetDialogExistingAssetPolicy::AllowButWarn;
-	SaveAssetDialogConfig.DialogTitleOverride = LOCTEXT("ExportRetargetPoseDialogTitle", "Export Retarget Pose");
-
-	const FContentBrowserModule& ContentBrowserModule = FModuleManager::LoadModuleChecked<FContentBrowserModule>("ContentBrowser");
-	FString SaveObjectPath = ContentBrowserModule.Get().CreateModalSaveAssetDialog(SaveAssetDialogConfig);
-	if (SaveObjectPath.IsEmpty())
-	{
-		return;
-	}
-
-	const FString PackagePath = FPackageName::ObjectPathToPackageName(SaveObjectPath);
-	const FString AssetName = FPaths::GetBaseFilename(PackagePath, true);
-	URetargetPoseFactory* NewFactory = NewObject<URetargetPoseFactory>();
-	const FAssetToolsModule& AssetToolsModule = FAssetToolsModule::GetModule();
-	URetargetPose* NewPoseAsset =  Cast<URetargetPose>(AssetToolsModule.Get().CreateAsset(AssetName, PackagePath, URetargetPose::StaticClass(), NewFactory));
-	TArray<UObject*> ObjectsToSync;
-	ObjectsToSync.Add(NewPoseAsset);
-	GEditor->SyncBrowserToObjects(ObjectsToSync);
+	SetRetargeterMode(ERetargeterOutputMode::ShowRetargetPose);
 	
-	// fill new pose asset with existing pose data
-	const FIKRetargetPose& CurrentRetargetPose = AssetController->GetCurrentRetargetPose();
-	NewPoseAsset->RootTranslationOffset = CurrentRetargetPose.RootTranslationOffset;
-	NewPoseAsset->BoneRotationOffsets = CurrentRetargetPose.BoneRotationOffsets;
-}
-
-void FIKRetargetEditorController::HandleDeletePose() const
-{
-	const FName CurrentPose = AssetController->GetCurrentRetargetPoseName();
-	AssetController->RemoveRetargetPose(CurrentPose);
-	DetailsView->ForceRefresh();
+	const ERetargetSourceOrTarget SourceOrTarget = GetSourceOrTarget();
+	const FName CurrentPose = AssetController->GetCurrentRetargetPoseName(SourceOrTarget);
+	AssetController->RemoveRetargetPose(CurrentPose, SourceOrTarget);
+	RefreshPoseList();
 }
 
 bool FIKRetargetEditorController::CanDeletePose() const
 {	
 	// cannot delete default pose
-	const bool bNotUsingDefaultPose = AssetController->GetCurrentRetargetPoseName() != UIKRetargeter::GetDefaultPoseName();
-	// cannot delete pose while editing
-	return bNotUsingDefaultPose && !IsEditingPose();
+	return AssetController->GetCurrentRetargetPoseName(GetSourceOrTarget()) != UIKRetargeter::GetDefaultPoseName();
 }
 
 void FIKRetargetEditorController::HandleResetAllBones() const
 {
-	const FName CurrentPose = AssetController->GetCurrentRetargetPoseName();
+	const FName CurrentPose = AssetController->GetCurrentRetargetPoseName(CurrentlyEditingSourceOrTarget);
 	static TArray<FName> Empty; // empty list will reset all bones
-	AssetController->ResetRetargetPose(CurrentPose, Empty);
+	AssetController->ResetRetargetPose(CurrentPose, Empty, GetSourceOrTarget());
 }
 
 void FIKRetargetEditorController::HandleResetSelectedBones() const
 {
-	const FName CurrentPose = AssetController->GetCurrentRetargetPoseName();
-	AssetController->ResetRetargetPose(CurrentPose, GetSelectedBones());
+	const FName CurrentPose = AssetController->GetCurrentRetargetPoseName(CurrentlyEditingSourceOrTarget);
+	AssetController->ResetRetargetPose(CurrentPose, GetSelectedBones(), GetSourceOrTarget());
 }
 
 void FIKRetargetEditorController::HandleResetSelectedAndChildrenBones() const
 {
-	const FName CurrentPose = AssetController->GetCurrentRetargetPoseName();
+	// get the reference skeleton we're operating on
+	const USkeletalMesh* SkeletalMesh = GetSkeletalMesh(GetSourceOrTarget());
+	if (!SkeletalMesh)
+	{
+		return;
+	}
+	const FReferenceSkeleton RefSkeleton = SkeletalMesh->GetRefSkeleton();
 	
 	// get list of all children of selected bones
-	const UIKRetargetProcessor* Processor = GetRetargetProcessor();
-	const FTargetSkeleton& Skeleton = Processor->GetTargetSkeleton();
 	TArray<int32> AllChildrenIndices;
 	for (const FName& SelectedBone : SelectedBones)
 	{
-		const int32 SelectedBoneIndex = Skeleton.FindBoneIndexByName(SelectedBone);
-		Skeleton.GetChildrenIndicesRecursive(SelectedBoneIndex, AllChildrenIndices);
+		const int32 SelectedBoneIndex = RefSkeleton.FindBoneIndex(SelectedBone);
+		AllChildrenIndices.Add(SelectedBoneIndex);
+		
+		for (int32 ChildIndex = 0; ChildIndex < RefSkeleton.GetNum(); ++ChildIndex)
+		{
+			const int32 ParentIndex = RefSkeleton.GetParentIndex(ChildIndex);
+			if (ParentIndex != INDEX_NONE && AllChildrenIndices.Contains(ParentIndex))
+			{
+				AllChildrenIndices.Add(ChildIndex);
+			}
+		}
 	}
 
 	// merge total list of all selected bones and their children
 	TArray<FName> BonesToReset = SelectedBones;
 	for (const int32 ChildIndex : AllChildrenIndices)
 	{
-		BonesToReset.AddUnique(Skeleton.BoneNames[ChildIndex]);
+		BonesToReset.AddUnique(RefSkeleton.GetBoneName(ChildIndex));
 	}
 	
-	// reset the bones
-	AssetController->ResetRetargetPose(CurrentPose, BonesToReset);
+	// reset the bones in the current pose
+	const FName CurrentPose = AssetController->GetCurrentRetargetPoseName(CurrentlyEditingSourceOrTarget);
+	AssetController->ResetRetargetPose(CurrentPose, BonesToReset, GetSourceOrTarget());
 }
 
 bool FIKRetargetEditorController::CanResetSelected() const
 {
-	return CanResetPose() && !GetSelectedBones().IsEmpty();
-}
-
-bool FIKRetargetEditorController::CanResetPose() const
-{
-	// only allow reseting pose while editing to avoid confusion
-	return IsEditingPose();
+	return !GetSelectedBones().IsEmpty();
 }
 
 void FIKRetargetEditorController::HandleRenamePose()
@@ -1328,27 +1034,38 @@ FReply FIKRetargetEditorController::RenamePose() const
 	const FName NewPoseName = FName(NewNameEditableText.Get()->GetText().ToString());
 	RenamePoseWindow->RequestDestroyWindow();
 	
-	AssetController->RenameCurrentRetargetPose(NewPoseName);
-	DetailsView->ForceRefresh();
+	AssetController->RenameCurrentRetargetPose(NewPoseName, GetSourceOrTarget());
+	RefreshPoseList();
 	return FReply::Handled();
 }
 
 bool FIKRetargetEditorController::CanRenamePose() const
 {
 	// cannot rename default pose
-	const bool bNotUsingDefaultPose = AssetController->GetCurrentRetargetPoseName() != UIKRetargeter::GetDefaultPoseName();
+	const bool bNotUsingDefaultPose = AssetController->GetCurrentRetargetPoseName(GetSourceOrTarget()) != UIKRetargeter::GetDefaultPoseName();
 	// cannot rename pose while editing
 	return bNotUsingDefaultPose && !IsEditingPose();
 }
 
+void FIKRetargetEditorController::AddReferencedObjects(FReferenceCollector& Collector)
+{
+	for (TTuple<FName, TObjectPtr<UIKRetargetBoneDetails>> Pair : AllBoneDetails)
+	{
+		Collector.AddReferencedObject(Pair.Value);
+	}
+};
+
 FText FIKRetargetEditorController::GetCurrentPoseName() const
 {
-	return FText::FromName(AssetController->GetCurrentRetargetPoseName());
+	return FText::FromName(AssetController->GetCurrentRetargetPoseName(GetSourceOrTarget()));
 }
 
-void FIKRetargetEditorController::OnPoseSelected(TSharedPtr<FName> InPosePose, ESelectInfo::Type SelectInfo) const
+void FIKRetargetEditorController::OnPoseSelected(TSharedPtr<FName> InPose, ESelectInfo::Type SelectInfo) const
 {
-	AssetController->SetCurrentRetargetPose(*InPosePose.Get());
+	if (InPose.IsValid())
+	{
+		AssetController->SetCurrentRetargetPose(*InPose.Get(), GetSourceOrTarget());
+	}
 }
 
 #undef LOCTEXT_NAMESPACE
diff --git a/Engine/Plugins/Animation/IKRig/Source/IKRigEditor/Private/RetargetEditor/IKRetargetFactory.cpp b/Engine/Plugins/Animation/IKRig/Source/IKRigEditor/Private/RetargetEditor/IKRetargetFactory.cpp
--- a/Engine/Plugins/Animation/IKRig/Source/IKRigEditor/Private/RetargetEditor/IKRetargetFactory.cpp
+++ b/Engine/Plugins/Animation/IKRig/Source/IKRigEditor/Private/RetargetEditor/IKRetargetFactory.cpp
@@ -118,45 +118,4 @@ FString UIKRetargetFactory::GetDefaultNewAssetName() const
 	return FString(TEXT("NewIKRetargeter"));
 }
 
-/** Start URetargetPose Factory ...*/
-
-URetargetPoseFactory::URetargetPoseFactory()
-{
-	bCreateNew = true;
-	bEditAfterNew = true;
-	SupportedClass = URetargetPose::StaticClass();
-}
-
-FText URetargetPoseFactory::GetDisplayName() const
-{
-	return LOCTEXT("RetargetPose_DisplayName", "Retarget Pose");
-}
-
-uint32 URetargetPoseFactory::GetMenuCategories() const
-{
-	return EAssetTypeCategories::Animation;
-}
-
-FText URetargetPoseFactory::GetToolTip() const
-{
-	return LOCTEXT("RetargetPose_Tooltip", "Defines a set of bone transform offsets to create a pose for retargeting.");
-}
-
-FString URetargetPoseFactory::GetDefaultNewAssetName() const
-{
-	return FString(TEXT("NewRetargetPose"));
-}
-
-UObject* URetargetPoseFactory::FactoryCreateNew(
-	UClass* Class, UObject* InParent, FName Name, EObjectFlags Flags,
-	UObject* Context, FFeedbackContext* Warn)
-{
-	return NewObject<URetargetPose>(InParent, Class, Name, Flags);
-}
-
-bool URetargetPoseFactory::ShouldShowInNewMenu() const
-{
-	return true;
-}
-
 #undef LOCTEXT_NAMESPACE
diff --git a/Engine/Plugins/Animation/IKRig/Source/IKRigEditor/Private/RetargetEditor/IKRetargetOutputLogTabSummoner.cpp b/Engine/Plugins/Animation/IKRig/Source/IKRigEditor/Private/RetargetEditor/IKRetargetOutputLogTabSummoner.cpp
--- a/Engine/Plugins/Animation/IKRig/Source/IKRigEditor/Private/RetargetEditor/IKRetargetOutputLogTabSummoner.cpp
+++ b/Engine/Plugins/Animation/IKRig/Source/IKRigEditor/Private/RetargetEditor/IKRetargetOutputLogTabSummoner.cpp
@@ -35,7 +35,7 @@ TSharedRef<SWidget> FIKRetargetOutputLogTabSummoner::CreateTabBody(const FWorkfl
 	const TSharedRef<FIKRetargetEditorController>& Controller = IKRetargetEditor.Pin()->GetController();
 	const FName LogName = Controller->GetRetargetProcessor()->Log.GetLogTarget();
 	
-	return SNew(SIKRigOutputLog, LogName, Controller->OutputLogView);
+	return SNew(SIKRigOutputLog, LogName);
 }
 
 #undef LOCTEXT_NAMESPACE 
diff --git a/Engine/Plugins/Animation/IKRig/Source/IKRigEditor/Private/RetargetEditor/IKRetargetPoseExporter.cpp b/Engine/Plugins/Animation/IKRig/Source/IKRigEditor/Private/RetargetEditor/IKRetargetPoseExporter.cpp
new file mode 100644
--- /dev/null
+++ b/Engine/Plugins/Animation/IKRig/Source/IKRigEditor/Private/RetargetEditor/IKRetargetPoseExporter.cpp
@@ -0,0 +1,667 @@
+﻿// Copyright Epic Games, Inc. All Rights Reserved.
+
+#include "RetargetEditor/IKRetargetPoseExporter.h"
+
+#include "AnimPose.h"
+#include "AssetToolsModule.h"
+#include "ContentBrowserModule.h"
+#include "IContentBrowserSingleton.h"
+#include "Animation/DebugSkelMeshComponent.h"
+#include "AssetRegistry/AssetRegistryModule.h"
+#include "Framework/Notifications/NotificationManager.h"
+#include "RetargetEditor/IKRetargetAnimInstance.h"
+#include "RetargetEditor/IKRetargetEditorController.h"
+#include "UObject/SavePackage.h"
+#include "Widgets/Notifications/SNotificationList.h"
+#include "Animation/AnimSequence.h"
+
+#define LOCTEXT_NAMESPACE "IKRetargetExporter"
+
+
+void FIKRetargetPoseExporter::Initialize(TSharedPtr<FIKRetargetEditorController> InController)
+{
+	Controller = InController;
+}
+
+void FIKRetargetPoseExporter::HandleImportFromPoseAsset()
+{
+	FIKRetargetEditorController* ControllerPtr = Controller.Pin().Get();
+	
+	ControllerPtr->SetRetargeterMode(ERetargeterOutputMode::ShowRetargetPose);
+	
+	RetargetPoseToImport = nullptr;
+	PosesInSelectedAsset.Empty();
+	SelectedPose.Reset();
+
+	// load the content browser module to display an asset picker
+	const FContentBrowserModule& ContentBrowserModule = FModuleManager::LoadModuleChecked<FContentBrowserModule>("ContentBrowser");
+
+	// the asset picker will only show pose assets
+	FAssetPickerConfig AssetPickerConfig;
+	AssetPickerConfig.Filter.ClassPaths.Add(UPoseAsset::StaticClass()->GetClassPathName());
+	AssetPickerConfig.Filter.bRecursiveClasses = true;
+	AssetPickerConfig.OnAssetSelected = FOnAssetSelected::CreateSP(this, &FIKRetargetPoseExporter::OnRetargetPoseSelected);
+	AssetPickerConfig.InitialAssetViewType = EAssetViewType::Tile;
+
+	ImportPoseWindow = SNew(SWindow)
+	.Title(LOCTEXT("ImportRetargetPose", "Import Retarget Pose"))
+	.ClientSize(FVector2D(500, 600))
+	.SupportsMinimize(false) .SupportsMaximize(false)
+	[
+		SNew(SBorder)
+		.BorderImage( FAppStyle::GetBrush("Menu.Background") )
+		[
+			SNew(SVerticalBox)
+			
+			+ SVerticalBox::Slot()
+			.Padding(4)
+			[
+				ContentBrowserModule.Get().CreateAssetPicker(AssetPickerConfig)
+			]
+
+			+ SVerticalBox::Slot()
+			.AutoHeight()
+			.Padding(4)
+			[
+				SNew(SHorizontalBox)
+				+ SHorizontalBox::Slot()
+				.HAlign(HAlign_Right)
+				[
+					SNew(SComboBox<TSharedPtr<FName> >) // pose list combo
+					.OptionsSource(&PosesInSelectedAsset)
+					.OnGenerateWidget(this, &FIKRetargetPoseExporter::OnGeneratePoseComboWidget)
+					.OnSelectionChanged(this, &FIKRetargetPoseExporter::OnSelectPoseFromPoseAsset)
+					[
+						SNew(STextBlock).Text_Lambda([this]()
+						{
+							if (SelectedPose.IsValid())
+							{
+								return FText::FromName(*SelectedPose.Get());
+							}
+
+							return FText::FromString("No Pose Selected");
+						})
+					]
+				]
+			]
+			
+			+ SVerticalBox::Slot()
+			.Padding(4)
+			.AutoHeight()
+			.HAlign(HAlign_Right)
+			[
+				SNew(SHorizontalBox)
+				+ SHorizontalBox::Slot()
+				.AutoWidth()
+				.Padding(4)
+				[
+					SNew(SButton)
+					.ButtonStyle(FAppStyle::Get(), "Button")
+					.TextStyle( FAppStyle::Get(), "DialogButtonText" )
+					.HAlign(HAlign_Center)
+					.VAlign(VAlign_Center)
+					.Text(LOCTEXT("ImportNewRetargetPoseButtonLabel", "Import New Retarget Pose"))
+					.OnClicked(this, &FIKRetargetPoseExporter::ImportPoseAsset)
+				]
+				
+				+ SHorizontalBox::Slot()
+				.AutoWidth()
+				.Padding(4)
+				[
+					SNew(SButton)
+					.ButtonStyle(FAppStyle::Get(), "Button")
+					.TextStyle( FAppStyle::Get(), "DialogButtonText" )
+					.HAlign(HAlign_Center)
+					.VAlign(VAlign_Center)
+					.Text(LOCTEXT("CancelButtonLabel", "Cancel"))
+					.OnClicked_Lambda( [this]()
+					{
+						ImportPoseWindow->RequestDestroyWindow();
+						return FReply::Handled();
+					})
+				]
+			]	
+		]
+	];
+
+	GEditor->EditorAddModalWindow(ImportPoseWindow.ToSharedRef());
+	ImportPoseWindow.Reset();
+}
+
+void FIKRetargetPoseExporter::OnRetargetPoseSelected(const FAssetData& SelectedAsset)
+{
+	RetargetPoseToImport = SelectedAsset.ToSoftObjectPath();
+
+	// get all poses in the selected asset
+	PosesInSelectedAsset.Empty();
+	const TObjectPtr<UPoseAsset> PoseAsset = Cast<UPoseAsset>(RetargetPoseToImport.TryLoad());
+	if (!PoseAsset)
+	{
+		return;
+	}
+
+	// store pose names in list used by combobox
+	const TArray<FSmartName> AllPoseNames = PoseAsset->GetPoseNames();
+	for (const FSmartName& PoseName : AllPoseNames)
+	{
+		PosesInSelectedAsset.Add(MakeShared<FName>(PoseName.DisplayName));
+	}
+
+	// set the selected pose to the first pose in the asset (or null if empty)
+	if (PosesInSelectedAsset.IsEmpty())
+	{
+		SelectedPose = nullptr;
+	}
+	else
+	{
+		SelectedPose = PosesInSelectedAsset[0];
+	}
+}
+
+TSharedRef<SWidget> FIKRetargetPoseExporter::OnGeneratePoseComboWidget(TSharedPtr<FName> Item) const
+{
+	return SNew(STextBlock).Text(FText::FromName(*Item.Get()));
+}
+
+void FIKRetargetPoseExporter::OnSelectPoseFromPoseAsset(TSharedPtr<FName> Item, ESelectInfo::Type SelectInfo)
+{
+	if (Item.IsValid())
+	{
+		SelectedPose = Item;
+	}
+}
+
+FReply FIKRetargetPoseExporter::ImportPoseAsset() const
+{
+	// close the import window
+	ImportPoseWindow->RequestDestroyWindow();
+
+	// check that user actually selected an asset to import
+	if (RetargetPoseToImport.IsNull())
+	{
+		const FText Message = LOCTEXT("NoPoseToImport", "No pose asset specified. Import aborted.");
+		NotifyUser(Message, SNotificationItem::CS_Fail);
+		return FReply::Handled();
+	}
+
+	// load the selected pose asset
+	const TObjectPtr<UPoseAsset> PoseAsset = Cast<UPoseAsset>(RetargetPoseToImport.TryLoad());
+	if (!PoseAsset)
+	{
+		const FText Message = LOCTEXT("PoseLoadFailed", "Pose asset could not be loaded. Import aborted.");
+		NotifyUser(Message, SNotificationItem::CS_Fail);
+		return FReply::Handled();
+	}
+
+	// get the requested pose out of the pose asset
+	if (!SelectedPose.IsValid())
+	{
+		const FText Message = LOCTEXT("PoseNotSelected", "No pose was specified to load. Import aborted.");
+		NotifyUser(Message, SNotificationItem::CS_Fail);
+		return FReply::Handled();
+	}
+
+	// validate that pose exists within the selected pose asset
+	const FName NameOfPoseToImport = *SelectedPose.Get();
+	const int32 IndexOfPoseToImport = PoseAsset->GetPoseIndexByName(NameOfPoseToImport);
+	if (IndexOfPoseToImport == INDEX_NONE)
+	{
+		const FText Message = LOCTEXT("PoseNotInAsset", "Specified pose not found in the pose asset. Import aborted.");
+		NotifyUser(Message, SNotificationItem::CS_Fail);
+		return FReply::Handled();
+	}
+
+	// validate mesh is available to load sequence on
+	FIKRetargetEditorController* ControllerPtr = Controller.Pin().Get();
+	ERetargetSourceOrTarget SourceOrTarget = ControllerPtr->GetSourceOrTarget();
+	USkeletalMesh* Mesh = ControllerPtr->GetSkeletalMesh(SourceOrTarget);
+	if (!Mesh)
+	{
+		const FText Message = LOCTEXT("ImportFailedNoMesh", "Skeletal Mesh not found. Import aborted.");
+		NotifyUser(Message, SNotificationItem::CS_Fail);
+		return FReply::Handled();
+	}
+
+	// validate IK Rig is available to sort the pose offsets
+	const UIKRigDefinition* IKRig = ControllerPtr->AssetController->GetIKRig(SourceOrTarget);
+	if (!IKRig)
+	{
+		const FText Message = LOCTEXT("ImportFailedNoIKRig", "IK Rig not found. Import aborted.");
+		NotifyUser(Message, SNotificationItem::CS_Fail);
+		return FReply::Handled();
+	}
+
+	// get retarget root
+	const FName RetargetRootBoneName = IKRig->GetRetargetRoot();
+
+	// get transforms of loaded pose
+	TArray<FTransform> LocalBoneTransformFromPose;
+	PoseAsset->GetFullPose(IndexOfPoseToImport, LocalBoneTransformFromPose);
+
+	// create a new retarget pose to store the data from the selected retarget pose asset
+	FIKRetargetPose NewPose;
+	
+	// iterate over all bones in the reference skeleton and compare against bone in the pose asset
+	// store sparse set of deltas in the new retarget pose 
+	FReferenceSkeleton& RefSkeleton = Mesh->GetRefSkeleton();
+	const TArray<FTransform>& RefPose = RefSkeleton.GetRefBonePose();
+	for (int32 BoneIndex=0; BoneIndex<RefSkeleton.GetNum(); ++BoneIndex)
+	{
+		const FName BoneName = RefSkeleton.GetBoneName(BoneIndex);
+		const int32 PoseTrackIndex = PoseAsset->GetTrackIndexByName(BoneName);
+		if (PoseTrackIndex == INDEX_NONE)
+		{
+			continue;
+		}
+		
+		// record a global space translation offset for the root bone
+		if (BoneName == RetargetRootBoneName)
+		{
+			FTransform GlobalRefTransform = FAnimationRuntime::GetComponentSpaceTransform(RefSkeleton, RefPose, BoneIndex);
+			FTransform GlobalImportedTransform = PoseAsset->GetComponentSpaceTransform(BoneName, LocalBoneTransformFromPose);
+			NewPose.SetRootTranslationDelta(GlobalImportedTransform.GetLocation() - GlobalRefTransform.GetLocation());
+		}
+		
+		// record a local space delta rotation (if there is one)
+		const FTransform& LocalRefTransform = RefPose[BoneIndex];
+		const FTransform& LocalImportedTransform = LocalBoneTransformFromPose[PoseTrackIndex];
+		const FQuat DeltaRotation = LocalRefTransform.GetRotation().Inverse() * LocalImportedTransform.GetRotation();
+		const float DeltaAngle = FMath::RadiansToDegrees(DeltaRotation.GetAngle());
+		constexpr float MinAngleThreshold = 0.05f;
+		if (DeltaAngle >= MinAngleThreshold)
+		{
+			NewPose.SetDeltaRotationForBone(BoneName, DeltaRotation);
+		}
+	}
+
+	// store the retarget pose in the retarget asset
+	ControllerPtr->AssetController->AddRetargetPose(FName(PoseAsset->GetName()), &NewPose, SourceOrTarget);
+
+	// update view with new pose
+	ControllerPtr->RefreshAllViews();
+	
+	return FReply::Unhandled();
+}
+
+void FIKRetargetPoseExporter::HandleImportFromSequenceAsset()
+{
+	FIKRetargetEditorController* ControllerPtr = Controller.Pin().Get();
+	
+	ControllerPtr->SetRetargeterMode(ERetargeterOutputMode::ShowRetargetPose);
+	
+	SequenceToImportAsPose = nullptr;
+
+	// get a unique pose name to use as suggestion
+	const FString DefaultImportedPoseName = LOCTEXT("ImportedRetargetPoseName", "ImportedRetargetPose").ToString();
+	const FName UniqueNewPoseName = ControllerPtr->AssetController->MakePoseNameUnique(DefaultImportedPoseName, ControllerPtr->GetSourceOrTarget());
+	ImportedPoseName = FText::FromName(UniqueNewPoseName);
+
+	// load the content browser module to display an asset picker
+	const FContentBrowserModule& ContentBrowserModule = FModuleManager::LoadModuleChecked<FContentBrowserModule>("ContentBrowser");
+
+	// the asset picker will only show animation sequences compatible with the preview mesh
+	FAssetPickerConfig AssetPickerConfig;
+	AssetPickerConfig.Filter.ClassPaths.Add(UAnimSequence::StaticClass()->GetClassPathName());
+	AssetPickerConfig.Filter.ClassPaths.Add(UAnimMontage::StaticClass()->GetClassPathName());
+	AssetPickerConfig.Filter.ClassPaths.Add(UPoseAsset::StaticClass()->GetClassPathName());
+	AssetPickerConfig.InitialAssetViewType = EAssetViewType::Column;
+	AssetPickerConfig.bAddFilterUI = true;
+	AssetPickerConfig.bShowPathInColumnView = true;
+	AssetPickerConfig.bShowTypeInColumnView = true;
+	AssetPickerConfig.DefaultFilterMenuExpansion = EAssetTypeCategories::Animation;
+	AssetPickerConfig.OnShouldFilterAsset = FOnShouldFilterAsset::CreateSP(this, &FIKRetargetPoseExporter::OnShouldFilterSequenceToImport);
+	AssetPickerConfig.OnAssetSelected = FOnAssetSelected::CreateSP(this, &FIKRetargetPoseExporter::OnSequenceSelectedForPose);
+	AssetPickerConfig.bAllowNullSelection = false;
+
+	// hide all asset registry columns by default (we only really want the name and path)
+	TArray<UObject::FAssetRegistryTag> AssetRegistryTags;
+	UAnimSequence::StaticClass()->GetDefaultObject()->GetAssetRegistryTags(AssetRegistryTags);
+	FName ColumnToKeep = FName("Number of Frames");
+	for(UObject::FAssetRegistryTag& AssetRegistryTag : AssetRegistryTags)
+	{
+		if (AssetRegistryTag.Name != ColumnToKeep)
+		{
+			AssetPickerConfig.HiddenColumnNames.Add(AssetRegistryTag.Name.ToString());
+		}
+	}
+
+	// also hide the type column by default (but allow users to enable it, so don't use bShowTypeInColumnView)
+	AssetPickerConfig.HiddenColumnNames.Add(TEXT("Class"));
+	AssetPickerConfig.HiddenColumnNames.Add(TEXT("HasVirtualizedData"));
+	AssetPickerConfig.HiddenColumnNames.Add(TEXT("DiskSize"));
+
+	// create pop-up window for user to select animation sequence asset to import as a retarget pose]
+	ImportPoseFromSequenceWindow = SNew(SWindow)
+	.Title(LOCTEXT("ImportRetargetPoseFromSequenceAsset", "Import Retarget Pose from Sequence Asset"))
+	.ClientSize(FVector2D(500, 600))
+	.SupportsMinimize(false) .SupportsMaximize(false)
+	[
+		SNew(SBorder)
+		.BorderImage( FAppStyle::GetBrush("Menu.Background") )
+		[
+			SNew(SVerticalBox)
+			
+			+ SVerticalBox::Slot()
+			.Padding(4)
+			[
+				ContentBrowserModule.Get().CreateAssetPicker(AssetPickerConfig)
+			]
+
+			+ SVerticalBox::Slot()
+			.Padding(4)
+			.AutoHeight()
+			.HAlign(HAlign_Right)
+			[
+				SNew(SHorizontalBox)
+				+SHorizontalBox::Slot()
+				.AutoWidth()
+				.VAlign(VAlign_Center)
+				.Padding(4.0f, 0.0f)
+				[
+					SNew(STextBlock)
+					.Text(LOCTEXT("ImportFrame_Label", "Sequence Frame: "))
+				]
+				+SHorizontalBox::Slot()
+				.AutoWidth()
+				.HAlign(HAlign_Right)
+				.Padding(2.0f, 0.0f)
+				[
+					SNew(SNumericEntryBox<int32>)
+					.ToolTipText(LOCTEXT("ArrayIndex", "Frame of sequence to import pose from."))
+					.AllowSpin(true)
+					.Font(IDetailLayoutBuilder::GetDetailFont())
+					.MinValue(0)
+					.Value_Lambda([this]
+					{
+						return FrameOfSequenceToImport;
+					})
+					.OnValueChanged(SNumericEntryBox<int32>::FOnValueChanged::CreateLambda([this](int32 Value)
+					{
+						FrameOfSequenceToImport = Value;
+					}))
+				]
+			]
+
+			+ SVerticalBox::Slot()
+			.Padding(4)
+			.AutoHeight()
+			.HAlign(HAlign_Right)
+			[
+				SNew(SHorizontalBox)
+				+SHorizontalBox::Slot()
+				.AutoWidth()
+				.VAlign(VAlign_Center)
+				.Padding(4.0f, 0.0f)
+				[
+					SNew(STextBlock)
+					.Text(LOCTEXT("ImportName_Label", "Pose Name: "))
+				]
+				
+				+SHorizontalBox::Slot()
+				.AutoWidth()
+				[
+					SNew(SEditableTextBox)
+					.Text(FText::FromName(UniqueNewPoseName))
+					.OnTextChanged(FOnTextChanged::CreateLambda([this](const FText InText)
+					{
+						ImportedPoseName = InText;
+					}))
+				]
+			]
+			
+			+ SVerticalBox::Slot()
+			.Padding(4)
+			.AutoHeight()
+			.HAlign(HAlign_Right)
+			[
+				SNew(SHorizontalBox)
+				+ SHorizontalBox::Slot()
+				.AutoWidth()
+				.Padding(4)
+				[
+					SNew(SButton)
+					.ButtonStyle(FAppStyle::Get(), "Button")
+					.TextStyle( FAppStyle::Get(), "DialogButtonText" )
+					.HAlign(HAlign_Center)
+					.VAlign(VAlign_Center)
+					.Text(LOCTEXT("ImportAsRetargetPoseButtonLabel", "Import As Retarget Pose"))
+					.OnClicked(this, &FIKRetargetPoseExporter::OnImportPoseFromSequence)
+				]
+				
+				+ SHorizontalBox::Slot()
+				.AutoWidth()
+				.Padding(4)
+				[
+					SNew(SButton)
+					.ButtonStyle(FAppStyle::Get(), "Button")
+					.TextStyle( FAppStyle::Get(), "DialogButtonText" )
+					.HAlign(HAlign_Center)
+					.VAlign(VAlign_Center)
+					.Text(LOCTEXT("CancelButtonLabel", "Cancel"))
+					.OnClicked_Lambda( [this]()
+					{
+						ImportPoseFromSequenceWindow->RequestDestroyWindow();
+						return FReply::Handled();
+					})
+				]
+			]	
+		]
+	];
+
+	GEditor->EditorAddModalWindow(ImportPoseFromSequenceWindow.ToSharedRef());
+	ImportPoseFromSequenceWindow.Reset();
+}
+
+bool FIKRetargetPoseExporter::OnShouldFilterSequenceToImport(const FAssetData& AssetData) const
+{
+	// is this an animation asset?
+	if (!AssetData.IsInstanceOf(UAnimationAsset::StaticClass()))
+	{
+		return true;
+	}
+
+	// get currently edited skeleton
+	const USkeleton* DesiredSkeleton = Controller.Pin()->GetSkeleton(Controller.Pin()->GetSourceOrTarget());
+	if (!DesiredSkeleton)
+	{
+		return true;
+	}
+
+	return !DesiredSkeleton->IsCompatibleSkeletonByAssetData(AssetData);
+}
+
+FReply FIKRetargetPoseExporter::OnImportPoseFromSequence()
+{
+	// close the pop-up window
+	ImportPoseFromSequenceWindow->RequestDestroyWindow();
+
+	// check that a sequence was selected to import
+	if (SequenceToImportAsPose.IsNull())
+	{
+		const FText Message = LOCTEXT("ImportFailedNoSequence", "No sequence selected. Import aborted.");
+		NotifyUser(Message, SNotificationItem::CS_Fail);
+		return FReply::Handled();
+	}
+
+	// load sequence to read pose from
+	const TObjectPtr<UAnimSequence> AnimSequence = Cast<UAnimSequence>(SequenceToImportAsPose.TryLoad());
+	if (!AnimSequence)
+	{
+		const FText Message = LOCTEXT("ImportFailedSeqNotLoaded", "Sequence could not be loaded. Import aborted.");
+		NotifyUser(Message, SNotificationItem::CS_Fail);
+		return FReply::Handled();
+	}
+
+	// validate mesh is available to load sequence on
+	FIKRetargetEditorController* ControllerPtr = Controller.Pin().Get();
+	ERetargetSourceOrTarget SourceOrTarget = ControllerPtr->GetSourceOrTarget();
+	USkeletalMesh* Mesh = ControllerPtr->GetSkeletalMesh(SourceOrTarget);
+	if (!Mesh)
+	{
+		const FText Message = LOCTEXT("ImportFailedNoMesh", "Skeletal Mesh not found. Import aborted.");
+		NotifyUser(Message, SNotificationItem::CS_Fail);
+		return FReply::Handled();
+	}
+	
+	// ensure we evaluate the source animation using the skeletal mesh proportions that were evaluated in the viewport
+	FAnimPoseEvaluationOptions EvaluationOptions = FAnimPoseEvaluationOptions();
+	EvaluationOptions.OptionalSkeletalMesh = Mesh;
+	
+	// evaluate the sequence at the desired frame
+	FAnimPose ImportedPose;
+	FrameOfSequenceToImport = FMath::Clamp(FrameOfSequenceToImport, 0, AnimSequence->GetNumberOfSampledKeys());
+	UAnimPoseExtensions::GetAnimPoseAtFrame(AnimSequence, FrameOfSequenceToImport, EvaluationOptions, ImportedPose);
+
+	// record delta pose for all bones being retargeted
+	FIKRetargetPose ImportedRetargetPose;
+	
+	// get all imported bone transforms and record them in the retarget pose
+	FReferenceSkeleton& RefSkeleton = Mesh->GetRefSkeleton();
+	const TArray<FTransform>& RefPose = RefSkeleton.GetRefBonePose();
+	int32 NumBones = RefSkeleton.GetNum();
+	const FName RootBoneName = ControllerPtr->AssetController->GetRetargetRootBone(ControllerPtr->GetSourceOrTarget());
+	for (int32 BoneIndex = 0; BoneIndex < NumBones; BoneIndex++)
+	{
+		const FName& BoneName = RefSkeleton.GetBoneName(BoneIndex);
+		
+		// if this is the retarget root, we want to record the translation delta as well
+		if (BoneName == RootBoneName)
+		{
+			const FTransform GlobalTransformImported = UAnimPoseExtensions::GetBonePose(ImportedPose, BoneName, EAnimPoseSpaces::World);
+			const FTransform GlobalTransformReference = UAnimPoseExtensions::GetRefBonePose(ImportedPose, BoneName, EAnimPoseSpaces::World);
+			const FVector TranslationDelta = GlobalTransformImported.GetLocation() - GlobalTransformReference.GetLocation();
+			ImportedRetargetPose.SetRootTranslationDelta(TranslationDelta);
+
+			// rotation offsets are interpreted as relative to the parent (local), but in the case of the retarget root bone,
+			// when we generate the retarget pose, it's parents will be left at ref pose, so we need to generate a local
+			// rotation offset relative to the ref pose parent, NOT the (potentially) posed parent transform from the animation.
+			FTransform GlobalParentTransformInRefPose = FTransform::Identity;
+			const int32 ParentIndex = RefSkeleton.GetParentIndex(BoneIndex);
+			if (ParentIndex != INDEX_NONE)
+			{
+				const FName& ParentBoneName = RefSkeleton.GetBoneName(ParentIndex);
+				GlobalParentTransformInRefPose = UAnimPoseExtensions::GetRefBonePose(ImportedPose, ParentBoneName, EAnimPoseSpaces::World);
+			}
+
+			// this is a bit crazy, but we have to generate a delta rotation in the local space of the retarget root
+			// bone while treating the root bone as being in global space since the retarget pose does not consider any
+			// bones above it.
+			const FQuat GlobalDeltaRotation = GlobalTransformImported.GetRotation() * GlobalTransformReference.GetRotation().Inverse();
+			const FQuat BoneGlobalOrig = GlobalTransformReference.GetRotation();
+			const FQuat BoneGlobalPlusOffset = GlobalDeltaRotation * BoneGlobalOrig;
+			const FQuat ParentInv = GlobalParentTransformInRefPose.GetRotation().Inverse();
+			const FQuat BoneLocal = ParentInv * BoneGlobalOrig;
+			const FQuat BoneLocalPlusOffset = ParentInv * BoneGlobalPlusOffset;
+			const FQuat BoneLocalOffset = BoneLocal * BoneLocalPlusOffset.Inverse();
+			
+			ImportedRetargetPose.SetDeltaRotationForBone(BoneName, BoneLocalOffset.Inverse());
+		}
+		else
+		{
+			// record the delta rotation
+			const FTransform LocalTransformImported = UAnimPoseExtensions::GetBonePose(ImportedPose, BoneName, EAnimPoseSpaces::Local);
+			const FTransform LocalTransformReference = RefPose[BoneIndex];
+			const FQuat DeltaRotation = LocalTransformReference.GetRotation().Inverse() * LocalTransformImported.GetRotation();
+			// only if it's different than the ref pose
+			const float DeltaAngle = FMath::RadiansToDegrees(DeltaRotation.GetAngle());
+			constexpr float MinAngleThreshold = 0.05f;
+			if (DeltaAngle >= MinAngleThreshold)
+			{
+				ImportedRetargetPose.SetDeltaRotationForBone(BoneName, DeltaRotation);
+			}
+		}
+	}
+
+	// store the newly imported retarget pose in the asset
+	ControllerPtr->AssetController->AddRetargetPose( FName(ImportedPoseName.ToString()), &ImportedRetargetPose, SourceOrTarget);
+
+	// notify user of new pose
+	const FText Message = FText::Format(LOCTEXT("ImportSuccess", "Imported pose from animation sequence: {0}"), ImportedPoseName);
+	NotifyUser(Message, SNotificationItem::CS_Success);
+	return FReply::Unhandled();
+}
+
+void FIKRetargetPoseExporter::OnSequenceSelectedForPose(const FAssetData& SelectedAsset)
+{
+	SequenceToImportAsPose = SelectedAsset.ToSoftObjectPath();
+}
+
+void FIKRetargetPoseExporter::HandleExportPoseAsset()
+{
+	const FIKRetargetEditorController* ControllerPtr = Controller.Pin().Get();
+
+	// get mesh to export pose
+	const ERetargetSourceOrTarget SourceOrTarget = ControllerPtr->GetSourceOrTarget();
+	const USkeletalMesh* Mesh = ControllerPtr->GetSkeletalMesh(SourceOrTarget);
+	UDebugSkelMeshComponent* MeshComponent = ControllerPtr->GetSkeletalMeshComponent(SourceOrTarget);
+	if (!(Mesh && MeshComponent))
+	{
+		const FText Message = LOCTEXT("IKRigExportPoseFailed", "Export pose failed. No mesh found.");
+		NotifyUser(Message, SNotificationItem::CS_Fail);
+		return;
+	}
+
+	FString PoseName = ControllerPtr->GetCurrentPoseName().ToString();
+	PoseName.RemoveSpacesInline();
+	
+	FSaveAssetDialogConfig SaveAssetDialogConfig;
+	SaveAssetDialogConfig.DefaultPath = ControllerPtr->AssetController->GetAsset()->GetPackage()->GetPathName();
+	SaveAssetDialogConfig.DefaultAssetName = PoseName;
+	SaveAssetDialogConfig.AssetClassNames.Add(UPoseAsset::StaticClass()->GetClassPathName());
+	SaveAssetDialogConfig.ExistingAssetPolicy = ESaveAssetDialogExistingAssetPolicy::AllowButWarn;
+	SaveAssetDialogConfig.DialogTitleOverride = LOCTEXT("ExportRetargetPoseDialogTitle", "Export Retarget Pose");
+
+	const FContentBrowserModule& ContentBrowserModule = FModuleManager::LoadModuleChecked<FContentBrowserModule>("ContentBrowser");
+	const FString SaveObjectPath = ContentBrowserModule.Get().CreateModalSaveAssetDialog(SaveAssetDialogConfig);
+	if (SaveObjectPath.IsEmpty())
+	{
+		const FText Message = LOCTEXT("IKRigExportPoseAborted", "Export pose cancelled.");
+		NotifyUser(Message, SNotificationItem::CS_Fail);
+		return;
+	}
+
+	// create new pose asset
+	const FString PackagePath = FPackageName::ObjectPathToPackageName(SaveObjectPath);
+	const FString AssetName = FPaths::GetBaseFilename(PackagePath, true);
+	UPackage *Package = CreatePackage(*PackagePath);
+	UPoseAsset* NewPoseAsset = NewObject<UPoseAsset>(Package, UPoseAsset::StaticClass(), FName(AssetName), RF_Public | RF_Standalone);
+
+	// fill new pose asset with current pose
+	NewPoseAsset->SetSkeleton(const_cast<USkeleton*>(Mesh->GetSkeleton()));
+	FSmartName NewPoseName;
+	bool bSuccess = NewPoseAsset->AddOrUpdatePoseWithUniqueName(MeshComponent, &NewPoseName);
+
+	if (bSuccess)
+	{
+		// mark asset dirty 
+		FAssetRegistryModule::AssetCreated(NewPoseAsset);
+		NewPoseAsset->MarkPackageDirty();
+	
+		// show in content browser
+		TArray<UObject*> ObjectsToSync;
+		ObjectsToSync.Add(NewPoseAsset);
+		GEditor->SyncBrowserToObjects(ObjectsToSync);
+		
+		// notify user of newly created asset
+		const FText Message = FText::Format(LOCTEXT("IKRigExportPoseAborted", "Exported Pose Asset: {0} with pose, '{1}'."), FText::FromString(AssetName), FText::FromName(NewPoseName.DisplayName));
+		NotifyUser(Message, SNotificationItem::CS_Success);
+	}
+	else
+	{
+		// remove the empty asset
+		NewPoseAsset->ConditionalBeginDestroy();
+		
+		// notify user that asset failed to export
+		const FText Message = FText::Format(LOCTEXT("IKRigExportPoseAborted", "Failed to exported pose asset, {0}."), FText::FromString(AssetName));
+		NotifyUser(Message, SNotificationItem::CS_Fail);
+	}
+}
+
+void FIKRetargetPoseExporter::NotifyUser(const FText& Message, SNotificationItem::ECompletionState NotificationType)
+{
+	FNotificationInfo NotifyInfo(Message);
+	NotifyInfo.ExpireDuration = 5.0f;
+	FSlateNotificationManager::Get().AddNotification(NotifyInfo)->SetCompletionState(NotificationType);
+}
+
+#undef LOCTEXT_NAMESPACE
diff --git a/Engine/Plugins/Animation/IKRig/Source/IKRigEditor/Private/RetargetEditor/IKRetargeterController.cpp b/Engine/Plugins/Animation/IKRig/Source/IKRigEditor/Private/RetargetEditor/IKRetargeterController.cpp
--- a/Engine/Plugins/Animation/IKRig/Source/IKRigEditor/Private/RetargetEditor/IKRetargeterController.cpp
+++ b/Engine/Plugins/Animation/IKRig/Source/IKRigEditor/Private/RetargetEditor/IKRetargeterController.cpp
@@ -27,11 +27,6 @@ UIKRetargeterController* UIKRetargeterController::GetController(UIKRetargeter* I
 	}
 
 	UIKRetargeterController* Controller = Cast<UIKRetargeterController>(InRetargeterAsset->Controller);
-	// clean the asset before editing
-	const bool bForceReinitialization = false;
-	Controller->CleanChainMapping(bForceReinitialization);
-	Controller->CleanPoseList(bForceReinitialization);
-	
 	return Controller;
 }
 
@@ -45,97 +40,59 @@ void UIKRetargeterController::SetSourceIKRig(UIKRigDefinition* SourceIKRig)
 	Asset->SourceIKRigAsset = SourceIKRig;
 }
 
-USkeletalMesh* UIKRetargeterController::GetSourcePreviewMesh() const
-{
-	// can't preview anything if target IK Rig is null
-	if (!Asset->GetSourceIKRig())
-	{
-		return nullptr;
-	}
-
-	// optionally prefer override if one is provided
-	if (!Asset->SourcePreviewMesh.IsNull())
-	{
-		return Asset->SourcePreviewMesh.LoadSynchronous();
-	}
-
-	// fallback to preview mesh from IK Rig asset
-	return Asset->GetSourceIKRig()->GetPreviewMesh();
-}
-
-USkeletalMesh* UIKRetargeterController::GetTargetPreviewMesh() const
+USkeletalMesh* UIKRetargeterController::GetPreviewMesh(const ERetargetSourceOrTarget& SourceOrTarget) const
 {
 	// can't preview anything if target IK Rig is null
-	if (!Asset->GetTargetIKRig())
+	const UIKRigDefinition* IKRig = GetIKRig(SourceOrTarget);
+	if (!IKRig)
 	{
 		return nullptr;
 	}
 
 	// optionally prefer override if one is provided
-	if (!Asset->TargetPreviewMesh.IsNull())
+	const TSoftObjectPtr<USkeletalMesh> OverrideMesh = SourceOrTarget == ERetargetSourceOrTarget::Source ? Asset->SourcePreviewMesh : Asset->TargetPreviewMesh;
+	if (!OverrideMesh.IsNull())
 	{
-		return Asset->TargetPreviewMesh.LoadSynchronous();
+		return OverrideMesh.LoadSynchronous();
 	}
 
 	// fallback to preview mesh from IK Rig asset
-	return Asset->GetTargetIKRig()->GetPreviewMesh();
+	return IKRig->GetPreviewMesh();
 }
 
-const UIKRigDefinition* UIKRetargeterController::GetSourceIKRig() const
+const UIKRigDefinition* UIKRetargeterController::GetIKRig(const ERetargetSourceOrTarget& SourceOrTarget) const
 {
-	return Asset->GetSourceIKRig();
+	return SourceOrTarget == ERetargetSourceOrTarget::Source ? Asset->GetSourceIKRig() : Asset->GetTargetIKRig();
 }
 
-const UIKRigDefinition* UIKRetargeterController::GetTargetIKRig() const
+void UIKRetargeterController::OnIKRigChanged(const ERetargetSourceOrTarget& SourceOrTarget) const
 {
-	return Asset->GetTargetIKRig();
-}
-
-void UIKRetargeterController::OnTargetIKRigChanged() const
-{
-	if (const UIKRigDefinition* IKRig = GetTargetIKRig())
-	{
-		Asset->TargetPreviewMesh = IKRig->PreviewSkeletalMesh;
-	}
-}
-
-void UIKRetargeterController::OnSourceIKRigChanged() const
-{
-	if (const UIKRigDefinition* IKRig = GetSourceIKRig())
+	if (const UIKRigDefinition* IKRig = GetIKRig(SourceOrTarget))
 	{
-		Asset->SourcePreviewMesh = IKRig->PreviewSkeletalMesh;
+		if (SourceOrTarget == ERetargetSourceOrTarget::Source)
+		{
+			Asset->SourcePreviewMesh = IKRig->PreviewSkeletalMesh;
+		}
+		else
+		{
+			Asset->TargetPreviewMesh = IKRig->PreviewSkeletalMesh;
+		}
+		
 	}
 }
 
-FName UIKRetargeterController::GetSourceRootBone() const
-{
-	const UIKRigDefinition* SourceIKRig = Asset->GetSourceIKRig();
-	return SourceIKRig ? SourceIKRig->GetRetargetRoot() : FName("None");
-}
-
-FName UIKRetargeterController::GetTargetRootBone() const
+FName UIKRetargeterController::GetRetargetRootBone(
+	const ERetargetSourceOrTarget& SourceOrTarget) const
 {
-	const UIKRigDefinition* TargetIKRig = Asset->GetTargetIKRig();
-	return TargetIKRig ? TargetIKRig->GetRetargetRoot() : FName("None");
-}
-
-void UIKRetargeterController::GetTargetChainNames(TArray<FName>& OutNames) const
-{
-	if (const UIKRigDefinition* TargetIKRig = Asset->GetTargetIKRig())
-	{
-		const TArray<FBoneChain>& Chains = TargetIKRig->GetRetargetChains();
-		for (const FBoneChain& Chain : Chains)
-		{
-			OutNames.Add(Chain.ChainName);
-		}
-	}
+	const UIKRigDefinition* IKRig = GetIKRig(SourceOrTarget);
+	return IKRig ? IKRig->GetRetargetRoot() : FName("None");
 }
 
-void UIKRetargeterController::GetSourceChainNames(TArray<FName>& OutNames) const
+void UIKRetargeterController::GetChainNames(const ERetargetSourceOrTarget& SourceOrTarget, TArray<FName>& OutNames) const
 {
-	if (const UIKRigDefinition* SourceIKRig = Asset->GetSourceIKRig())
+	if (const UIKRigDefinition* IKRig = GetIKRig(SourceOrTarget))
 	{
-		const TArray<FBoneChain>& Chains = SourceIKRig->GetRetargetChains();
+		const TArray<FBoneChain>& Chains = IKRig->GetRetargetChains();
 		for (const FBoneChain& Chain : Chains)
 		{
 			OutNames.Add(Chain.ChainName);
@@ -148,7 +105,7 @@ void UIKRetargeterController::CleanChainMapping(const bool bForceReinitializatio
 	if (IsValid(Asset->GetTargetIKRig()))
 	{
 		TArray<FName> TargetChainNames;
-		GetTargetChainNames(TargetChainNames);
+		GetChainNames(ERetargetSourceOrTarget::Target, TargetChainNames);
 
 		// remove all target chains that are no longer in the target IK rig asset
 		TArray<FName> TargetChainsToRemove;
@@ -187,7 +144,7 @@ void UIKRetargeterController::CleanChainMapping(const bool bForceReinitializatio
 	if (IsValid(Asset->GetSourceIKRig()))
 	{
 		TArray<FName> SourceChainNames;
-		GetSourceChainNames(SourceChainNames);
+		GetChainNames(ERetargetSourceOrTarget::Source,SourceChainNames);
 	
 		// reset any sources that are no longer present to "None"
 		for (URetargetChainSettings* ChainMap : Asset->ChainSettings)
@@ -208,26 +165,27 @@ void UIKRetargeterController::CleanChainMapping(const bool bForceReinitializatio
 	}
 }
 
-void UIKRetargeterController::CleanPoseList(const bool bForceReinitialization)
+void UIKRetargeterController::CleanPoseLists(const bool bForceReinitialization) const
 {
-	// enforce the existence of a default pose
-	const bool HasDefaultPose = Asset->RetargetPoses.Contains(UIKRetargeter::GetDefaultPoseName());
-	if (!HasDefaultPose)
-	{
-		Asset->RetargetPoses.Emplace(UIKRetargeter::GetDefaultPoseName());
-	}
-	
-	// use default pose unless set to something else
-	if (Asset->CurrentRetargetPose == NAME_None)
+	CleanPoseList(ERetargetSourceOrTarget::Source);
+	CleanPoseList(ERetargetSourceOrTarget::Target);
+
+	if (bForceReinitialization)
 	{
-		Asset->CurrentRetargetPose = UIKRetargeter::GetDefaultPoseName();
+		BroadcastNeedsReinitialized();	
 	}
+}
+
+void UIKRetargeterController::CleanPoseList(const ERetargetSourceOrTarget& SourceOrTarget) const
+{
+	TMap<FName, FIKRetargetPose>& RetargetPoses = GetRetargetPoses(SourceOrTarget);
 
-	// remove all bone offsets that are no longer part of the target skeleton
-	if (const UIKRigDefinition* TargetIKRig = Asset->GetTargetIKRig())
+	// remove all bone offsets that are no longer part of the skeleton
+	const UIKRigDefinition* IKRig = GetIKRig(SourceOrTarget);
+	if (IKRig)
 	{
-		const TArray<FName> AllowedBoneNames = TargetIKRig->Skeleton.BoneNames;
-		for (TTuple<FName, FIKRetargetPose>& Pose : Asset->RetargetPoses)
+		const TArray<FName> AllowedBoneNames = IKRig->Skeleton.BoneNames;
+		for (TTuple<FName, FIKRetargetPose>& Pose : RetargetPoses)
 		{
 			// find bone offsets no longer in target skeleton
 			TArray<FName> BonesToRemove;
@@ -246,20 +204,15 @@ void UIKRetargeterController::CleanPoseList(const bool bForceReinitialization)
 			}
 
 			// sort the pose offset from leaf to root
-			Pose.Value.SortHierarchically(TargetIKRig->Skeleton);
+			Pose.Value.SortHierarchically(IKRig->Skeleton);
 		}
 	}
-
-	if (bForceReinitialization)
-	{
-		BroadcastNeedsReinitialized();	
-	}
 }
 
 void UIKRetargeterController::AutoMapChains() const
 {
 	TArray<FName> SourceChainNames;
-	GetSourceChainNames(SourceChainNames);
+	GetChainNames(ERetargetSourceOrTarget::Source, SourceChainNames);
 	
 	// auto-map any chains that have no value using a fuzzy string search
 	for (URetargetChainSettings* ChainMap : Asset->ChainSettings)
@@ -365,27 +318,36 @@ const TArray<TObjectPtr<URetargetChainSettings>>& UIKRetargeterController::GetCh
 	return Asset->ChainSettings;
 }
 
-void UIKRetargeterController::AddRetargetPose(FName NewPoseName, const FIKRetargetPose* ToDuplicate) const
+void UIKRetargeterController::AddRetargetPose(
+	const FName& NewPoseName,
+	const FIKRetargetPose* ToDuplicate,
+	const ERetargetSourceOrTarget& SourceOrTarget) const
 {
 	FScopedTransaction Transaction(LOCTEXT("AddRetargetPose", "Add Retarget Pose"));
 	Asset->Modify();
 	
-	NewPoseName = MakePoseNameUnique(NewPoseName.ToString());
-	FIKRetargetPose& NewPose = Asset->RetargetPoses.Add(NewPoseName);
+	const FName UniqueNewPoseName = MakePoseNameUnique(NewPoseName.ToString(), SourceOrTarget);
+	TMap<FName, FIKRetargetPose>& Poses = GetRetargetPoses(SourceOrTarget);
+	FIKRetargetPose& NewPose = Poses.Add(UniqueNewPoseName);
 	if (ToDuplicate)
 	{
 		NewPose.RootTranslationOffset = ToDuplicate->RootTranslationOffset;
 		NewPose.BoneRotationOffsets = ToDuplicate->BoneRotationOffsets;
 	}
-	Asset->CurrentRetargetPose = NewPoseName;
+	
+	FName& CurrentRetargetPoseName = SourceOrTarget == ERetargetSourceOrTarget::Source ? Asset->CurrentSourceRetargetPose : Asset->CurrentTargetRetargetPose;
+	CurrentRetargetPoseName = UniqueNewPoseName;
 
 	BroadcastNeedsReinitialized();
 }
 
-void UIKRetargeterController::RenameCurrentRetargetPose(FName NewPoseName) const
+void UIKRetargeterController::RenameCurrentRetargetPose(
+	const FName& NewPoseName,
+	const ERetargetSourceOrTarget& SourceOrTarget) const
 {
 	// do we already have a retarget pose with this name?
-	if (Asset->RetargetPoses.Contains(NewPoseName))
+	
+	if (GetRetargetPoses(SourceOrTarget).Contains(NewPoseName))
 	{
 		return;
 	}
@@ -394,26 +356,30 @@ void UIKRetargeterController::RenameCurrentRetargetPose(FName NewPoseName) const
 	Asset->Modify();
 
 	// replace key in the map
-	TMap<FName, FIKRetargetPose>& Poses = Asset->RetargetPoses;
-	const FIKRetargetPose CurrentPose = Poses[Asset->CurrentRetargetPose];
-	Poses.Remove(Asset->CurrentRetargetPose);
+	TMap<FName, FIKRetargetPose>& Poses = GetRetargetPoses(SourceOrTarget);
+	const FName CurrentPoseName = GetCurrentRetargetPoseName(SourceOrTarget);
+	const FIKRetargetPose CurrentPoseData = GetCurrentRetargetPose(SourceOrTarget);
+	Poses.Remove(CurrentPoseName);
 	Poses.Shrink();
-	Poses.Add(NewPoseName, CurrentPose);
+	Poses.Add(NewPoseName, CurrentPoseData);
 
 	// update current pose name
-	Asset->CurrentRetargetPose = NewPoseName;
+	SetCurrentRetargetPose(NewPoseName, SourceOrTarget);
 
 	BroadcastNeedsReinitialized();
 }
 
-void UIKRetargeterController::RemoveRetargetPose(FName PoseToRemove) const
+void UIKRetargeterController::RemoveRetargetPose(
+	const FName& PoseToRemove,
+	const ERetargetSourceOrTarget& SourceOrTarget) const
 {
 	if (PoseToRemove == Asset->GetDefaultPoseName())
 	{
 		return; // cannot remove default pose
 	}
 
-	if (!Asset->RetargetPoses.Contains(PoseToRemove))
+	TMap<FName, FIKRetargetPose>& Poses = GetRetargetPoses(SourceOrTarget);
+	if (!Poses.Contains(PoseToRemove))
 	{
 		return; // cannot remove pose that doesn't exist
 	}
@@ -421,25 +387,29 @@ void UIKRetargeterController::RemoveRetargetPose(FName PoseToRemove) const
 	FScopedTransaction Transaction(LOCTEXT("RemoveRetargetPose", "Remove Retarget Pose"));
 	Asset->Modify();
 
-	Asset->RetargetPoses.Remove(PoseToRemove);
+	Poses.Remove(PoseToRemove);
 
 	// did we remove the currently used pose?
-	if (Asset->CurrentRetargetPose == PoseToRemove)
+	if (GetCurrentRetargetPoseName(SourceOrTarget) == PoseToRemove)
 	{
-		Asset->CurrentRetargetPose = UIKRetargeter::GetDefaultPoseName();
+		SetCurrentRetargetPose(UIKRetargeter::GetDefaultPoseName(), SourceOrTarget);
 	}
 
 	BroadcastNeedsReinitialized();
 }
 
-void UIKRetargeterController::ResetRetargetPose(FName PoseToReset, const TArray<FName>& BonesToReset) const
+void UIKRetargeterController::ResetRetargetPose(
+	const FName& PoseToReset,
+	const TArray<FName>& BonesToReset,
+	const ERetargetSourceOrTarget& SourceOrTarget) const
 {
-	if (!Asset->RetargetPoses.Contains(PoseToReset))
+	TMap<FName, FIKRetargetPose>& Poses = GetRetargetPoses(SourceOrTarget);
+	if (!Poses.Contains(PoseToReset))
 	{
 		return; // cannot reset pose that doesn't exist
 	}
 	
-	FIKRetargetPose& PoseToEdit = Asset->RetargetPoses[PoseToReset];
+	FIKRetargetPose& PoseToEdit = Poses[PoseToReset];
 	
 	if (BonesToReset.IsEmpty())
 	{
@@ -454,7 +424,7 @@ void UIKRetargeterController::ResetRetargetPose(FName PoseToReset, const TArray<
 		FScopedTransaction Transaction(LOCTEXT("ResetRetargetBonePose", "Reset Bone Pose"));
 		Asset->Modify();
 		
-		const FName RootBoneName = GetTargetRootBone();
+		const FName RootBoneName = GetRetargetRootBone(SourceOrTarget);
 		for (const FName& BoneToReset : BonesToReset)
 		{
 			if (PoseToEdit.BoneRotationOffsets.Contains(BoneToReset))
@@ -472,41 +442,50 @@ void UIKRetargeterController::ResetRetargetPose(FName PoseToReset, const TArray<
 	BroadcastNeedsReinitialized();
 }
 
-FName UIKRetargeterController::GetCurrentRetargetPoseName() const
+FName UIKRetargeterController::GetCurrentRetargetPoseName(const ERetargetSourceOrTarget& SourceOrTarget) const
 {
-	return GetAsset()->CurrentRetargetPose;
+	return SourceOrTarget == ERetargetSourceOrTarget::Source ? GetAsset()->CurrentSourceRetargetPose : GetAsset()->CurrentTargetRetargetPose;
 }
 
-void UIKRetargeterController::SetCurrentRetargetPose(FName CurrentPose) const
+void UIKRetargeterController::SetCurrentRetargetPose(FName NewCurrentPose, const ERetargetSourceOrTarget& SourceOrTarget) const
 {
-	check(Asset->RetargetPoses.Contains(CurrentPose));
+	const TMap<FName, FIKRetargetPose>& Poses = GetRetargetPoses(SourceOrTarget);
+	check(Poses.Contains(NewCurrentPose));
 
 	FScopedTransaction Transaction(LOCTEXT("SetCurrentPose", "Set Current Pose"));
 	Asset->Modify();
-	Asset->CurrentRetargetPose = CurrentPose;
+	FName& CurrentPose = SourceOrTarget == ERetargetSourceOrTarget::Source ? Asset->CurrentSourceRetargetPose : Asset->CurrentTargetRetargetPose;
+	CurrentPose = NewCurrentPose;
 	
 	BroadcastNeedsReinitialized();
 }
 
-const TMap<FName, FIKRetargetPose>& UIKRetargeterController::GetRetargetPoses()
+TMap<FName, FIKRetargetPose>& UIKRetargeterController::GetRetargetPoses(const ERetargetSourceOrTarget& SourceOrTarget) const
 {
-	return GetAsset()->RetargetPoses;
+	return SourceOrTarget == ERetargetSourceOrTarget::Source ? GetAsset()->SourceRetargetPoses : GetAsset()->TargetRetargetPoses;
 }
 
-const FIKRetargetPose& UIKRetargeterController::GetCurrentRetargetPose() const
+FIKRetargetPose& UIKRetargeterController::GetCurrentRetargetPose(const ERetargetSourceOrTarget& SourceOrTarget) const
 {
-	return GetAsset()->RetargetPoses[GetCurrentRetargetPoseName()];
+	return GetRetargetPoses(SourceOrTarget)[GetCurrentRetargetPoseName(SourceOrTarget)];
 }
 
-void UIKRetargeterController::SetRotationOffsetForRetargetPoseBone(FName BoneName, const FQuat& RotationOffset) const
+void UIKRetargeterController::SetRotationOffsetForRetargetPoseBone(
+	const FName& BoneName,
+	const FQuat& RotationOffset,
+	const ERetargetSourceOrTarget& SourceOrTarget) const
 {
-	const FIKRigSkeleton& Skeleton = Asset->GetTargetIKRig()->Skeleton;
-	Asset->RetargetPoses[Asset->CurrentRetargetPose].SetBoneRotationOffset(BoneName, RotationOffset, Skeleton);
+	const UIKRigDefinition* IKRig = SourceOrTarget == ERetargetSourceOrTarget::Source ? GetAsset()->GetSourceIKRig() : GetAsset()->GetTargetIKRig();
+	FIKRetargetPose& Pose = GetCurrentRetargetPose(SourceOrTarget);
+	Pose.SetDeltaRotationForBone(BoneName, RotationOffset);
+	Pose.SortHierarchically(IKRig->Skeleton);
 }
 
-FQuat UIKRetargeterController::GetRotationOffsetForRetargetPoseBone(FName BoneName) const
+FQuat UIKRetargeterController::GetRotationOffsetForRetargetPoseBone(
+	const FName& BoneName,
+	const ERetargetSourceOrTarget& SourceOrTarget) const
 {
-	TMap<FName, FQuat>& BoneOffsets = Asset->RetargetPoses[Asset->CurrentRetargetPose].BoneRotationOffsets;
+	TMap<FName, FQuat>& BoneOffsets = GetCurrentRetargetPose(SourceOrTarget).BoneRotationOffsets;
 	if (!BoneOffsets.Contains(BoneName))
 	{
 		return FQuat::Identity;
@@ -515,26 +494,31 @@ FQuat UIKRetargeterController::GetRotationOffsetForRetargetPoseBone(FName BoneNa
 	return BoneOffsets[BoneName];
 }
 
-void UIKRetargeterController::SetTranslationOffsetOnRetargetRootBone(FVector TranslationOffset) const
+void UIKRetargeterController::SetTranslationOffsetOnRetargetRootBone(
+	const FVector& TranslationOffset,
+	const ERetargetSourceOrTarget& SourceOrTarget) const
 {
-	Asset->RetargetPoses[Asset->CurrentRetargetPose].SetRootTranslationDelta(TranslationOffset);
+	GetCurrentRetargetPose(SourceOrTarget).SetRootTranslationDelta(TranslationOffset);
 }
 
-const FVector& UIKRetargeterController::GetTranslationOffsetOnRetargetRootBone() const
+const FVector& UIKRetargeterController::GetTranslationOffsetOnRetargetRootBone(const ERetargetSourceOrTarget& SourceOrTarget) const
 {
-	return Asset->RetargetPoses[Asset->CurrentRetargetPose].RootTranslationOffset;
+	return GetCurrentRetargetPose(SourceOrTarget).RootTranslationOffset;
 }
 
-void UIKRetargeterController::AddTranslationOffsetToRetargetRootBone(FVector TranslationOffset) const
+void UIKRetargeterController::AddTranslationOffsetToRetargetRootBone(
+	const FVector& TranslationOffset,
+	const ERetargetSourceOrTarget& SourceOrTarget) const
 {
-	Asset->RetargetPoses[Asset->CurrentRetargetPose].AddToRootTranslationDelta(TranslationOffset);
+	GetCurrentRetargetPose(SourceOrTarget).AddToRootTranslationDelta(TranslationOffset);
 }
 
-FName UIKRetargeterController::MakePoseNameUnique(FString PoseName) const
+FName UIKRetargeterController::MakePoseNameUnique(const FString& PoseName, const ERetargetSourceOrTarget& SourceOrTarget) const
 {
 	FString UniqueName = PoseName;
 	int32 Suffix = 1;
-	while (Asset->RetargetPoses.Contains(FName(UniqueName)))
+	const TMap<FName, FIKRetargetPose>& Poses = GetRetargetPoses(SourceOrTarget);
+	while (Poses.Contains(FName(UniqueName)))
 	{
 		UniqueName = PoseName + "_" + FString::FromInt(Suffix);
 		++Suffix;
diff --git a/Engine/Plugins/Animation/IKRig/Source/IKRigEditor/Private/RetargetEditor/SIKRetargetAssetBrowser.cpp b/Engine/Plugins/Animation/IKRig/Source/IKRigEditor/Private/RetargetEditor/SIKRetargetAssetBrowser.cpp
--- a/Engine/Plugins/Animation/IKRig/Source/IKRigEditor/Private/RetargetEditor/SIKRetargetAssetBrowser.cpp
+++ b/Engine/Plugins/Animation/IKRig/Source/IKRigEditor/Private/RetargetEditor/SIKRetargetAssetBrowser.cpp
@@ -25,7 +25,7 @@ void SIKRetargetAssetBrowser::Construct(
 	TSharedRef<FIKRetargetEditorController> InEditorController)
 {
 	EditorController = InEditorController;
-	EditorController.Pin()->AssetBrowserView = SharedThis(this);
+	EditorController.Pin()->SetAssetBrowserView(SharedThis(this));
 	
 	ChildSlot
     [
@@ -109,8 +109,8 @@ FReply SIKRetargetAssetBrowser::OnExportButtonClicked()
 	// assemble the data for the assets we want to batch duplicate/retarget
 	FIKRetargetBatchOperationContext BatchContext = Dialog.Get().BatchContext;
 	BatchContext.NameRule.FolderPath = PrevBatchOutputPath;
-	BatchContext.SourceMesh = Controller->GetSourceSkeletalMesh();
-	BatchContext.TargetMesh = Controller->GetTargetSkeletalMesh();
+	BatchContext.SourceMesh = Controller->GetSkeletalMesh(ERetargetSourceOrTarget::Source);
+	BatchContext.TargetMesh = Controller->GetSkeletalMesh(ERetargetSourceOrTarget::Target);
 	BatchContext.IKRetargetAsset = Controller->AssetController->GetAsset();
 	BatchContext.bRemapReferencedAssets = false;
 
@@ -186,7 +186,7 @@ bool SIKRetargetAssetBrowser::OnShouldFilterAsset(const struct FAssetData& Asset
 	}
 
 	// get source skeleton
-	const USkeleton* DesiredSkeleton = Controller->GetSourceSkeleton();
+	const USkeleton* DesiredSkeleton = Controller->GetSkeleton(ERetargetSourceOrTarget::Source);
 	if (!DesiredSkeleton)
 	{
 		return true;
diff --git a/Engine/Plugins/Animation/IKRig/Source/IKRigEditor/Private/RetargetEditor/SIKRetargetChainMapList.cpp b/Engine/Plugins/Animation/IKRig/Source/IKRigEditor/Private/RetargetEditor/SIKRetargetChainMapList.cpp
--- a/Engine/Plugins/Animation/IKRig/Source/IKRigEditor/Private/RetargetEditor/SIKRetargetChainMapList.cpp
+++ b/Engine/Plugins/Animation/IKRig/Source/IKRigEditor/Private/RetargetEditor/SIKRetargetChainMapList.cpp
@@ -126,7 +126,7 @@ void SIKRetargetChainMapList::Construct(
 	TSharedRef<FIKRetargetEditorController> InEditorController)
 {
 	EditorController = InEditorController;
-	EditorController.Pin()->ChainsView = SharedThis(this);
+	EditorController.Pin()->SetChainsView(SharedThis(this));
 	
 	ChildSlot
     [
@@ -259,7 +259,7 @@ FText SIKRetargetChainMapList::GetSourceRootBone() const
 		return FText::FromName(NAME_None); 
 	}
 	
-	return FText::FromName(RetargeterController->GetSourceRootBone());
+	return FText::FromName(RetargeterController->GetRetargetRootBone(ERetargetSourceOrTarget::Source));
 }
 
 FText SIKRetargetChainMapList::GetTargetRootBone() const
@@ -270,7 +270,7 @@ FText SIKRetargetChainMapList::GetTargetRootBone() const
 		return FText::FromName(NAME_None); 
 	}
 	
-	return FText::FromName(RetargeterController->GetTargetRootBone());
+	return FText::FromName(RetargeterController->GetRetargetRootBone(ERetargetSourceOrTarget::Target));
 }
 
 bool SIKRetargetChainMapList::IsChainMapEnabled() const
@@ -358,12 +358,12 @@ void SIKRetargetChainMapList::OnSelectionChanged()
 	if (SelectedChainSettings.IsEmpty())
 	{
 		// show asset settings in the details view
-		Controller->DetailsView->SetObject(Controller->AssetController->GetAsset());
+		Controller->SetDetailsObject(Controller->AssetController->GetAsset());
 	}
 	else
 	{
 		// show chain settings in the details view
-		Controller->DetailsView->SetObjects(SelectedChainSettings);
+		Controller->SetDetailsObjects(SelectedChainSettings);
 	}
 }
 
@@ -406,7 +406,7 @@ FReply SIKRetargetChainMapList::OnEditSettingsButtonClicked() const
 		return FReply::Unhandled();
 	}
 	
-	Controller->DetailsView->SetObject(RetargeterController->GetAsset()->GetRetargetRootSettings());
+	Controller->SetDetailsObject(RetargeterController->GetAsset()->GetRetargetRootSettings());
 	return FReply::Handled();
 }
 
diff --git a/Engine/Plugins/Animation/IKRig/Source/IKRigEditor/Private/RetargetEditor/SIKRetargetHierarchy.cpp b/Engine/Plugins/Animation/IKRig/Source/IKRigEditor/Private/RetargetEditor/SIKRetargetHierarchy.cpp
--- a/Engine/Plugins/Animation/IKRig/Source/IKRigEditor/Private/RetargetEditor/SIKRetargetHierarchy.cpp
+++ b/Engine/Plugins/Animation/IKRig/Source/IKRigEditor/Private/RetargetEditor/SIKRetargetHierarchy.cpp
@@ -8,6 +8,8 @@
 #include "Retargeter/IKRetargetProcessor.h"
 #include "RigEditor/IKRigEditorStyle.h"
 #include "Widgets/Layout/SScrollBox.h"
+#include "Widgets/Layout/SSeparator.h"
+#include "Engine/SkeletalMesh.h"
 
 #define LOCTEXT_NAMESPACE "SIKRetargetHierarchy"
 
@@ -25,7 +27,7 @@ FIKRetargetHierarchyElement::FIKRetargetHierarchyElement(
 TSharedRef<SWidget> SIKRetargetHierarchyRow::GenerateWidgetForColumn(const FName& ColumnName)
 {
 	FName BoneName = WeakTreeElement.Pin()->Name;
-	EIKRetargetSkeletonMode CurrentSkeleton = EditorController.Pin()->GetSkeletonMode();
+	ERetargetSourceOrTarget CurrentSkeleton = EditorController.Pin()->GetSourceOrTarget();
 	const bool bIsBoneRetargeted = EditorController.Pin()->IsBoneRetargeted(BoneName, CurrentSkeleton);
 	const FText ChainName = FText::FromName(EditorController.Pin()->GetChainNameFromBone(BoneName, CurrentSkeleton));
 
@@ -115,7 +117,7 @@ void SIKRetargetHierarchy::Construct(
 	TSharedRef<FIKRetargetEditorController> InEditorController)
 {
 	EditorController = InEditorController;
-	EditorController.Pin()->HierarchyView = SharedThis(this);
+	EditorController.Pin()->SetHierarchyView(SharedThis(this));
 	CommandList = MakeShared<FUICommandList>();
 	
 	ChildSlot
@@ -130,27 +132,41 @@ void SIKRetargetHierarchy::Construct(
 			.Padding(2.0f)
 			.HAlign(HAlign_Center)
 			[
-				SNew(SSegmentedControl<EIKRetargetSkeletonMode>)
+				SNew(SSegmentedControl<ERetargetSourceOrTarget>)
 				.Value_Lambda([this]()
 				{
-					return EditorController.Pin()->GetSkeletonMode();
+					return EditorController.Pin()->GetSourceOrTarget();
 				})
-				.OnValueChanged_Lambda([this](EIKRetargetSkeletonMode Mode)
+				.OnValueChanged_Lambda([this](ERetargetSourceOrTarget Mode)
 				{
-					EditorController.Pin()->SetSkeletonMode(Mode);
+					EditorController.Pin()->SetSourceOrTargetMode(Mode);
 				})
-				+SSegmentedControl<EIKRetargetSkeletonMode>::Slot(EIKRetargetSkeletonMode::Source)
+				+SSegmentedControl<ERetargetSourceOrTarget>::Slot(ERetargetSourceOrTarget::Source)
 				.Text(LOCTEXT("SourceSkeleton", "Source"))
-				+ SSegmentedControl<EIKRetargetSkeletonMode>::Slot(EIKRetargetSkeletonMode::Target)
+				+ SSegmentedControl<ERetargetSourceOrTarget>::Slot(ERetargetSourceOrTarget::Target)
 				.Text(LOCTEXT("TargetSkeleton", "Target"))
 			]
 		]
+
+		+SVerticalBox::Slot()
+		.Padding(2.0f)
+		.AutoHeight()
+		[
+			SNew(SSeparator)
+		]
 		
 		+SVerticalBox::Slot()
 		.Padding(2.0f)
 		.AutoHeight()
 		[
-			SNew(SIKRetargetPoseEditor, InEditorController)
+			SAssignNew(EditPoseView, SIKRetargetPoseEditor, InEditorController)
+		]
+
+		+SVerticalBox::Slot()
+		.Padding(2.0f)
+		.AutoHeight()
+		[
+			SNew(SSeparator)
 		]
 
 		+SVerticalBox::Slot()
@@ -253,19 +269,25 @@ void SIKRetargetHierarchy::RefreshTreeView(bool IsInitialSetup)
 	AllElements.Reset();
 
 	// validate we have a skeleton to load
-	const UIKRetargetProcessor* Processor = Controller->GetRetargetProcessor();
-	if (!(Processor && Processor->IsInitialized()))
+	const ERetargetSourceOrTarget SourceOrTarget = Controller->GetSourceOrTarget();
+	const USkeletalMesh* Mesh = Controller->GetSkeletalMesh(SourceOrTarget);
+	if (!Mesh)
 	{
 		TreeView->RequestTreeRefresh();
 		return;
 	}
 
 	// get the skeleton that is currently being viewed in the editor
-	const FTargetSkeleton& TargetSkeleton = Processor->GetTargetSkeleton();
-	const FRetargetSkeleton& SourceSkeleton = Processor->GetSourceSkeleton();
-	const bool bViewTarget = Controller->GetSkeletonMode() == EIKRetargetSkeletonMode::Target;
-	const TArray<FName>& BoneNames = bViewTarget ? TargetSkeleton.BoneNames : SourceSkeleton.BoneNames;
-	const TArray<int32>& ParentIndices = bViewTarget ? TargetSkeleton.ParentIndices : SourceSkeleton.ParentIndices;
+	const FReferenceSkeleton& Skeleton = Mesh->GetRefSkeleton();
+	TArray<FName> BoneNames;
+	TArray<int32> ParentIndices;
+	BoneNames.AddUninitialized(Skeleton.GetNum());
+	ParentIndices.AddUninitialized(Skeleton.GetNum());
+	for (int32 BoneIndex=0; BoneIndex<Skeleton.GetNum(); ++BoneIndex)
+	{
+		BoneNames[BoneIndex] = Skeleton.GetBoneName(BoneIndex);
+		ParentIndices[BoneIndex] = Skeleton.GetParentIndex(BoneIndex);
+	}
 	
 	// record bone element indices
 	TMap<FName, int32> BoneTreeElementIndices;
diff --git a/Engine/Plugins/Animation/IKRig/Source/IKRigEditor/Private/RetargetEditor/SIKRetargetPoseEditor.cpp b/Engine/Plugins/Animation/IKRig/Source/IKRigEditor/Private/RetargetEditor/SIKRetargetPoseEditor.cpp
--- a/Engine/Plugins/Animation/IKRig/Source/IKRigEditor/Private/RetargetEditor/SIKRetargetPoseEditor.cpp
+++ b/Engine/Plugins/Animation/IKRig/Source/IKRigEditor/Private/RetargetEditor/SIKRetargetPoseEditor.cpp
@@ -6,6 +6,8 @@
 #include "RetargetEditor/IKRetargetEditor.h"
 #include "RetargetEditor/IKRetargetEditorController.h"
 
+#include "Widgets/Layout/SSeparator.h"
+
 #define LOCTEXT_NAMESPACE "SIKRetargetPoseEditor"
 
 void SIKRetargetPoseEditor::Construct(
@@ -13,54 +15,126 @@ void SIKRetargetPoseEditor::Construct(
 	TSharedRef<FIKRetargetEditorController> InEditorController)
 {
 	EditorController = InEditorController;
-	
-	// fill list of pose names
-	PoseNames.Reset();
-	for (const TTuple<FName, FIKRetargetPose>& Pose : EditorController.Pin()->AssetController->GetRetargetPoses())
-	{
-		PoseNames.Add(MakeShareable(new FName(Pose.Key)));
-	}
+
+	// the editor controller
+	FIKRetargetEditorController* Controller = EditorController.Pin().Get();
 
 	// the commands for the menus
-	TSharedPtr<FUICommandList> Commands = EditorController.Pin()->Editor.Pin()->GetToolkitCommands();
+	TSharedPtr<FUICommandList> Commands = Controller->Editor.Pin()->GetToolkitCommands();
+	
+	// build the list of poses
+	Refresh();
 	
 	ChildSlot
 	[
 		SNew(SVerticalBox)
 
+		// poses
 		+SVerticalBox::Slot()
 		.Padding(2.0f)
 		.AutoHeight()
 		[
-			SNew(SHorizontalBox)
-
-			// add pose selection label
-			+SHorizontalBox::Slot()
-			.Padding(2.0f)
-			.HAlign(HAlign_Right)
-			.VAlign(VAlign_Center)
+			SNew(SBox)
+			.HAlign(HAlign_Center)
 			[
-				SNew(STextBlock)
-				.Text(LOCTEXT("CurrentPose", "Current Retarget Pose:"))
-			]
+				SNew(SHorizontalBox)
+				+SHorizontalBox::Slot()
+				.Padding(4,0)
+				[
+					SNew(SVerticalBox)
 
-			// add pose selection combo box
-			+SHorizontalBox::Slot()
-			.Padding(2.0f)
-			[
-				SNew(SComboBox<TSharedPtr<FName>>)
-				.OptionsSource(&PoseNames)
-				.OnGenerateWidget_Lambda([](TSharedPtr<FName> InItem)
-				{
-					return SNew(STextBlock).Text(FText::FromName(*InItem.Get()));
-				})
-				.OnSelectionChanged(EditorController.Pin().Get(), &FIKRetargetEditorController::OnPoseSelected)
+					// add pose selection label and combobox
+					+SVerticalBox::Slot()
+					.VAlign(VAlign_Center)
+					.Padding(4,0)
+					[
+						SNew(STextBlock).MinDesiredWidth(180).Text(LOCTEXT("CurrentPose", "Current Retarget Pose:"))
+					]
+						
+					// pose selection combobox
+					+SVerticalBox::Slot()
+					[
+						SNew(SComboBox<TSharedPtr<FName>>)
+						.OptionsSource(&PoseNames)
+						.OnGenerateWidget_Lambda([](TSharedPtr<FName> InItem)
+						{
+							return SNew(STextBlock).Text(FText::FromName(*InItem.Get()));
+						})
+						.OnSelectionChanged(Controller, &FIKRetargetEditorController::OnPoseSelected)
+						.Content()
+						[
+							SNew(STextBlock).Text(Controller, &FIKRetargetEditorController::GetCurrentPoseName)
+						]
+					]
+				]
+
+				+SHorizontalBox::Slot()
+				.Padding(4,0)
 				[
-					SNew(STextBlock).Text(EditorController.Pin().Get(), &FIKRetargetEditorController::GetCurrentPoseName)
+					SNew(SVerticalBox)
+
+					// show retarget pose / run retarget
+					+SVerticalBox::Slot()
+					.Padding(0,4)
+					[
+						SNew(SButton)
+						.OnClicked(Controller, &FIKRetargetEditorController::HandleShowRetargetPose)
+						.IsEnabled_Lambda([Controller]
+						{
+							return Controller->GetRetargeterMode() != ERetargeterOutputMode::EditRetargetPose;
+						})
+						.Content()
+						[
+							SNew(SHorizontalBox)
+							+SHorizontalBox::Slot()
+							.AutoWidth()
+							.HAlign(HAlign_Left)
+							[
+								SNew(SImage)
+								.Image_Lambda([Controller]() -> const FSlateBrush*
+								{
+									if (Controller->GetRetargeterMode() == ERetargeterOutputMode::RunRetarget)
+									{
+										return FAppStyle::GetBrush("GenericPause");
+									}
+									return FAppStyle::GetBrush("GenericPlay");
+								})
+								.ColorAndOpacity(FSlateColor::UseForeground())
+							]
+				
+							+SHorizontalBox::Slot()
+							.HAlign(HAlign_Right)
+							.Padding(4,0)
+							[
+								SNew(STextBlock)
+								.ToolTipText(LOCTEXT("BlendPoseToolTip", "Pause playback and show the retarget pose. Use slider below to blend between reference pose (0) and retarget pose (1)."))
+								.Text_Lambda([Controller]() -> const FText
+								{
+									if (Controller->GetRetargeterMode() == ERetargeterOutputMode::RunRetarget)
+									{
+										return LOCTEXT("ShowPoseText", "Show Retarget Pose");
+									}
+									return LOCTEXT("RunRetargetText", "Run Retargeter");
+								})
+							]
+						]
+					]
+
+					// blend retarget pose slider
+					+SVerticalBox::Slot()
+					[
+						SNew(SSpinBox<float>)
+						.Font(FAppStyle::Get().GetFontStyle("PropertyWindow.NormalFont"))
+						.MinValue(0.0f)
+						.MaxValue(1.0f)
+						.Value(Controller, &FIKRetargetEditorController::GetRetargetPoseAmount)
+						.OnValueChanged(Controller, &FIKRetargetEditorController::SetRetargetPoseAmount)
+					]
 				]
-			]
+			]	
 		]
 
+		// pose editing toolbar
 		+SVerticalBox::Slot()
 		.Padding(2.0f)
 		.AutoHeight()
@@ -76,6 +150,20 @@ void SIKRetargetPoseEditor::Construct(
 	];
 }
 
+void SIKRetargetPoseEditor::Refresh()
+{
+	// get the retarget poses from the editor controller
+	const FIKRetargetEditorController* Controller = EditorController.Pin().Get();
+	const TMap<FName, FIKRetargetPose>& RetargetPoses = Controller->AssetController->GetRetargetPoses(Controller->GetSourceOrTarget());
+
+	// fill list of pose names
+	PoseNames.Reset();
+	for (const TTuple<FName, FIKRetargetPose>& Pose : RetargetPoses)
+	{
+		PoseNames.Add(MakeShareable(new FName(Pose.Key)));
+	}
+}
+
 TSharedRef<SWidget>  SIKRetargetPoseEditor::MakeToolbar(TSharedPtr<FUICommandList> Commands)
 {
 	FToolBarBuilder ToolbarBuilder(Commands, FMultiBoxCustomization::None);
diff --git a/Engine/Plugins/Animation/IKRig/Source/IKRigEditor/Private/RetargetEditor/SRetargetAnimAssetsWindow.cpp b/Engine/Plugins/Animation/IKRig/Source/IKRigEditor/Private/RetargetEditor/SRetargetAnimAssetsWindow.cpp
--- a/Engine/Plugins/Animation/IKRig/Source/IKRigEditor/Private/RetargetEditor/SRetargetAnimAssetsWindow.cpp
+++ b/Engine/Plugins/Animation/IKRig/Source/IKRigEditor/Private/RetargetEditor/SRetargetAnimAssetsWindow.cpp
@@ -664,8 +664,8 @@ void SRetargetAnimAssetsWindow::RetargeterAssigned(const FAssetData& InAssetData
 	UIKRetargeter* InRetargeter = Cast<UIKRetargeter>(InAssetData.GetAsset());
 	BatchContext.IKRetargetAsset = InRetargeter;
 	const UIKRetargeterController* Controller = UIKRetargeterController::GetController(InRetargeter);
-	SourceMeshAssigned(FAssetData(Controller->GetSourcePreviewMesh()));
-	TargetMeshAssigned(FAssetData(Controller->GetTargetPreviewMesh()));
+	SourceMeshAssigned(FAssetData(Controller->GetPreviewMesh(ERetargetSourceOrTarget::Source)));
+	TargetMeshAssigned(FAssetData(Controller->GetPreviewMesh(ERetargetSourceOrTarget::Target)));
 }
 
 ECheckBoxState SRetargetAnimAssetsWindow::IsRemappingReferencedAssets() const
diff --git a/Engine/Plugins/Animation/IKRig/Source/IKRigEditor/Private/RigEditor/IKRigOutputLogTabSummoner.cpp b/Engine/Plugins/Animation/IKRig/Source/IKRigEditor/Private/RigEditor/IKRigOutputLogTabSummoner.cpp
--- a/Engine/Plugins/Animation/IKRig/Source/IKRigEditor/Private/RigEditor/IKRigOutputLogTabSummoner.cpp
+++ b/Engine/Plugins/Animation/IKRig/Source/IKRigEditor/Private/RigEditor/IKRigOutputLogTabSummoner.cpp
@@ -35,7 +35,7 @@ TSharedRef<SWidget> FIKRigOutputLogTabSummoner::CreateTabBody(const FWorkflowTab
 	const TSharedRef<FIKRigEditorController>& Controller = IKRigEditor.Pin()->GetController();
 	const FName LogName = Controller->GetIKRigProcessor()->Log.GetLogTarget();
 	
-	return SNew(SIKRigOutputLog, LogName, Controller->OutputLogView);
+	return SNew(SIKRigOutputLog, LogName);
 }
 
 #undef LOCTEXT_NAMESPACE 
diff --git a/Engine/Plugins/Animation/IKRig/Source/IKRigEditor/Private/RigEditor/SIKRigOutputLog.cpp b/Engine/Plugins/Animation/IKRig/Source/IKRigEditor/Private/RigEditor/SIKRigOutputLog.cpp
--- a/Engine/Plugins/Animation/IKRig/Source/IKRigEditor/Private/RigEditor/SIKRigOutputLog.cpp
+++ b/Engine/Plugins/Animation/IKRig/Source/IKRigEditor/Private/RigEditor/SIKRigOutputLog.cpp
@@ -8,11 +8,8 @@
 
 void SIKRigOutputLog::Construct(
 	const FArguments& InArgs,
-	const FName& InLogName,
-	TSharedPtr<SIKRigOutputLog>& OutLog)
+	const FName& InLogName)
 {
-	OutLog = SharedThis(this);
-
 	FMessageLogModule& MessageLogModule = FModuleManager::LoadModuleChecked<FMessageLogModule>("MessageLog");
 
 	if (MessageLogModule.IsRegisteredLogListing(InLogName))
diff --git a/Engine/Plugins/Animation/IKRig/Source/IKRigEditor/Public/IKRigEditor.h b/Engine/Plugins/Animation/IKRig/Source/IKRigEditor/Public/IKRigEditor.h
--- a/Engine/Plugins/Animation/IKRig/Source/IKRigEditor/Public/IKRigEditor.h
+++ b/Engine/Plugins/Animation/IKRig/Source/IKRigEditor/Public/IKRigEditor.h
@@ -15,7 +15,6 @@ private:
 	TSharedPtr<class FAssetTypeActions_AnimationAssetRetarget> RetargetAnimationAssetAction;
 	TSharedPtr<class FAssetTypeActions_IKRigDefinition> IKRigDefinitionAssetAction;
 	TSharedPtr<class FAssetTypeActions_IKRetargeter> IKRetargeterAssetAction;
-	TSharedPtr<class FAssetTypeActions_RetargetPose> RetargetPoseAssetAction;
 
 	TArray<FName> ClassesToUnregisterOnShutdown;
 };
diff --git a/Engine/Plugins/Animation/IKRig/Source/IKRigEditor/Public/RetargetEditor/AssetTypeActions_RetargetPose.h b/Engine/Plugins/Animation/IKRig/Source/IKRigEditor/Public/RetargetEditor/AssetTypeActions_RetargetPose.h
deleted file mode 100644
diff --git a/Engine/Plugins/Animation/IKRig/Source/IKRigEditor/Public/RetargetEditor/IKRetargetAnimInstance.h b/Engine/Plugins/Animation/IKRig/Source/IKRigEditor/Public/RetargetEditor/IKRetargetAnimInstance.h
--- a/Engine/Plugins/Animation/IKRig/Source/IKRigEditor/Public/RetargetEditor/IKRetargetAnimInstance.h
+++ b/Engine/Plugins/Animation/IKRig/Source/IKRigEditor/Public/RetargetEditor/IKRetargetAnimInstance.h
@@ -3,12 +3,47 @@
 #pragma once
 
 #include "AnimPreviewInstance.h"
+#include "Animation/AnimNodeBase.h"
 #include "AnimNodes/AnimNode_RetargetPoseFromMesh.h"
 
 #include "IKRetargetAnimInstance.generated.h"
 
+enum class ERetargeterOutputMode : uint8;
 class UIKRetargeter;
 
+// a node just to preview a retarget pose
+USTRUCT(BlueprintInternalUseOnly)
+struct FAnimNode_PreviewRetargetPose : public FAnimNode_Base
+{
+	GENERATED_BODY()
+
+	FPoseLink InputPose;
+
+	/** Retarget asset holding the pose to preview.*/
+	TObjectPtr<UIKRetargeter> IKRetargeterAsset = nullptr;
+
+	/** Whether to preview the source or the target */
+	ERetargetSourceOrTarget SourceOrTarget;
+
+	/** Amount to blend the retarget pose with the reference pose when shoing the retarget pose */
+	float RetargetPoseBlend = 1.0f;
+
+	// FAnimNode_Base interface
+	virtual void CacheBones_AnyThread(const FAnimationCacheBonesContext& Context) override;
+	virtual void Evaluate_AnyThread(FPoseContext& Output) override;
+	// End of FAnimNode_Base interface
+
+	const TArray<FTransform>& GetLocalRetargetPose() const { return RetargetLocalPose; };
+	const TArray<FTransform>& GetGlobalRetargetPose() const { return RetargetGlobalPose; };
+
+private:
+	// mapping from required bone indices (from LODs) to actual bone indices within the skeletal mesh 
+	TArray<int32> RequiredBoneToMeshBoneMap;
+
+	TArray<FTransform> RetargetLocalPose;
+	TArray<FTransform> RetargetGlobalPose;
+};
+
 UCLASS(transient, NotBlueprintable)
 class UIKRetargetAnimInstance : public UAnimPreviewInstance
 {
@@ -16,15 +51,26 @@ class UIKRetargetAnimInstance : public UAnimPreviewInstance
 
 public:
 
-	void SetRetargetAssetAndSourceComponent(
+	void SetRetargetMode(const ERetargeterOutputMode& OutputMode);
+
+	void SetRetargetPoseBlend(const float& RetargetPoseBlend);
+
+	void ConfigureAnimInstance(
+		const ERetargetSourceOrTarget& SourceOrTarget,
 		UIKRetargeter* InIKRetargetAsset,
 		TWeakObjectPtr<USkeletalMeshComponent> InSourceMeshComponent);
 
 	UIKRetargetProcessor* GetRetargetProcessor() const;
 
+	const TArray<FTransform>& GetLocalRetargetPose() const { return PreviewPoseNode.GetLocalRetargetPose(); };
+	const TArray<FTransform>& GetGlobalRetargetPose() const { return PreviewPoseNode.GetGlobalRetargetPose(); };
+
 protected:
 	virtual FAnimInstanceProxy* CreateAnimInstanceProxy() override;
 
 	UPROPERTY(Transient)
-	FAnimNode_RetargetPoseFromMesh IKRetargeterNode;
+	FAnimNode_PreviewRetargetPose PreviewPoseNode;
+	
+	UPROPERTY(Transient)
+	FAnimNode_RetargetPoseFromMesh RetargetNode;
 };
diff --git a/Engine/Plugins/Animation/IKRig/Source/IKRigEditor/Public/RetargetEditor/IKRetargetAnimInstanceProxy.h b/Engine/Plugins/Animation/IKRig/Source/IKRigEditor/Public/RetargetEditor/IKRetargetAnimInstanceProxy.h
--- a/Engine/Plugins/Animation/IKRig/Source/IKRigEditor/Public/RetargetEditor/IKRetargetAnimInstanceProxy.h
+++ b/Engine/Plugins/Animation/IKRig/Source/IKRigEditor/Public/RetargetEditor/IKRetargetAnimInstanceProxy.h
@@ -3,11 +3,16 @@
 #pragma once
 
 #include "AnimPreviewInstance.h"
+#include "IKRetargetAnimInstance.h"
 #include "Animation/AnimNode_LinkedInputPose.h"
-#include "AnimNodes/AnimNode_RetargetPoseFromMesh.h"
 
 #include "IKRetargetAnimInstanceProxy.generated.h"
 
+class UIKRetargeter;
+enum class ERetargetSourceOrTarget : uint8;
+struct FAnimNode_PreviewRetargetPose;
+struct FAnimNode_RetargetPoseFromMesh;
+
 /** Proxy override for this UAnimInstance-derived class */
 USTRUCT()
 struct FIKRetargetAnimInstanceProxy : public FAnimPreviewInstanceProxy
@@ -17,11 +22,15 @@ struct FIKRetargetAnimInstanceProxy : public FAnimPreviewInstanceProxy
 public:
 	
 	FIKRetargetAnimInstanceProxy() = default;
-	FIKRetargetAnimInstanceProxy(UAnimInstance* InAnimInstance, FAnimNode_RetargetPoseFromMesh* IKRetargeterNode);
+	FIKRetargetAnimInstanceProxy(
+		UAnimInstance* InAnimInstance,
+		FAnimNode_PreviewRetargetPose* InPreviewPoseNode,
+		FAnimNode_RetargetPoseFromMesh* InRetargetNode);
 	virtual ~FIKRetargetAnimInstanceProxy() override = default;
 
 	/** FAnimPreviewInstanceProxy interface */
 	virtual void Initialize(UAnimInstance* InAnimInstance) override;
+	virtual void CacheBones() override;
 	virtual bool Evaluate(FPoseContext& Output) override;
 	virtual void UpdateAnimationNode(const FAnimationUpdateContext& InContext) override;
 	/** END FAnimPreviewInstanceProxy interface */
@@ -32,9 +41,18 @@ public:
 	virtual void GetCustomNodes(TArray<FAnimNode_Base*>& OutNodes) override;
 	/** END FAnimInstanceProxy interface */
 
-	void SetRetargetAssetAndSourceComponent(
+	void ConfigureAnimInstance(
+		const ERetargetSourceOrTarget& InSourceOrTarget,
 		UIKRetargeter* InIKRetargetAsset,
-		TWeakObjectPtr<USkeletalMeshComponent> InSourceMeshComponent) const;
+		TWeakObjectPtr<USkeletalMeshComponent> InSourceMeshComponent);
+
+	void SetRetargetMode(const ERetargeterOutputMode& InOutputMode);
+
+	void SetRetargetPoseBlend(const float& InRetargetPoseBlend) const;
+
+	FAnimNode_PreviewRetargetPose* PreviewPoseNode;
+	FAnimNode_RetargetPoseFromMesh* RetargetNode;
 
-	FAnimNode_RetargetPoseFromMesh* IKRetargetNode;
+	ERetargetSourceOrTarget SourceOrTarget;
+	ERetargeterOutputMode OutputMode;
 };
\ No newline at end of file
diff --git a/Engine/Plugins/Animation/IKRig/Source/IKRigEditor/Public/RetargetEditor/IKRetargetCommands.h b/Engine/Plugins/Animation/IKRig/Source/IKRigEditor/Public/RetargetEditor/IKRetargetCommands.h
--- a/Engine/Plugins/Animation/IKRig/Source/IKRigEditor/Public/RetargetEditor/IKRetargetCommands.h
+++ b/Engine/Plugins/Animation/IKRig/Source/IKRigEditor/Public/RetargetEditor/IKRetargetCommands.h
@@ -19,7 +19,7 @@ public:
 	{}
 
 	/** reset retarget pose */
-	TSharedPtr< FUICommandInfo > GoToRetargetPose;
+	TSharedPtr< FUICommandInfo > ShowRetargetPose;
 	
 	/** edit retarget pose */
 	TSharedPtr< FUICommandInfo > EditRetargetPose;
diff --git a/Engine/Plugins/Animation/IKRig/Source/IKRigEditor/Public/RetargetEditor/IKRetargetDefaultMode.h b/Engine/Plugins/Animation/IKRig/Source/IKRigEditor/Public/RetargetEditor/IKRetargetDefaultMode.h
--- a/Engine/Plugins/Animation/IKRig/Source/IKRigEditor/Public/RetargetEditor/IKRetargetDefaultMode.h
+++ b/Engine/Plugins/Animation/IKRig/Source/IKRigEditor/Public/RetargetEditor/IKRetargetDefaultMode.h
@@ -61,7 +61,7 @@ private:
 
 	// the skeleton currently being edited
 	UDebugSkelMeshComponent* GetCurrentlyEditedMesh() const;
-	EIKRetargetSkeletonMode SkeletonMode;
+	ERetargetSourceOrTarget SkeletonMode;
 	
 	/** The hosting app */
 	TWeakPtr<FIKRetargetEditorController> EditorController;
diff --git a/Engine/Plugins/Animation/IKRig/Source/IKRigEditor/Public/RetargetEditor/IKRetargetEditPoseMode.h b/Engine/Plugins/Animation/IKRig/Source/IKRigEditor/Public/RetargetEditor/IKRetargetEditPoseMode.h
--- a/Engine/Plugins/Animation/IKRig/Source/IKRigEditor/Public/RetargetEditor/IKRetargetEditPoseMode.h
+++ b/Engine/Plugins/Animation/IKRig/Source/IKRigEditor/Public/RetargetEditor/IKRetargetEditPoseMode.h
@@ -82,8 +82,7 @@ private:
 	// render skeleton in viewport
 	void RenderSkeleton(
 		FPrimitiveDrawInterface* PDI,
-		const FIKRetargetEditorController* Controller,
-		const UIKRetargetProcessor* Processor);
+		const FIKRetargetEditorController* Controller);
 
 	void GetSelectedAndAffectedBones(
 		const FIKRetargetEditorController* Controller,
@@ -101,8 +100,7 @@ private:
 	BoneEdit BoneEdit;
 
 	// the skeleton currently being edited
-	const FRetargetSkeleton& GetCurrentlyEditedSkeleton(const UIKRetargetProcessor& Processor) const;
-	EIKRetargetSkeletonMode SkeletonMode;
+	ERetargetSourceOrTarget SourceOrTarget;
 	
 	/** The hosting app */
 	TWeakPtr<FIKRetargetEditorController> EditorController;
diff --git a/Engine/Plugins/Animation/IKRig/Source/IKRigEditor/Public/RetargetEditor/IKRetargetEditor.h b/Engine/Plugins/Animation/IKRig/Source/IKRigEditor/Public/RetargetEditor/IKRetargetEditor.h
--- a/Engine/Plugins/Animation/IKRig/Source/IKRigEditor/Public/RetargetEditor/IKRetargetEditor.h
+++ b/Engine/Plugins/Animation/IKRig/Source/IKRigEditor/Public/RetargetEditor/IKRetargetEditor.h
@@ -37,7 +37,7 @@ class FIKRetargetEditor :
 public:
 
 	FIKRetargetEditor();
-	virtual ~FIKRetargetEditor() override;
+	virtual ~FIKRetargetEditor() override {};
 
 	void InitAssetEditor(
 		const EToolkitMode::Type Mode,
diff --git a/Engine/Plugins/Animation/IKRig/Source/IKRigEditor/Public/RetargetEditor/IKRetargetEditorController.h b/Engine/Plugins/Animation/IKRig/Source/IKRigEditor/Public/RetargetEditor/IKRetargetEditorController.h
--- a/Engine/Plugins/Animation/IKRig/Source/IKRigEditor/Public/RetargetEditor/IKRetargetEditorController.h
+++ b/Engine/Plugins/Animation/IKRig/Source/IKRigEditor/Public/RetargetEditor/IKRetargetEditorController.h
@@ -2,16 +2,15 @@
 #pragma once
 
 #include "IKRetargetDetails.h"
-#include "IKRetargeterController.h"
+#include "IKRetargetPoseExporter.h"
 #include "IPersonaToolkit.h"
 #include "SIKRetargetAssetBrowser.h"
-#include "UObject/ObjectMacros.h"
 #include "UObject/ObjectPtr.h"
 #include "Input/Reply.h"
-#include "Retargeter/IKRetargetProcessor.h"
 #include "Widgets/Input/SEditableTextBox.h"
 #include "Widgets/SWindow.h"
 
+enum class ERetargetSourceOrTarget : uint8;
 class SIKRetargetHierarchy;
 class SIKRigOutputLog;
 class UIKRetargetProcessor;
@@ -20,27 +19,24 @@ class UIKRetargetAnimInstance;
 class FIKRetargetEditor;
 class UDebugSkelMeshComponent;
 class UIKRigDefinition;
+class UIKRetargeterController;
+class UIKRetargetBoneDetails;
+struct FRetargetSkeleton;
 
-// which skeleton are we editing/viewing?
-enum class EIKRetargetSkeletonMode : uint8
-{
-	/** Editing / viewing the SOURCE skeleton (copy FROM) */
-	Source,
 
-	/** Editing / viewing the TARGET skeleton (copy TO) */
-	Target,
+// retarget editor modes
+enum class ERetargeterOutputMode : uint8
+{
+	RunRetarget,		// output the retargeted target pose
+	ShowRetargetPose,	// output the retarget pose for viewing purposes
+	EditRetargetPose,	// allow editing the retarget pose
 };
 
 enum class EBoneSelectionEdit : uint8
 {
-	/** add to selection set */
-	Add,
-
-	/** remove from selection */
-	Remove,
-
-	/** replace selection entirely*/
-	Replace
+	Add,	// add to selection set
+	Remove,	// remove from selection
+	Replace	// replace selection entirely
 };
 
 /** a home for cross-widget communication to synchronize state across all tabs and viewport */
@@ -48,6 +44,8 @@ class FIKRetargetEditorController : public TSharedFromThis<FIKRetargetEditorCont
 {
 public:
 
+	virtual ~FIKRetargetEditorController() override {};
+
 	/** Initialize the editor */
 	void Initialize(TSharedPtr<FIKRetargetEditor> InEditor, UIKRetargeter* InAsset);
 	/** Bind callbacks to this IK Rig */
@@ -62,7 +60,7 @@ public:
 	void OnRetargeterNeedsInitialized(UIKRetargeter* Retargeter) const;
 	
 	/** all modifications to the data model should go through this controller */
-	UIKRetargeterController* AssetController;
+	TObjectPtr<UIKRetargeterController> AssetController;
 
 	/** Preview scene to be supplied by IHasPersonaToolkit::GetPersonaToolkit */
 	TSharedPtr<IPersonaToolkit> PersonaToolkit;
@@ -70,39 +68,46 @@ public:
 	/** the persona toolkit */
 	TWeakPtr<FIKRetargetEditor> Editor;
 
+	/** import / export retarget poses */
+	TSharedPtr<FIKRetargetPoseExporter> PoseExporter;
+
 	/** viewport skeletal mesh */
+	UDebugSkelMeshComponent* GetSkeletalMeshComponent(const ERetargetSourceOrTarget& SourceOrTarget) const;
 	UDebugSkelMeshComponent* SourceSkelMeshComponent;
 	UDebugSkelMeshComponent* TargetSkelMeshComponent;
 
 	/** viewport anim instance */
+	UIKRetargetAnimInstance* GetAnimInstance(const ERetargetSourceOrTarget& SourceOrTarget) const;
 	UPROPERTY(transient, NonTransactional)
-	TWeakObjectPtr<class UAnimPreviewInstance> SourceAnimInstance;
+	TWeakObjectPtr<class UIKRetargetAnimInstance> SourceAnimInstance;
 	UPROPERTY(transient, NonTransactional)
 	TWeakObjectPtr<class UIKRetargetAnimInstance> TargetAnimInstance;
+	
+	/** store pointers to various tabs of UI,
+	 * have to manage access to these because they can be null if the tabs are closed */
+	void SetDetailsView(TSharedPtr<IDetailsView> InDetailsView) { DetailsView = InDetailsView; };
+	void SetChainsView(TSharedPtr<SIKRetargetChainMapList> InChainsView) { ChainsView = InChainsView; };
+	void SetAssetBrowserView(TSharedPtr<SIKRetargetAssetBrowser> InAssetBrowserView) { AssetBrowserView = InAssetBrowserView; };
+	void SetOutputLogView(TSharedPtr<SIKRigOutputLog> InOutputLogView) { OutputLogView = InOutputLogView; };
+	void SetHierarchyView(TSharedPtr<SIKRetargetHierarchy> InHierarchyView) { HierarchyView = InHierarchyView; };
+	
+	/** force refresh all views in the editor */
+	void RefreshAllViews() const;
+	void RefreshDetailsView() const;
+	void RefreshChainsView() const;
+	void RefreshAssetBrowserView() const;
+	void RefreshHierarchyView() const;
+	void RefreshPoseList() const;
+	void SetDetailsObject(UObject* DetailsObject) const;
+	void SetDetailsObjects(const TArray<UObject*>& DetailsObjects) const;
 
-	/** asset properties tab */
-	TSharedPtr<IDetailsView> DetailsView;
-	/** chain list view */
-	TSharedPtr<SIKRetargetChainMapList> ChainsView;
-	/** asset browser view */
-	TSharedPtr<SIKRetargetAssetBrowser> AssetBrowserView;
-	/** output log view */
-	TSharedPtr<SIKRigOutputLog> OutputLogView;
-	/** hierarchy view */
-	TSharedPtr<SIKRetargetHierarchy> HierarchyView;
 	/** clear the output log */
 	void ClearOutputLog() const;
-	/** force refresh all views in the editor */
-	void RefreshAllViews() const;
 
-	/** get the source skeletal mesh we are copying FROM */
-	USkeletalMesh* GetSourceSkeletalMesh() const;
-	/** get the target skeletal mesh we are copying TO */
-	USkeletalMesh* GetTargetSkeletalMesh() const;
-	/** get the source skeleton asset */
-	const USkeleton* GetSourceSkeleton() const;
-	/** get the target skeleton asset */
-	const USkeleton* GetTargetSkeleton() const;
+	/** get the USkeletalMesh we are transferring animation between (either source or target)*/
+	USkeletalMesh* GetSkeletalMesh(const ERetargetSourceOrTarget SourceOrTarget) const;
+	/** get the USkeleton we are transferring animation between (either source or target)*/
+	const USkeleton* GetSkeleton(const ERetargetSourceOrTarget SourceOrTarget) const;
 	/** get currently edited debug skeletal mesh */
 	UDebugSkelMeshComponent* GetEditedSkeletalMesh() const;
 	/** get the currently edited retarget skeleton */
@@ -110,7 +115,7 @@ public:
 	
 	/** get world space pose of a bone (with component scale / offset applied) */
 	FTransform GetGlobalRetargetPoseOfBone(
-		const FRetargetSkeleton& Skeleton,
+		const ERetargetSourceOrTarget SourceOrTarget,
 		const int32& BoneIndex,
 		const float& Scale,
 		const FVector& Offset) const;
@@ -133,17 +138,24 @@ public:
 
 	/** Sequence Browser**/
 	void PlayAnimationAsset(UAnimationAsset* AssetToPlay);
+	void PausePlayback();
+	void ResumePlayback();
+	UAnimationAsset* AnimThatWasPlaying = nullptr;
+	float TimeWhenPaused = 0.0f;
+	bool bWasPlayingAnim = false;
 	/** END Sequence Browser */
 
 	/** Set viewport / editor tool mode */
 	void SetRetargeterMode(ERetargeterOutputMode Mode);
-	UAnimationAsset* AnimThatWasPlaying = nullptr;
-	bool bWasPlayingAnim = false;
+	ERetargeterOutputMode GetRetargeterMode() const { return OutputMode; }
+	ERetargeterOutputMode PreviousMode;
+	float GetRetargetPoseAmount() const;
+	void SetRetargetPoseAmount(float InValue);
 	/** END viewport / editor tool mode */
 	
 	/** general editor mode can be either viewing/editing source or target */
-	EIKRetargetSkeletonMode GetSkeletonMode() const { return SkeletonMode; };
-	void SetSkeletonMode(EIKRetargetSkeletonMode Mode);
+	ERetargetSourceOrTarget GetSourceOrTarget() const { return CurrentlyEditingSourceOrTarget; };
+	void SetSourceOrTargetMode(ERetargetSourceOrTarget SourceOrTarget);
 
 	/** bone selection management (viewport or hierarchy view) */
 	void EditBoneSelection(
@@ -158,9 +170,9 @@ public:
 	void AddOffsetToMeshComponent(const FVector& Offset, USceneComponent* MeshComponent) const;
 
 	/** determine if bone in the specified skeleton is part of the retarget (in a mapped chain) */
-	bool IsBoneRetargeted(const FName& BoneName, EIKRetargetSkeletonMode WhichSkeleton) const;
+	bool IsBoneRetargeted(const FName& BoneName, ERetargetSourceOrTarget SourceOrTarget) const;
 	/** get the name of the chain that contains this bone */
-	FName GetChainNameFromBone(const FName& BoneName, EIKRetargetSkeletonMode WhichSkeleton) const;
+	FName GetChainNameFromBone(const FName& BoneName, ERetargetSourceOrTarget SourceOrTarget) const;
 
 	/** factory to get/create bone details object */
 	TObjectPtr<UIKRetargetBoneDetails> GetDetailsObjectForBone(const FName& BoneName);
@@ -168,7 +180,9 @@ public:
 	/* START RETARGET POSES */
 	
 	/** go to retarget pose */
-	void HandleGoToRetargetPose();
+	FReply HandleShowRetargetPose();
+	bool CanShowRetargetPose() const;
+	bool IsShowingRetargetPose() const;
 	
 	/** toggle current retarget pose */
 	TArray<TSharedPtr<FName>> PoseNames;
@@ -184,7 +198,6 @@ public:
 	void HandleResetAllBones() const;
 	void HandleResetSelectedBones() const;
 	void HandleResetSelectedAndChildrenBones() const;
-	bool CanResetPose() const;
 	bool CanResetSelected() const;
 
 	/** create new retarget pose */
@@ -198,28 +211,8 @@ public:
 	void HandleDuplicatePose();
 	FReply CreateDuplicatePose() const;
 
-	/** import retarget pose from asset*/
-	void HandleImportPose();
-	FReply ImportRetargetPose() const;
-	void OnRetargetPoseSelected(const FAssetData& SelectedAsset);
-	FSoftObjectPath RetargetPoseToImport;
-	TSharedPtr<SWindow> ImportPoseWindow;
-
-	/** import retarget pose from animation sequence*/
-	void HandleImportPoseFromSequence();
-	bool OnShouldFilterSequenceToImport(const struct FAssetData& AssetData) const;
-	FReply OnImportPoseFromSequence();
-	void OnSequenceSelectedForPose(const FAssetData& SelectedAsset);
-	TSharedPtr<SWindow> ImportPoseFromSequenceWindow;
-	FSoftObjectPath SequenceToImportAsPose;
-	int32 FrameOfSequenceToImport;
-	FText ImportedPoseName;
-
-	/** export retarget pose to asset*/
-	void HandleExportPose();
-
 	/** delete retarget pose */
-	void HandleDeletePose() const;
+	void HandleDeletePose();
 	bool CanDeletePose() const;
 
 	/** rename retarget pose */
@@ -232,20 +225,30 @@ public:
 	/* END RETARGET POSES */
 
 	/** FGCObject interface */
-	virtual void AddReferencedObjects(FReferenceCollector& Collector) override
-	{
-		for (TTuple<FName, TObjectPtr<UIKRetargetBoneDetails>> Pair : AllBoneDetails)
-		{
-			Collector.AddReferencedObject(Pair.Value);
-		}
-	};
+	virtual void AddReferencedObjects(FReferenceCollector& Collector) override;
 	virtual FString GetReferencerName() const override { return TEXT("Retarget Editor"); };
 	/** END FGCObject interface */
 
 private:
+
+	/** asset properties tab */
+	TSharedPtr<IDetailsView> DetailsView;
+	/** chain list view */
+	TSharedPtr<SIKRetargetChainMapList> ChainsView;
+	/** asset browser view */
+	TSharedPtr<SIKRetargetAssetBrowser> AssetBrowserView;
+	/** output log view */
+	TSharedPtr<SIKRigOutputLog> OutputLogView;
+	/** hierarchy view */
+	TSharedPtr<SIKRetargetHierarchy> HierarchyView;
+
+	/** the current output mode of the retargeter */
+	ERetargeterOutputMode OutputMode;
+	/** slider value to blend between reference pose and retarget pose */
+	float RetargetPosePreviewBlend = 1.0f;
 	
 	/** which skeleton are we editing / viewing? */
-	EIKRetargetSkeletonMode SkeletonMode;
+	ERetargetSourceOrTarget CurrentlyEditingSourceOrTarget;
 
 	/** current selection set */
 	TArray<FName> SelectedBones;
diff --git a/Engine/Plugins/Animation/IKRig/Source/IKRigEditor/Public/RetargetEditor/IKRetargetEditorStyle.h b/Engine/Plugins/Animation/IKRig/Source/IKRigEditor/Public/RetargetEditor/IKRetargetEditorStyle.h
--- a/Engine/Plugins/Animation/IKRig/Source/IKRigEditor/Public/RetargetEditor/IKRetargetEditorStyle.h
+++ b/Engine/Plugins/Animation/IKRig/Source/IKRigEditor/Public/RetargetEditor/IKRetargetEditorStyle.h
@@ -14,16 +14,15 @@ public:
 	
 	FIKRetargetEditorStyle() : FSlateStyleSet("IKRetargetEditorStyle")
 	{
-		SetContentRoot(FPaths::EnginePluginsDir() / TEXT("Animation/IKRig/Content"));
-		const FVector2D Icon10x10(10.0f, 10.0f);
-		const FVector2D Icon14x14(14.0f, 14.0f);
 		const FVector2D Icon16x16(16.0f, 16.0f);
-		const FVector2D Icon20x20(20.0f, 20.0f);
-		const FVector2D Icon24x24(24.0f, 24.0f);
-		const FVector2D Icon32x32(32.0f, 32.0f);
-		const FVector2D Icon40x40(40.0f, 40.0f);
-
+		
+		const FString IKRigPluginContentDir = FPaths::EnginePluginsDir() / TEXT("Animation/IKRig/Content");
+		SetContentRoot(IKRigPluginContentDir);
 		Set("IKRetarget.Tree.Bone", new IMAGE_BRUSH("Slate/Bone_16x", Icon16x16));
+
+		const FString EngineEditorSlateDir = FPaths::EngineContentDir() / TEXT("Editor/Slate");
+		SetContentRoot(EngineEditorSlateDir);
+		Set( "IKRetarget.Viewport.Border", new BOX_BRUSH( "Old/Window/ViewportDebugBorder", 0.8f, FLinearColor(1.0f,1.0f,1.0f,1.0f) ) );
 		
 		FSlateStyleRegistry::RegisterSlateStyle(*this);
 	}
diff --git a/Engine/Plugins/Animation/IKRig/Source/IKRigEditor/Public/RetargetEditor/IKRetargetFactory.h b/Engine/Plugins/Animation/IKRig/Source/IKRigEditor/Public/RetargetEditor/IKRetargetFactory.h
--- a/Engine/Plugins/Animation/IKRig/Source/IKRigEditor/Public/RetargetEditor/IKRetargetFactory.h
+++ b/Engine/Plugins/Animation/IKRig/Source/IKRigEditor/Public/RetargetEditor/IKRetargetFactory.h
@@ -36,22 +36,3 @@ private:
 
 	TSharedPtr<SWindow> PickerWindow;
 };
-
-
-UCLASS(hidecategories=Object, MinimalAPI)
-class URetargetPoseFactory : public UFactory
-{
-	GENERATED_BODY()
-
-public:
-
-	URetargetPoseFactory();
-
-	// UFactory Interface
-	virtual FText GetDisplayName() const override;
-	virtual uint32 GetMenuCategories() const override;
-	virtual FText GetToolTip() const override;
-	virtual FString GetDefaultNewAssetName() const override;
-	virtual UObject* FactoryCreateNew(UClass* Class, UObject* InParent, FName Name, EObjectFlags Flags, UObject* Context, FFeedbackContext* Warn) override;
-	virtual bool ShouldShowInNewMenu() const override;
-};
diff --git a/Engine/Plugins/Animation/IKRig/Source/IKRigEditor/Public/RetargetEditor/IKRetargetPoseExporter.h b/Engine/Plugins/Animation/IKRig/Source/IKRigEditor/Public/RetargetEditor/IKRetargetPoseExporter.h
new file mode 100644
--- /dev/null
+++ b/Engine/Plugins/Animation/IKRig/Source/IKRigEditor/Public/RetargetEditor/IKRetargetPoseExporter.h
@@ -0,0 +1,46 @@
+﻿// Copyright Epic Games, Inc. All Rights Reserved.
+
+#pragma once
+
+#include "Widgets/Notifications/SNotificationList.h"
+
+class FIKRetargetEditorController;
+
+class FIKRetargetPoseExporter : public TSharedFromThis<FIKRetargetPoseExporter>
+{
+public:
+	
+	FIKRetargetPoseExporter() = default;
+
+	void Initialize(TSharedPtr<FIKRetargetEditorController> InController);
+
+	/** import retarget pose from asset*/
+	void HandleImportFromPoseAsset();
+	FReply ImportPoseAsset() const;
+	void OnRetargetPoseSelected(const FAssetData& SelectedAsset);
+	TSharedRef<SWidget> OnGeneratePoseComboWidget(TSharedPtr<FName> Item) const;
+	void OnSelectPoseFromPoseAsset(TSharedPtr<FName> Item, ESelectInfo::Type SelectInfo);
+	FSoftObjectPath RetargetPoseToImport;
+	TSharedPtr<SWindow> ImportPoseWindow;
+	TArray<TSharedPtr<FName>> PosesInSelectedAsset;
+	TSharedPtr<FName> SelectedPose;
+
+	/** import retarget pose from animation sequence*/
+	void HandleImportFromSequenceAsset();
+	bool OnShouldFilterSequenceToImport(const struct FAssetData& AssetData) const;
+	FReply OnImportPoseFromSequence();
+	void OnSequenceSelectedForPose(const FAssetData& SelectedAsset);
+	TSharedPtr<SWindow> ImportPoseFromSequenceWindow;
+	FSoftObjectPath SequenceToImportAsPose;
+	int32 FrameOfSequenceToImport;
+	FText ImportedPoseName;
+
+	/** export retarget pose to asset*/
+	void HandleExportPoseAsset();
+
+private:
+
+	static void NotifyUser(const FText& Message, SNotificationItem::ECompletionState NotificationType);
+	
+	TWeakPtr<FIKRetargetEditorController> Controller;
+};
diff --git a/Engine/Plugins/Animation/IKRig/Source/IKRigEditor/Public/RetargetEditor/IKRetargeterController.h b/Engine/Plugins/Animation/IKRig/Source/IKRigEditor/Public/RetargetEditor/IKRetargeterController.h
--- a/Engine/Plugins/Animation/IKRig/Source/IKRigEditor/Public/RetargetEditor/IKRetargeterController.h
+++ b/Engine/Plugins/Animation/IKRig/Source/IKRigEditor/Public/RetargetEditor/IKRetargeterController.h
@@ -1,21 +1,21 @@
 ﻿// Copyright Epic Games, Inc. All Rights Reserved.
-
 #pragma once
 
 #include "CoreMinimal.h"
-
-#include "Retargeter/IKRetargeter.h"
 #include "UObject/Object.h"
 
 #include "IKRetargeterController.generated.h"
 
+struct FIKRetargetPose;
+enum class ERetargetSourceOrTarget : uint8;
 class FIKRetargetEditorController;
 class URetargetChainSettings;
 class UIKRigDefinition;
 class UIKRetargeter;
+class USkeletalMesh;
 
 /** A singleton (per-asset) class used to make modifications to a UIKRetargeter asset.
-* Call the static UIKRetargeterController() function to get the controller for the asset you want to modify. */ 
+* Call the static GetController(UIKRetargeter* Asset) function to get the controller for the asset you want to modify. */ 
 UCLASS(config = Engine, hidecategories = UObject)
 class IKRIGEDITOR_API UIKRetargeterController : public UObject
 {
@@ -28,41 +28,27 @@ public:
 	/** Get access to the retargeter asset.
 	 *@warning Do not make modifications to the asset directly. Use this API instead. */
 	UIKRetargeter* GetAsset() const;
-	/** Get access to the editor controller.*/
-	FIKRetargetEditorController* GetEditorController() const { return EditorController; };
-	/** Set the currently used editor controller.*/
-	void SetEditorController(FIKRetargetEditorController* InEditorController) { EditorController = InEditorController; };
 
 	/** SOURCE / TARGET
 	* 
 	*/
 	/** Set the IK Rig to use as the source (to copy animation FROM) */
 	void SetSourceIKRig(UIKRigDefinition* SourceIKRig);
-	/** Get source skeletal mesh */
-	USkeletalMesh* GetSourcePreviewMesh() const;
-	/** Get target skeletal mesh */
-	USkeletalMesh* GetTargetPreviewMesh() const;
-	/** Get source IK Rig asset */
-	const UIKRigDefinition* GetSourceIKRig() const;
-	/** Get source IK Rig asset */
-	const UIKRigDefinition* GetTargetIKRig() const;
-	/** Set the target preview mesh based on the mesh in the target IK Rig asset */
-	void OnTargetIKRigChanged() const;
-	/** Set the source preview mesh based on the mesh in the source IK Rig asset */
-	void OnSourceIKRigChanged() const;
-
-	/** Get name of the Root bone used for retargeting the Source skeleton. */
-	FName GetSourceRootBone() const;
-	/** Get name of the Root bone used for retargeting the Target skeleton. */
-	FName GetTargetRootBone() const;
+	/** Get the preview skeletal mesh */
+	USkeletalMesh* GetPreviewMesh(const ERetargetSourceOrTarget& SourceOrTarget) const;
+	/** Get either source or target IK Rig */
+	const UIKRigDefinition* GetIKRig(const ERetargetSourceOrTarget& SourceOrTarget) const;
+	/** Set the SOURCE or TARGET preview mesh based on the mesh in the corresponding IK Rig asset */
+	void OnIKRigChanged(const ERetargetSourceOrTarget& SourceOrTarget) const;
+
+	/** Get name of the Root bone used for retargeting. */
+	FName GetRetargetRootBone(const ERetargetSourceOrTarget& SourceOrTarget) const;
 
 	/** RETARGET CHAIN MAPPING
 	* 
 	*/
-	/** Get names of all the target bone chains. */
-	void GetTargetChainNames(TArray<FName>& OutNames) const;
-	/** Get names of all the source bone chains. */
-	void GetSourceChainNames(TArray<FName>& OutNames) const;
+	/** Get names of all the bone chains. */
+	void GetChainNames(const ERetargetSourceOrTarget& SourceOrTarget, TArray<FName>& OutNames) const;
 	/** Remove invalid chain mappings (no longer existing in currently referenced source/target IK Rig assets) */
 	void CleanChainMapping(const bool bForceReinitialization=true) const;
 	/** Use fuzzy string search to find "best" Source chain to map to each Target chain */
@@ -81,36 +67,59 @@ public:
 	 * 
 	 */
 	/** Remove bones from retarget poses that are no longer in skeleton */
-	void CleanPoseList(const bool bForceReinitialization=true);
+	void CleanPoseLists(const bool bForceReinitialization=true) const;
+	/** Remove bones from retarget poses that are no longer in skeleton */
+	void CleanPoseList(const ERetargetSourceOrTarget& SourceOrTarget) const;
 	/** Add new retarget pose. */
-	void AddRetargetPose(FName NewPoseName, const FIKRetargetPose* ToDuplicate = nullptr) const;
+	void AddRetargetPose(
+		const FName& NewPoseName,
+		const FIKRetargetPose* ToDuplicate,
+		const ERetargetSourceOrTarget& SourceOrTarget) const;
 	/** Rename current retarget pose. */
-	void RenameCurrentRetargetPose(FName NewPoseName) const;
+	void RenameCurrentRetargetPose(
+		const FName& NewPoseName,
+		const ERetargetSourceOrTarget& SourceOrTarget) const;
 	/** Remove a retarget pose. */
-	void RemoveRetargetPose(FName PoseToRemove) const;
+	void RemoveRetargetPose(
+	const FName& PoseToRemove, 
+	const ERetargetSourceOrTarget& SourceOrTarget) const;
 	/** Reset a retarget pose for the specified bones.
 	 *If BonesToReset is Empty, will removes all stored deltas, returning pose to reference pose */
-	void ResetRetargetPose(FName PoseToReset, const TArray<FName>& BonesToReset) const;
+	void ResetRetargetPose(
+		const FName& PoseToReset,
+		const TArray<FName>& BonesToReset,
+		const ERetargetSourceOrTarget& SourceOrTarget) const;
 	/** Get the current retarget pose */
-    FName GetCurrentRetargetPoseName() const;
+    FName GetCurrentRetargetPoseName(const ERetargetSourceOrTarget& SourceOrTarget) const;
 	/** Change which retarget pose is used by the retargeter at runtime */
-	void SetCurrentRetargetPose(FName CurrentPose) const;
-	/** Get read-only access to list of retarget poses */
-	const TMap<FName, FIKRetargetPose>& GetRetargetPoses();
+	void SetCurrentRetargetPose(FName CurrentPose, const ERetargetSourceOrTarget& SourceOrTarget) const;
+	/** Get access to array of retarget poses */
+	TMap<FName, FIKRetargetPose>& GetRetargetPoses(const ERetargetSourceOrTarget& SourceOrTarget) const;
 	/** Get the current retarget pose */
-	const FIKRetargetPose& GetCurrentRetargetPose() const;
+	FIKRetargetPose& GetCurrentRetargetPose(const ERetargetSourceOrTarget& SourceOrTarget) const;
 	/** Set a delta rotation for a given bone for the current retarget pose (used in Edit Mode in the retarget editor) */
-	void SetRotationOffsetForRetargetPoseBone(FName BoneName, const FQuat& RotationOffset) const;
+	void SetRotationOffsetForRetargetPoseBone(
+		const FName& BoneName,
+		const FQuat& RotationOffset,
+		const ERetargetSourceOrTarget& SkeletonMode) const;
 	/** Get a delta rotation for a given bone for the current retarget pose (used in Edit Mode in the retarget editor) */
-	FQuat GetRotationOffsetForRetargetPoseBone(FName BoneName) const;
+	FQuat GetRotationOffsetForRetargetPoseBone(
+		const FName& BoneName,
+		const ERetargetSourceOrTarget& SourceOrTarget) const;
 	/** Set the delta translation of the root bone (used in Edit Mode in the retarget editor) */
-	void SetTranslationOffsetOnRetargetRootBone(FVector TranslationOffset) const;
+	void SetTranslationOffsetOnRetargetRootBone(
+		const FVector& TranslationOffset,
+		const ERetargetSourceOrTarget& SourceOrTarget) const;
 	/** Get the delta translation of the root bone (used in Edit Mode in the retarget editor) */
-	const FVector& GetTranslationOffsetOnRetargetRootBone() const;
+	const FVector& GetTranslationOffsetOnRetargetRootBone(const ERetargetSourceOrTarget& SourceOrTarget) const;
 	/** Add a delta translation to the root bone (used in Edit Mode in the retarget editor) */
-	void AddTranslationOffsetToRetargetRootBone(FVector TranslationOffset) const;
+	void AddTranslationOffsetToRetargetRootBone(
+		const FVector& TranslationOffset,
+		const ERetargetSourceOrTarget& SourceOrTarget) const;
 	/** Add a numbered suffix to the given pose name to make it unique. */
-	FName MakePoseNameUnique(FString PoseName) const;
+	FName MakePoseNameUnique(
+		const FString& PoseName,
+		const ERetargetSourceOrTarget& SourceOrTarget) const;
 	/** END RETARGET POSE EDITING */
 
 private:
@@ -139,5 +148,5 @@ private:
 	TObjectPtr<UIKRetargeter> Asset = nullptr;
 
 	/** The editor controller for this asset. */
-	TObjectPtr<FIKRetargetEditorController> EditorController = nullptr;
+	TSharedPtr<FIKRetargetEditorController> EditorController = nullptr;
 };
diff --git a/Engine/Plugins/Animation/IKRig/Source/IKRigEditor/Public/RetargetEditor/SIKRetargetHierarchy.h b/Engine/Plugins/Animation/IKRig/Source/IKRigEditor/Public/RetargetEditor/SIKRetargetHierarchy.h
--- a/Engine/Plugins/Animation/IKRig/Source/IKRigEditor/Public/RetargetEditor/SIKRetargetHierarchy.h
+++ b/Engine/Plugins/Animation/IKRig/Source/IKRigEditor/Public/RetargetEditor/SIKRetargetHierarchy.h
@@ -6,10 +6,11 @@
 #include "EditorUndoClient.h"
 #include "IKRetargetEditorController.h"
 #include "SBaseHierarchyTreeView.h"
+#include "SIKRetargetPoseEditor.h"
 #include "DragAndDrop/DecoratedDragDropOp.h"
 #include "Widgets/SCompoundWidget.h"
 
-class SIKRetargetHierarchy;
+class SIKRetargetPoseEditor;
 
 class FIKRetargetHierarchyElement : public TSharedFromThis<FIKRetargetHierarchyElement>
 {
@@ -70,6 +71,8 @@ public:
 
 	void ShowItemAfterSelection(FName ItemKey);
 
+	void RefreshPoseList() const { EditPoseView.Get()->Refresh(); };
+
 private:
 	
 	/** centralized editor controls (facilitate cross-communication between multiple UI elements)*/
@@ -77,6 +80,9 @@ private:
 	
 	/** command list we bind to */
 	TSharedPtr<FUICommandList> CommandList;
+
+	/** edit pose widget */
+	TSharedPtr<SIKRetargetPoseEditor> EditPoseView;
 	
 	/** tree view widget */
 	TSharedPtr<SIKRetargetHierarchyTreeView> TreeView;
diff --git a/Engine/Plugins/Animation/IKRig/Source/IKRigEditor/Public/RetargetEditor/SIKRetargetPoseEditor.h b/Engine/Plugins/Animation/IKRig/Source/IKRigEditor/Public/RetargetEditor/SIKRetargetPoseEditor.h
--- a/Engine/Plugins/Animation/IKRig/Source/IKRigEditor/Public/RetargetEditor/SIKRetargetPoseEditor.h
+++ b/Engine/Plugins/Animation/IKRig/Source/IKRigEditor/Public/RetargetEditor/SIKRetargetPoseEditor.h
@@ -18,7 +18,11 @@ public:
 
 	void Construct(const FArguments& InArgs, TSharedRef<FIKRetargetEditorController> InEditorController);
 
+	void Refresh();
+	
 private:
+
+	TSharedPtr<SVerticalBox> Contents;
 	
 	TSharedRef<SWidget> MakeToolbar(TSharedPtr<FUICommandList> Commands);
 	TSharedRef<SWidget> GenerateResetMenuContent(TSharedPtr<FUICommandList> Commands);
diff --git a/Engine/Plugins/Animation/IKRig/Source/IKRigEditor/Public/RigEditor/SIKRigOutputLog.h b/Engine/Plugins/Animation/IKRig/Source/IKRigEditor/Public/RigEditor/SIKRigOutputLog.h
--- a/Engine/Plugins/Animation/IKRig/Source/IKRigEditor/Public/RigEditor/SIKRigOutputLog.h
+++ b/Engine/Plugins/Animation/IKRig/Source/IKRigEditor/Public/RigEditor/SIKRigOutputLog.h
@@ -16,10 +16,7 @@ public:
 	SLATE_BEGIN_ARGS(SIKRigOutputLog) {}
 	SLATE_END_ARGS()
 
-	void Construct(
-		const FArguments& InArgs,
-		const FName& InLogName,
-		TSharedPtr<SIKRigOutputLog>& OutLog);
+	void Construct(const FArguments& InArgs, const FName& InLogName);
 
 	void ClearLog() const;
 	
diff --git a/Engine/Source/Editor/AnimGraph/Private/AnimGraphNode_PoseDriver.cpp b/Engine/Source/Editor/AnimGraph/Private/AnimGraphNode_PoseDriver.cpp
--- a/Engine/Source/Editor/AnimGraph/Private/AnimGraphNode_PoseDriver.cpp
+++ b/Engine/Source/Editor/AnimGraph/Private/AnimGraphNode_PoseDriver.cpp
@@ -288,35 +288,6 @@ FAnimNode_PoseDriver* UAnimGraphNode_PoseDriver::GetPreviewPoseDriverNode() cons
 	return PreviewNode;
 }
 
-/** Util to return transform of a bone from the pose asset in component space, by walking up tracks in pose asset */
-FTransform GetComponentSpaceTransform(FName BoneName, TArray<FTransform>& LocalTransforms, UPoseAsset* PoseAsset)
-{
-	const FReferenceSkeleton& RefSkel = PoseAsset->GetSkeleton()->GetReferenceSkeleton();
-
-	// Init component space transform with local transform
-	FTransform ComponentSpaceTransform = FTransform::Identity;
-
-	// Start to walk up parent chain until we reach root (ParentIndex == INDEX_NONE)
-	int32 BoneIndex = RefSkel.FindBoneIndex(BoneName);
-	while (BoneIndex != INDEX_NONE)
-	{
-		BoneName = RefSkel.GetBoneName(BoneIndex);
-		int32 TrackIndex = PoseAsset->GetTrackIndexByName(BoneName);
-
-		// If a track for parent, get local space transform from that
-		// If not, get from ref pose
-		FTransform BoneLocalTM = (TrackIndex != INDEX_NONE) ? LocalTransforms[TrackIndex] : RefSkel.GetRefBonePose()[BoneIndex];
-
-		// Continue to build component space transform
-		ComponentSpaceTransform = ComponentSpaceTransform * BoneLocalTM;
-
-		// Now move up to parent
-		BoneIndex = RefSkel.GetParentIndex(BoneIndex);
-	}
-
-	return ComponentSpaceTransform;
-}
-
 void UAnimGraphNode_PoseDriver::CopyTargetsFromPoseAsset()
 {
 	UPoseAsset* PoseAsset = Node.PoseAsset;
@@ -355,8 +326,8 @@ void UAnimGraphNode_PoseDriver::CopyTargetsFromPoseAsset()
 						// If eval'ing in different space (and that space is valid)
 						if (Node.EvalSpaceBone.BoneName != NAME_None)
 						{
-							FTransform SourceCompSpace = GetComponentSpaceTransform(SourceBoneRef.BoneName, PoseTransforms, PoseAsset);
-							FTransform EvalCompSpace = GetComponentSpaceTransform(Node.EvalSpaceBone.BoneName, PoseTransforms, PoseAsset);
+							FTransform SourceCompSpace = PoseAsset->GetComponentSpaceTransform(SourceBoneRef.BoneName, PoseTransforms);
+							FTransform EvalCompSpace = PoseAsset->GetComponentSpaceTransform(Node.EvalSpaceBone.BoneName, PoseTransforms);
 
 							SourceBoneTransform = SourceCompSpace.GetRelativeTransform(EvalCompSpace);
 						}
diff --git a/Engine/Source/Runtime/Engine/Classes/Animation/PoseAsset.h b/Engine/Source/Runtime/Engine/Classes/Animation/PoseAsset.h
--- a/Engine/Source/Runtime/Engine/Classes/Animation/PoseAsset.h
+++ b/Engine/Source/Runtime/Engine/Classes/Animation/PoseAsset.h
@@ -284,6 +284,9 @@ public:
 	virtual void PostEditChangeProperty(FPropertyChangedEvent& PropertyChangedEvent) override;
 	// Return full (local space, non additive) pose. Will do conversion if PoseAsset is Additive. 
 	ENGINE_API bool GetFullPose(int32 PoseIndex, TArray<FTransform>& OutTransforms) const;
+	
+	// util to return transform of a bone from the pose asset in component space, by walking up tracks in pose asset */
+	ENGINE_API FTransform GetComponentSpaceTransform(FName TrackName, const TArray<FTransform>& LocalTransforms) const;
 
 	ENGINE_API int32 DeletePoses(TArray<FName> PoseNamesToDelete);
 	ENGINE_API int32 DeleteCurves(TArray<FName> CurveNamesToDelete);
diff --git a/Engine/Source/Runtime/Engine/Private/Animation/PoseAsset.cpp b/Engine/Source/Runtime/Engine/Private/Animation/PoseAsset.cpp
--- a/Engine/Source/Runtime/Engine/Private/Animation/PoseAsset.cpp
+++ b/Engine/Source/Runtime/Engine/Private/Animation/PoseAsset.cpp
@@ -1625,6 +1625,34 @@ bool UPoseAsset::GetFullPose(int32 PoseIndex, TArray<FTransform>& OutTransforms)
 	return true;
 }
 
+FTransform UPoseAsset::GetComponentSpaceTransform(FName TrackName, const TArray<FTransform>& LocalTransforms) const
+{
+	const FReferenceSkeleton& RefSkel = GetSkeleton()->GetReferenceSkeleton();
+
+	// Init component space transform with identity
+	FTransform ComponentSpaceTransform = FTransform::Identity;
+
+	// Start to walk up parent chain until we reach root (ParentIndex == INDEX_NONE)
+	int32 BoneIndex = RefSkel.FindBoneIndex(TrackName);
+	while (BoneIndex != INDEX_NONE)
+	{
+		TrackName = RefSkel.GetBoneName(BoneIndex);
+		const int32 TrackIndex = GetTrackIndexByName(TrackName);
+
+		// If a track for parent, get local space transform from that
+		// If not, get from ref pose
+		FTransform BoneLocalTM = (TrackIndex != INDEX_NONE) ? LocalTransforms[TrackIndex] : RefSkel.GetRefBonePose()[BoneIndex];
+
+		// Continue to build component space transform
+		ComponentSpaceTransform = ComponentSpaceTransform * BoneLocalTM;
+
+		// Now move up to parent
+		BoneIndex = RefSkel.GetParentIndex(BoneIndex);
+	}
+
+	return ComponentSpaceTransform;
+}
+
 bool UPoseAsset::ConvertSpace(bool bNewAdditivePose, int32 NewBasePoseIndex)
 {
 	// first convert to full pose first
