Fixing localization issues in IK Rig text.

diff --git a/Engine/Plugins/Animation/IKRig/Source/IKRigEditor/Private/RetargetEditor/IKRetargetPoseExporter.cpp b/Engine/Plugins/Animation/IKRig/Source/IKRigEditor/Private/RetargetEditor/IKRetargetPoseExporter.cpp
--- a/Engine/Plugins/Animation/IKRig/Source/IKRigEditor/Private/RetargetEditor/IKRetargetPoseExporter.cpp
+++ b/Engine/Plugins/Animation/IKRig/Source/IKRigEditor/Private/RetargetEditor/IKRetargetPoseExporter.cpp
@@ -44,7 +44,7 @@ void FIKRetargetPoseExporter::HandleImportFromPoseAsset()
 	AssetPickerConfig.InitialAssetViewType = EAssetViewType::Tile;
 
 	ImportPoseWindow = SNew(SWindow)
-	.Title(LOCTEXT("ImportRetargetPose", "Import Retarget Pose"))
+	.Title(LOCTEXT("ImportRetargetPose_Label", "Import Retarget Pose"))
 	.ClientSize(FVector2D(500, 600))
 	.SupportsMinimize(false) .SupportsMaximize(false)
 	[
@@ -332,7 +332,7 @@ void FIKRetargetPoseExporter::HandleImportFromSequenceAsset()
 
 	// create pop-up window for user to select animation sequence asset to import as a retarget pose]
 	ImportPoseFromSequenceWindow = SNew(SWindow)
-	.Title(LOCTEXT("ImportRetargetPoseFromSequenceAsset", "Import Retarget Pose from Sequence Asset"))
+	.Title(LOCTEXT("ImportRetargetPoseFromSequenceAsset_Label", "Import Retarget Pose from Sequence Asset"))
 	.ClientSize(FVector2D(500, 600))
 	.SupportsMinimize(false) .SupportsMaximize(false)
 	[
@@ -615,7 +615,7 @@ void FIKRetargetPoseExporter::HandleExportPoseAsset()
 	const FString SaveObjectPath = ContentBrowserModule.Get().CreateModalSaveAssetDialog(SaveAssetDialogConfig);
 	if (SaveObjectPath.IsEmpty())
 	{
-		const FText Message = LOCTEXT("IKRigExportPoseAborted", "Export pose cancelled.");
+		const FText Message = LOCTEXT("IKRigExportPoseCancelled", "Export pose cancelled.");
 		NotifyUser(Message, SNotificationItem::CS_Fail);
 		return;
 	}
@@ -643,7 +643,7 @@ void FIKRetargetPoseExporter::HandleExportPoseAsset()
 		GEditor->SyncBrowserToObjects(ObjectsToSync);
 		
 		// notify user of newly created asset
-		const FText Message = FText::Format(LOCTEXT("IKRigExportPoseAborted", "Exported Pose Asset: {0} with pose, '{1}'."), FText::FromString(AssetName), FText::FromName(NewPoseName.DisplayName));
+		const FText Message = FText::Format(LOCTEXT("IKRigExportPoseSuccess", "Exported Pose Asset: {0} with pose, '{1}'."), FText::FromString(AssetName), FText::FromName(NewPoseName.DisplayName));
 		NotifyUser(Message, SNotificationItem::CS_Success);
 	}
 	else
@@ -652,7 +652,7 @@ void FIKRetargetPoseExporter::HandleExportPoseAsset()
 		NewPoseAsset->ConditionalBeginDestroy();
 		
 		// notify user that asset failed to export
-		const FText Message = FText::Format(LOCTEXT("IKRigExportPoseAborted", "Failed to exported pose asset, {0}."), FText::FromString(AssetName));
+		const FText Message = FText::Format(LOCTEXT("IKRigExportPoseFailedAfter", "Failed to exported pose asset, {0}."), FText::FromString(AssetName));
 		NotifyUser(Message, SNotificationItem::CS_Fail);
 	}
 }
