Fixing crash when opening IK Rig that cannot initialize the processor.

diff --git a/Engine/Plugins/Animation/IKRig/Source/IKRig/Private/AnimNodes/AnimNode_IKRig.cpp b/Engine/Plugins/Animation/IKRig/Source/IKRig/Private/AnimNodes/AnimNode_IKRig.cpp
--- a/Engine/Plugins/Animation/IKRig/Source/IKRig/Private/AnimNodes/AnimNode_IKRig.cpp
+++ b/Engine/Plugins/Animation/IKRig/Source/IKRig/Private/AnimNodes/AnimNode_IKRig.cpp
@@ -167,6 +167,15 @@ void FAnimNode_IKRig::Initialize_AnyThread(const FAnimationInitializeContext& Co
 	DECLARE_SCOPE_HIERARCHICAL_COUNTER_FUNC()
 	FAnimNode_Base::Initialize_AnyThread(Context);
 	Source.Initialize(Context);
+
+	// Initial update of the node, so we dont have a frame-delay on setup
+	GetEvaluateGraphExposedInputs().Execute(Context);
+
+	// ensure there is always a processor available
+	if (IKRigProcessor.IsNull() && IsInGameThread())
+	{
+		IKRigProcessor = NewObject<UIKRigProcessor>(Context.AnimInstanceProxy->GetSkelMeshComponent());	
+	}
 }
 
 void FAnimNode_IKRig::Update_AnyThread(const FAnimationUpdateContext& Context)
