Fixes non-clearing output log in retarget editor

diff --git a/Engine/Plugins/Animation/IKRig/Source/IKRigEditor/Private/RetargetEditor/IKRetargetEditorController.cpp b/Engine/Plugins/Animation/IKRig/Source/IKRigEditor/Private/RetargetEditor/IKRetargetEditorController.cpp
--- a/Engine/Plugins/Animation/IKRig/Source/IKRigEditor/Private/RetargetEditor/IKRetargetEditorController.cpp
+++ b/Engine/Plugins/Animation/IKRig/Source/IKRigEditor/Private/RetargetEditor/IKRetargetEditorController.cpp
@@ -652,11 +652,8 @@ void FIKRetargetEditorController::SetRetargeterMode(ERetargeterOutputMode Mode)
 			OutputMode = ERetargeterOutputMode::RunRetarget;
 			SourceAnimInstance->SetRetargetMode(ERetargeterOutputMode::RunRetarget);
 			TargetAnimInstance->SetRetargetMode(ERetargeterOutputMode::RunRetarget);
-			if (PreviousMode == ERetargeterOutputMode::EditRetargetPose)
-			{
-				// must reinitialize after editing the retarget pose
-				AssetController->BroadcastNeedsReinitialized();
-			}
+			// must reinitialize after editing the retarget pose
+			AssetController->BroadcastNeedsReinitialized();
 			ResumePlayback();
 			break;
 
diff --git a/Engine/Plugins/Animation/IKRig/Source/IKRigEditor/Private/RetargetEditor/IKRetargetOutputLogTabSummoner.cpp b/Engine/Plugins/Animation/IKRig/Source/IKRigEditor/Private/RetargetEditor/IKRetargetOutputLogTabSummoner.cpp
--- a/Engine/Plugins/Animation/IKRig/Source/IKRigEditor/Private/RetargetEditor/IKRetargetOutputLogTabSummoner.cpp
+++ b/Engine/Plugins/Animation/IKRig/Source/IKRigEditor/Private/RetargetEditor/IKRetargetOutputLogTabSummoner.cpp
@@ -34,8 +34,9 @@ TSharedRef<SWidget> FIKRetargetOutputLogTabSummoner::CreateTabBody(const FWorkfl
 {
 	const TSharedRef<FIKRetargetEditorController>& Controller = IKRetargetEditor.Pin()->GetController();
 	const FName LogName = Controller->GetRetargetProcessor()->Log.GetLogTarget();
-	
-	return SNew(SIKRigOutputLog, LogName);
+	TSharedRef<SIKRigOutputLog> LogView = SNew(SIKRigOutputLog, LogName);
+	Controller->SetOutputLogView(LogView);
+	return LogView;
 }
 
 #undef LOCTEXT_NAMESPACE 
