Inclusive Terminology Push
Renamed terms at AnimNode_RetargetPoseFromMesh

diff --git a/Engine/Plugins/Animation/IKRig/Source/IKRig/Public/AnimNodes/AnimNode_RetargetPoseFromMesh.h b/Engine/Plugins/Animation/IKRig/Source/IKRig/Public/AnimNodes/AnimNode_RetargetPoseFromMesh.h
--- a/Engine/Plugins/Animation/IKRig/Source/IKRig/Public/AnimNodes/AnimNode_RetargetPoseFromMesh.h
+++ b/Engine/Plugins/Animation/IKRig/Source/IKRig/Public/AnimNodes/AnimNode_RetargetPoseFromMesh.h
@@ -60,7 +60,7 @@ private:
 	bool EnsureProcessorIsInitialized(const TObjectPtr<USkeletalMeshComponent> TargetMeshComponent);
 	/** copies the source mesh pose (on main thread) */
 	void CopyBoneTransformsFromSource(USkeletalMeshComponent* TargetMeshComponent);
-	/** indirection to account for Master Pose Component setup */
+	/** indirection to account for Leader Pose Component setup */
 	TObjectPtr<USkeletalMeshComponent> GetComponentToCopyPoseFrom() const;
 
 	/** the runtime processor used to run the retarget and generate new poses */
