Added profiles to the IK retargeter
Added filtering to all IK Rig editors
Added global retarget settings with stride warping
Added root settings and detail customization for all settings

diff --git a/Engine/Plugins/Animation/IKRig/Config/BaseIKRig.ini b/Engine/Plugins/Animation/IKRig/Config/BaseIKRig.ini
--- a/Engine/Plugins/Animation/IKRig/Config/BaseIKRig.ini
+++ b/Engine/Plugins/Animation/IKRig/Config/BaseIKRig.ini
@@ -1 +1,3 @@
 +ClassRedirects=(OldName="AnimNode_IKRetargeter",NewName="AnimNode_RetargetPoseFromMesh")
+
++PropertyRedirects=(OldName="/Script/IKRig.TargetChainIKSettings.StaticOffset",NewName="TranslationGlobalOffset")
diff --git a/Engine/Plugins/Animation/IKRig/Source/IKRig/Private/AnimNodes/AnimNode_IKRig.cpp b/Engine/Plugins/Animation/IKRig/Source/IKRig/Private/AnimNodes/AnimNode_IKRig.cpp
--- a/Engine/Plugins/Animation/IKRig/Source/IKRig/Private/AnimNodes/AnimNode_IKRig.cpp
+++ b/Engine/Plugins/Animation/IKRig/Source/IKRig/Private/AnimNodes/AnimNode_IKRig.cpp
@@ -66,7 +66,7 @@ void FAnimNode_IKRig::CopyInputPoseToSolver(FCompactPose& InputPose)
 	
 	// start Solve() from INPUT pose
 	// copy local bone transforms into IKRigProcessor skeleton
-	FIKRigSkeleton& IKRigSkeleton = IKRigProcessor->GetSkeleton();
+	FIKRigSkeleton& IKRigSkeleton = IKRigProcessor->GetSkeletonWriteable();
 	for (FCompactPoseBoneIndex CPIndex : InputPose.ForEachBoneIndex())
 	{
 		if (int32* Index = CompactPoseToRigIndices.Find(CPIndex))
@@ -118,7 +118,7 @@ void FAnimNode_IKRig::AssignGoalTargets()
 
 void FAnimNode_IKRig::CopyOutputPoseToAnimGraph(FCompactPose& OutputPose)
 {
-	FIKRigSkeleton& IKRigSkeleton = IKRigProcessor->GetSkeleton();
+	FIKRigSkeleton& IKRigSkeleton = IKRigProcessor->GetSkeletonWriteable();
 	
 	// update local transforms of current IKRig pose
 	IKRigSkeleton.UpdateAllLocalTransformFromGlobal();
@@ -300,6 +300,7 @@ void FAnimNode_IKRig::CacheBones_AnyThread(const FAnimationCacheBonesContext& Co
 	CompactPoseToRigIndices.Reset();
 	const TArray<FBoneIndexType>& RequiredBonesArray = RequiredBones.GetBoneIndicesArray();
 	const FReferenceSkeleton& RefSkeleton = RequiredBones.GetReferenceSkeleton();
+	const FIKRigSkeleton& IKRigSkeleton = IKRigProcessor->GetSkeleton();
 	const int32 NumBones = RequiredBonesArray.Num();
 	for (uint16 Index = 0; Index < NumBones; ++Index)
 	{
@@ -311,7 +312,7 @@ void FAnimNode_IKRig::CacheBones_AnyThread(const FAnimationCacheBonesContext& Co
 		
 		FCompactPoseBoneIndex CPIndex = RequiredBones.MakeCompactPoseIndex(FMeshPoseBoneIndex(MeshBone));
 		const FName Name = RefSkeleton.GetBoneName(MeshBone);
-		CompactPoseToRigIndices.Add(CPIndex) = IKRigProcessor->GetSkeleton().GetBoneIndexFromName(Name);
+		CompactPoseToRigIndices.Add(CPIndex) = IKRigSkeleton.GetBoneIndexFromName(Name);
 	}
 }
 
diff --git a/Engine/Plugins/Animation/IKRig/Source/IKRig/Private/AnimNodes/AnimNode_RetargetPoseFromMesh.cpp b/Engine/Plugins/Animation/IKRig/Source/IKRig/Private/AnimNodes/AnimNode_RetargetPoseFromMesh.cpp
--- a/Engine/Plugins/Animation/IKRig/Source/IKRig/Private/AnimNodes/AnimNode_RetargetPoseFromMesh.cpp
+++ b/Engine/Plugins/Animation/IKRig/Source/IKRig/Private/AnimNodes/AnimNode_RetargetPoseFromMesh.cpp
@@ -99,10 +99,17 @@ void FAnimNode_RetargetPoseFromMesh::Evaluate_AnyThread(FPoseContext& Output)
 	// NOTE: this copies goal targets as well, but these are overwritten by IK chain goals
 	if (bDriveWithAsset)
 	{
-		Processor->CopyAllSettingsFromAsset();
+		Processor->ApplySettingsFromAsset();
+		if (const FRetargetProfile* CurrentProfile = IKRetargeterAsset->GetCurrentProfile())
+		{
+			Processor->ApplySettingsFromProfile(*CurrentProfile);
+		}
 	}
 #endif
 
+	// apply custom profile settings to the processor
+	Processor->ApplySettingsFromProfile(CustomRetargetProfile);
+
 	// run the retargeter
 	const TArray<FTransform>& RetargetedPose = Processor->RunRetargeter(SourceMeshComponentSpaceBoneTransforms, SpeedValuesFromCurves, DeltaTime);
 	DeltaTime = 0.0f;
@@ -180,10 +187,10 @@ void FAnimNode_RetargetPoseFromMesh::UpdateSpeedValuesFromCurves()
 		return;
 	}
 
-	const TArray<TObjectPtr<URetargetChainSettings>> ChainSettings = IKRetargeterAsset->GetAllChainSettings();
+	const TArray<TObjectPtr<URetargetChainSettings>>& ChainSettings = IKRetargeterAsset->GetAllChainSettings();
 	for (const TObjectPtr<URetargetChainSettings>& ChainSetting : ChainSettings)
 	{
-		FName CurveName = ChainSetting->SpeedCurveName;
+		FName CurveName = ChainSetting->Settings.SpeedPlanting.SpeedCurveName;
 		if (CurveName == NAME_None)
 		{
 			continue;
diff --git a/Engine/Plugins/Animation/IKRig/Source/IKRig/Private/IKRigProcessor.cpp b/Engine/Plugins/Animation/IKRig/Source/IKRig/Private/IKRigProcessor.cpp
--- a/Engine/Plugins/Animation/IKRig/Source/IKRig/Private/IKRigProcessor.cpp
+++ b/Engine/Plugins/Animation/IKRig/Source/IKRig/Private/IKRigProcessor.cpp
@@ -14,7 +14,10 @@ UIKRigProcessor::UIKRigProcessor()
 	Log.SetLogTarget(LogName);
 }
 
-void UIKRigProcessor::Initialize(const UIKRigDefinition* InRigAsset, const USkeletalMesh* SkeletalMesh)
+void UIKRigProcessor::Initialize(
+	const UIKRigDefinition* InRigAsset,
+	const USkeletalMesh* SkeletalMesh,
+	const TArray<FName>& ExcludedGoals)
 {
 	// we instantiate UObjects here which MUST be done on game thread...
 	check(IsInGameThread());
@@ -52,17 +55,29 @@ void UIKRigProcessor::Initialize(const UIKRigDefinition* InRigAsset, const USkel
 	
 	// initialize goals based on source asset
 	GoalContainer.Empty();
-	const TArray<UIKRigEffectorGoal*>& GoalsInAsset = InRigAsset->GetGoalArray();
-	for (const UIKRigEffectorGoal* GoalInAsset : GoalsInAsset)
+	const TArray<UIKRigEffectorGoal*>& EffectorGoals = InRigAsset->GetGoalArray();
+	for (const UIKRigEffectorGoal* EffectorGoal : EffectorGoals)
 	{
+		// skip excluded goals
+		if (ExcludedGoals.Contains(EffectorGoal->GoalName))
+		{
+			continue;
+		}
+		
 		// add a copy of the Goal to the container
-		GoalContainer.SetIKGoal(GoalInAsset);
+		GoalContainer.SetIKGoal(EffectorGoal);
 	}
 	
 	// initialize goal bones from asset
 	GoalBones.Reset();
-	for (const UIKRigEffectorGoal* EffectorGoal : GoalsInAsset)
-	{	
+	for (const UIKRigEffectorGoal* EffectorGoal : EffectorGoals)
+	{
+		// skip excluded goals
+		if (ExcludedGoals.Contains(EffectorGoal->GoalName))
+		{
+			continue;
+		}
+		
 		FGoalBone NewGoalBone;
 		NewGoalBone.BoneName = EffectorGoal->BoneName;
 		NewGoalBone.BoneIndex = Skeleton.GetBoneIndexFromName(EffectorGoal->BoneName);
@@ -110,10 +125,18 @@ void UIKRigProcessor::Initialize(const UIKRigDefinition* InRigAsset, const USkel
 			continue;
 		}
 
-		// new solver name
+		// create duplicate solver instance with unique name
 		FString Name = IKRigSolver->GetName() + "_SolverInstance_";
 		Name.AppendInt(SolverIndex++);
 		UIKRigSolver* Solver = DuplicateObject(IKRigSolver, this, FName(*Name));
+
+		// remove excluded goals from the solver
+		for (const FName& ExcludedGoal : ExcludedGoals)
+		{
+			Solver->RemoveGoal(ExcludedGoal);
+		}
+
+		// initialize it and store in the processor
 		Solver->Initialize(Skeleton);
 		Solvers.Add(Solver);
 	}
@@ -333,7 +356,12 @@ const FIKRigGoalContainer& UIKRigProcessor::GetGoalContainer() const
 	return GoalContainer;
 }
 
-FIKRigSkeleton& UIKRigProcessor::GetSkeleton()
+FIKRigSkeleton& UIKRigProcessor::GetSkeletonWriteable()
+{
+	return Skeleton;
+}
+
+const FIKRigSkeleton& UIKRigProcessor::GetSkeleton() const
 {
 	return Skeleton;
 }
diff --git a/Engine/Plugins/Animation/IKRig/Source/IKRig/Private/IKRigSkeleton.cpp b/Engine/Plugins/Animation/IKRig/Source/IKRig/Private/IKRigSkeleton.cpp
--- a/Engine/Plugins/Animation/IKRig/Source/IKRig/Private/IKRigSkeleton.cpp
+++ b/Engine/Plugins/Animation/IKRig/Source/IKRig/Private/IKRigSkeleton.cpp
@@ -2,6 +2,7 @@
 
 #include "IKRigSkeleton.h"
 
+#include "IKRigDefinition.h"
 #include "Engine/SkeletalMesh.h"
 
 // This is the default end of branch index value, meaning we haven't cached it yet
@@ -321,7 +322,7 @@ void FIKRigSkeleton::NormalizeRotations(TArray<FTransform>& Transforms)
 	}
 }
 
-void FIKRigSkeleton::GetChainsInList(const TArray<int32>& SelectedBones, TArray<FIKRigSkeletonChain>& OutChains) const
+void FIKRigSkeleton::GetChainsInList(const TArray<int32>& SelectedBones, TArray<FBoneChain>& OutChains) const
 {	
 	// no bones provided, return empty list
 	if (SelectedBones.IsEmpty())
@@ -365,6 +366,96 @@ void FIKRigSkeleton::GetChainsInList(const TArray<int32>& SelectedBones, TArray<
 
     	const FName StartBoneName = GetBoneNameFromIndex(ChainStart);
     	const FName EndBoneName = GetBoneNameFromIndex(ChainEnd);
-    	OutChains.Emplace(StartBoneName, EndBoneName);	
+    	OutChains.Add(FBoneChain(NAME_None, StartBoneName, EndBoneName));	
     }
 }
+
+bool FIKRigSkeleton::GetBonesInChain(const FBoneChain& BoneChain, TArray<int32>& OutBoneIndices) const
+{
+	// validate start and end bones exist and are not the root
+	const int32 StartIndex = GetBoneIndexFromName(BoneChain.StartBone.BoneName);
+	const int32 EndIndex = GetBoneIndexFromName(BoneChain.EndBone.BoneName);
+	const bool bFoundStartBone = StartIndex > INDEX_NONE;
+	const bool bFoundEndBone = EndIndex > INDEX_NONE;
+
+	// no need to build the chain if start/end indices are wrong 
+	const bool bIsWellFormed = bFoundStartBone && bFoundEndBone && EndIndex >= StartIndex;
+	if (!bIsWellFormed)
+	{
+		return false;
+	}
+	
+	// init array with end bone 
+	OutBoneIndices = {EndIndex};
+
+	// if only one bone in the chain
+	if (EndIndex == StartIndex)
+	{
+		return true;
+	}
+
+	// record all bones in chain while walking up the hierarchy (tip to root of chain)
+	int32 ParentIndex = GetParentIndex(EndIndex);
+	while (ParentIndex > INDEX_NONE && ParentIndex >= StartIndex)
+	{
+		OutBoneIndices.Add(ParentIndex);
+		ParentIndex = GetParentIndex(ParentIndex);
+	}
+
+	// if we walked up till the start bone
+	if (OutBoneIndices.Last() == StartIndex)
+	{
+		// reverse the indices (we want root to tip order)
+		Algo::Reverse(OutBoneIndices);
+		return true;
+	}
+
+	// oops, we walked all the way up without finding the start bone
+	OutBoneIndices.Reset();
+	return false;
+}
+
+bool FIKRigSkeleton::GetMirroredBoneIndices(
+	const TArray<int32>& BoneIndices,
+	TArray<int32>& OutMirroredIndices,
+	float MaxDistanceThreshold) const
+{
+	constexpr float MirrorDistanceThreshold = 5.0f;
+
+	OutMirroredIndices.Reset();
+	for (const int32 BoneIndex : BoneIndices)
+	{
+		const FVector BoneLocation = RefPoseGlobal[BoneIndex].GetLocation();
+		const FVector MirroredLocation = BoneLocation * FVector(-1.f,1.f,1.f);
+		int32 ClosestBoneIndex;
+		float Distance;
+		GetClosestBone(MirroredLocation, ClosestBoneIndex, Distance);
+		if (Distance > MirrorDistanceThreshold || ClosestBoneIndex == INDEX_NONE)
+		{
+			OutMirroredIndices.Reset();
+			return false;
+		}
+		OutMirroredIndices.Add(ClosestBoneIndex);
+	}
+	return true;
+}
+
+void FIKRigSkeleton::GetClosestBone(
+	const FVector& InPoint,
+	int32& OutBoneIndex,
+	float& OutDistance) const
+{
+	int32 Index = 0;
+	OutBoneIndex = INDEX_NONE;
+	OutDistance = TNumericLimits<float>::Max();
+	for (const FTransform& BonePose : RefPoseGlobal)
+	{
+		const float Distance = FVector::Distance(InPoint, BonePose.GetLocation());
+		if (Distance < OutDistance)
+		{
+			OutBoneIndex = Index;
+			OutDistance = Distance;
+		}
+		++Index;
+	}
+}
diff --git a/Engine/Plugins/Animation/IKRig/Source/IKRig/Private/Retargeter/IKRetargetProcessor.cpp b/Engine/Plugins/Animation/IKRig/Source/IKRig/Private/Retargeter/IKRetargetProcessor.cpp
--- a/Engine/Plugins/Animation/IKRig/Source/IKRig/Private/Retargeter/IKRetargetProcessor.cpp
+++ b/Engine/Plugins/Animation/IKRig/Source/IKRig/Private/Retargeter/IKRetargetProcessor.cpp
@@ -8,6 +8,9 @@
 #include "Engine/SkeletalMesh.h"
 #include "Kismet/KismetMathLibrary.h"
 #include "Retargeter/IKRetargeter.h"
+#include "Retargeter/IKRetargetProfile.h"
+#include "Solvers/PointsToRotation.h"
+#include "UObject/UnrealTypePrivate.h"
 
 #define LOCTEXT_NAMESPACE "IKRetargetProcessor"
 
@@ -17,9 +20,13 @@
 void FRetargetSkeleton::Initialize(
 	USkeletalMesh* InSkeletalMesh,
 	const TArray<FBoneChain>& BoneChains,
+	const FName InRetargetPoseName,
 	const FIKRetargetPose* RetargetPose,
 	const FName& RetargetRootBone)
-{	
+{
+	// reset all skeleton data
+	Reset();
+	
 	// record which skeletal mesh this is running on
 	SkeletalMesh = InSkeletalMesh;
 	
@@ -51,7 +58,7 @@ void FRetargetSkeleton::Initialize(
 
 	// update retarget pose to reflect custom offsets (applies stored offsets)
 	// NOTE: this must be done AFTER generating IsBoneInAnyTargetChain array above
-	GenerateRetargetPose(RetargetPose, RetargetRootBone);
+	GenerateRetargetPose(InRetargetPoseName, RetargetPose, RetargetRootBone);
 }
 
 void FRetargetSkeleton::Reset()
@@ -63,8 +70,14 @@ void FRetargetSkeleton::Reset()
 	SkeletalMesh = nullptr;
 }
 
-void FRetargetSkeleton::GenerateRetargetPose(const FIKRetargetPose* InRetargetPose, const FName& RetargetRootBone)
+void FRetargetSkeleton::GenerateRetargetPose(
+	const FName InRetargetPoseName,
+	const FIKRetargetPose* InRetargetPose,
+	const FName& RetargetRootBone)
 {
+	// record the name of the retarget pose (prevents re-initialization if profile swaps it)
+	RetargetPoseName = InRetargetPoseName;
+	
 	// initialize retarget pose to the skeletal mesh reference pose
 	RetargetLocalPose = SkeletalMesh->GetRefSkeleton().GetRefBonePose();
 	// copy local pose to global
@@ -294,10 +307,13 @@ int32 FRetargetSkeleton::GetParentIndex(const int32 BoneIndex) const
 void FTargetSkeleton::Initialize(
 	USkeletalMesh* InSkeletalMesh,
 	const TArray<FBoneChain>& BoneChains,
+	const FName& InRetargetPoseName,
 	const FIKRetargetPose* RetargetPose,
 	const FName& RetargetRootBone)
 {
-	FRetargetSkeleton::Initialize(InSkeletalMesh, BoneChains, RetargetPose, RetargetRootBone);
+	Reset();
+	
+	FRetargetSkeleton::Initialize(InSkeletalMesh, BoneChains, InRetargetPoseName, RetargetPose, RetargetRootBone);
 
 	// make storage for per-bone "Is Retargeted" flag (used for hierarchy updates)
 	// these are bones that are in a target chain that is mapped to a source chain (ie, will actually be retargeted)
@@ -383,38 +399,16 @@ void FTargetSkeleton::SetBoneIsRetargeted(const int32 BoneIndex, const bool IsRe
 	IsBoneRetargeted[BoneIndex] = IsRetargeted;
 }
 
-void FRetargetChainSettings::CopySettingsFromAsset(const URetargetChainSettings* AssetChainSettings)
-{
-	TargetChainName = AssetChainSettings->TargetChain;
-
-	CopyPoseUsingFK = AssetChainSettings->CopyPoseUsingFK;
-	RotationMode = AssetChainSettings->RotationMode;
-	RotationAlpha = AssetChainSettings->RotationAlpha;
-	TranslationMode = AssetChainSettings->TranslationMode;
-	TranslationAlpha = AssetChainSettings->TranslationAlpha;
-		
-	DriveIKGoal = AssetChainSettings->DriveIKGoal;
-	Extension = AssetChainSettings->Extension;
-	StaticOffset = AssetChainSettings->StaticOffset;
-	StaticLocalOffset = AssetChainSettings->StaticLocalOffset;
-	StaticRotationOffset = AssetChainSettings->StaticRotationOffset;
-	BlendToSource = AssetChainSettings->BlendToSource;
-	BlendToSourceWeights = AssetChainSettings->BlendToSourceWeights;
-
-	UseSpeedCurveToPlantIK = AssetChainSettings->UseSpeedCurveToPlantIK;
-	UnplantStiffness = AssetChainSettings->UnplantStiffness;
-	UnplantCriticalDamping = AssetChainSettings->UnplantCriticalDamping;
-	SpeedThreshold = AssetChainSettings->VelocityThreshold;
-	SpeedCurveName = AssetChainSettings->SpeedCurveName;
-};
-
 bool FChainFK::Initialize(
 	const FRetargetSkeleton& Skeleton,
-	const TArray<int32>& BoneIndices,
+	const TArray<int32>& InBoneIndices,
 	const TArray<FTransform>& InitialGlobalPose,
 	FIKRigLogger& Log)
 {
-	check(!BoneIndices.IsEmpty());
+	check(!InBoneIndices.IsEmpty());
+
+	// store for debugging purposes
+	BoneIndices = InBoneIndices;
 
 	// store all the initial bone transforms in the bone chain
 	InitialGlobalTransforms.Reset();
@@ -509,23 +503,23 @@ void FChainFK::FillTransformsWithLocalSpaceOfChain(
 }
 
 void FChainFK::PutCurrentTransformsInRefPose(
-	const TArray<int32>& BoneIndices,
+	const TArray<int32>& InBoneIndices,
 	const FRetargetSkeleton& Skeleton,
 	const TArray<FTransform>& InCurrentGlobalPose)
 {
 	// update chain current transforms to the retarget pose in global space
-	for (int32 ChainIndex=0; ChainIndex<BoneIndices.Num(); ++ChainIndex)
+	for (int32 ChainIndex=0; ChainIndex<InBoneIndices.Num(); ++ChainIndex)
 	{
 		// update first bone in chain based on the incoming parent
 		if (ChainIndex == 0)
 		{
-			const int32 BoneIndex = BoneIndices[ChainIndex];
+			const int32 BoneIndex = InBoneIndices[ChainIndex];
 			CurrentGlobalTransforms[ChainIndex] = Skeleton.GetGlobalRefPoseOfSingleBone(BoneIndex, InCurrentGlobalPose);
 		}
 		else
 		{
 			// all subsequent bones in chain are based on previous parent
-			const int32 BoneIndex = BoneIndices[ChainIndex];
+			const int32 BoneIndex = InBoneIndices[ChainIndex];
 			const FTransform& ParentGlobalTransform = CurrentGlobalTransforms[ChainIndex-1];
 			const FTransform& ChildLocalTransform = Skeleton.RetargetLocalPose[BoneIndex];
 			CurrentGlobalTransforms[ChainIndex] = ChildLocalTransform * ParentGlobalTransform;
@@ -573,7 +567,7 @@ void FChainEncoderFK::TransformCurrentChainTransforms(const FTransform& NewParen
 
 void FChainDecoderFK::DecodePose(
 	const FRootRetargeter& RootRetargeter,
-	const FRetargetChainSettings& Settings,
+	const FTargetChainSettings& Settings,
 	const TArray<int32>& TargetBoneIndices,
     FChainEncoderFK& SourceChain,
     const FTargetSkeleton& TargetSkeleton,
@@ -602,7 +596,7 @@ void FChainDecoderFK::DecodePose(
 	SourceChain.TransformCurrentChainTransforms(SourceChainParentTransform);
 
 	// if FK retargeting has been disabled for this chain, then simply set it to the retarget pose
-	if (!Settings.CopyPoseUsingFK)
+	if (!Settings.FK.EnableFK)
 	{
 		// put the chain in the global ref pose (globally rotated by parent bone in it's currently retargeted state)
 		PutCurrentTransformsInRefPose(TargetBoneIndices, TargetSkeleton, InOutGlobalPose);
@@ -630,7 +624,7 @@ void FChainDecoderFK::DecodePose(
 		FTransform SourceInitialTransform;
 
 		// get source current / initial transforms for this bone
-		switch (Settings.RotationMode)
+		switch (Settings.FK.RotationMode)
 		{
 			case ERetargetRotationMode::Interpolated:
 			{
@@ -703,7 +697,7 @@ void FChainDecoderFK::DecodePose(
 			ParentGlobalTransform = InOutGlobalPose[ParentIndex];
 		}
 		FVector OutPosition;
-		switch (Settings.TranslationMode)
+		switch (Settings.FK.TranslationMode)
 		{
 			case ERetargetTranslationMode::None:
 				{
@@ -731,13 +725,14 @@ void FChainDecoderFK::DecodePose(
 		const FVector OutScale = SourceCurrentScale + (TargetInitialScale - SourceInitialScale);
 		
 		// apply output transform
-		InOutGlobalPose[BoneIndex] = FTransform(OutRotation, OutPosition, OutScale);
+		CurrentGlobalTransforms[ChainIndex] = FTransform(OutRotation, OutPosition, OutScale);
+		InOutGlobalPose[BoneIndex] = CurrentGlobalTransforms[ChainIndex];
 	}
 
 	// apply final blending between retarget pose of chain and newly retargeted pose
 	// blend must be done in local space, so we do it in a separate loop after full chain pose is generated
 	// (skipped if the alphas are not near 1.0)
-	if (!FMath::IsNearlyEqual(Settings.RotationAlpha, 1.0f) || !FMath::IsNearlyEqual(Settings.TranslationAlpha, 1.0f))
+	if (!FMath::IsNearlyEqual(Settings.FK.RotationAlpha, 1.0f) || !FMath::IsNearlyEqual(Settings.FK.TranslationAlpha, 1.0f))
 	{
 		TArray<FTransform> NewLocalTransforms;
 		NewLocalTransforms.SetNum(InitialLocalTransforms.Num());
@@ -748,8 +743,8 @@ void FChainDecoderFK::DecodePose(
 			// blend between current local pose and initial local pose
 			FTransform& NewLocalTransform = NewLocalTransforms[ChainIndex];
 			const FTransform& RefPoseLocalTransform = InitialLocalTransforms[ChainIndex];
-			NewLocalTransform.SetTranslation(FMath::Lerp(RefPoseLocalTransform.GetTranslation(), NewLocalTransform.GetTranslation(), Settings.TranslationAlpha));
-			NewLocalTransform.SetRotation(FQuat::FastLerp(RefPoseLocalTransform.GetRotation(), NewLocalTransform.GetRotation(), Settings.RotationAlpha).GetNormalized());
+			NewLocalTransform.SetTranslation(FMath::Lerp(RefPoseLocalTransform.GetTranslation(), NewLocalTransform.GetTranslation(), Settings.FK.TranslationAlpha));
+			NewLocalTransform.SetRotation(FQuat::FastLerp(RefPoseLocalTransform.GetRotation(), NewLocalTransform.GetRotation(), Settings.FK.RotationAlpha).GetNormalized());
 
 			// put blended transforms back in global space and store in final output pose
 			const int32 BoneIndex = TargetBoneIndices[ChainIndex];
@@ -924,68 +919,63 @@ bool FChainRetargeterIK::InitializeTarget(
 }
 	
 void FChainRetargeterIK::DecodePose(
-	const FRetargetChainSettings& Settings,
+	const FTargetChainSettings& Settings,
+	const FRootRetargeter& RootRetargeter,
 	const TMap<FName, float>& SpeedValuesFromCurves,
 	const float DeltaTime,
-    const TArray<FTransform>& OutGlobalPose,
-    FDecodedIKChain& OutResults)
+    const TArray<FTransform>& InGlobalPose)
 {
-	// record the end bone rotation on the input pose
-	const FQuat InputEndRotation = OutGlobalPose[Target.BoneIndexC].GetRotation();
-	// we have to "undo" the end bone delta, otherwise we will get double-transformations because the FK pass has already rotated the foot
-	const FQuat InputToInitialDeltaRotation = InputEndRotation * Target.InitialEndRotation.Inverse();
-	const FQuat Rotation = InputToInitialDeltaRotation * Target.InitialEndRotation;
-	
-	if (!Settings.DriveIKGoal)
-	{
-		// set goal transform to the input coming from the previous retarget phase (FK if enabled)
-		OutResults.EndEffectorPosition = OutGlobalPose[Target.BoneIndexC].GetTranslation();
-		OutResults.EndEffectorRotation = Rotation;
-		return;
-	}
-	
+	// apply delta rotation from input
+	const FQuat DeltaRotation = Source.CurrentEndRotation * Source.InitialEndRotation.Inverse();
+	FQuat GoalRotation = DeltaRotation * Target.InitialEndRotation;
+
 	// apply static rotation offset in the local space of the foot
-	const FQuat GoalRotation = Rotation * Settings.StaticRotationOffset.Quaternion();
+	GoalRotation = GoalRotation * Settings.IK.StaticRotationOffset.Quaternion();
 
 	//
 	// calculate position of IK goal ...
 	//
 	
 	// set position to length-scaled direction from source limb
-	const FVector Start = OutGlobalPose[Target.BoneIndexA].GetTranslation();
+	const FVector InvAffectIKWeights = FVector::OneVector - RootRetargeter.Settings.GetAffectIKWeightVector();
+	const FVector InvRootModification = RootRetargeter.Target.RootTranslationDelta * InvAffectIKWeights;
+	const FVector Start = InGlobalPose[Target.BoneIndexA].GetTranslation() - InvRootModification;
 	FVector GoalPosition = Start + (Source.CurrentEndDirectionNormalized * Target.InitialLength);
 
 	// blend to source location
-	if (Settings.BlendToSource > KINDA_SMALL_NUMBER)
+	if (Settings.IK.BlendToSource > KINDA_SMALL_NUMBER)
 	{
-		const FVector Weight = Settings.BlendToSource * Settings.BlendToSourceWeights;
-		GoalPosition.X = FMath::Lerp(GoalPosition.X, Source.CurrentEndPosition.X, Weight.X);
-		GoalPosition.Y = FMath::Lerp(GoalPosition.Y, Source.CurrentEndPosition.Y, Weight.Y);
-		GoalPosition.Z = FMath::Lerp(GoalPosition.Z, Source.CurrentEndPosition.Z, Weight.Z);
+		const FVector RootModification = RootRetargeter.Target.RootTranslationDelta * RootRetargeter.Settings.GetAffectIKWeightVector();
+		const FVector Weight = Settings.IK.BlendToSource * Settings.IK.BlendToSourceWeights;
+		const FVector SourceLocation = Source.CurrentEndPosition + RootModification;
+		GoalPosition.X = FMath::Lerp(GoalPosition.X, SourceLocation.X, Weight.X);
+		GoalPosition.Y = FMath::Lerp(GoalPosition.Y, SourceLocation.Y, Weight.Y);
+		GoalPosition.Z = FMath::Lerp(GoalPosition.Z, SourceLocation.Z, Weight.Z);
 	}
 
 	// apply global static offset
-	GoalPosition += Settings.StaticOffset;
+	GoalPosition += Settings.IK.StaticOffset;
 
 	// apply local static offset
-	GoalPosition += GoalRotation.RotateVector(Settings.StaticLocalOffset);
+	GoalPosition += GoalRotation.RotateVector(Settings.IK.StaticLocalOffset);
 
 	// apply extension
-	if (!FMath::IsNearlyEqual(Settings.Extension, 1.0f))
+	if (!FMath::IsNearlyEqual(Settings.IK.Extension, 1.0f))
 	{
-		GoalPosition = Start + (GoalPosition - Start) * Settings.Extension;	
+		GoalPosition = Start + (GoalPosition - Start) * Settings.IK.Extension;	
 	}
 	
 	// match velocity
-	if (!ResetThisTick && Settings.UseSpeedCurveToPlantIK && SpeedValuesFromCurves.Contains(Settings.SpeedCurveName))
+	const FTargetChainSpeedPlantSettings& PlantSettings = Settings.SpeedPlanting;
+	if (!ResetThisTick && PlantSettings.EnableSpeedPlanting && SpeedValuesFromCurves.Contains(PlantSettings.SpeedCurveName))
 	{
-		const float SourceSpeed = SpeedValuesFromCurves[Settings.SpeedCurveName];
-		if (SourceSpeed < 0.0f || SourceSpeed > Settings.SpeedThreshold)
+		const float SourceSpeed = SpeedValuesFromCurves[PlantSettings.SpeedCurveName];
+		if (SourceSpeed < 0.0f || SourceSpeed > PlantSettings.SpeedThreshold)
 		{
 			GoalPosition = UKismetMathLibrary::VectorSpringInterp(
 				Target.PrevEndPosition, GoalPosition, PlantingSpringState,
-				Settings.UnplantStiffness,
-				Settings.UnplantCriticalDamping,
+				PlantSettings.UnplantStiffness,
+				PlantSettings.UnplantCriticalDamping,
 				DeltaTime, 1.0f, 0.0f);
 		}
 		else
@@ -996,15 +986,23 @@ void FChainRetargeterIK::DecodePose(
 	}
 	
 	// output transform
-	OutResults.EndEffectorPosition = GoalPosition;
-	OutResults.EndEffectorRotation = GoalRotation;
-	OutResults.PoleVectorPosition = FVector::OneVector; // TODO calc pole vector position
+	Results.EndEffectorPosition = GoalPosition;
+	Results.EndEffectorRotation = GoalRotation;
+	Results.PoleVectorPosition = FVector::OneVector; // TODO calc pole vector position
 	Target.PrevEndPosition = GoalPosition;
 	ResetThisTick = false;
 }
 
+void FChainRetargeterIK::SaveDebugInfo(const TArray<FTransform>& InGlobalPose)
+{
+#if WITH_EDITOR
+	DebugData.InputTransformStart = InGlobalPose[Target.BoneIndexA];
+	DebugData.InputTransformEnd = InGlobalPose[Target.BoneIndexC];
+	DebugData.OutputTransformEnd = FTransform(Results.EndEffectorRotation, Results.EndEffectorPosition);
+#endif
+}
+
 bool FRetargetChainPair::Initialize(
-	URetargetChainSettings* InSettings,
     const FBoneChain& SourceBoneChain,
     const FBoneChain& TargetBoneChain,
     const FRetargetSkeleton& SourceSkeleton,
@@ -1031,17 +1029,7 @@ bool FRetargetChainPair::Initialize(
 		return false;
     }
 
-	// ensure valid settings object
-	if (InSettings == nullptr)
-	{
-		Log.LogWarning( FText::Format(
-			LOCTEXT("MissingChainSettings", "IK Retargeter target bone chain, '{0}', has null settings."),
-			FText::FromName(TargetBoneChain.ChainName)));
-		return false;
-	}
-
 	// store attributes of chain
-	Settings.CopySettingsFromAsset(InSettings);
 	SourceBoneChainName = SourceBoneChain.ChainName;
 	TargetBoneChainName = TargetBoneChain.ChainName;
 	
@@ -1090,14 +1078,13 @@ bool FRetargetChainPair::ValidateBoneChainWithSkeletalMesh(
 }
 
 bool FRetargetChainPairFK::Initialize(
-	URetargetChainSettings* InSettings,
 	const FBoneChain& SourceBoneChain,
 	const FBoneChain& TargetBoneChain,
 	const FRetargetSkeleton& SourceSkeleton,
 	const FTargetSkeleton& TargetSkeleton,
 	FIKRigLogger& Log)
 {
-	const bool bChainInitialized = FRetargetChainPair::Initialize(InSettings, SourceBoneChain, TargetBoneChain, SourceSkeleton, TargetSkeleton, Log);
+	const bool bChainInitialized = FRetargetChainPair::Initialize(SourceBoneChain, TargetBoneChain, SourceSkeleton, TargetSkeleton, Log);
 	if (!bChainInitialized)
 	{
 		return false;
@@ -1127,7 +1114,6 @@ bool FRetargetChainPairFK::Initialize(
 }
 
 bool FRetargetChainPairIK::Initialize(
-	URetargetChainSettings* InSettings,
 	const FBoneChain& SourceBoneChain,
 	const FBoneChain& TargetBoneChain,
 	const FRetargetSkeleton& SourceSkeleton,
@@ -1145,7 +1131,7 @@ bool FRetargetChainPairIK::Initialize(
 	IKGoalName = TargetBoneChain.IKGoalName;
 
 	// initialize bone chains
-	const bool bChainInitialized = FRetargetChainPair::Initialize(InSettings, SourceBoneChain, TargetBoneChain, SourceSkeleton, TargetSkeleton, Log);
+	const bool bChainInitialized = FRetargetChainPair::Initialize(SourceBoneChain, TargetBoneChain, SourceSkeleton, TargetSkeleton, Log);
 	if (!bChainInitialized)
 	{
 		return false;
@@ -1227,20 +1213,10 @@ bool FRootRetargeter::InitializeTarget(
 	Target.InitialHeight = TargetInitialTransform.GetTranslation().Z;
 	Target.InitialRotation = TargetInitialTransform.GetRotation();
 	Target.InitialPosition = TargetInitialTransform.GetTranslation();
-
+	
 	return true;
 }
 
-void FRetargetRootSettings::CopySettingsFromAsset(const URetargetRootSettings* AssetRootSettings)
-{
-	GlobalScaleHorizontal = AssetRootSettings->GlobalScaleHorizontal;
-	GlobalScaleVertical = AssetRootSettings->GlobalScaleVertical;
-	BlendToSource = AssetRootSettings->BlendToSource;
-	BlendToSourceWeights = AssetRootSettings->BlendToSourceWeights;
-	StaticOffset = AssetRootSettings->StaticOffset;
-	StaticRotationOffset = AssetRootSettings->StaticRotationOffset;
-}
-
 void FRootRetargeter::Reset()
 {
 	Source = FRootSource();
@@ -1255,25 +1231,51 @@ void FRootRetargeter::EncodePose(const TArray<FTransform>& SourceGlobalPose)
 	Source.CurrentRotation = SourceTransform.GetRotation();	
 }
 
-void FRootRetargeter::DecodePose(TArray<FTransform>& OutTargetGlobalPose) const
+void FRootRetargeter::DecodePose(TArray<FTransform>& OutTargetGlobalPose)
 {
-	// scale normalized position by root height
-	FVector RetargetedPosition = Source.CurrentPositionNormalized * Target.InitialHeight;// * GlobalScaleVertical;
-	RetargetedPosition.Z *= Settings.GlobalScaleVertical;
-	// globally scale offset of root
-	const FVector RootOffset = (RetargetedPosition - Target.InitialPosition) * FVector(Settings.GlobalScaleHorizontal, Settings.GlobalScaleHorizontal, 1.0f);
-	RetargetedPosition = Target.InitialPosition + RootOffset;
-	// blend the retarget root position towards the source retarget root position
-	FVector Position = FMath::Lerp(RetargetedPosition, Source.CurrentPosition, Settings.BlendToSource*Settings.BlendToSourceWeights);
-	// apply a static offset
-	Position += Settings.StaticOffset;
-
-	// calc offset between initial source/target root rotations
-	const FQuat RotationDelta = Source.CurrentRotation * Source.InitialRotation.Inverse();
-	// add retarget pose delta to the current source rotation
-	FQuat Rotation = RotationDelta * Target.InitialRotation;
-	// add static rotation offset
-	Rotation = Settings.StaticRotationOffset.Quaternion() * Rotation;
+	// retarget position
+	FVector Position;
+	{
+		// generate basic retarget root position by scaling the normalized position by root height
+		const FVector RetargetedPosition = Source.CurrentPositionNormalized * Target.InitialHeight;
+		
+		// blend the retarget root position towards the source retarget root position
+		Position = FMath::Lerp(RetargetedPosition, Source.CurrentPosition, Settings.BlendToSource*Settings.BlendToSourceWeights);
+
+		// apply vertical / horizontal scaling of motion
+		FVector ScaledRetargetedPosition = Position;
+		ScaledRetargetedPosition.Z *= Settings.ScaleVertical;
+		const FVector HorizontalOffset = (ScaledRetargetedPosition - Target.InitialPosition) * FVector(Settings.ScaleHorizontal, Settings.ScaleHorizontal, 1.0f);
+		Position = Target.InitialPosition + HorizontalOffset;
+		
+		// apply a static offset
+		Position += Settings.TranslationOffset;
+
+		// blend with alpha
+		Position = FMath::Lerp(Target.InitialPosition, Position, Settings.TranslationAlpha);
+
+		// record the delta created by all the modifications made to the root translation
+		Target.RootTranslationDelta = Position - RetargetedPosition;
+	}
+
+	// retarget rotation
+	FQuat Rotation;
+	{
+		// calc offset between initial source/target root rotations
+		const FQuat RotationDelta = Source.CurrentRotation * Source.InitialRotation.Inverse();
+		// add retarget pose delta to the current source rotation
+		const FQuat RetargetedRotation = RotationDelta * Target.InitialRotation;
+
+		// add static rotation offset
+		Rotation = RetargetedRotation * Settings.RotationOffset.Quaternion();
+
+		// blend with alpha
+		Rotation = FQuat::FastLerp(Target.InitialRotation, Rotation, Settings.RotationAlpha);
+		Rotation.Normalize();
+
+		// record the delta created by all the modifications made to the root rotation
+		Target.RootRotationDelta = RetargetedRotation * Target.InitialRotation.Inverse();
+	}
 
 	// apply to target
 	FTransform& TargetRootTransform = OutTargetGlobalPose[Target.BoneIndex];
@@ -1293,19 +1295,15 @@ void UIKRetargetProcessor::Initialize(
 		UIKRetargeter* InRetargeterAsset,
 		const bool bSuppressWarnings)
 {
+	// reset all initialized flags
 	bIsInitialized = false;
+	bRootsInitialized = false;
+	bAtLeastOneValidBoneChainPair = false;
+	bIKRigInitialized = false;
 	
 	// record source asset
 	RetargeterAsset = InRetargeterAsset;
 
-	// reset all data structures
-	SourceSkeleton.Reset();
-	TargetSkeleton.Reset();
-	IKRigProcessor = nullptr;
-	ChainPairsFK.Reset();
-	ChainPairsIK.Reset();
-	RootRetargeter.Reset();
-
 	const UIKRigDefinition* SourceIKRig = RetargeterAsset->GetSourceIKRig();
 	const UIKRigDefinition* TargetIKRig = RetargeterAsset->GetTargetIKRig();
 
@@ -1330,21 +1328,18 @@ void UIKRetargetProcessor::Initialize(
 		Log.LogError(LOCTEXT("MissingTargetIKRig", "IK Retargeter unable to initialize. Missing target IK Rig asset."));
 		return;
 	}
-	if (!RetargeterAsset->GetCurrentRetargetPose(ERetargetSourceOrTarget::Target))
-	{
-		Log.LogError(LOCTEXT("MissingRetargetPose", "IK Retargeter unable to initialize. Missing retarget pose."));
-		return;
-	}
 	
 	// initialize skeleton data for source and target
 	SourceSkeleton.Initialize(
 		SourceSkeletalMesh,
 		SourceIKRig->GetRetargetChains(),
+		RetargeterAsset->GetCurrentRetargetPoseName(ERetargetSourceOrTarget::Source),
 		RetargeterAsset->GetCurrentRetargetPose(ERetargetSourceOrTarget::Source),
 		SourceIKRig->GetRetargetRoot());
 	TargetSkeleton.Initialize(
 		TargetSkeletalMesh,
 		TargetIKRig->GetRetargetChains(),
+		RetargeterAsset->GetCurrentRetargetPoseName(ERetargetSourceOrTarget::Target),
 		RetargeterAsset->GetCurrentRetargetPose(ERetargetSourceOrTarget::Target),
 		TargetIKRig->GetRetargetRoot());
 
@@ -1379,14 +1374,17 @@ void UIKRetargetProcessor::Initialize(
 				LOCTEXT("SuccessfulInit", "Success! The IK Retargeter is ready to transfer animation from the source, {0} to the target, {1}"),
 				FText::FromString(SourceSkeletalMesh->GetName()), FText::FromString(TargetSkeletalMesh->GetName())));
 	}
+
+	// copy the initial settings from the asset
+	ApplySettingsFromAsset();
 	
 	bIsInitialized = true;
 }
 
 bool UIKRetargetProcessor::InitializeRoots()
 {
-	// copy the retarget root settings
-	RootRetargeter.Settings.CopySettingsFromAsset(RetargeterAsset->GetRetargetRootSettings());
+	// reset root data
+	RootRetargeter.Reset();
 	
 	// initialize root encoder
 	const FName SourceRootBoneName = RetargeterAsset->GetSourceIKRig()->GetRetargetRoot();
@@ -1413,6 +1411,9 @@ bool UIKRetargetProcessor::InitializeRoots()
 
 bool UIKRetargetProcessor::InitializeBoneChainPairs()
 {
+	ChainPairsFK.Reset();
+	ChainPairsIK.Reset();
+	
 	const UIKRigDefinition* TargetIKRig = RetargeterAsset->GetTargetIKRig();
 	const UIKRigDefinition* SourceIKRig = RetargeterAsset->GetSourceIKRig();
 	
@@ -1450,16 +1451,35 @@ bool UIKRetargetProcessor::InitializeBoneChainPairs()
 
 		// all chains are loaded as FK (giving IK better starting pose)
 		FRetargetChainPairFK ChainPair;
-		if (ChainPair.Initialize(ChainMap, *SourceBoneChain, *TargetBoneChain, SourceSkeleton, TargetSkeleton, Log))
+		if (ChainPair.Initialize(*SourceBoneChain, *TargetBoneChain, SourceSkeleton, TargetSkeleton, Log))
 		{
 			ChainPairsFK.Add(ChainPair);
 		}
 		
 		// load IK chain
-		FRetargetChainPairIK ChainPairIK;
-		if (ChainPairIK.Initialize(ChainMap, *SourceBoneChain, *TargetBoneChain, SourceSkeleton, TargetSkeleton, Log))
+		if (ChainMap->Settings.IK.EnableIK)
+		{
+			FRetargetChainPairIK ChainPairIK;
+			if (ChainPairIK.Initialize(*SourceBoneChain, *TargetBoneChain, SourceSkeleton, TargetSkeleton, Log))
+			{
+				ChainPairsIK.Add(ChainPairIK);
+			}
+		}
+
+		// warn user if IK goal is not on the END bone of the target chain. It will still work, but may produce bad results.
+		const TArray<UIKRigEffectorGoal*> AllGoals = TargetIKRig->GetGoalArray();
+		for (const UIKRigEffectorGoal*Goal : AllGoals)
 		{
-			ChainPairsIK.Add(ChainPairIK);
+			if (Goal->GoalName == TargetBoneChain->IKGoalName)
+			{
+				if (Goal->BoneName != TargetBoneChain->EndBone.BoneName)
+				{
+					Log.LogWarning( FText::Format(
+				LOCTEXT("TargetIKNotOnEndBone", "Retarget chain, '{0}' has an IK goal that is not on the End Bone of the chain."),
+					FText::FromString(TargetBoneChain->ChainName.ToString())));
+				}
+				break;
+			}
 		}
 	}
 
@@ -1508,13 +1528,32 @@ bool UIKRetargetProcessor::InitializeBoneChainPairs()
 }
 
 bool UIKRetargetProcessor::InitializeIKRig(UObject* Outer, const USkeletalMesh* InSkeletalMesh)
-{	
-	// initialize IK Rig runtime processor
-	if (!IKRigProcessor)
+{
+	// get the target IK Rig asset
+	const UIKRigDefinition* IKRig = RetargeterAsset->GetTargetIKRig();
+	
+	// gather list of excluded goals based on any chain mapping that has it's IK disabled
+	TArray<FName> GoalsToExclude;
+	const TArray<TObjectPtr<URetargetChainSettings>>& ChainMapping = RetargeterAsset->GetAllChainSettings();
+	const TArray<FBoneChain>& IKRigChains = IKRig->GetRetargetChains();
+	for (const URetargetChainSettings* ChainMap : ChainMapping)
 	{
-		IKRigProcessor = NewObject<UIKRigProcessor>(Outer);	
+		if (ChainMap->Settings.IK.EnableIK)
+		{
+			continue;
+		}
+		for (const FBoneChain& Chain : IKRigChains)
+		{
+			if (Chain.ChainName == ChainMap->TargetChain)
+			{
+				GoalsToExclude.Add(Chain.IKGoalName);
+			}
+		}
 	}
-	IKRigProcessor->Initialize(RetargeterAsset->GetTargetIKRig(), InSkeletalMesh);
+	
+	// initialize IK Rig runtime processor
+	IKRigProcessor = NewObject<UIKRigProcessor>(Outer);	
+	IKRigProcessor->Initialize(IKRig, InSkeletalMesh, GoalsToExclude);
 	if (!IKRigProcessor->IsInitialized())
 	{
 		return false;
@@ -1546,8 +1585,11 @@ TArray<FTransform>&  UIKRetargetProcessor::RunRetargeter(
 	// start from retarget pose
 	TargetSkeleton.OutputGlobalPose = TargetSkeleton.RetargetGlobalPose;
 
+	// get the global settings
+	const FRetargetGlobalSettings& GlobalSettings = RetargeterAsset->GetGlobalSettings();
+
 	// ROOT retargeting
-	if (RetargeterAsset->bRetargetRoot && bRootsInitialized)
+	if (GlobalSettings.bEnableRoot && bRootsInitialized)
 	{
 		RunRootRetarget(InSourceGlobalPose, TargetSkeleton.OutputGlobalPose);
 		// update global transforms below root
@@ -1555,7 +1597,7 @@ TArray<FTransform>&  UIKRetargetProcessor::RunRetargeter(
 	}
 	
 	// FK CHAIN retargeting
-	if (RetargeterAsset->bRetargetFK && bAtLeastOneValidBoneChainPair)
+	if (GlobalSettings.bEnableFK && bAtLeastOneValidBoneChainPair)
 	{
 		RunFKRetarget(InSourceGlobalPose, TargetSkeleton.OutputGlobalPose);
 		// update all the bones that are not controlled by FK chains or root
@@ -1563,7 +1605,7 @@ TArray<FTransform>&  UIKRetargetProcessor::RunRetargeter(
 	}
 	
 	// IK CHAIN retargeting
-	if (RetargeterAsset->bRetargetIK && bAtLeastOneValidBoneChainPair && bIKRigInitialized)
+	if (GlobalSettings.bEnableIK && bAtLeastOneValidBoneChainPair && bIKRigInitialized)
 	{
 		RunIKRetarget(InSourceGlobalPose, TargetSkeleton.OutputGlobalPose, SpeedValuesFromCurves, DeltaTime);
 	}
@@ -1603,7 +1645,7 @@ void UIKRetargetProcessor::RunFKRetarget(
 
 void UIKRetargetProcessor::RunIKRetarget(
 	const TArray<FTransform>& InSourceGlobalPose,
-    TArray<FTransform>& OutTargeGlobalPose,
+    TArray<FTransform>& OutTargetGlobalPose,
     const TMap<FName, float>& SpeedValuesFromCurves,
     const float DeltaTime)
 {
@@ -1617,53 +1659,170 @@ void UIKRetargetProcessor::RunIKRetarget(
 		return; // skip IK
 	}
 	
-	// spin through IK chains
+	// retarget the IK goals to their new locations based on input pose
 	for (FRetargetChainPairIK& ChainPair : ChainPairsIK)
 	{
 		// encode them all using the input pose
 		ChainPair.IKChainRetargeter.EncodePose(InSourceGlobalPose);
 		// decode the IK goal and apply to IKRig
-		FDecodedIKChain OutIKGoal;
 		ChainPair.IKChainRetargeter.DecodePose(
 			ChainPair.Settings,
+			RootRetargeter,
 			SpeedValuesFromCurves,
 			DeltaTime,
-			OutTargeGlobalPose,
-			OutIKGoal);
-		// set the goal transform on the IK Rig
+			OutTargetGlobalPose);
+	}
+
+	// apply stride warping to the IK goals
+	RunStrideWarping(OutTargetGlobalPose);
+
+	// set the goal transform on the IK Rig
+	for (FRetargetChainPairIK& ChainPair : ChainPairsIK)
+	{
 		FIKRigGoal Goal = FIKRigGoal(
 			ChainPair.IKGoalName,
-			OutIKGoal.EndEffectorPosition,
-			OutIKGoal.EndEffectorRotation,
+			ChainPair.IKChainRetargeter.Results.EndEffectorPosition,
+			ChainPair.IKChainRetargeter.Results.EndEffectorRotation,
 			1.0f,
 			1.0f,
 			EIKRigGoalSpace::Component,
 			EIKRigGoalSpace::Component);
 		IKRigProcessor->SetIKGoal(Goal);
+
+		ChainPair.IKChainRetargeter.SaveDebugInfo(OutTargetGlobalPose);
 	}
 
 	// copy input pose to start IK solve from
-	IKRigProcessor->SetInputPoseGlobal(OutTargeGlobalPose);
+	IKRigProcessor->SetInputPoseGlobal(OutTargetGlobalPose);
 	// run IK solve
 	IKRigProcessor->Solve();
 	// copy results of solve
-	IKRigProcessor->CopyOutputGlobalPoseToArray(OutTargeGlobalPose);
+	IKRigProcessor->CopyOutputGlobalPoseToArray(OutTargetGlobalPose);
 }
 
-void UIKRetargetProcessor::ResetPlanting()
+void UIKRetargetProcessor::RunStrideWarping(const TArray<FTransform>& InTargetGlobalPose)
 {
+	// check if we should apply stride warping
+	const FRetargetGlobalSettings& Settings = RetargeterAsset->GetGlobalSettings();
+	if (!(Settings.bIKWarping && bRootsInitialized))
+	{
+		return;
+	}
+	
+	// get sample points to use for "best fit" global body rotation
+	TArray<FVector> InitialPoints;
+	TArray<FVector> CurrentPoints;
+	if (Settings.DirectionSource == EWarpingDirectionSource::Goals)
+	{
+		// use goals to determine the body's rotation
+		for (const FRetargetChainPairIK& ChainPair : ChainPairsIK)
+		{
+			if (!ChainPair.Settings.IK.bAffectedByIKWarping)
+			{
+				continue;
+			}
+			FVector A = ChainPair.IKChainRetargeter.Target.InitialEndPosition;
+			FVector B = ChainPair.IKChainRetargeter.Results.EndEffectorPosition;
+
+			// flatten into 2D for more robust yaw construction (what really matters)
+			A.Z = 0.0f;
+			B.Z = 0.0f;
+		
+			InitialPoints.Add(A);
+			CurrentPoints.Add(B);
+		}
+	}
+	else
+	{
+		// use chain to determine the body's rotation
+		const FRetargetChainPairFK* SpineChain = ChainPairsFK.FindByPredicate(
+				[&Settings](const FRetargetChainPairFK& Other)
+				{
+					return Other.TargetBoneChainName == Settings.DirectionChain;
+				});
+		if (!SpineChain)
+		{
+			return;
+		}
+
+		const TArray<FTransform>& InitialChainTransforms = SpineChain->FKDecoder.InitialGlobalTransforms;
+		const TArray<FTransform>& CurrentChainTransforms = SpineChain->FKDecoder.CurrentGlobalTransforms;
+
+		// calculate a "best fit" transform from deformed chain
+		for (int32 ChainIndex=0; ChainIndex<InitialChainTransforms.Num(); ++ChainIndex)
+		{
+			FVector A = InitialChainTransforms[ChainIndex].GetLocation();
+			FVector B = CurrentChainTransforms[ChainIndex].GetLocation();
+
+			// flatten into 2D for more robust yaw construction (what really matters)
+			A.Z = 0.0f;
+			B.Z = 0.0f;
+		
+			InitialPoints.Add(A);
+			CurrentPoints.Add(B);
+		}
+	}
+
+	// calculate "best fit" global body rotation based on deformation of sample points
+	FVector InitialGoalCentroid;
+	FVector CurrentGoalCentroid;
+	const FQuat BestFitBodyRotation = GetRotationFromDeformedPoints(
+		InitialPoints,
+		CurrentPoints,
+		InitialGoalCentroid,
+		CurrentGoalCentroid);
+
+	// generate axes based on body rotation
+	FTransform InitialGlobalGoalTransform = FTransform(FQuat::Identity, InitialGoalCentroid);
+	FTransform CurrentGlobalGoalTransform = FTransform(BestFitBodyRotation, CurrentGoalCentroid);
+	const FVector& Fwd = CurrentGlobalGoalTransform.TransformVector(Settings.GetAxisVector(Settings.ForwardDirection));
+	const FVector& Side = FVector::CrossProduct(Fwd, FVector::ZAxisVector);
+	const FVector SideOrig = FVector::CrossProduct(Settings.GetAxisVector(Settings.ForwardDirection), FVector::ZAxisVector);
+
+	// warp goal positions...
 	for (FRetargetChainPairIK& ChainPair : ChainPairsIK)
 	{
-		ChainPair.IKChainRetargeter.ResetThisTick = true;
+		if (!ChainPair.Settings.IK.bAffectedByIKWarping)
+		{
+			continue;
+		}
+		
+		FChainRetargeterIK& ChainRetargeter = ChainPair.IKChainRetargeter;
+
+		// warp this goal position in various ways...
+		FVector WarpedGoalPosition = ChainRetargeter.Results.EndEffectorPosition;
+		
+		// forward warping
+		const FVector InitialGoalInOrigSpace = InitialGlobalGoalTransform.InverseTransformPosition(ChainRetargeter.Target.InitialEndPosition);
+		const FVector IntialGoalInCurrentSpace = CurrentGlobalGoalTransform.TransformPosition(InitialGoalInOrigSpace);
+		const FPlane FwdPlane(IntialGoalInCurrentSpace, Fwd);
+		const FVector GoalProjOnFwdPlane = FwdPlane.PointPlaneProject(WarpedGoalPosition, FwdPlane);
+		WarpedGoalPosition = GoalProjOnFwdPlane + ((WarpedGoalPosition - GoalProjOnFwdPlane) * Settings.WarpForwards);
+
+		// sideways offset
+		// first determine which side the goal is on originally
+		float GoalSideMultiplier = FVector::DotProduct(ChainRetargeter.Target.InitialEndPosition.GetSafeNormal(), SideOrig);
+		// push goal by offset in the newly calculated sideways direction
+		WarpedGoalPosition += Side * Settings.SidewaysOffset * GoalSideMultiplier;
+
+		// splay warping
+		WarpedGoalPosition = CurrentGoalCentroid + (WarpedGoalPosition - CurrentGoalCentroid) * Settings.WarpSplay; 
+
+		// assign warped position
+		ChainRetargeter.Results.EndEffectorPosition = WarpedGoalPosition;
 	}
+
+#if WITH_EDITOR
+	DebugData.StrideWarpingFrame = CurrentGlobalGoalTransform;
+#endif
 }
 
-FTransform UIKRetargetProcessor::GetTargetBoneRetargetPoseLocalTransform(const int32& TargetBoneIndex) const
+void UIKRetargetProcessor::ResetPlanting()
 {
-	check(TargetSkeleton.BoneNames.IsValidIndex(TargetBoneIndex))
-
-	// get the current retarget pose
-	return TargetSkeleton.RetargetLocalPose[TargetBoneIndex];
+	for (FRetargetChainPairIK& ChainPair : ChainPairsIK)
+	{
+		ChainPair.IKChainRetargeter.ResetThisTick = true;
+	}
 }
 
 bool UIKRetargetProcessor::WasInitializedWithTheseAssets(
@@ -1695,14 +1854,76 @@ void UIKRetargetProcessor::SetNeedsInitialized()
 	}
 }
 
-void UIKRetargetProcessor::CopyAllSettingsFromAsset()
+void UIKRetargetProcessor::ApplySettingsFromProfile(const FRetargetProfile& Profile)
+{
+	auto ApplyRetargetPoseFromProfile = [this, Profile](ERetargetSourceOrTarget SourceOrTarget)
+	{
+		const bool bIsSource = SourceOrTarget == ERetargetSourceOrTarget::Source;
+		const FName PoseName = bIsSource ? Profile.SourceRetargetPoseName : Profile.TargetRetargetPoseName;
+		FRetargetSkeleton& RetargetSkeleton = bIsSource ? SourceSkeleton : TargetSkeleton;
+		if (PoseName == NAME_None || RetargetSkeleton.RetargetPoseName == PoseName)
+		{
+			return; // retarget pose not specified, or already in use
+		}
+		
+		const UIKRigDefinition* IKRig = bIsSource ? RetargeterAsset->GetSourceIKRig() : RetargeterAsset->GetTargetIKRig();
+		RetargetSkeleton.Initialize(
+			RetargetSkeleton.SkeletalMesh,
+			IKRig->GetRetargetChains(),
+			PoseName,
+			RetargeterAsset->GetRetargetPoseByName(SourceOrTarget, PoseName),
+			IKRig->GetRetargetRoot());		
+	};
+
+	// assign retarget poses specified in the profile
+	if (Profile.bApplySourceRetargetPose)
+	{
+		ApplyRetargetPoseFromProfile(ERetargetSourceOrTarget::Source);
+	}
+	if (Profile.bApplyTargetRetargetPose)
+	{
+		ApplyRetargetPoseFromProfile(ERetargetSourceOrTarget::Target);
+	}
+
+	// assign chain settings
+	if (Profile.bApplyChainSettings)
+	{
+		for (const TTuple<FName, FTargetChainSettings>& ChainSettings : Profile.ChainSettings)
+		{
+			for (FRetargetChainPairFK& Chain : ChainPairsFK)
+			{
+				if (Chain.TargetBoneChainName == ChainSettings.Key)
+				{
+					Chain.Settings = ChainSettings.Value;
+				}
+			}
+		
+			for (FRetargetChainPairIK& Chain : ChainPairsIK)
+			{
+				if (Chain.TargetBoneChainName == ChainSettings.Key)
+				{
+					Chain.Settings = ChainSettings.Value;
+				}
+			}
+		}
+	}
+
+	// assign root settings
+	if (Profile.bApplyRootSettings)
+	{
+		RootRetargeter.Settings = Profile.RootSettings;
+	}
+}
+
+void UIKRetargetProcessor::ApplySettingsFromAsset()
 {
+	// copy IK Rig settings
 	if (const UIKRigDefinition* TargetIKRig = RetargeterAsset->GetTargetIKRig())
 	{
 		IKRigProcessor->CopyAllInputsFromSourceAssetAtRuntime(TargetIKRig);
 	}
-
-	// copy most recent settings from asset for each chain
+	
+	// copy chain settings
 	const TArray<TObjectPtr<URetargetChainSettings>>& AllChainSettings = RetargeterAsset->GetAllChainSettings();
 	for (const TObjectPtr<URetargetChainSettings>& ChainSettings : AllChainSettings)
 	{
@@ -1724,7 +1945,7 @@ void UIKRetargetProcessor::CopyAllSettingsFromAsset()
 	}
 
 	// copy root settings
-	RootRetargeter.Settings.CopySettingsFromAsset(RetargeterAsset->GetRetargetRootSettings());
+	RootRetargeter.Settings.CopySettingsFromAsset(RetargeterAsset->GetRootSettingsUObject());
 }
 
 bool UIKRetargetProcessor::IsBoneRetargeted(
diff --git a/Engine/Plugins/Animation/IKRig/Source/IKRig/Private/Retargeter/IKRetargetSettings.cpp b/Engine/Plugins/Animation/IKRig/Source/IKRig/Private/Retargeter/IKRetargetSettings.cpp
new file mode 100644
--- /dev/null
+++ b/Engine/Plugins/Animation/IKRig/Source/IKRig/Private/Retargeter/IKRetargetSettings.cpp
@@ -0,0 +1,14 @@
+﻿// Copyright Epic Games, Inc. All Rights Reserved.
+
+#include "Retargeter/IKRetargetSettings.h"
+#include "Retargeter/IKRetargeter.h"
+
+void FTargetChainSettings::CopySettingsFromAsset(const URetargetChainSettings* AssetChainSettings)
+{
+	*this = AssetChainSettings->Settings;
+}
+
+void FTargetRootSettings::CopySettingsFromAsset(const URetargetRootSettings* AssetRootSettings)
+{
+	*this = AssetRootSettings->Settings;
+}
diff --git a/Engine/Plugins/Animation/IKRig/Source/IKRig/Private/Retargeter/IKRetargeter.cpp b/Engine/Plugins/Animation/IKRig/Source/IKRig/Private/Retargeter/IKRetargeter.cpp
--- a/Engine/Plugins/Animation/IKRig/Source/IKRig/Private/Retargeter/IKRetargeter.cpp
+++ b/Engine/Plugins/Animation/IKRig/Source/IKRig/Private/Retargeter/IKRetargeter.cpp
@@ -3,6 +3,7 @@
 #include "Retargeter/IKRetargeter.h"
 
 #include "IKRigObjectVersion.h"
+#include "Retargeter/IKRetargetProfile.h"
 
 #if WITH_EDITOR
 const FName UIKRetargeter::GetSourceIKRigPropertyName() { return GET_MEMBER_NAME_STRING_CHECKED(UIKRetargeter, SourceIKRigAsset); };
@@ -14,9 +15,9 @@ void UIKRetargeter::GetSpeedCurveNames(TArray<FName>& OutSpeedCurveNames) const
 {
 	for (const URetargetChainSettings* ChainSetting : ChainSettings)
 	{
-		if (ChainSetting->SpeedCurveName != NAME_None)
+		if (ChainSetting->Settings.SpeedPlanting.SpeedCurveName != NAME_None)
 		{
-			OutSpeedCurveNames.Add(ChainSetting->SpeedCurveName);
+			OutSpeedCurveNames.Add(ChainSetting->Settings.SpeedPlanting.SpeedCurveName);
 		}
 	}
 }
@@ -29,6 +30,9 @@ UIKRetargeter::UIKRetargeter(const FObjectInitializer& ObjectInitializer)
 	RootSettings = CreateDefaultSubobject<URetargetRootSettings>(TEXT("RootSettings"));
 	RootSettings->SetFlags(RF_Transactional);
 
+	GlobalSettings = CreateDefaultSubobject<UIKRetargetGlobalSettings>(TEXT("GlobalSettings"));
+	GlobalSettings->SetFlags(RF_Transactional);
+
 	CleanAndInitialize();
 }
 
@@ -84,8 +88,21 @@ void UIKRetargeter::PostLoad()
 		{
 			TargetMeshScale = TargetActorScale_DEPRECATED;
 		}
-	#endif
 
+		// load deprecated global settings
+		if (bRetargetRoot_DEPRECATED)
+		{
+			GlobalSettings->Settings.bEnableRoot = bRetargetRoot_DEPRECATED;
+		}
+		if (bRetargetFK_DEPRECATED)
+		{
+			GlobalSettings->Settings.bEnableFK = bRetargetFK_DEPRECATED;
+		}
+		if (bRetargetIK_DEPRECATED)
+		{
+			GlobalSettings->Settings.bEnableIK = bRetargetIK_DEPRECATED;
+		}
+	#endif
 
 	// load deprecated retarget poses (pre adding retarget poses for source)
 	if (!RetargetPoses_DEPRECATED.IsEmpty())
@@ -142,13 +159,67 @@ void UIKRetargeter::CleanAndInitialize()
 	}
 };
 
+void URetargetChainSettings::Serialize(FArchive& Ar)
+{
+	Super::Serialize(Ar);
+	Ar.UsingCustomVersion(FIKRigObjectVersion::GUID);
+
+	if (Ar.IsLoading())
+	{
+		// load the old chain settings into the new struct format
+		if (Ar.CustomVer(FIKRigObjectVersion::GUID) < FIKRigObjectVersion::ChainSettingsConvertedToStruct)
+		{
+			PRAGMA_DISABLE_DEPRECATION_WARNINGS
+			Settings.FK.EnableFK =  CopyPoseUsingFK_DEPRECATED;
+			Settings.FK.RotationMode =  RotationMode_DEPRECATED;
+			Settings.FK.RotationAlpha =  RotationAlpha_DEPRECATED;
+			Settings.FK.TranslationMode =  TranslationMode_DEPRECATED;
+			Settings.FK.TranslationAlpha =  TranslationAlpha_DEPRECATED;
+			Settings.IK.EnableIK =  DriveIKGoal_DEPRECATED;
+			Settings.IK.BlendToSource =  BlendToSource_DEPRECATED;
+			Settings.IK.BlendToSourceWeights =  BlendToSourceWeights_DEPRECATED;
+			Settings.IK.StaticOffset =  StaticOffset_DEPRECATED;
+			Settings.IK.StaticLocalOffset =  StaticLocalOffset_DEPRECATED;
+			Settings.IK.StaticRotationOffset =  StaticRotationOffset_DEPRECATED;
+			Settings.IK.Extension =  Extension_DEPRECATED;
+			Settings.SpeedPlanting.EnableSpeedPlanting =  UseSpeedCurveToPlantIK_DEPRECATED;
+			Settings.SpeedPlanting.SpeedCurveName =  SpeedCurveName_DEPRECATED;
+			Settings.SpeedPlanting.SpeedThreshold =  VelocityThreshold_DEPRECATED;
+			Settings.SpeedPlanting.UnplantStiffness =  UnplantStiffness_DEPRECATED;
+			Settings.SpeedPlanting.UnplantCriticalDamping =  UnplantCriticalDamping_DEPRECATED;
+			PRAGMA_ENABLE_DEPRECATION_WARNINGS
+		}
+	}
+}
+
+void URetargetRootSettings::Serialize(FArchive& Ar)
+{
+	Super::Serialize(Ar);
+	Ar.UsingCustomVersion(FIKRigObjectVersion::GUID);
+
+	if (Ar.IsLoading())
+	{
+		// load the old root settings into the new struct format
+		if (Ar.CustomVer(FIKRigObjectVersion::GUID) < FIKRigObjectVersion::ChainSettingsConvertedToStruct)
+		{
+			PRAGMA_DISABLE_DEPRECATION_WARNINGS
+			Settings.ScaleHorizontal = GlobalScaleHorizontal_DEPRECATED;
+			Settings.ScaleVertical = GlobalScaleVertical_DEPRECATED;
+			Settings.BlendToSource = BlendToSource_DEPRECATED.Size();
+			Settings.TranslationOffset = StaticOffset_DEPRECATED;
+			Settings.RotationOffset = StaticRotationOffset_DEPRECATED;
+			PRAGMA_ENABLE_DEPRECATION_WARNINGS
+		}
+	}
+}
+
 FQuat FIKRetargetPose::GetDeltaRotationForBone(const FName BoneName) const
 {
 	const FQuat* BoneRotationOffset = BoneRotationOffsets.Find(BoneName);
 	return BoneRotationOffset != nullptr ? *BoneRotationOffset : FQuat::Identity;
 }
 
-void FIKRetargetPose::SetDeltaRotationForBone(FName BoneName, FQuat RotationDelta)
+void FIKRetargetPose::SetDeltaRotationForBone(FName BoneName, const FQuat& RotationDelta)
 {
 	FQuat* RotOffset = BoneRotationOffsets.Find(BoneName);
 	if (RotOffset == nullptr)
@@ -186,13 +257,173 @@ void FIKRetargetPose::SortHierarchically(const FIKRigSkeleton& Skeleton)
 	});
 }
 
+const FTargetChainSettings* UIKRetargeter::GetChainSettingsByName(const FName& TargetChainName) const
+{
+	const TObjectPtr<URetargetChainSettings>* ChainMap = ChainSettings.FindByPredicate(
+		[TargetChainName](const TObjectPtr<URetargetChainSettings> ChainMap)
+		{
+			return ChainMap->TargetChain == TargetChainName;
+		});
+	if (ChainMap)
+	{
+		return &ChainMap->Get()->Settings;
+	}
+
+	return nullptr;
+}
+
 const FIKRetargetPose* UIKRetargeter::GetCurrentRetargetPose(const ERetargetSourceOrTarget& SourceOrTarget) const
 {
 	return SourceOrTarget == ERetargetSourceOrTarget::Source ? &SourceRetargetPoses[CurrentSourceRetargetPose] : &TargetRetargetPoses[CurrentTargetRetargetPose];
 }
 
+FName UIKRetargeter::GetCurrentRetargetPoseName(const ERetargetSourceOrTarget& SourceOrTarget) const
+{
+	return SourceOrTarget == ERetargetSourceOrTarget::Source ? CurrentSourceRetargetPose : CurrentTargetRetargetPose;
+}
+
+const FIKRetargetPose* UIKRetargeter::GetRetargetPoseByName(
+	const ERetargetSourceOrTarget& SourceOrTarget,
+	const FName PoseName) const
+{
+	return SourceOrTarget == ERetargetSourceOrTarget::Source ? SourceRetargetPoses.Find(PoseName) : TargetRetargetPoses.Find(PoseName);
+}
+
 const FName UIKRetargeter::GetDefaultPoseName()
 {
 	static const FName DefaultPoseName = "Default Pose";
 	return DefaultPoseName;
 }
+
+const FRetargetProfile* UIKRetargeter::GetCurrentProfile() const
+{
+	return GetProfileByName(CurrentProfile);
+}
+
+const FRetargetProfile* UIKRetargeter::GetProfileByName(const FName& ProfileName) const
+{
+	return Profiles.Find(ProfileName);
+}
+
+FTargetChainSettings UIKRetargeter::GetChainSettingsFromRetargetAsset(
+	const UIKRetargeter* RetargetAsset,
+	const FName TargetChainName,
+	const FName OptionalProfileName)
+{
+	FTargetChainSettings OutSettings;
+	
+	if (!RetargetAsset)
+	{
+		return OutSettings;
+	}
+	
+	// optionally get the chain settings from a profile
+	if (OptionalProfileName != NAME_None)
+	{
+		if (const FRetargetProfile* RetargetProfile = RetargetAsset->GetProfileByName(OptionalProfileName))
+		{
+			if (const FTargetChainSettings* ProfileChainSettings = RetargetProfile->ChainSettings.Find(TargetChainName))
+			{
+				return *ProfileChainSettings;
+			}
+		}
+
+		// no profile with this chain found, return default settings
+		return OutSettings;
+	}
+	
+	// return the chain settings stored in the retargeter (if it has one matching specified name)
+	if (const FTargetChainSettings* AssetChainSettings = RetargetAsset->GetChainSettingsByName(TargetChainName))
+	{
+		return *AssetChainSettings;
+	}
+
+	// no chain map with the given target chain, so return default settings
+	return OutSettings;
+}
+
+FTargetChainSettings UIKRetargeter::GetChainSettingsFromRetargetProfile(
+	FRetargetProfile& RetargetProfile,
+	const FName TargetChainName)
+{
+	return RetargetProfile.ChainSettings.FindOrAdd(TargetChainName);
+}
+
+void UIKRetargeter::GetRootSettingsFromRetargetAsset(
+	const UIKRetargeter* RetargetAsset,
+	const FName OptionalProfileName,
+	FTargetRootSettings& OutSettings)
+{
+	if (!RetargetAsset)
+	{
+		OutSettings = FTargetRootSettings();
+		return;
+	}
+	
+	// optionally get the root settings from a profile
+	if (OptionalProfileName != NAME_None)
+	{
+		if (const FRetargetProfile* RetargetProfile = RetargetAsset->GetProfileByName(OptionalProfileName))
+		{
+			if (RetargetProfile->bApplyRootSettings)
+			{
+				OutSettings =  RetargetProfile->RootSettings;
+				return;
+			}
+		}
+		
+		// could not find profile, so return default settings
+		OutSettings = FTargetRootSettings();
+		return;
+	}
+
+	// return the base root settings
+	OutSettings =  RetargetAsset->GetRootSettingsUObject()->Settings;
+}
+
+void UIKRetargeter::SetRootSettingsInRetargetProfile(
+	FRetargetProfile& RetargetProfile,
+	const FTargetRootSettings& RootSettings)
+{
+	RetargetProfile.RootSettings = RootSettings;
+	RetargetProfile.bApplyRootSettings = true;
+}
+
+void UIKRetargeter::SetChainSettingsInRetargetProfile(
+	FRetargetProfile& RetargetProfile,
+	const FTargetChainSettings& ChainSettings,
+	const FName TargetChainName)
+{
+	RetargetProfile.ChainSettings.Add(TargetChainName, ChainSettings);
+	RetargetProfile.bApplyChainSettings = true;
+}
+
+void UIKRetargeter::SetChainFKSettingsInRetargetProfile(
+	FRetargetProfile& RetargetProfile,
+	const FTargetChainFKSettings& FKSettings,
+	const FName TargetChainName)
+{
+	FTargetChainSettings& ChainSettings = RetargetProfile.ChainSettings.FindOrAdd(TargetChainName);
+	ChainSettings.FK = FKSettings;
+	RetargetProfile.bApplyChainSettings = true;
+}
+
+void UIKRetargeter::SetChainIKSettingsInRetargetProfile(
+	FRetargetProfile& RetargetProfile,
+	const FTargetChainIKSettings& IKSettings,
+	const FName TargetChainName)
+{
+	FTargetChainSettings& ChainSettings = RetargetProfile.ChainSettings.FindOrAdd(TargetChainName);
+	ChainSettings.IK = IKSettings;
+	RetargetProfile.bApplyChainSettings = true;
+}
+
+void UIKRetargeter::SetChainSpeedPlantSettingsInRetargetProfile(
+	FRetargetProfile& RetargetProfile,
+	const FTargetChainSpeedPlantSettings& SpeedPlantSettings,
+	const FName TargetChainName)
+{
+	FTargetChainSettings& ChainSettings = RetargetProfile.ChainSettings.FindOrAdd(TargetChainName);
+	ChainSettings.SpeedPlanting = SpeedPlantSettings;
+	RetargetProfile.bApplyChainSettings = true;
+}
diff --git a/Engine/Plugins/Animation/IKRig/Source/IKRig/Private/Solvers/IKRig_BodyMover.cpp b/Engine/Plugins/Animation/IKRig/Source/IKRig/Private/Solvers/IKRig_BodyMover.cpp
--- a/Engine/Plugins/Animation/IKRig/Source/IKRig/Private/Solvers/IKRig_BodyMover.cpp
+++ b/Engine/Plugins/Animation/IKRig/Source/IKRig/Private/Solvers/IKRig_BodyMover.cpp
@@ -3,6 +3,7 @@
 #include "Solvers/IKRig_BodyMover.h"
 #include "IKRigDataTypes.h"
 #include "IKRigSkeleton.h"
+#include "Solvers/PointsToRotation.h"
 
 #define LOCTEXT_NAMESPACE "UIKRig_BodyMover"
 
@@ -27,37 +28,10 @@ void UIKRig_BodyMover::Solve(FIKRigSkeleton& IKRigSkeleton, const FIKRigGoalCont
 	
 	// ensure body bone exists
 	check(IKRigSkeleton.RefPoseGlobal.IsValidIndex(BodyBoneIndex));
-
-	// the bone transform to modify
-	FTransform& CurrentBodyTransform = IKRigSkeleton.CurrentPoseGlobal[BodyBoneIndex];
-
-	// calculate initial and current centroids
-	FVector InitialCentroid = FVector::ZeroVector;
-	FVector CurrentCentroid = FVector::ZeroVector;
-	for (UIKRig_BodyMoverEffector* Effector : Effectors)
-	{
-		const FIKRigGoal* Goal = Goals.FindGoalByName(Effector->GoalName);
-		if (!Goal)
-		{
-			return;
-		}
-
-		const int32 BoneIndex = IKRigSkeleton.GetBoneIndexFromName(Effector->BoneName);
-		const FTransform InitialEffector = IKRigSkeleton.CurrentPoseGlobal[BoneIndex];
-
-		InitialCentroid += InitialEffector.GetTranslation();
-		CurrentCentroid += Goal->FinalBlendedPosition;
-	}
-
-	// average centroids
-	const float InvNumEffectors = 1.0f / static_cast<float>(Effectors.Num());
-	InitialCentroid *= InvNumEffectors;
-	CurrentCentroid *= InvNumEffectors;
-
-	// accumulate deformation gradient to extract a rotation from
-	FVector DX = FVector::ZeroVector; // DX, DY, DZ are rows of 3x3 deformation gradient tensor
-	FVector DY = FVector::ZeroVector;
-	FVector DZ = FVector::ZeroVector;
+	
+	// calculate a "best fit" transform from deformed goal locations
+	TArray<FVector> InitialPoints;
+	TArray<FVector> CurrentPoints;
 	for (UIKRig_BodyMoverEffector* Effector : Effectors)
 	{
 		const FIKRigGoal* Goal = Goals.FindGoalByName(Effector->GoalName);
@@ -69,33 +43,27 @@ void UIKRig_BodyMover::Solve(FIKRigSkeleton& IKRigSkeleton, const FIKRigGoalCont
 		const int32 BoneIndex = IKRigSkeleton.GetBoneIndexFromName(Effector->BoneName);
 		const FTransform InitialEffector = IKRigSkeleton.CurrentPoseGlobal[BoneIndex];
 
-		//
-		// accumulate the deformation gradient tensor for all points
-		// "Meshless Deformations Based on Shape Matching"
-		// Equation 7 describes accumulation of deformation gradient from points
-		//
-		// P is normalized vector from INITIAL centroid to INITIAL point
-		// Q is normalized vector from CURRENT centroid to CURRENT point
-		FVector P = (InitialEffector.GetTranslation() - InitialCentroid).GetSafeNormal();
-		FVector Q = (Goal->FinalBlendedPosition - CurrentCentroid).GetSafeNormal();
-		Q = FMath::Lerp(P,Q, Effector->InfluenceMultiplier);
-		// PQ^T is the outer product of P and Q which is a 3x3 matrix
-		// https://en.m.wikipedia.org/wiki/Outer_product
-		DX += FVector(P[0]*Q[0], P[0]*Q[1], P[0]*Q[2]);
-		DY += FVector(P[1]*Q[0], P[1]*Q[1], P[1]*Q[2]);
-		DZ += FVector(P[2]*Q[0], P[2]*Q[1], P[2]*Q[2]);
+		InitialPoints.Add(InitialEffector.GetTranslation());
+		CurrentPoints.Add(Goal->FinalBlendedPosition);
 	}
-
-	// extract "best fit" rotation from deformation gradient
-	FQuat RotationOffset = FQuat::Identity;
-	ExtractRotation(DX, DY, DZ, RotationOffset, 50);
+	
+	FVector InitialCentroid;
+	FVector CurrentCentroid;
+	const FQuat RotationOffset = GetRotationFromDeformedPoints(
+		InitialPoints,
+		CurrentPoints,
+		InitialCentroid,
+		CurrentCentroid);
 
 	// alpha blend the position offset and add it to the current bone location
-	FVector Offset = (CurrentCentroid - InitialCentroid);
-	FVector Weight(
+	const FVector Offset = (CurrentCentroid - InitialCentroid);
+	const FVector Weight(
 		Offset.X > 0.f ? PositionPositiveX : PositionNegativeX,
 		Offset.Y > 0.f ? PositionPositiveY : PositionNegativeY,
 		Offset.Z > 0.f ? PositionPositiveZ : PositionNegativeZ);
+
+	// the bone transform to modify
+	FTransform& CurrentBodyTransform = IKRigSkeleton.CurrentPoseGlobal[BodyBoneIndex];
 	CurrentBodyTransform.AddToTranslation(Offset * (Weight*PositionAlpha));
 
 	// do per-axis alpha blend
@@ -110,33 +78,6 @@ void UIKRig_BodyMover::Solve(FIKRigSkeleton& IKRigSkeleton, const FIKRigGoalCont
 	IKRigSkeleton.PropagateGlobalPoseBelowBone(BodyBoneIndex);
 }
 
-void UIKRig_BodyMover::ExtractRotation(
-	const FVector& DX,
-	const FVector& DY,
-	const FVector& DZ,
-	FQuat &Q,
-	const unsigned int MaxIter)
-{
-	// "A Robust Method to Extract the Rotational Part of Deformations" equation 7
-	// https://matthias-research.github.io/pages/publications/stablePolarDecomp.pdf
-	for (unsigned int Iter = 0; Iter < MaxIter; Iter++)
-	{
-		FMatrix R = FRotationMatrix::Make(Q);
-		FVector RCol0(R.M[0][0], R.M[0][1], R.M[0][2]);
-		FVector RCol1(R.M[1][0], R.M[1][1], R.M[1][2]);
-		FVector RCol2(R.M[2][0], R.M[2][1], R.M[2][2]);
-		FVector Omega = RCol0.Cross(DX) + RCol1.Cross(DY) + RCol2.Cross(DZ);
-		Omega *= 1.0f / (fabs(RCol0.Dot(DX) + RCol1.Dot(DY) + RCol2.Dot(DZ)) + SMALL_NUMBER);
-		const float W = Omega.Size();
-		if (W < SMALL_NUMBER)
-		{
-			break;
-		}
-		Q = FQuat(FQuat((1.0 / W) * Omega, W)) * Q;
-		Q.Normalize();
-	}
-}
-
 #if WITH_EDITOR
 
 void UIKRig_BodyMover::UpdateSolverSettings(UIKRigSolver* InSettings)
diff --git a/Engine/Plugins/Animation/IKRig/Source/IKRig/Public/AnimNodes/AnimNode_RetargetPoseFromMesh.h b/Engine/Plugins/Animation/IKRig/Source/IKRig/Public/AnimNodes/AnimNode_RetargetPoseFromMesh.h
--- a/Engine/Plugins/Animation/IKRig/Source/IKRig/Public/AnimNodes/AnimNode_RetargetPoseFromMesh.h
+++ b/Engine/Plugins/Animation/IKRig/Source/IKRig/Public/AnimNodes/AnimNode_RetargetPoseFromMesh.h
@@ -7,6 +7,7 @@
 #include "Animation/AnimNodeBase.h"
 #include "Retargeter/IKRetargeter.h"
 #include "Retargeter/IKRetargetProcessor.h"
+#include "Retargeter/IKRetargetProfile.h"
 
 #include "AnimNode_RetargetPoseFromMesh.generated.h"
 
@@ -29,6 +30,10 @@ struct IKRIG_API FAnimNode_RetargetPoseFromMesh : public FAnimNode_Base
 	UPROPERTY(EditAnywhere, BlueprintReadWrite, Category = Settings, meta=(PinHiddenByDefault))
 	TObjectPtr<UIKRetargeter> IKRetargeterAsset = nullptr;
 
+	/** connect a custom retarget profile to modify the retargeter's settings at runtime.*/
+	UPROPERTY(BlueprintReadWrite, transient, Category=Settings, meta=(PinHiddenByDefault))
+	FRetargetProfile CustomRetargetProfile;
+
 	/* Toggle whether to print warnings about missing or incorrectly configured retarget configurations. */
 	UPROPERTY(EditAnywhere, BlueprintReadWrite, Category=Debug, meta = (NeverAsPin))
 	bool bSuppressWarnings = false;
diff --git a/Engine/Plugins/Animation/IKRig/Source/IKRig/Public/IKRigObjectVersion.h b/Engine/Plugins/Animation/IKRig/Source/IKRig/Public/IKRigObjectVersion.h
--- a/Engine/Plugins/Animation/IKRig/Source/IKRig/Public/IKRigObjectVersion.h
+++ b/Engine/Plugins/Animation/IKRig/Source/IKRig/Public/IKRigObjectVersion.h
@@ -17,6 +17,9 @@ struct IKRIG_API FIKRigObjectVersion
 		// Retarget pose quaternions changed from pre to post multiplied
 		RetargetPoseQuatPostMultiplied,
 
+		// Chain settings moved to struct to be used in profiles
+		ChainSettingsConvertedToStruct,
+
 		// -----<new versions can be added above this line>-------------------------------------------------
 		VersionPlusOne,
 		LatestVersion = VersionPlusOne - 1
diff --git a/Engine/Plugins/Animation/IKRig/Source/IKRig/Public/IKRigProcessor.h b/Engine/Plugins/Animation/IKRig/Source/IKRig/Public/IKRigProcessor.h
--- a/Engine/Plugins/Animation/IKRig/Source/IKRig/Public/IKRigProcessor.h
+++ b/Engine/Plugins/Animation/IKRig/Source/IKRig/Public/IKRigProcessor.h
@@ -46,8 +46,12 @@ public:
 	 *  NOTE!! this function creates new UObjects and consequently MUST be called from the main thread!!
 	 *  @param InRigAsset - the IK Rig defining the collection of solvers to execute and all the rig settings
 	 *  @param SkeletalMesh - the skeletal mesh you want to solve the IK on
+	 *  @param ExcludedGoals - a list of goal names to exclude from this processor instance (support per-instance removal of goals)
 	 */
-	void Initialize(const UIKRigDefinition* InRigAsset, const USkeletalMesh* SkeletalMesh);
+	void Initialize(
+		const UIKRigDefinition* InRigAsset,
+		const USkeletalMesh* SkeletalMesh,
+		const TArray<FName>& ExcludedGoals = TArray<FName>());
 
 	//
 	// BEGIN UPDATE SEQUENCE FUNCTIONS
@@ -84,8 +88,10 @@ public:
 	/** Get access to the internal goal data (read only) */
 	const FIKRigGoalContainer& GetGoalContainer() const;
 	
-	/** Get access to the internal skeleton data */
-	FIKRigSkeleton& GetSkeleton();
+	/** Get read/write access to the internal skeleton data */
+	FIKRigSkeleton& GetSkeletonWriteable();
+	/** Get read-only access to the internal skeleton data */
+	const FIKRigSkeleton& GetSkeleton() const;
 
 	/** Used to determine if the IK Rig asset is compatible with a given skeleton. */
 	static bool IsIKRigCompatibleWithSkeleton(
diff --git a/Engine/Plugins/Animation/IKRig/Source/IKRig/Public/IKRigSkeleton.h b/Engine/Plugins/Animation/IKRig/Source/IKRig/Public/IKRigSkeleton.h
--- a/Engine/Plugins/Animation/IKRig/Source/IKRig/Public/IKRigSkeleton.h
+++ b/Engine/Plugins/Animation/IKRig/Source/IKRig/Public/IKRigSkeleton.h
@@ -6,22 +6,11 @@
 
 #include "IKRigSkeleton.generated.h"
 
+struct FBoneChain;
 class USkeletalMesh;
 class UIKRigDefinition;
 class UIKRigSolver;
 
-USTRUCT()
-struct IKRIG_API FIKRigSkeletonChain
-{
-	GENERATED_BODY()
-	
-	FIKRigSkeletonChain(){}
-	FIKRigSkeletonChain(const FName& StartBone, const FName& EndBone): StartBone(StartBone), EndBone(EndBone){}
-
-	FName StartBone;
-	FName EndBone;
-};
-
 /** Data used just to initialize an IKRigSkeleton from outside systems
  *
  * The input skeleton may be different than the skeleton that the IK Rig asset was created for, within some limits.
@@ -120,7 +109,20 @@ struct IKRIG_API FIKRigSkeleton
 
 	static void NormalizeRotations(TArray<FTransform>& Transforms);
 
-	void GetChainsInList(const TArray<int32>& SelectedBones, TArray<FIKRigSkeletonChain>& OutChains) const;
+	void GetChainsInList(const TArray<int32>& SelectedBones, TArray<FBoneChain>& OutChains) const;
+
+	// get indices of bones in this chain, ordered from tip to root. returns false if chain is invalid.
+	bool GetBonesInChain(const FBoneChain& Chain, TArray<int32>& OutBoneIndices) const;
+
+	// given a set of bones, returns a set of indices of bones that are mirrored across X axis (YZ plane)
+	// return false if any of the bones do not have a mirrored pair within the MaxDistanceThreshold
+	bool GetMirroredBoneIndices(
+		const TArray<int32>& BoneIndices,
+		TArray<int32>& OutMirroredIndices,
+		float MaxDistanceThreshold=5.0f) const;
+	
+	// get the bone that is closest to the given point (sets bone index and distance)
+	void GetClosestBone(const FVector& InPoint, int32& OutBoneIndex, float& OutDistance) const;
 
 private:
 
diff --git a/Engine/Plugins/Animation/IKRig/Source/IKRig/Public/Retargeter/IKRetargetProcessor.h b/Engine/Plugins/Animation/IKRig/Source/IKRig/Public/Retargeter/IKRetargetProcessor.h
--- a/Engine/Plugins/Animation/IKRig/Source/IKRig/Public/Retargeter/IKRetargetProcessor.h
+++ b/Engine/Plugins/Animation/IKRig/Source/IKRig/Public/Retargeter/IKRetargetProcessor.h
@@ -4,8 +4,10 @@
 
 #include "CoreMinimal.h"
 #include "IKRetargeter.h"
+#include "IKRetargetSettings.h"
 #include "IKRigLogger.h"
 #include "Kismet/KismetMathLibrary.h"
+
 #include "IKRetargetProcessor.generated.h"
 
 enum class ERetargetTranslationMode : uint8;
@@ -22,28 +24,27 @@ class USkeletalMesh;
 
 struct IKRIG_API FRetargetSkeleton
 {
-	TArray<FName> BoneNames;
-	
-	TArray<int32> ParentIndices;
-	
-	TArray<FTransform> RetargetLocalPose;
-	
-	TArray<FTransform> RetargetGlobalPose;
-	
-	USkeletalMesh* SkeletalMesh;
-
-	// record which chain is actually controlling each bone
-	TArray<FName> ChainThatContainsBone;
+	TArray<FName> BoneNames;				// list of all bone names in ref skeleton order
+	TArray<int32> ParentIndices;			// per-bone indices of parent bones (the hierarchy)
+	TArray<FTransform> RetargetLocalPose;	// local space retarget pose
+	TArray<FTransform> RetargetGlobalPose;	// global space retarget pose
+	FName RetargetPoseName;					// the name of the retarget pose this was initialized with
+	USkeletalMesh* SkeletalMesh;			// the skeletal mesh this was initialized with
+	TArray<FName> ChainThatContainsBone;	// record which chain is actually controlling each bone
 
 	void Initialize(
 		USkeletalMesh* InSkeletalMesh,
 		const TArray<FBoneChain>& BoneChains,
+		const FName InRetargetPoseName,
 		const FIKRetargetPose* RetargetPose,
 		const FName& RetargetRootBone);
 
 	void Reset();
 
-	void GenerateRetargetPose(const FIKRetargetPose* InRetargetPose, const FName& RetargetRootBone);
+	void GenerateRetargetPose(
+		const FName InRetargetPoseName,
+		const FIKRetargetPose* InRetargetPose,
+		const FName& RetargetRootBone);
 
 	int32 FindBoneIndexByName(const FName InName) const;
 
@@ -98,6 +99,7 @@ struct FTargetSkeleton : public FRetargetSkeleton
 	void Initialize(
 		USkeletalMesh* InSkeletalMesh,
 		const TArray<FBoneChain>& BoneChains,
+		const FName& InRetargetPoseName,
 		const FIKRetargetPose* RetargetPose,
 		const FName& RetargetRootBone);
 
@@ -140,30 +142,19 @@ struct FRootSource
 struct FRootTarget
 {
 	int32 BoneIndex;
-
 	FVector InitialPosition;
 	FQuat InitialRotation;
-	
 	float InitialHeight;
-};
 
-struct FRetargetRootSettings
-{
-	float GlobalScaleHorizontal = 1.0f;
-	float GlobalScaleVertical = 1.0f;
-	float BlendToSource = 0.0f;
-	FVector BlendToSourceWeights = FVector::OneVector;
-	FVector StaticOffset = FVector::ZeroVector;
-	FRotator StaticRotationOffset = FRotator::ZeroRotator;
-	
-	void CopySettingsFromAsset(const URetargetRootSettings* AssetRootSettings);
+	FVector RootTranslationDelta;
+	FQuat RootRotationDelta;
 };
 
 struct FRootRetargeter
 {	
 	FRootSource Source;
 	FRootTarget Target;
-	FRetargetRootSettings Settings;
+	FTargetRootSettings Settings;
 
 	void Reset();
 	
@@ -179,44 +170,14 @@ struct FRootRetargeter
 
 	void EncodePose(const TArray<FTransform> &SourceGlobalPose);
 	
-	void DecodePose( TArray<FTransform> &OutTargetGlobalPose) const;
+	void DecodePose(TArray<FTransform> &OutTargetGlobalPose);
 
 	FVector GetGlobalScaleVector() const
 	{
-		return FVector(Settings.GlobalScaleHorizontal, Settings.GlobalScaleHorizontal, Settings.GlobalScaleVertical);
+		return FVector(Settings.ScaleHorizontal, Settings.ScaleHorizontal, Settings.ScaleVertical);
 	}
 };
 
-struct FRetargetChainSettings
-{
-	FName TargetChainName;
-
-	bool CopyPoseUsingFK = true;
-	ERetargetRotationMode RotationMode;
-	float RotationAlpha = 1.0f;
-	ERetargetTranslationMode TranslationMode;
-	float TranslationAlpha = 1.0f;
-
-	bool DriveIKGoal = true;
-	FVector StaticOffset;
-	FVector StaticLocalOffset;
-	FRotator StaticRotationOffset;
-	float Extension = 1.0f;
-	float BlendToSource = 0.0f;
-	FVector BlendToSourceWeights = FVector::OneVector;
-
-	bool UseSpeedCurveToPlantIK = false;
-	float SpeedThreshold = 0.0f;
-	FName SpeedCurveName;
-	float UnplantStiffness = 250.0f;
-	float UnplantCriticalDamping = 1.0f;
-
-public:
-	
-	void CopySettingsFromAsset(const URetargetChainSettings* AssetChainSettings);
-};
-
-
 struct FChainFK
 {
 	TArray<FTransform> InitialGlobalTransforms;
@@ -226,13 +187,14 @@ struct FChainFK
 	TArray<FTransform> CurrentGlobalTransforms;
 
 	TArray<float> Params;
+	TArray<int32> BoneIndices;
 
 	int32 ChainParentBoneIndex;
 	FTransform ChainParentInitialGlobalTransform;
 
 	bool Initialize(
 		const FRetargetSkeleton& Skeleton,
-		const TArray<int32>& BoneIndices,
+		const TArray<int32>& InBoneIndices,
 		const TArray<FTransform> &InitialGlobalPose,
 		FIKRigLogger& Log);
 
@@ -277,7 +239,7 @@ struct FChainDecoderFK : public FChainFK
 	
 	void DecodePose(
 		const FRootRetargeter& RootRetargeter,
-		const FRetargetChainSettings& Settings,
+		const FTargetChainSettings& Settings,
 		const TArray<int32>& TargetBoneIndices,
 		FChainEncoderFK& SourceChain,
 		const FTargetSkeleton& TargetSkeleton,
@@ -299,64 +261,62 @@ private:
 
 struct FDecodedIKChain
 {
-	FVector EndEffectorPosition;
-	
-	FQuat EndEffectorRotation;
-	
-	FVector PoleVectorPosition;
+	FVector EndEffectorPosition = FVector::ZeroVector;
+	FQuat EndEffectorRotation = FQuat::Identity;
+	FVector PoleVectorPosition = FVector::ZeroVector;
 };
 
 struct FSourceChainIK
 {
-	int32 BoneIndexA;
-	
-	int32 BoneIndexB;
-	
-	int32 BoneIndexC;
-	
-	FVector InitialEndPosition;
-	
-	FQuat InitialEndRotation;
+	int32 BoneIndexA = INDEX_NONE;
+	int32 BoneIndexB = INDEX_NONE;
+	int32 BoneIndexC = INDEX_NONE;
 	
-	float InvInitialLength;
+	FVector InitialEndPosition = FVector::ZeroVector;
+	FQuat InitialEndRotation = FQuat::Identity;
+	float InvInitialLength = 1.0f;
 
 	// results after encoding...
-	FVector PreviousEndPosition;
-	FVector CurrentEndPosition;
-	
-	FVector CurrentEndDirectionNormalized;
-	
-	FQuat CurrentEndRotation;
-	
-	float CurrentHeightFromGroundNormalized;
-	
-	FVector PoleVectorDirection;
+	FVector PreviousEndPosition = FVector::ZeroVector;
+	FVector CurrentEndPosition = FVector::ZeroVector;
+	FVector CurrentEndDirectionNormalized = FVector::ZeroVector;
+	FQuat CurrentEndRotation = FQuat::Identity;
+	float CurrentHeightFromGroundNormalized = 0.0f;
+	FVector PoleVectorDirection = FVector::ZeroVector;
 };
 
 struct FTargetChainIK
 {
-	int32 BoneIndexA;
-
-	int32 BoneIndexC;
-	
-	float InitialLength;
-	
-	FVector InitialEndPosition;
+	int32 BoneIndexA = INDEX_NONE;
+	int32 BoneIndexC = INDEX_NONE;
 	
-	FQuat InitialEndRotation;
+	float InitialLength = 1.0f;
+	FVector InitialEndPosition = FVector::ZeroVector;
+	FQuat InitialEndRotation = FQuat::Identity;
+	FVector PrevEndPosition = FVector::ZeroVector;
+};
 
-	FVector PrevEndPosition;
+struct FChainDebugData
+{
+	FTransform InputTransformStart = FTransform::Identity;
+	FTransform InputTransformEnd = FTransform::Identity;
+	FTransform OutputTransformEnd = FTransform::Identity;
 };
 
 struct FChainRetargeterIK
 {
 	FSourceChainIK Source;
-	
 	FTargetChainIK Target;
+	
+	FDecodedIKChain Results;
 
 	bool ResetThisTick;
 	FVectorSpringState PlantingSpringState;
 
+	#if WITH_EDITOR
+	FChainDebugData DebugData;
+	#endif
+
 	bool InitializeSource(
 		const TArray<int32>& BoneIndices,
 		const TArray<FTransform> &SourceInitialGlobalPose,
@@ -370,16 +330,18 @@ struct FChainRetargeterIK
 	void EncodePose(const TArray<FTransform> &SourceInputGlobalPose);
 	
 	void DecodePose(
-		const FRetargetChainSettings& Settings,
+		const FTargetChainSettings& Settings,
+		const FRootRetargeter& RootRetargeter,
 		const TMap<FName, float>& SpeedValuesFromCurves,
 		const float DeltaTime,
-		const TArray<FTransform>& OutGlobalPose,
-		FDecodedIKChain& OutResults);
+		const TArray<FTransform>& InGlobalPose);
+
+	void SaveDebugInfo(const TArray<FTransform>& InGlobalPose);
 };
 
 struct FRetargetChainPair
 {
-	FRetargetChainSettings Settings;
+	FTargetChainSettings Settings;
 	
 	TArray<int32> SourceBoneIndices;
 	
@@ -392,7 +354,6 @@ struct FRetargetChainPair
 	virtual ~FRetargetChainPair() = default;
 	
 	virtual bool Initialize(
-		URetargetChainSettings* InSettings,
 		const FBoneChain& SourceBoneChain,
 		const FBoneChain& TargetBoneChain,
 		const FRetargetSkeleton& SourceSkeleton,
@@ -415,7 +376,6 @@ struct FRetargetChainPairFK : FRetargetChainPair
 	FChainDecoderFK FKDecoder;
 	
 	virtual bool Initialize(
-		URetargetChainSettings* InSettings,
         const FBoneChain& SourceBoneChain,
         const FBoneChain& TargetBoneChain,
         const FRetargetSkeleton& SourceSkeleton,
@@ -432,7 +392,6 @@ struct FRetargetChainPairIK : FRetargetChainPair
 	FName PoleVectorGoalName;
 
 	virtual bool Initialize(
-		URetargetChainSettings* Settings,
         const FBoneChain& SourceBoneChain,
         const FBoneChain& TargetBoneChain,
         const FRetargetSkeleton& SourceSkeleton,
@@ -440,6 +399,11 @@ struct FRetargetChainPairIK : FRetargetChainPair
         FIKRigLogger& Log) override;
 };
 
+struct FRetargetDebugData
+{
+	FTransform StrideWarpingFrame;
+};
+
 /** The runtime processor that converts an input pose from a source skeleton into an output pose on a target skeleton.
  * To use:
  * 1. Initialize a processor with a Source/Target skeletal mesh and a UIKRetargeter asset.
@@ -474,11 +438,18 @@ public:
 	* @param InSourceGlobalPose -  is the source mesh input pose in Component/Global space
 	* @return The retargeted Component/Global space pose for the target skeleton
 	*/
-	TArray<FTransform>& RunRetargeter(const TArray<FTransform>& InSourceGlobalPose, const TMap<FName, float>& SpeedValuesFromCurves, const float DeltaTime);
+	TArray<FTransform>& RunRetargeter(
+		const TArray<FTransform>& InSourceGlobalPose,
+		const TMap<FName,
+		float>& SpeedValuesFromCurves,
+		const float DeltaTime);
 
+	/** Apply the settings stored in a retarget profile. Call this before RunRetargeter() to use the settings stored in a profile. */
+	void ApplySettingsFromProfile(const FRetargetProfile& Profile);
+	
+	/** Apply the settings stored in the retargeter asset. */
+	void ApplySettingsFromAsset();
 	
-	FTransform GetTargetBoneRetargetPoseLocalTransform(const int32& TargetBoneIndex) const;
-
 	/** Get read-only access to the target skeleton. */
 	const FTargetSkeleton& GetTargetSkeleton() const { return TargetSkeleton; };
 
@@ -512,25 +483,28 @@ public:
 #if WITH_EDITOR
 	/** Set that this processor needs to be reinitialized. */
 	void SetNeedsInitialized();
-	/** During editor preview, drive the target IK Rig with the settings from it's source asset */
-	void CopyAllSettingsFromAsset();
 	/** Returns true if the bone is part of a retarget chain or root bone, false otherwise. */
 	bool IsBoneRetargeted(const int32& BoneIndex, const int8& SkeletonToCheck) const;
 	/** Returns name of the chain associated with this bone. Returns NAME_None if bone is not in a chain. */
 	FName GetChainNameForBone(const int32& BoneIndex, const int8& SkeletonToCheck) const;
+	/** Get read only access to the core IK retarget chains for debug purposes */
+	const TArray<FRetargetChainPairIK>& GetIKChainPairs() const { return ChainPairsIK; };
+	/** Get read only access to the core FK retarget chains for debug purposes */
+	const TArray<FRetargetChainPairFK>& GetFKChainPairs() const { return ChainPairsFK; };
+	/** Get read only access to the core root retargeter for debug purposes */
+	const FRootRetargeter& GetRootRetargeter() const { return RootRetargeter; };
+	/** store data for debug drawing */
+	FRetargetDebugData DebugData;
 #endif
 
 private:
 
 	/** Only true once Initialize() has successfully completed.*/
 	bool bIsInitialized = false;
-
 	/** true when roots are able to be retargeted */
 	bool bRootsInitialized = false;
-
 	/** true when at least one pair of bone chains is able to be retargeted */
 	bool bAtLeastOneValidBoneChainPair = false;
-
 	/** true when roots are able to be retargeted */
 	bool bIKRigInitialized = false;
 
@@ -574,7 +548,10 @@ private:
 	/** Internal retarget phase for the IK chains. */
 	void RunIKRetarget(
 		const TArray<FTransform>& InSourceGlobalPose,
-		TArray<FTransform>& OutTargeGlobalPose,
+		TArray<FTransform>& OutTargetGlobalPose,
 		const TMap<FName, float>& SpeedValuesFromCurves,
 		const float DeltaTime);
+
+	/** Runs in the after the base IK retarget to apply stride warping to IK goals. */
+	void RunStrideWarping(const TArray<FTransform>& InTargeGlobalPose);
 };
diff --git a/Engine/Plugins/Animation/IKRig/Source/IKRig/Public/Retargeter/IKRetargetProfile.h b/Engine/Plugins/Animation/IKRig/Source/IKRig/Public/Retargeter/IKRetargetProfile.h
new file mode 100644
--- /dev/null
+++ b/Engine/Plugins/Animation/IKRig/Source/IKRig/Public/Retargeter/IKRetargetProfile.h
@@ -0,0 +1,52 @@
+// Copyright Epic Games, Inc. All Rights Reserved.
+
+#pragma once
+
+#include "CoreMinimal.h"
+#include "Retargeter/IKRetargetSettings.h"
+
+#include "IKRetargetProfile.generated.h"
+
+class UIKRetargeter;
+class URetargetChainSettings;
+
+USTRUCT(BlueprintType)
+struct FRetargetProfile
+{
+	GENERATED_BODY()
+	
+public:
+	// Override the TARGET Retarget Pose to use when this profile is active.
+	// The pose must be present in the Retarget Asset and is not applied unless bApplyTargetRetargetPose is true.
+	UPROPERTY(BlueprintReadWrite, Category=RetargetPoses)
+	FName TargetRetargetPoseName;
+
+	// If true, the TARGET Retarget Pose specified in this profile will be applied to the Retargeter (when plugged into the Retargeter).
+	UPROPERTY(BlueprintReadWrite, Category=RetargetPoses)
+	bool bApplyTargetRetargetPose = false;
+
+	// Override the SOURCE Retarget Pose to use when this profile is active.
+	// The pose must be present in the Retarget Asset and is not applied unless bApplySourceRetargetPose is true.
+	UPROPERTY(BlueprintReadWrite, Category=RetargetPoses)
+	FName SourceRetargetPoseName;
+	
+	// If true, the Source Retarget Pose specified in this profile will be applied to the Retargeter (when plugged into the Retargeter).
+	UPROPERTY(BlueprintReadWrite, Category=RetargetPoses)
+	bool bApplySourceRetargetPose = false;
+
+	// A (potentially sparse) set of setting overrides for the target chains (only applied when bApplyChainSettings is true).
+	UPROPERTY(BlueprintReadWrite, Category=ChainSettings)
+	TMap<FName, FTargetChainSettings> ChainSettings;
+
+	// If true, the Chain Settings stored in this profile will be applied to the Retargeter (when plugged into the Retargeter).
+	UPROPERTY(BlueprintReadWrite, Category=ChainSettings)
+	bool bApplyChainSettings = true;
+
+	// Retarget settings to control behavior of the retarget root motion (not applied unless bApplyRootSettings is true)
+	UPROPERTY(BlueprintReadWrite, Category=RootSettings)
+	FTargetRootSettings RootSettings;
+
+	// If true, the root settings stored in this profile will be applied to the Retargeter (when plugged into the Retargeter).
+	UPROPERTY(BlueprintReadWrite, Category=RootSettings)
+	bool bApplyRootSettings = false;
+};
