Added profiles to the IK retargeter

diff --git a/Engine/Plugins/Animation/IKRig/Source/IKRig/Public/Retargeter/IKRetargetSettings.h b/Engine/Plugins/Animation/IKRig/Source/IKRig/Public/Retargeter/IKRetargetSettings.h
new file mode 100644
--- /dev/null
+++ b/Engine/Plugins/Animation/IKRig/Source/IKRig/Public/Retargeter/IKRetargetSettings.h
@@ -0,0 +1,334 @@
+﻿// Copyright Epic Games, Inc. All Rights Reserved.
+
+#pragma once
+
+#include "IKRetargetSettings.generated.h"
+
+class URetargetRootSettings;
+class URetargetChainSettings;
+
+UENUM(BlueprintType)
+enum class ERetargetTranslationMode : uint8
+{
+	None			UMETA(DisplayName = "None"),
+	GloballyScaled	UMETA(DisplayName = "Globally Scaled"),
+	Absolute		UMETA(DisplayName = "Absolute"),
+};
+
+UENUM(BlueprintType)
+enum class ERetargetRotationMode : uint8
+{
+	
+	Interpolated		UMETA(DisplayName = "Interpolated"),
+	OneToOne			UMETA(DisplayName = "One to One"),
+	OneToOneReversed	UMETA(DisplayName = "One to One Reversed"),
+	None				UMETA(DisplayName = "None"),
+};
+
+USTRUCT(BlueprintType)
+struct FTargetChainSpeedPlantSettings
+{
+	GENERATED_USTRUCT_BODY()
+
+	/** The name of the curve on the source animation that contains the speed of the end effector bone.*/
+	UPROPERTY(EditAnywhere, BlueprintReadWrite, Category = "Plant IK by Speed", meta = (ClampMin = "0.0", ClampMax = "100.0", UIMin = "0.0", UIMax = "100.0"))
+	bool EnableSpeedPlanting = false;
+	
+	/** The name of the curve on the source animation that contains the speed of the end effector bone.*/
+	UPROPERTY(EditAnywhere, BlueprintReadWrite, BlueprintReadWrite, Category = "Plant IK by Speed", meta = (ClampMin = "0.0", ClampMax = "100.0", UIMin = "0.0", UIMax = "100.0"))
+	FName SpeedCurveName;
+
+	/** Range 0 to 1000. Default 15. The maximum speed a source bone can be moving while being considered 'planted'.
+	*  The target IK goal will not be allowed to move whenever the source bone speed drops below this threshold speed. */
+	UPROPERTY(EditAnywhere, BlueprintReadWrite, Category = "Plant IK by Speed", meta = (ClampMin = "0.0", ClampMax = "100.0", UIMin = "0.0", UIMax = "100.0"))
+	float SpeedThreshold = 15.0f;
+
+	// How stiff the spring model is that smoothly pulls the IK position after unplanting (more stiffness means more oscillation around the target value)
+	UPROPERTY(EditAnywhere, BlueprintReadWrite, Category = "Plant IK by Speed", meta = (ClampMin = "0.0", UIMin = "0.0"))
+	float UnplantStiffness = 250.0f;
+
+	// How much damping to apply to the spring (0 means no damping, 1 means critically damped which means no oscillation)
+	UPROPERTY(EditAnywhere, BlueprintReadWrite, Category = "Plant IK by Speed", meta = (ClampMin = "0.0", ClampMax = "10.0", UIMin = "0.0", UIMax = "1.0"))
+	float UnplantCriticalDamping = 1.0f;
+};
+
+USTRUCT(BlueprintType)
+struct FTargetChainFKSettings
+{
+	GENERATED_BODY()
+
+	/** Whether to copy the shape of the chain from the source skeleton using the Rotation and Translation modes. Default is true.
+	* NOTE: All FK operations run before the IK pass to copy the shape of the FK chain from the source skeleton. */
+	UPROPERTY(EditAnywhere, BlueprintReadWrite, Category = "FK Adjustments")
+	bool EnableFK = true;
+	
+	/** Determines how rotation is copied from the source chain to the target chain. Default is Interpolated.
+	* Interpolated: Source and target chains are normalized by their length, then each target bone rotation is generated by finding the rotation at the same normalized distance on the source chain and interpolating between the neighboring bones.
+	* One to One: Each target bone rotation is copied from the equivalent bone in the source chain, based on the order in the chain, starting at the root of the chain. If the target chain has more bones than the source, the extra bones will remain at their reference pose.
+	* One to One Reversed: Same as One-to-One, but starting from the tip of the chain.
+	* None: The rotation of each target bone in the chain is left at the reference pose. */
+	UPROPERTY(EditAnywhere, BlueprintReadWrite, Category = "FK Adjustments")
+	ERetargetRotationMode RotationMode = ERetargetRotationMode::Interpolated;
+
+	/** Range +/- infinity. Default 1. Scales the amount of rotation that is applied.
+	*  If Rotation Mode is None this parameter has no effect.
+	*  Otherwise, this parameter blends the rotation of each bone in the chain from the base retarget pose (0) to the retargeted pose (1).*/
+	UPROPERTY(EditAnywhere, BlueprintReadWrite, Category = "FK Adjustments", meta = (UIMin = "0.0", UIMax = "1.0"))
+	float RotationAlpha = 1.0f;
+
+	/** Determines how translation is copied from the source chain to the target chain. Default is None.
+	* None: Translation of target bones are left unmodified from the retarget pose.
+	* Globally Scaled: Translation of target bone is set to the source bone offset multiplied by the global scale of the skeleton (determined by the relative height difference between retarget root bones).
+	* Absolute: Translation of target bone is set to the absolute position of the source bone. */
+	UPROPERTY(EditAnywhere, BlueprintReadWrite, Category = "FK Adjustments")
+	ERetargetTranslationMode TranslationMode = ERetargetTranslationMode::None;
+
+	/** Range +/- infinity. Default 1. Scales the amount of translation that is applied. Exact behavior depends on the Translation Mode.
+	*  In None Mode, this parameter has no effect.
+	*  In Globally Scaled and Absolute modes, the translation offset is scaled by this parameter.*/
+	UPROPERTY(EditAnywhere, BlueprintReadWrite, Category = "FK Adjustments", meta = (UIMin = "0.0", UIMax = "1.0"))
+	float TranslationAlpha = 1.0f;
+};
+
+USTRUCT(BlueprintType)
+struct FTargetChainIKSettings
+{
+	GENERATED_USTRUCT_BODY()
+
+	/** Whether to modify the location of the IK goal on this chain. Default is true.
+	 * NOTE: This only has an effect if the chain has an IK Goal assigned to it in the Target IK Rig asset.
+	 * NOTE: If off, and this chain has an IK Goal, the IK will still be evaluated, but the Goal is set to the input bone location (from the FK pass).*/
+	UPROPERTY(EditAnywhere, BlueprintReadWrite, Category = "IK Adjustments")
+	bool EnableIK = true;
+
+	/** Range 0 to 1. Default 0. Blends IK goal position from retargeted location (0) to source bone location (1).
+	*  At 0 the goal is placed at the retargeted location.
+	*  At 1 the goal is placed at the location of the source chain's end bone. */
+	UPROPERTY(EditAnywhere, BlueprintReadWrite, Category = "IK Adjustments", meta = (ClampMin = "0.0", ClampMax = "1.0", UIMin = "0.0", UIMax = "1.0"))
+	float BlendToSource = 0.0f;
+	
+	/** Range 0 to 1. Default 1. Weight each axis separately when using Blend To Source.
+	*  At 0 the goal is placed at the retargeted location.
+	*  At 1 the goal is placed at the location of the source chain's end bone. */
+	UPROPERTY(EditAnywhere, BlueprintReadWrite, Category = "IK Adjustments", meta = (ClampMin = "0.0", ClampMax = "1.0", UIMin = "0.0", UIMax = "1.0"))
+	FVector BlendToSourceWeights = FVector::OneVector;
+
+	/** Default 0, 0, 0. Apply a static global-space offset to IK goal position. */
+	UPROPERTY(EditAnywhere, BlueprintReadWrite, Category = "IK Adjustments")
+	FVector StaticOffset;
+
+	/** Default 0, 0, 0. Apply a static local-space offset to IK goal position. */
+	UPROPERTY(EditAnywhere, BlueprintReadWrite, Category = "IK Adjustments")
+	FVector StaticLocalOffset;
+
+	/** Default 0, 0, 0. Apply a static local-space offset to IK goal rotation. */
+	UPROPERTY(EditAnywhere, BlueprintReadWrite, Category = "IK Adjustments")
+	FRotator StaticRotationOffset;
+	
+	/** Range 0 to 5. Default 1. Brings IK goal closer (0) or further (1+) from origin of chain.
+	*  At 0 the effector is placed at the origin of the chain (ie Shoulder, Hip etc).
+	*  At 1 the effector is left at the end of the chain (ie Wrist, Foot etc)
+	*  Values in-between 0-1 will slide the effector along the vector from the start to the end of the chain.
+	*  Values greater than 1 will stretch the chain beyond the retargeted length. */
+	UPROPERTY(EditAnywhere, BlueprintReadWrite, Category = "IK Adjustments", meta = (ClampMin = "0.0", ClampMax = "5.0", UIMin = "0.1", UIMax = "2.0"))
+	float Extension = 1.0f;
+
+	/** Is this IK goal affected by the stride warping (in Global Settings). Typically this is true for all feet, but not for hands.  */
+	UPROPERTY(EditAnywhere, Category = Warping, meta = (DisplayName = "Affected by IK Warping"))
+	bool bAffectedByIKWarping = true;
+};
+
+USTRUCT(BlueprintType)
+struct FTargetChainSettings
+{
+	GENERATED_USTRUCT_BODY()
+
+public:
+
+	/** All settings for FK retargeting of this chain.
+	 *FK retargeting runs before the IK pass. */
+	UPROPERTY(EditAnywhere, BlueprintReadWrite, Category = "FK")
+	FTargetChainFKSettings FK;
+
+	/** All settings controlling the IK Goal assigned to this chain.
+	 *The IK pass runs AFTER the FK pass and can be used to fix contacts. */
+	UPROPERTY(EditAnywhere, BlueprintReadWrite, Category = "IK")
+	FTargetChainIKSettings IK;
+
+	/** All settings associated with planting IK goals based on the speed of the source.
+	 *Speed planting will pin the IK goal to the location determined by the IK settings above. */
+	UPROPERTY(EditAnywhere, BlueprintReadWrite, Category = "Speed Planting")
+	FTargetChainSpeedPlantSettings SpeedPlanting;
+	
+	void CopySettingsFromAsset(const URetargetChainSettings* AssetChainSettings);
+};
+
+USTRUCT(BlueprintType)
+struct IKRIG_API FTargetRootSettings
+{
+	GENERATED_USTRUCT_BODY()
+
+public:
+
+	/** Range 0 to 1. Default 1. Blends the amount of retargeted root rotation to apply.
+	*  At 0 the root is left at the rotation from the retarget pose.
+	*  At 1 the root is rotated fully to match the source root rotation. */
+	UPROPERTY(EditAnywhere, BlueprintReadWrite, Category = "Root Retarget Settings", meta = (ClampMin = "-5.0", ClampMax = "5.0", UIMin = "0.0", UIMax = "1.0"))
+	float RotationAlpha = 1.0f;
+	
+	/** Range 0 to 1. Default 1. Blends the amount of retargeted root translation to apply.
+	*  At 0 the root is left at the position from the retarget pose.
+	*  At 1 the root will follow the source motion according to the behavior defined in the subsequent settings. */
+	UPROPERTY(EditAnywhere, BlueprintReadWrite, Category = "Root Retarget Settings", meta = (ClampMin = "-5.0", ClampMax = "5.0", UIMin = "0.0", UIMax = "1.0"))
+	float TranslationAlpha = 1.0f;
+	
+	/** Range 0 to 1. Default 0. Blends the retarget root's translation to the exact source location.
+	*  At 0 the root is placed at the retargeted location.
+	*  At 1 the root is placed at the location of the source's retarget root bone.*/
+	UPROPERTY(EditAnywhere, BlueprintReadWrite, Category = "Root Retarget Settings", meta = (ClampMin = "0.0", ClampMax = "1.0", UIMin = "0.0", UIMax = "1.0"))
+	float BlendToSource = 0.0f;
+
+	/** Per-axis weights for the Blend to Source. */
+	UPROPERTY(EditAnywhere, BlueprintReadWrite, Category = "Root Retarget Settings", meta = (ClampMin = "0.0", ClampMax = "1.0", UIMin = "0.0", UIMax = "1.0"))
+	FVector BlendToSourceWeights = FVector::OneVector;
+
+	/** Default 1. Scales the motion of the root position in the horizontal plane (X,Y). */
+	UPROPERTY(EditAnywhere, BlueprintReadWrite, Category = "Root Retarget Settings", meta = (UIMin = "0.0", UIMax = "3.0"))
+	float ScaleHorizontal = 1.0f;
+
+	/** Default 1. Scales the motion of the root position in the vertical direction (Z). */
+	UPROPERTY(EditAnywhere, BlueprintReadWrite, Category = "Root Retarget Settings", meta = (UIMin = "0.0", UIMax = "3.0"))
+	float ScaleVertical = 1.0f;
+
+	/** Applies a static component-space translation offset to the retarget root.*/
+	UPROPERTY(EditAnywhere, BlueprintReadWrite, Category = "Root Retarget Settings")
+	FVector TranslationOffset = FVector::ZeroVector;
+
+	/** Applies a static local-space rotation offset to the retarget root.*/
+	UPROPERTY(EditAnywhere, BlueprintReadWrite, Category = "Root Retarget Settings", meta = (ClampMin = "-180.0", ClampMax = "180.0", UIMin = "-180.0", UIMax = "180.0"))
+	FRotator RotationOffset = FRotator::ZeroRotator;
+
+	/** Range 0 to 1. Default 1. Control whether modifications made to the root will affect the horizontal component of IK positions.
+	*  At 0 the IK positions are independent of the root modifications.
+	*  At 1 the IK positions are calculated relative to the modified root location.*/
+	UPROPERTY(EditAnywhere, BlueprintReadWrite, Category = "Root Retarget Settings", meta = (ClampMin = "0.0", ClampMax = "1.0", UIMin = "0.0", UIMax = "1.0", DisplayName = "Affect IK Horizontal"))
+	float AffectIKHorizontal = 1.0f;
+	
+	/** Range 0 to 1. Default 0. Control whether modifications made to the root will affect the vertical component of IK positions.
+	*  At 0 the IK positions are independent of the root modifications.
+	*  At 1 the IK positions are calculated relative to the modified root location.*/
+	UPROPERTY(EditAnywhere, BlueprintReadWrite, Category = "Root Retarget Settings", meta = (ClampMin = "0.0", ClampMax = "1.0", UIMin = "0.0", UIMax = "1.0", DisplayName = "Affect IK Vertical"))
+	float AffectIKVertical = 0.0f;
+
+	void CopySettingsFromAsset(const URetargetRootSettings* AssetChainSettings);
+
+	FVector GetAffectIKWeightVector() const
+	{
+		return FVector(AffectIKHorizontal, AffectIKHorizontal, AffectIKVertical);
+	}
+};
+
+UENUM()
+enum class EBasicAxis
+{
+	X,
+	Y,
+	Z,
+	NegX,
+	NegY,
+	NegZ
+};
+
+UENUM()
+enum class EWarpingDirectionSource
+{
+	Goals,
+	Chain,
+};
+
+USTRUCT()
+struct IKRIG_API FRetargetGlobalSettings
+{
+	GENERATED_USTRUCT_BODY()
+	
+	/** When false, the motion of the Retarget Root bone is not copied from the source. Useful for debugging issues with the root settings.
+	 * Note: the retargeting order is: Root > FK > IK
+	 */
+	UPROPERTY(EditAnywhere, Category = RetargetPhases)
+	bool bEnableRoot = true;
+	
+	/** When false, limbs are not copied via FK. Useful for debugging limb issues suspected to be caused by FK chain settings.
+	 * Note: the retargeting order is: Root > FK > IK
+	 */
+	UPROPERTY(EditAnywhere, Category = RetargetPhases)
+	bool bEnableFK = true;
+	
+	/** When false, IK is not applied as part of retargeter. Useful for debugging limb issues suspected to be caused by IK.
+	 * Note: the retargeting order is: Root > FK > IK
+	 */
+	UPROPERTY(EditAnywhere, Category = RetargetPhases)
+	bool bEnableIK = true;
+
+	/** Enable IK Warping.
+	 * These options allow for global modifications to all IK Goals that have "Affected by IK Warping" turned on (the default).
+	 * "Affected by IK Warping" can be found in the IK settings of a retarget chain. */
+	UPROPERTY(EditAnywhere, Category = Warping)
+	bool bIKWarping = false;
+
+	/** Defines the source used to determine the forward direction as the character animates in world space. Default is "Goals".
+	 * This method uses various positions on the character to create a "best fit" global rotation that approximates the facing direction of the character over time.
+	 * This global rotation is used to define the forward and sideways directions used when warping goals along those axes.
+	 * The options are:
+	 * Goals: uses the positions of the IK goals to approximate the facing direction. This is best used on characters with a vertical spine, like bipeds.
+	 * Chain: uses the positions of the bones in a retarget chain to approximate the facing direction. This is best for characters with a horizontal spine, like quadrupeds.
+	 */
+	UPROPERTY(EditAnywhere, Category = Warping)
+	EWarpingDirectionSource DirectionSource = EWarpingDirectionSource::Goals;
+	
+	/** The world space axis that represents the forward facing direction for your character. By default, in Unreal, this is +Y.  */
+	UPROPERTY(EditAnywhere, Category = Warping)
+	EBasicAxis ForwardDirection = EBasicAxis::Y;
+
+	/** When using the "Chain" option as a Direction Source, this defines the chain to use to determine the facing direction of the character.
+	 * Typically this would be the chain that contains the Spine bones. */
+	UPROPERTY(EditAnywhere, Category = Warping, meta=(EditCondition="DirectionSource == EWarpingDirectionSource::Chain"))
+	FName DirectionChain;
+
+	/** Range 0 to Inf. Default 1. Warps IK goal positions in the forward direction. Useful for stride warping.
+	 * Values below 1 will create smaller, squashed strides. Values greater than 1 will create stretched, longer strides.*/
+	UPROPERTY(EditAnywhere, Category = Warping, meta = (UIMin = "0.0", UIMax = "5.0"))
+	float WarpForwards = 1.0f;
+
+	/** Range -+Inf. Default is 0. A static offset in world units to move the IK goals perpendicular to the forward direction.
+	 * Values less than zero will move the goals towards the center-line of the character. Values greater than zero push the goals outwards.*/
+	UPROPERTY(EditAnywhere, Category = Warping, meta = (UIMin = "-10.0", UIMax = "100.0"))
+	float SidewaysOffset = 0.0f;
+
+	/** Range 0 to +Inf. Default is 1.0f.
+	 * Values below 1 pull all the goals towards the average of all the goals (towards each other).
+	 * Values greater than 1 push the goals apart.*/
+	UPROPERTY(EditAnywhere, Category = Warping, meta = (UIMin = "0.0", UIMax = "2.0"))
+	float WarpSplay = 1.0f;
+
+	static FVector GetAxisVector(const EBasicAxis& Axis)
+	{
+		switch (Axis)
+		{
+		case EBasicAxis::X:
+			return FVector::XAxisVector;
+		case EBasicAxis::Y:
+			return FVector::YAxisVector;
+		case EBasicAxis::Z:
+			return FVector::ZAxisVector;
+		case EBasicAxis::NegX:
+			return -FVector::XAxisVector;
+		case EBasicAxis::NegY:
+			return -FVector::YAxisVector;
+		case EBasicAxis::NegZ:
+			return -FVector::ZAxisVector;
+		default:
+			checkNoEntry();
+			return FVector::ZeroVector;
+		}
+	}
+};
\ No newline at end of file
diff --git a/Engine/Plugins/Animation/IKRig/Source/IKRig/Public/Retargeter/IKRetargeter.h b/Engine/Plugins/Animation/IKRig/Source/IKRig/Public/Retargeter/IKRetargeter.h
--- a/Engine/Plugins/Animation/IKRig/Source/IKRig/Public/Retargeter/IKRetargeter.h
+++ b/Engine/Plugins/Animation/IKRig/Source/IKRig/Public/Retargeter/IKRetargeter.h
@@ -2,11 +2,18 @@
 
 #pragma once
 
+#include "IKRetargetProfile.h"
 #include "IKRigDefinition.h"
+#include "IKRetargetSettings.h"
 
 #include "IKRetargeter.generated.h"
 
+#if WITH_EDITOR
+class FIKRetargetEditorController;
+#endif
 struct FIKRetargetPose;
+class URetargetChainSettings;
+class URetargetRootSettings;
 
 struct UE_DEPRECATED(5.1, "Use URetargetChainSettings instead.") FRetargetChainMap;
 USTRUCT()
@@ -24,35 +31,21 @@ struct IKRIG_API FRetargetChainMap
 	FName TargetChain = NAME_None;
 };
 
-UENUM(BlueprintType)
-enum class ERetargetTranslationMode : uint8
-{
-	None			UMETA(DisplayName = "None"),
-	GloballyScaled	UMETA(DisplayName = "Globally Scaled"),
-	Absolute		UMETA(DisplayName = "Absolute"),
-};
-
-UENUM(BlueprintType)
-enum class ERetargetRotationMode : uint8
-{
-	
-	Interpolated		UMETA(DisplayName = "Interpolated"),
-	OneToOne			UMETA(DisplayName = "One to One"),
-	OneToOneReversed	UMETA(DisplayName = "One to One Reversed"),
-	None				UMETA(DisplayName = "None"),
-};
-
 UCLASS()
 class IKRIG_API URetargetChainSettings : public UObject
 {
 	GENERATED_BODY()
 
-	public:
+public:
 	
 	URetargetChainSettings() = default;
 	
 	URetargetChainSettings(const FName& TargetChain) : TargetChain(TargetChain){}
 	
+	/** UObject */
+	virtual void Serialize(FArchive& Ar) override;
+	/** END UObject */
+	
 	/** The chain on the Source IK Rig asset to copy animation FROM. */
 	UPROPERTY(VisibleAnywhere, Category = "Chain Mapping")
 	FName SourceChain;
@@ -61,132 +54,109 @@ class IKRIG_API URetargetChainSettings : public UObject
 	UPROPERTY(VisibleAnywhere, Category = "Chain Mapping")
 	FName TargetChain;
 
-	/** Whether to copy the shape of the chain from the source skeleton using the Rotation and Translation modes. Default is true.
-	* NOTE: All FK operations run before the IK pass to copy the shape of the FK chain from the source skeleton. */
-	UPROPERTY(EditAnywhere, Category = "FK Adjustments")
-	bool CopyPoseUsingFK = true;
-	
-	/** Determines how rotation is copied from the source chain to the target chain. Default is Interpolated.
-	* Interpolated: Source and target chains are normalized by their length, then each target bone rotation is generated by finding the rotation at the same normalized distance on the source chain and interpolating between the neighboring bones.
-	* One to One: Each target bone rotation is copied from the equivalent bone in the source chain, based on the order in the chain, starting at the root of the chain. If the target chain has more bones than the source, the extra bones will remain at their reference pose.
-	* One to One Reversed: Same as One-to-One, but starting from the tip of the chain.
-	* None: The rotation of each target bone in the chain is left at the reference pose. */
-	UPROPERTY(EditAnywhere, Category = "FK Adjustments")
-	ERetargetRotationMode RotationMode;
-
-	/** Range +/- infinity. Default 1. Scales the amount of rotation that is applied.
-	*  If Rotation Mode is None this parameter has no effect.
-	*  Otherwise, this parameter blends the rotation of each bone in the chain from the base retarget pose (0) to the retargeted pose (1).*/
-	UPROPERTY(EditAnywhere, Category = "FK Adjustments", meta = (UIMin = "0.0", UIMax = "1.0"))
-	float RotationAlpha = 1.0f;
-
-	/** Determines how translation is copied from the source chain to the target chain. Default is None.
-	* None: Translation of target bones are left unmodified from the retarget pose.
-	* Globally Scaled: Translation of target bone is set to the source bone offset multiplied by the global scale of the skeleton (determined by the relative height difference between retarget root bones).
-	* Absolute: Translation of target bone is set to the absolute position of the source bone. */
-	UPROPERTY(EditAnywhere, Category = "FK Adjustments")
-	ERetargetTranslationMode TranslationMode;
-
-	/** Range +/- infinity. Default 1. Scales the amount of translation that is applied. Exact behavior depends on the Translation Mode.
-	*  In None Mode, this parameter has no effect.
-	*  In Globally Scaled and Absolute modes, the translation offset is scaled by this parameter.*/
-	UPROPERTY(EditAnywhere, Category = "FK Adjustments", meta = (UIMin = "0.0", UIMax = "1.0"))
-	float TranslationAlpha = 1.0f;
-
-	/** Whether to modify the location of the IK goal on this chain. Default is true.
-	 * NOTE: This only has an effect if the chain has an IK Goal assigned to it in the Target IK Rig asset.
-	 * NOTE: If off, and this chain has an IK Goal, the IK will still be evaluated, but the Goal is set to the input bone location (from the FK pass).*/
-	UPROPERTY(EditAnywhere, Category = "IK Adjustments")
-	bool DriveIKGoal = true;
-
-	/** Range 0 to 1. Default 0. Blends IK goal position from retargeted location (0) to source bone location (1).
-	*  At 0 the goal is placed at the retargeted location.
-	*  At 1 the goal is placed at the location of the source chain's end bone. */
-	UPROPERTY(EditAnywhere, Category = "IK Adjustments", meta = (ClampMin = "0.0", ClampMax = "1.0", UIMin = "0.0", UIMax = "1.0"))
-	float BlendToSource = 0.0f;
-	
-	/** Range 0 to 1. Default 1. Weight each axis separately when using Blend To Source.
-	*  At 0 the goal is placed at the retargeted location.
-	*  At 1 the goal is placed at the location of the source chain's end bone. */
-	UPROPERTY(EditAnywhere, Category = "IK Adjustments", meta = (ClampMin = "0.0", ClampMax = "1.0", UIMin = "0.0", UIMax = "1.0"))
-	FVector BlendToSourceWeights = FVector::OneVector;
-
-	/** Default 0, 0, 0. Apply static global-space offset to IK goal position. */
-	UPROPERTY(EditAnywhere, Category = "IK Adjustments")
-	FVector StaticOffset;
-
-	/** Default 0, 0, 0. Apply static local-space offset to IK goal position. */
-	UPROPERTY(EditAnywhere, Category = "IK Adjustments")
-	FVector StaticLocalOffset;
-
-	/** Default 0, 0, 0. Apply static local-space offset to IK goal rotation. */
-	UPROPERTY(EditAnywhere, Category = "IK Adjustments")
-	FRotator StaticRotationOffset;
-	
-	/** Range 0 to 5. Default 1. Brings IK goal closer (-) or further (+) from origin of chain.
-	*  At 0 the effector is placed at the origin of the chain.
-	*  Values greater than 1 will stretch the chain beyond the retargeted length. */
-	UPROPERTY(EditAnywhere, Category = "IK Adjustments", meta = (ClampMin = "0.0", ClampMax = "5.0", UIMin = "0.1", UIMax = "2.0"))
-	float Extension = 1.0f;
-
-	/** The name of the curve on the source animation that contains the speed of the end effector bone.*/
-	UPROPERTY(EditAnywhere, Category = "Plant IK by Speed", meta = (ClampMin = "0.0", ClampMax = "100.0", UIMin = "0.0", UIMax = "100.0"))
-	bool UseSpeedCurveToPlantIK = false;
-	
-	/** The name of the curve on the source animation that contains the speed of the end effector bone.*/
-	UPROPERTY(EditAnywhere, Category = "Plant IK by Speed", meta = (ClampMin = "0.0", ClampMax = "100.0", UIMin = "0.0", UIMax = "100.0"))
-	FName SpeedCurveName;
-
-	/** Range 0 to 1000. Default 15. The maximum speed a source bone can be moving while being considered 'planted'.
-	*  The target IK goal will not be allowed to move whenever the source bone speed drops below this threshold speed. */
-	UPROPERTY(EditAnywhere, Category = "Plant IK by Speed", meta = (ClampMin = "0.0", ClampMax = "100.0", UIMin = "0.0", UIMax = "100.0"))
-	float VelocityThreshold = 15.0f;
-
-	// How stiff the spring model is that smoothly pulls the IK position after unplanting (more stiffness means more oscillation around the target value)
-	UPROPERTY(EditAnywhere, Category = "Plant IK by Speed", meta = (ClampMin = "0.0", UIMin = "0.0"))
-	float UnplantStiffness = 250.0f;
-
-	// How much damping to apply to the spring (0 means no damping, 1 means critically damped which means no oscillation)
-	UPROPERTY(EditAnywhere, Category = "Plant IK by Speed", meta = (ClampMin = "0.0", ClampMax = "10.0", UIMin = "0.0", UIMax = "1.0"))
-	float UnplantCriticalDamping = 1.0f;
+	/** The settings used to control the motion on this target chain. */
+	UPROPERTY(EditAnywhere, Category = "Settings")
+	FTargetChainSettings Settings;
+
+	/** Deprecated properties from before FTargetChainSettings / profile refactor  (July 2022)*/
+	UPROPERTY()
+	bool CopyPoseUsingFK_DEPRECATED = true;
+	UPROPERTY()
+	ERetargetRotationMode RotationMode_DEPRECATED;
+	UPROPERTY()
+	float RotationAlpha_DEPRECATED = 1.0f;
+	UPROPERTY()
+	ERetargetTranslationMode TranslationMode_DEPRECATED;
+	UPROPERTY()
+	float TranslationAlpha_DEPRECATED = 1.0f;
+	UPROPERTY()
+	bool DriveIKGoal_DEPRECATED = true;
+	UPROPERTY()
+	float BlendToSource_DEPRECATED = 0.0f;
+	UPROPERTY()
+	FVector BlendToSourceWeights_DEPRECATED = FVector::OneVector;
+	UPROPERTY()
+	FVector StaticOffset_DEPRECATED;
+	UPROPERTY()
+	FVector StaticLocalOffset_DEPRECATED;
+	UPROPERTY()
+	FRotator StaticRotationOffset_DEPRECATED;
+	UPROPERTY()
+	float Extension_DEPRECATED = 1.0f;
+	UPROPERTY()
+	bool UseSpeedCurveToPlantIK_DEPRECATED = false;
+	UPROPERTY()
+	FName SpeedCurveName_DEPRECATED;
+	UPROPERTY()
+	float VelocityThreshold_DEPRECATED = 15.0f;
+	UPROPERTY()
+	float UnplantStiffness_DEPRECATED = 250.0f;
+	UPROPERTY()
+	float UnplantCriticalDamping_DEPRECATED = 1.0f;
+	/** END deprecated properties */
+
+	// pointer to editor for details customization
+	#if WITH_EDITOR
+	TSharedPtr<FIKRetargetEditorController> EditorController;
+	#endif
 };
 
 UCLASS()
 class IKRIG_API URetargetRootSettings: public UObject
 {
 	GENERATED_BODY()
-
-	public:
 	
-	URetargetRootSettings() = default;
+public:
+	
+	/** UObject */
+	virtual void Serialize(FArchive& Ar) override;
+	/** END UObject */
+
+	/** The settings used to control the motion of the target root bone. */
+	UPROPERTY(EditAnywhere, Category = "Settings")
+	FTargetRootSettings Settings;
 
-	/** Range 0 to 1. Default 0. Blends the retarget root's translation to the exact source location.
-	*  At 0 the root is placed at the retargeted location.
-	*  At 1 the root is placed at the location of the source's retarget root bone. */
-	UPROPERTY(EditAnywhere, Category = "Root Retarget Settings", meta = (ClampMin = "0.0", ClampMax = "1.0", UIMin = "0.0", UIMax = "1.0"))
-	float BlendToSource = 0.0f;
+	// pointer to editor for details customization
+	#if WITH_EDITOR
+	TSharedPtr<FIKRetargetEditorController> EditorController;
+	#endif
 
-	/** Per-axis weights for the Blend to Source. */
-	UPROPERTY(EditAnywhere, Category = "Root Retarget Settings", meta = (ClampMin = "0.0", ClampMax = "1.0", UIMin = "0.0", UIMax = "1.0"))
-	FVector BlendToSourceWeights = FVector::OneVector;
+private:
+	/** Deprecated properties from before FTargetRootSettings / profile refactor */
+	UPROPERTY()
+	bool RetargetRootTranslation_DEPRECATED = true;
+	UPROPERTY()
+	float GlobalScaleHorizontal_DEPRECATED = 1.0f;
+	UPROPERTY()
+	float GlobalScaleVertical_DEPRECATED = 1.0f;
+	UPROPERTY()
+	FVector BlendToSource_DEPRECATED = FVector::ZeroVector;
+	UPROPERTY()
+	FVector StaticOffset_DEPRECATED = FVector::ZeroVector;
+	UPROPERTY()
+	FRotator StaticRotationOffset_DEPRECATED = FRotator::ZeroRotator;
+	/** END deprecated properties */
+};
 
-	/** Default 1. Scales the motion of the root position in the horizontal plane (X,Y). */
-	UPROPERTY(EditAnywhere, Category = "Root Retarget Settings", meta = (UIMin = "0.0", UIMax = "3.0"))
-	float GlobalScaleHorizontal = 1.0f;
 
-	/** Default 1. Scales the motion of the root position in the vertical direction (Z). */
-	UPROPERTY(EditAnywhere, Category = "Root Retarget Settings", meta = (UIMin = "0.0", UIMax = "3.0"))
-	float GlobalScaleVertical = 1.0f;
+UCLASS()
+class IKRIG_API UIKRetargetGlobalSettings: public UObject
+{
+	GENERATED_BODY()
+	
+public:
 
-	/** Applies a static component-space translation offset to the retarget root.*/
-	UPROPERTY(EditAnywhere, Category = "Root Retarget Settings")
-	FVector StaticOffset = FVector::ZeroVector;
+	/** Global retargeter settings. */
+	UPROPERTY(EditAnywhere, Category = "Settings")
+	FRetargetGlobalSettings Settings;
 
-	/** Applies a static component-space rotation offset to the retarget root.*/
-	UPROPERTY(EditAnywhere, Category = "Root Retarget Settings", meta = (ClampMin = "-180.0", ClampMax = "180.0", UIMin = "-180.0", UIMax = "180.0"))
-	FRotator StaticRotationOffset = FRotator::ZeroRotator;
+	// pointer to editor for details customization
+	#if WITH_EDITOR
+	TSharedPtr<FIKRetargetEditorController> EditorController;
+	#endif
 };
 
+
 USTRUCT()
 struct IKRIG_API FIKRetargetPose
 {
@@ -197,7 +167,7 @@ public:
 	FIKRetargetPose() = default;
 
 	FQuat GetDeltaRotationForBone(const FName BoneName) const;
-	void SetDeltaRotationForBone(FName BoneName, FQuat RotationOffset);
+	void SetDeltaRotationForBone(FName BoneName, const FQuat& RotationDelta);
 	const TMap<FName, FQuat>& GetAllDeltaRotations() const { return BoneRotationOffsets; };
 
 	FVector GetRootTranslationDelta() const;
@@ -244,14 +214,85 @@ public:
 	/** Get read-write access to the target IK Rig asset.
 	 * WARNING: do not use for editing the data model. Use Controller class instead. */
 	UIKRigDefinition* GetTargetIKRigWriteable() const { return TargetIKRigAsset.LoadSynchronous(); };
+
 	/** Get read-only access to the chain mapping */
-	const TArray<TObjectPtr<URetargetChainSettings>> GetAllChainSettings() const { return ChainSettings; };
+	const TArray<TObjectPtr<URetargetChainSettings>>& GetAllChainSettings() const { return ChainSettings; };
+	/** Get read-only access to the chain settings for a given chain (null if chain not in retargeter) */
+	const FTargetChainSettings* GetChainSettingsByName(const FName& TargetChainName) const;
 	/** Get access to the root settings */
-	URetargetRootSettings* GetRetargetRootSettings() const { return RootSettings; };
+	URetargetRootSettings* GetRootSettingsUObject() const { return RootSettings; };
+	/** Get access to the global settings uobject*/
+	UIKRetargetGlobalSettings* GetGlobalSettingsUObject() const { return GlobalSettings; };
+	/** Get access to the global settings itself */
+	const FRetargetGlobalSettings& GetGlobalSettings() const { return GlobalSettings->Settings; };
+
 	/** Get read-only access to a retarget pose */
 	const FIKRetargetPose* GetCurrentRetargetPose(const ERetargetSourceOrTarget& SourceOrTarget) const;
+	/** Get name of the current retarget pose */
+	FName GetCurrentRetargetPoseName(const ERetargetSourceOrTarget& SourceOrTarget) const;
+	/** Get read-only access to a retarget pose */
+	const FIKRetargetPose* GetRetargetPoseByName(const ERetargetSourceOrTarget& SourceOrTarget, const FName PoseName) const;
 	/* Get name of default pose */
 	static const FName GetDefaultPoseName();
+	
+	/* Get the current retarget profile (may be null) */
+	const FRetargetProfile* GetCurrentProfile() const;
+	/* Get the retarget profile by name (may be null) */
+	const FRetargetProfile* GetProfileByName(const FName& ProfileName) const;
+
+	/** Returns the chain settings associated with a given target chain in an IK Retargeter Asset using the given profile name (optional) */
+	UFUNCTION(BlueprintCallable, BlueprintPure, Category=RetargetAsset)
+	static FTargetChainSettings GetChainSettingsFromRetargetAsset(
+		const UIKRetargeter* RetargetAsset,
+		const FName TargetChainName,
+		const FName OptionalProfileName);
+
+	/** Returns the root settings in an IK Retargeter Asset using the given profile name (optional) */
+	UFUNCTION(BlueprintCallable, BlueprintPure, Category=RetargetAsset)
+	static void GetRootSettingsFromRetargetAsset(
+		const UIKRetargeter* RetargetAsset,
+		const FName OptionalProfileName,
+		FTargetRootSettings& OutSettings);
+
+	/** Returns the chain settings associated with a given target chain in the supplied Retarget Profile. */
+	UFUNCTION(BlueprintCallable, BlueprintPure, Category=RetargetProfile)
+	static FTargetChainSettings GetChainSettingsFromRetargetProfile(
+		UPARAM(ref) FRetargetProfile& RetargetProfile,
+		const FName TargetChainName);
+
+	/** Set the root settings in a retarget profile (will set bApplyRootSettings to true). */
+	UFUNCTION(BlueprintCallable, Category=RetargetProfile)
+	static void SetRootSettingsInRetargetProfile(
+		UPARAM(ref) FRetargetProfile& RetargetProfile,
+		const FTargetRootSettings& RootSettings);
+	
+	/** Set the chain settings in a retarget profile (will set bApplyChainSettings to true). */
+	UFUNCTION(BlueprintCallable, Category=RetargetProfile)
+	static void SetChainSettingsInRetargetProfile(
+		UPARAM(ref) FRetargetProfile& RetargetProfile,
+		const FTargetChainSettings& ChainSettings,
+		const FName TargetChainName);
+
+	/** Set the chain FK settings in a retarget profile (will set bApplyChainSettings to true). */
+	UFUNCTION(BlueprintCallable, Category=RetargetProfile)
+	static void SetChainFKSettingsInRetargetProfile(
+		UPARAM(ref) FRetargetProfile& RetargetProfile,
+		const FTargetChainFKSettings& FKSettings,
+		const FName TargetChainName);
+
+	/** Set the chain IK settings in a retarget profile (will set bApplyChainSettings to true). */
+	UFUNCTION(BlueprintCallable, Category=RetargetProfile)
+	static void SetChainIKSettingsInRetargetProfile(
+		UPARAM(ref) FRetargetProfile& RetargetProfile,
+		const FTargetChainIKSettings& IKSettings,
+		const FName TargetChainName);
+
+	/** Set the chain Speed Plant settings in a retarget profile (will set bApplyChainSettings to true). */
+	UFUNCTION(BlueprintCallable, Category=RetargetProfile)
+	static void SetChainSpeedPlantSettingsInRetargetProfile(
+		UPARAM(ref) FRetargetProfile& RetargetProfile,
+		const FTargetChainSpeedPlantSettings& SpeedPlantSettings,
+		const FName TargetChainName);
 
 	/** UObject */
 	virtual void Serialize(FArchive& Ar) override;
@@ -294,26 +335,16 @@ private:
 	UPROPERTY(EditAnywhere, Category = Target)
 	TSoftObjectPtr<USkeletalMesh> TargetPreviewMesh = nullptr;
 #endif
-
-	/** the retarget root settings */
-	UPROPERTY()
-	TObjectPtr<URetargetRootSettings> RootSettings;
 	
 public:
-	
-	/** When false, translational motion of skeleton root is not copied. Useful for debugging.*/
-	UPROPERTY(EditAnywhere, Category = RetargetPhases)
-	bool bRetargetRoot = true;
-	
-	/** When false, limbs are not copied via FK. Useful for debugging limb issues suspected to be caused by FK pose.*/
-	UPROPERTY(EditAnywhere, Category = RetargetPhases)
-	bool bRetargetFK = true;
-	
-	/** When false, IK is not applied as part of retargeter. Useful for debugging limb issues suspected to be caused by IK.*/
-	UPROPERTY(EditAnywhere, Category = RetargetPhases)
-	bool bRetargetIK = true;
 
 #if WITH_EDITORONLY_DATA
+	UPROPERTY()
+	bool bRetargetRoot_DEPRECATED = true;
+	UPROPERTY()
+	bool bRetargetFK_DEPRECATED = true;
+	UPROPERTY()
+	bool bRetargetIK_DEPRECATED = true;
 	UPROPERTY()
 	float TargetActorOffset_DEPRECATED = 0.0f;
 	UPROPERTY()
@@ -331,6 +362,26 @@ public:
 	UPROPERTY(EditAnywhere, Category = PreviewSettings)
 	FVector SourceMeshOffset;
 
+	/** Toggle debug drawing for retargeting in the viewport. */
+	UPROPERTY(EditAnywhere, Category = DebugSettings)
+	bool bDebugDraw = true;
+
+	/** Draw final IK goal locations. */
+	UPROPERTY(EditAnywhere, Category = DebugSettings)
+	bool bDrawFinalGoals = true;
+
+	/** Draw goal locations from source skeleton. */
+	UPROPERTY(EditAnywhere, Category = DebugSettings)
+	bool bDrawSourceLocations = true;
+	
+	/** The visual size of the IK goals in the viewport. */
+	UPROPERTY(EditAnywhere, Category = DebugSettings)
+	float ChainDrawSize = 5.0f;
+
+	/** The thickness of lines on the IK goals in the viewport. */
+	UPROPERTY(EditAnywhere, Category = DebugSettings)
+	float ChainDrawThickness = 0.5f;
+	
 	/** The visual size of the bones in the viewport (saved between sessions). This is set from the viewport Character>Bones menu*/
 	UPROPERTY()
 	float BoneDrawSize = 8.0f;
@@ -354,6 +405,20 @@ private:
 	UPROPERTY()
 	TArray<TObjectPtr<URetargetChainSettings>> ChainSettings;
 	
+	/** the retarget root settings */
+	UPROPERTY()
+	TObjectPtr<URetargetRootSettings> RootSettings;
+
+	/** the retarget root settings */
+	UPROPERTY()
+	TObjectPtr<UIKRetargetGlobalSettings> GlobalSettings;
+
+	/** settings profiles stored in this asset */
+	UPROPERTY()
+	TMap<FName, FRetargetProfile> Profiles;
+	UPROPERTY()
+	FName CurrentProfile = NAME_None;
+	
 	/** The set of retarget poses for the SOURCE skeleton.*/
 	UPROPERTY()
 	TMap<FName, FIKRetargetPose> SourceRetargetPoses;
diff --git a/Engine/Plugins/Animation/IKRig/Source/IKRig/Public/Solvers/IKRig_BodyMover.h b/Engine/Plugins/Animation/IKRig/Source/IKRig/Public/Solvers/IKRig_BodyMover.h
--- a/Engine/Plugins/Animation/IKRig/Source/IKRig/Public/Solvers/IKRig_BodyMover.h
+++ b/Engine/Plugins/Animation/IKRig/Source/IKRig/Public/Solvers/IKRig_BodyMover.h
@@ -111,13 +111,6 @@ private:
 
 	int32 GetIndexOfGoal(const FName& OldName) const;
 
-	static void ExtractRotation(
-	    const FVector& DX,
-	    const FVector& DY,
-	    const FVector& DZ,
-	    FQuat &Q,
-	    const unsigned int maxIter);
-
 	int32 BodyBoneIndex;
 };
 
diff --git a/Engine/Plugins/Animation/IKRig/Source/IKRigEditor/Private/IKRigEditor.cpp b/Engine/Plugins/Animation/IKRig/Source/IKRigEditor/Private/IKRigEditor.cpp
--- a/Engine/Plugins/Animation/IKRig/Source/IKRigEditor/Private/IKRigEditor.cpp
+++ b/Engine/Plugins/Animation/IKRig/Source/IKRigEditor/Private/IKRigEditor.cpp
@@ -52,15 +52,24 @@ void FIKRigEditor::StartupModule()
 
 	// register detail customizations
 	FPropertyEditorModule& PropertyEditorModule = FModuleManager::LoadModuleChecked<FPropertyEditorModule>("PropertyEditor");
-	// custom IK rig bone widget
+	// custom IK rig bone details
 	PropertyEditorModule.RegisterCustomClassLayout(UIKRigBoneDetails::StaticClass()->GetFName(), FOnGetDetailCustomizationInstance::CreateStatic(&FIKRigGenericDetailCustomization::MakeInstance));
 	ClassesToUnregisterOnShutdown.Add(UIKRigBoneDetails::StaticClass()->GetFName());
-	// custom IK goal widget
+	// custom IK rig goal details
 	PropertyEditorModule.RegisterCustomClassLayout(UIKRigEffectorGoal::StaticClass()->GetFName(), FOnGetDetailCustomizationInstance::CreateStatic(&FIKRigGenericDetailCustomization::MakeInstance));
 	ClassesToUnregisterOnShutdown.Add(UIKRigEffectorGoal::StaticClass()->GetFName());
-	// custom retargeter bone widget
+	// custom retargeter bone details
 	PropertyEditorModule.RegisterCustomClassLayout(UIKRetargetBoneDetails::StaticClass()->GetFName(), FOnGetDetailCustomizationInstance::CreateStatic(&FIKRetargetBoneDetailCustomization::MakeInstance));
 	ClassesToUnregisterOnShutdown.Add(UIKRetargetBoneDetails::StaticClass()->GetFName());
+	// custom retargeter chain details
+	PropertyEditorModule.RegisterCustomClassLayout(URetargetChainSettings::StaticClass()->GetFName(), FOnGetDetailCustomizationInstance::CreateStatic(&FRetargetChainSettingsCustomization::MakeInstance));
+	ClassesToUnregisterOnShutdown.Add(URetargetChainSettings::StaticClass()->GetFName());
+	// custom retargeter root details
+	PropertyEditorModule.RegisterCustomClassLayout(URetargetRootSettings::StaticClass()->GetFName(), FOnGetDetailCustomizationInstance::CreateStatic(&FRetargetRootSettingsCustomization::MakeInstance));
+	ClassesToUnregisterOnShutdown.Add(URetargetRootSettings::StaticClass()->GetFName());
+	// custom retargeter global details
+	PropertyEditorModule.RegisterCustomClassLayout(UIKRetargetGlobalSettings::StaticClass()->GetFName(), FOnGetDetailCustomizationInstance::CreateStatic(&FRetargetGlobalSettingsCustomization::MakeInstance));
+	ClassesToUnregisterOnShutdown.Add(UIKRetargetGlobalSettings::StaticClass()->GetFName());
 }
 
 void FIKRigEditor::ShutdownModule()
diff --git a/Engine/Plugins/Animation/IKRig/Source/IKRigEditor/Private/RetargetEditor/IKRetargetBatchOperation.cpp b/Engine/Plugins/Animation/IKRig/Source/IKRigEditor/Private/RetargetEditor/IKRetargetBatchOperation.cpp
--- a/Engine/Plugins/Animation/IKRig/Source/IKRigEditor/Private/RetargetEditor/IKRetargetBatchOperation.cpp
+++ b/Engine/Plugins/Animation/IKRig/Source/IKRigEditor/Private/RetargetEditor/IKRetargetBatchOperation.cpp
@@ -348,7 +348,7 @@ void FIKRetargetBatchOperation::ConvertAnimation(
 			}
 
 			// update goals 
-			Processor->CopyAllSettingsFromAsset();
+			Processor->ApplySettingsFromAsset();
 			
 			// calculate the delta time
 			const float TimeAtCurrentFrame = SourceSequence->GetTimeAtFrame(FrameIndex);
diff --git a/Engine/Plugins/Animation/IKRig/Source/IKRigEditor/Private/RetargetEditor/IKRetargetDefaultMode.cpp b/Engine/Plugins/Animation/IKRig/Source/IKRigEditor/Private/RetargetEditor/IKRetargetDefaultMode.cpp
--- a/Engine/Plugins/Animation/IKRig/Source/IKRigEditor/Private/RetargetEditor/IKRetargetDefaultMode.cpp
+++ b/Engine/Plugins/Animation/IKRig/Source/IKRigEditor/Private/RetargetEditor/IKRetargetDefaultMode.cpp
@@ -9,7 +9,7 @@
 #include "IPersonaPreviewScene.h"
 #include "Animation/DebugSkelMeshComponent.h"
 #include "RetargetEditor/IKRetargetEditor.h"
-#include "RetargetEditor/SIKRetargetChainMapList.h"
+#include "RetargetEditor/IKRetargetHitProxies.h"
 
 
 #define LOCTEXT_NAMESPACE "IKRetargetDefaultMode"
@@ -67,6 +67,176 @@ void FIKRetargetDefaultMode::Initialize()
 void FIKRetargetDefaultMode::Render(const FSceneView* View, FViewport* Viewport, FPrimitiveDrawInterface* PDI)
 {
 	FEdMode::Render(View, Viewport, PDI);
+
+	if (!EditorController.IsValid())
+	{
+		return;
+	}
+	
+	FIKRetargetEditorController* Controller = EditorController.Pin().Get();
+	if (Controller->GetRetargeterMode() == ERetargeterOutputMode::EditRetargetPose)
+	{
+		return;
+	}
+	
+	const UIKRetargeter* Asset = Controller->AssetController->GetAsset();
+	if (!Asset->bDebugDraw)
+	{
+		return;
+	}
+	
+	const UIKRetargetProcessor* RetargetProcessor = Controller->GetRetargetProcessor();
+	if (!RetargetProcessor && RetargetProcessor->IsInitialized())
+	{
+		return;
+	}
+
+	UDebugSkelMeshComponent* TargetSkelMesh = Controller->GetSkeletalMeshComponent(ERetargetSourceOrTarget::Target);
+	const FTransform ComponentTransform = TargetSkelMesh->GetComponentTransform();
+	const FVector ComponentOffset = ComponentTransform.GetTranslation();
+
+	const TArray<FName>& SelectedChains = Controller->GetSelectedChains();
+
+	constexpr FLinearColor Muted = FLinearColor(0.5,0.5,0.5, 0.5);
+	const FLinearColor SourceColor = (FLinearColor::Gray * FLinearColor::Blue) * Muted;
+	const FLinearColor GoalColor = FLinearColor::Yellow;
+	const FLinearColor MainColor = FLinearColor::Green;
+	const FLinearColor NonSelected = FLinearColor::Gray * 0.3f;
+
+	// get the root modification
+	const FRootRetargeter& RootRetargeter = RetargetProcessor->GetRootRetargeter();
+	const FVector RootModification = RootRetargeter.Target.RootTranslationDelta * RootRetargeter.Settings.GetAffectIKWeightVector();
+
+	// draw IK goals on each IK chain
+	const TArray<FRetargetChainPairIK>& IKChainPairs = RetargetProcessor->GetIKChainPairs();
+	for (const FRetargetChainPairIK& IKChainPair : IKChainPairs)
+	{
+		const FChainDebugData& ChainDebugData = IKChainPair.IKChainRetargeter.DebugData;
+		FTransform FinalTransform = ChainDebugData.OutputTransformEnd;
+		FinalTransform.AddToTranslation(ComponentOffset);
+
+		const bool bIsSelected = SelectedChains.Contains(IKChainPair.TargetBoneChainName);
+
+		PDI->SetHitProxy(new HIKRetargetEditorChainProxy(IKChainPair.TargetBoneChainName));
+
+		if (Asset->bDrawFinalGoals)
+		{
+			IKRigDebugRendering::DrawWireCube(
+			PDI,
+			FinalTransform,
+			bIsSelected ? GoalColor : GoalColor * NonSelected,
+			Asset->ChainDrawSize,
+			Asset->ChainDrawThickness);
+		}
+		
+		if (Asset->bDrawSourceLocations)
+		{
+			const FSourceChainIK& SourceChain = IKChainPair.IKChainRetargeter.Source;
+			FTransform SourceGoalTransform;
+			SourceGoalTransform.SetTranslation(SourceChain.CurrentEndPosition + ComponentOffset + RootModification);
+			SourceGoalTransform.SetRotation(SourceChain.CurrentEndRotation);
+
+			FLinearColor Color = bIsSelected ? SourceColor : SourceColor * NonSelected;
+
+			DrawWireSphere(
+				PDI,
+				SourceGoalTransform,
+				Color,
+				Asset->ChainDrawSize * 0.5f,
+				12,
+				SDPG_World,
+				0.0f,
+				0.001f,
+				false);
+
+			if (Asset->bDrawFinalGoals)
+			{
+				DrawDashedLine(
+					PDI,
+					SourceGoalTransform.GetLocation(),
+					FinalTransform.GetLocation(),
+					Color,
+					1.0f,
+					SDPG_Foreground);
+			}
+		}
+
+		// done drawing chain proxies
+		PDI->SetHitProxy(nullptr);
+	}
+
+	// draw lines on each FK chain
+	const TArray<FRetargetChainPairFK>& FKChainPairs = RetargetProcessor->GetFKChainPairs();
+	for (const FRetargetChainPairFK& FKChainPair : FKChainPairs)
+	{
+		const TArray<int32>& TargetChainBoneIndices = FKChainPair.FKDecoder.BoneIndices;
+		if (TargetChainBoneIndices.IsEmpty())
+		{
+			continue;
+		}
+		
+		const bool bIsSelected = SelectedChains.Contains(FKChainPair.TargetBoneChainName);
+		
+		// draw a line from start to end of chain, or in the case of a chain with only 1 bone in it, draw a sphere
+		PDI->SetHitProxy(new HIKRetargetEditorChainProxy(FKChainPair.TargetBoneChainName));
+		if (TargetChainBoneIndices.Num() > 1)
+		{
+			FTransform StartTransform = TargetSkelMesh->GetBoneTransform(TargetChainBoneIndices[0], ComponentTransform);
+			FTransform EndTransform = TargetSkelMesh->GetBoneTransform(TargetChainBoneIndices.Last(), ComponentTransform);
+			PDI->DrawLine(
+			StartTransform.GetLocation(),
+			EndTransform.GetLocation(),
+			bIsSelected ? MainColor : MainColor * NonSelected,
+			SDPG_Foreground,
+			Asset->ChainDrawThickness);
+		}
+		else
+		{
+			FTransform StartTransform = TargetSkelMesh->GetBoneTransform(TargetChainBoneIndices[0], ComponentTransform);
+			
+			DrawWireSphere(
+				PDI,
+				StartTransform.GetLocation(),
+				bIsSelected ? MainColor : MainColor * NonSelected,
+				Asset->ChainDrawSize,
+				12,
+				SDPG_World,
+				Asset->ChainDrawThickness,
+				0.001f,
+				false);
+		}
+		
+		PDI->SetHitProxy(nullptr);
+	}
+
+	// draw stride warping frame
+	FTransform WarpingFrame = RetargetProcessor->DebugData.StrideWarpingFrame;
+	DrawCoordinateSystem(PDI, WarpingFrame.GetLocation(), WarpingFrame.GetRotation().Rotator(), Asset->ChainDrawSize, SDPG_World, Asset->ChainDrawThickness);
+
+	// root bone name
+	const FName RootBoneName = Controller->AssetController->GetRetargetRootBone(ERetargetSourceOrTarget::Target);
+	const int32 RootBoneIndex = TargetSkelMesh->GetReferenceSkeleton().FindBoneIndex(RootBoneName);
+	if (RootBoneIndex == INDEX_NONE)
+	{
+		return;
+	}
+	const FTransform RootTransform = TargetSkelMesh->GetBoneTransform(RootBoneIndex, ComponentTransform);
+	const FVector RootCircleLocation = RootTransform.GetLocation() * FVector(1,1,0);
+	const bool bIsSelected = EditorController.Pin()->GetRootSelected();
+	const FLinearColor RootColor = bIsSelected ? MainColor : MainColor * NonSelected;
+	
+	PDI->SetHitProxy(new HIKRetargetEditorRootProxy());
+	DrawCircle(
+		PDI,
+		RootCircleLocation,
+		FVector(1, 0, 0),
+		FVector(0, 1, 0),
+		RootColor,
+		Asset->ChainDrawSize * 10.f,
+		12,
+		SDPG_World,
+		Asset->ChainDrawThickness * 2.0f);
+	PDI->SetHitProxy(nullptr);
 }
 
 bool FIKRetargetDefaultMode::AllowWidgetMove()
@@ -121,7 +291,8 @@ bool FIKRetargetDefaultMode::HandleClick(FEditorViewportClient* InViewportClient
 	}
 
 	const bool bLeftButtonClicked = Click.GetKey() == EKeys::LeftMouseButton;
-	constexpr bool bFromHierarchy = false;
+	const bool bCtrlOrShiftHeld = Click.IsControlDown() || Click.IsShiftDown();
+	const ESelectionEdit EditMode = bCtrlOrShiftHeld ? ESelectionEdit::Add : ESelectionEdit::Replace;
 	
 	// did we click on an actor in the viewport?
 	const bool bHitActor = HitProxy && HitProxy->IsA(HActor::StaticGetType());
@@ -137,23 +308,33 @@ bool FIKRetargetDefaultMode::HandleClick(FEditorViewportClient* InViewportClient
 	if (bLeftButtonClicked && bHitBone)
 	{
 		const HPersonaBoneHitProxy* BoneProxy = static_cast<HPersonaBoneHitProxy*>(HitProxy);
-		const TArray<FName> BoneNames{BoneProxy->BoneName};
-		
-		const bool bCtrlOrShiftHeld = Click.IsControlDown() || Click.IsShiftDown();
-		const EBoneSelectionEdit EditMode = bCtrlOrShiftHeld ? EBoneSelectionEdit::Add : EBoneSelectionEdit::Replace;
-		
+		const TArray BoneNames{BoneProxy->BoneName};
+		constexpr bool bFromHierarchy = false;
 		Controller->EditBoneSelection(BoneNames, EditMode, bFromHierarchy);
-		
 		return true;
 	}
 
-	// did we click in empty space in viewport?
-	if (!(bHitActor || bHitBone))
+	// did we click on a chain in the viewport?
+	const bool bHitChain = HitProxy && HitProxy->IsA(HIKRetargetEditorChainProxy::StaticGetType());
+	if (bLeftButtonClicked && bHitChain)
 	{
-		// deselect all meshes, bones, chains and update details view
-		Controller->ClearSelection();
+		const HIKRetargetEditorChainProxy* ChainProxy = static_cast<HIKRetargetEditorChainProxy*>(HitProxy);
+		const TArray ChainNames{ChainProxy->TargetChainName};
+		constexpr bool bFromChainView = false;
+		Controller->EditChainSelection(ChainNames, EditMode, bFromChainView);
+		return true;
 	}
-	
+
+	// did we click on the root in the viewport?
+	const bool bHitRoot = HitProxy && HitProxy->IsA(HIKRetargetEditorRootProxy::StaticGetType());
+	if (bLeftButtonClicked && bHitRoot)
+	{
+		Controller->SetRootSelected(true);
+		return true;
+	}
+
+	// we didn't hit anything, therefore clicked in empty space in viewport
+	Controller->ClearSelection(); // deselect all meshes, bones, chains and update details view
 	return true;
 }
 
diff --git a/Engine/Plugins/Animation/IKRig/Source/IKRigEditor/Private/RetargetEditor/IKRetargetDetails.cpp b/Engine/Plugins/Animation/IKRig/Source/IKRigEditor/Private/RetargetEditor/IKRetargetDetails.cpp
--- a/Engine/Plugins/Animation/IKRig/Source/IKRigEditor/Private/RetargetEditor/IKRetargetDetails.cpp
+++ b/Engine/Plugins/Animation/IKRig/Source/IKRigEditor/Private/RetargetEditor/IKRetargetDetails.cpp
@@ -15,6 +15,8 @@
 #include "Widgets/Input/SSegmentedControl.h"
 #include "ScopedTransaction.h"
 #include "AnimationRuntime.h"
+#include "SSearchableComboBox.h"
+#include "UObject/UnrealTypePrivate.h"
 
 #if WITH_EDITOR
 #include "HAL/PlatformApplicationMisc.h"
@@ -80,7 +82,7 @@ FEulerTransform UIKRetargetBoneDetails::GetTransform(EIKRetargetTransformType Tr
 			// this is the only stored data we have for bone pose offsets
 			const ERetargetSourceOrTarget SourceOrTarget = EditorController->GetSourceOrTarget();
 			const FRotator LocalRotationDelta = EditorController->AssetController->GetRotationOffsetForRetargetPoseBone(SelectedBone, SourceOrTarget).Rotator();
-			const FVector GlobalTranslationDelta = IsRootBone() ? EditorController->AssetController->GetTranslationOffsetOnRetargetRootBone(SourceOrTarget) : FVector::Zero();
+			const FVector GlobalTranslationDelta = IsRootBone() ? EditorController->AssetController->GetCurrentRetargetPose(SourceOrTarget).GetRootTranslationDelta() : FVector::Zero();
 			const int32 ParentIndex = RefSkeleton.GetParentIndex(BoneIndex);
 
 			if (bLocalSpace)
@@ -306,7 +308,7 @@ void UIKRetargetBoneDetails::OnPasteFromClipboard(
 			if (Result && ErrorPipe.NumErrors == 0)
 			{
 				Transform.SetLocation(Data);
-				EditorController->AssetController->SetTranslationOffsetOnRetargetRootBone(Transform.GetLocation(), SourceOrTarget);
+				EditorController->AssetController->GetCurrentRetargetPose(SourceOrTarget).SetRootTranslationDelta(Transform.GetLocation());
 			}
 			break;
 		}
@@ -404,7 +406,7 @@ void UIKRetargetBoneDetails::OnNumericValueCommitted(
 		{
 			const bool bIsTranslationLocal = RelativeOffsetTransformRelative[0];
 			FTransform CurrentGlobalOffset = FTransform::Identity;
-			CurrentGlobalOffset.SetTranslation(AssetController->GetTranslationOffsetOnRetargetRootBone(SourceOrTarget));
+			CurrentGlobalOffset.SetTranslation(AssetController->GetCurrentRetargetPose(SourceOrTarget).GetRootTranslationDelta());
 			
 			if (bIsTranslationLocal)
 			{	
@@ -433,7 +435,7 @@ void UIKRetargetBoneDetails::OnNumericValueCommitted(
 			// store the new transform in the retarget pose
 			FScopedTransaction Transaction(LOCTEXT("EditRootTranslation", "Edit Retarget Root Pose Translation"));
 			EditorController->AssetController->GetAsset()->Modify();
-			EditorController->AssetController->SetTranslationOffsetOnRetargetRootBone(CurrentGlobalOffset.GetTranslation(), SourceOrTarget);
+			EditorController->AssetController->GetCurrentRetargetPose(SourceOrTarget).SetRootTranslationDelta(CurrentGlobalOffset.GetTranslation());
 			
 			break;
 		}
@@ -625,14 +627,14 @@ void FIKRetargetBoneDetailCustomization::CustomizeDetails(IDetailLayoutBuilder&
 		// get/set relative
 		TransformWidgetArgs.OnGetIsComponentRelative_Lambda( [bIsEditable, BonesView, TransformType](ESlateTransformComponent::Type InComponent)
 		{
-			return BonesView.ContainsByPredicate( [&](const UIKRetargetBoneDetails* Bone)
+			return BonesView.ContainsByPredicate( [&](const TObjectPtr<UIKRetargetBoneDetails> Bone)
 			{
 				return Bone->IsComponentRelative(InComponent, TransformType);
 			} );
 		})
 		.OnIsComponentRelativeChanged_Lambda( [bIsEditable, BonesView, TransformType](ESlateTransformComponent::Type InComponent, bool bIsRelative)
 		{
-			for (UIKRetargetBoneDetails* Bone: BonesView)
+			for (const TObjectPtr<UIKRetargetBoneDetails> Bone: BonesView)
 			{
 				Bone->OnComponentRelativeChanged(InComponent, bIsRelative, TransformType);
 			}
@@ -755,4 +757,429 @@ void FIKRetargetBoneDetailCustomization::GetTransformUIData(
 	}
 }
 
+void FRetargetChainSettingsCustomization::CustomizeDetails(IDetailLayoutBuilder& DetailBuilder)
+{
+	TArray<TWeakObjectPtr<UObject>> ObjectsBeingCustomized = DetailBuilder.GetSelectedObjects();
+	ChainSettingsObjects.Reset();
+	for (const TWeakObjectPtr<UObject>& Object : ObjectsBeingCustomized)
+	{
+		if (URetargetChainSettings* ChainSettings = Cast<URetargetChainSettings>(Object.Get()))
+		{
+			ChainSettingsObjects.Add(ChainSettings);
+		}
+	}
+	
+	if (ChainSettingsObjects.IsEmpty())
+	{
+		return;
+	}
+
+	Controller = ChainSettingsObjects[0]->EditorController;
+	if (!Controller.IsValid())
+	{
+		return;
+	}
+
+	// hide the settings category and build our own
+	DetailBuilder.HideCategory("Settings");
+	DetailBuilder.HideCategory("Chain Mapping");
+
+	// determine the current setup state of the chain (is it ready for FK and/or IK?)
+	// FK: it can only run FK retarget if it's mapped to a source chain
+	// IK: can only run IK if FK is setup AND it's mapped to a goal that is connected to a solver (in the target IK rig)
+	auto CheckChainSetupForIKAndFK([this](const TObjectPtr<URetargetChainSettings> ChainSettings, bool& OutFKSetup, bool&OutIKSetup)
+	{
+		OutFKSetup = true;
+		OutIKSetup = true;
+		
+		if (ChainSettings->SourceChain == NAME_None || ChainSettings->TargetChain == NAME_None)
+		{
+			OutFKSetup = false;
+			OutIKSetup = false;
+		}
+
+		if (!Controller->AssetController->IsChainRunningIK(ChainSettings.Get()))
+		{
+			OutIKSetup = false;
+		}
+	});
+
+	// get the FK and IK status of the first selected chain
+	bool bFirstChainFKSetup;
+	bool bFirstChainIKSetup;
+	CheckChainSetupForIKAndFK(ChainSettingsObjects[0].Get(), bFirstChainFKSetup, bFirstChainIKSetup);
+
+	// check to see if there is a mix of chains with valid FK or IK setups
+	bool bIsFKSetupMultipleValues = false;
+	bool bIsIKSetupMultipleValues = false;
+	for (const TWeakObjectPtr<URetargetChainSettings> ChainSettings : ChainSettingsObjects)
+	{
+		bool bIsFKSetup;
+		bool bIsIKSetup;
+		CheckChainSetupForIKAndFK(ChainSettings.Get(), bIsFKSetup, bIsIKSetup);
+		bIsFKSetupMultipleValues = bIsFKSetup != bFirstChainFKSetup ? true : bIsFKSetupMultipleValues;
+		bIsIKSetupMultipleValues = bIsIKSetup != bFirstChainIKSetup ? true : bIsIKSetupMultipleValues;
+	}
+
+	// are we editing multiple chains?
+	const bool bEditingMultipleChains = ChainSettingsObjects.Num() > 1;
+	TObjectPtr<URetargetChainSettings> PrimaryChainSettingsObject = ChainSettingsObjects[0].Get();
+
+	// get source chain options for combobox
+	SourceChainOptions.Reset();
+	SourceChainOptions.Add(MakeShareable(new FString(TEXT("None"))));
+	if (const UIKRigDefinition* SourceIKRig = Controller->AssetController->GetIKRig(ERetargetSourceOrTarget::Source))
+	{
+		const TArray<FBoneChain>& Chains = SourceIKRig->GetRetargetChains();
+		for (const FBoneChain& BoneChain : Chains)
+		{
+			SourceChainOptions.Add(MakeShareable(new FString(BoneChain.ChainName.ToString())));
+		}
+	}
+	
+	// replace the "Settings" struct property
+	const FText ChainTitleName = bEditingMultipleChains ?  LOCTEXT("MultipleChains_Label", "Multiple Chain") : FText::FromName(PrimaryChainSettingsObject->TargetChain);
+	const FText ChainCategoryTitle = FText::Format( LOCTEXT("ChainSettingsHeaderLabel", "{0} Settings"), ChainTitleName);
+	IDetailCategoryBuilder& SettingsCategory = DetailBuilder.EditCategory("Chain Settings", ChainCategoryTitle);
+	
+	// add row to select source chain to map to
+	SettingsCategory.AddCustomRow(LOCTEXT("ChainStatus_Label", "Status"))
+	.NameContent()
+	[
+		SNew(SBox)
+		.Content()
+		[
+			SNew(STextBlock)
+			.Text(LOCTEXT( "ChainStatus_Label", "Source Chain" ))
+			.Font(IDetailLayoutBuilder::GetDetailFont())
+		]
+	]
+	.ValueContent()
+	.MinDesiredWidth(166.0f)
+	[
+		SNew(SSearchableComboBox)
+		.ToolTipText(LOCTEXT("SourceChainOptionsToolTip", "Select source chain to map to this target chain."))
+		.OptionsSource(&SourceChainOptions)
+		.OnGenerateWidget_Lambda([](TSharedPtr<FString> InItem)
+		{
+			return SNew(STextBlock).Text(FText::FromString(*InItem.Get()));
+		})
+		.OnSelectionChanged_Lambda([this](TSharedPtr<FString> InString, ESelectInfo::Type SelectInfo)
+		{
+			const FName SourceChainName = FName(*InString.Get());
+			for (TWeakObjectPtr<URetargetChainSettings> ChainMapSettings : ChainSettingsObjects)
+			{
+				Controller->AssetController->SetSourceChainForTargetChain(ChainMapSettings.Get(), SourceChainName);
+			}
+		})
+		[
+			SNew(STextBlock)
+			.Font(IPropertyTypeCustomizationUtils::GetRegularFont())
+			.Text_Lambda([PrimaryChainSettingsObject]()
+			{
+				return FText::FromName(PrimaryChainSettingsObject->SourceChain);
+			})
+		]
+	];
+
+	// FK settings
+	const FText FKNotMappedWarning = LOCTEXT( "FKDisabled_Label", "Disabled: No source chain specified." );
+	const bool bFKEditingEnabled = bFirstChainFKSetup && !bIsFKSetupMultipleValues;
+	AddSettingsSection(
+		DetailBuilder,
+		SettingsCategory,
+		GET_MEMBER_NAME_STRING_CHECKED(FTargetChainSettings, FK),
+		FName("FK"),
+		LOCTEXT("FKGroup_Label", "FK"),
+		FTargetChainFKSettings::StaticStruct(),
+		GET_MEMBER_NAME_STRING_CHECKED(FTargetChainFKSettings, EnableFK),
+		bFKEditingEnabled,
+		FKNotMappedWarning);
+
+	// IK settings
+	const FText NoIKGoalWarning = LOCTEXT( "IKDisabled_Label", "Disabled: No IK Goal specified." );
+	const FText IKEnabledMultipleValuesWarning = LOCTEXT( "IKDisabledMultipleValues_Label", "Some selected chain(s) do not have an IK Goal." );
+	const bool bIKEditingEnabled = bFirstChainIKSetup && !bIsIKSetupMultipleValues;
+	const FText IKWarningToUse = !bFKEditingEnabled ? FKNotMappedWarning : !bIKEditingEnabled ? NoIKGoalWarning : FText();
+	AddSettingsSection(
+		DetailBuilder,
+		SettingsCategory,
+		GET_MEMBER_NAME_STRING_CHECKED(FTargetChainSettings, IK),
+		FName("IK"),
+		LOCTEXT("IKGroup_Label", "IK"),
+		FTargetChainIKSettings::StaticStruct(),
+		GET_MEMBER_NAME_STRING_CHECKED(FTargetChainIKSettings, EnableIK),
+		bIKEditingEnabled,
+		IKWarningToUse);
+
+	// Plant settings
+	AddSettingsSection(
+		DetailBuilder,
+		SettingsCategory,
+		GET_MEMBER_NAME_STRING_CHECKED(FTargetChainSettings, SpeedPlanting),
+		FName("Speed Planting"),
+		LOCTEXT("PlantingGroup_Label", "Speed Planting"),
+		FTargetChainSpeedPlantSettings::StaticStruct(),
+		GET_MEMBER_NAME_STRING_CHECKED(FTargetChainSpeedPlantSettings, EnableSpeedPlanting),
+		bIKEditingEnabled,
+		IKWarningToUse);
+}
+
+void FRetargetChainSettingsCustomization::AddSettingsSection(
+	const IDetailLayoutBuilder& DetailBuilder,
+	IDetailCategoryBuilder& SettingsCategory,
+	const FString& StructPropertyName,
+	const FName& GroupName,
+	const FText& LocalizedGroupName,
+	const UScriptStruct* SettingsClass,
+	const FString& EnabledPropertyName,
+	const bool& bIsSectionEnabled,
+	const FText& DisabledMessage) const
+{
+	// create a group of all the properties in a settings struct,
+	// with a single "Enable X" property in the header of the details row
+	
+	const FString SettingsPropertyPath = GET_MEMBER_NAME_STRING_CHECKED(URetargetChainSettings, Settings);
+	const FString SettingsStructPropertyPath = FString::Printf(TEXT("%s.%s"), *SettingsPropertyPath, *StructPropertyName);
+	const FString EnablePropertyPath = FString::Printf(TEXT("%s.%s"), *SettingsStructPropertyPath, *EnabledPropertyName);
+	const TSharedRef<IPropertyHandle> IsEnabledProperty = DetailBuilder.GetProperty(*EnablePropertyPath);
+
+	IDetailGroup& Group = SettingsCategory.AddGroup(GroupName, LocalizedGroupName, false, true);
+
+	if (bIsSectionEnabled)
+	{
+		Group.HeaderProperty(IsEnabledProperty).CustomWidget()
+			.NameContent()
+			.VAlign(VAlign_Center)
+			[
+				SNew( STextBlock )
+				.Text( FText::FromName(GroupName) )
+				.Font( FAppStyle::GetFontStyle( TEXT( "PropertyWindow.NormalFont" ) ) )
+				.ToolTipText( IsEnabledProperty->GetToolTipText() )
+			]
+			.ValueContent()
+			[
+				SNew(SCheckBox)
+				.ToolTipText(IsEnabledProperty->GetToolTipText())
+				.IsChecked_Lambda([IsEnabledProperty]() -> ECheckBoxState
+				{
+					bool IsChecked;
+					IsEnabledProperty.Get().GetValue(IsChecked);
+					return IsChecked ? ECheckBoxState::Checked : ECheckBoxState::Unchecked;
+				})
+				.OnCheckStateChanged_Lambda([this, IsEnabledProperty](ECheckBoxState State)
+				{
+					IsEnabledProperty->SetValue(State == ECheckBoxState::Checked);
+					Controller->OnRetargeterNeedsInitialized(Controller->AssetController->GetAsset());
+				})
+			];
+	}
+	else
+	{
+		Group.HeaderRow()
+		.NameContent()
+		[
+			SNew(STextBlock)
+			.Font(IDetailLayoutBuilder::GetDetailFont())
+			.Text(LocalizedGroupName)
+		]
+		.ValueContent()
+		[
+			SNew(STextBlock)
+			.Font(IDetailLayoutBuilder::GetDetailFont())
+			.Text(DisabledMessage)
+		];
+	}
+	
+	const TAttribute<bool> EditCondition = TAttribute<bool>::Create([IsEnabledProperty, bIsSectionEnabled]()
+	{
+		bool bIsEnabledPropertyTrue = true;
+		IsEnabledProperty->GetValue(bIsEnabledPropertyTrue);
+		return bIsSectionEnabled && bIsEnabledPropertyTrue;
+	});
+	
+	for (TFieldIterator<FProperty> It(SettingsClass); It; ++It)
+	{
+		if (It->GetName() == EnabledPropertyName)
+		{
+			continue;
+		}
+		const FString PropertyPath = FString::Printf(TEXT("%s.%s"), *SettingsStructPropertyPath, *It->GetName());
+		Group.AddPropertyRow(DetailBuilder.GetProperty(*PropertyPath)).EditCondition(EditCondition, nullptr);
+	}
+}
+
+void FRetargetRootSettingsCustomization::CustomizeDetails(IDetailLayoutBuilder& DetailBuilder)
+{
+	TArray<TWeakObjectPtr<UObject>> ObjectsBeingCustomized = DetailBuilder.GetSelectedObjects();
+	if (ObjectsBeingCustomized.Num() < 0)
+	{
+		return;
+	}
+	
+	RootSettingsObject = Cast<URetargetRootSettings>(ObjectsBeingCustomized[0].Get());
+	if (!RootSettingsObject.IsValid())
+	{
+		return;
+	}
+
+	Controller = RootSettingsObject->EditorController;
+	if (!Controller.IsValid())
+	{
+		return;
+	}
+
+	const FString StructPropertyPath = GET_MEMBER_NAME_STRING_CHECKED(URetargetRootSettings, Settings);
+
+	auto AddPropertyRowToGroup = [&StructPropertyPath, &DetailBuilder](IDetailGroup& Group, const FString& PropertyName)
+	{
+		const FString PropertyPath = FString::Printf(TEXT("%s.%s"), *StructPropertyPath, *PropertyName);
+		Group.AddPropertyRow(DetailBuilder.GetProperty(*PropertyPath));
+	};
+
+	DetailBuilder.HideCategory("Settings");
+
+	// replace the "Settings" struct property
+	IDetailCategoryBuilder& SettingsCategory = DetailBuilder.EditCategory(FName("Root Settings"));
+
+	// alpha group
+	IDetailGroup& AlphaGroup = SettingsCategory.AddGroup("Alpha", LOCTEXT("Alpha_Label", "Alpha"), false, true);
+	AddPropertyRowToGroup(AlphaGroup,GET_MEMBER_NAME_STRING_CHECKED(FTargetRootSettings, RotationAlpha));
+	AddPropertyRowToGroup(AlphaGroup,GET_MEMBER_NAME_STRING_CHECKED(FTargetRootSettings, TranslationAlpha));
+
+	// offset group
+	IDetailGroup& OffsetGroup = SettingsCategory.AddGroup("Offsets", LOCTEXT("OffsetRoot_Label", "Offsets"), false, true);
+	AddPropertyRowToGroup(OffsetGroup,GET_MEMBER_NAME_STRING_CHECKED(FTargetRootSettings, TranslationOffset));
+	AddPropertyRowToGroup(OffsetGroup,GET_MEMBER_NAME_STRING_CHECKED(FTargetRootSettings, RotationOffset));
+
+	// scale group
+	IDetailGroup& ScaleGroup = SettingsCategory.AddGroup("Scale Translation", LOCTEXT("ScaleRoot_Label", "Scale Translation"), false, true);
+	AddPropertyRowToGroup(ScaleGroup,GET_MEMBER_NAME_STRING_CHECKED(FTargetRootSettings, ScaleHorizontal));
+	AddPropertyRowToGroup(ScaleGroup,GET_MEMBER_NAME_STRING_CHECKED(FTargetRootSettings, ScaleVertical));
+
+	// blend to source group
+	IDetailGroup& BlendToSourceGroup = SettingsCategory.AddGroup("Blend To Source", LOCTEXT("BlendToSource_Label", "Blend to Source"), false, true);
+	AddPropertyRowToGroup(BlendToSourceGroup,GET_MEMBER_NAME_STRING_CHECKED(FTargetRootSettings, BlendToSource));
+	AddPropertyRowToGroup(BlendToSourceGroup,GET_MEMBER_NAME_STRING_CHECKED(FTargetRootSettings, BlendToSourceWeights));
+
+	// affect IK directionally
+	IDetailGroup& AffectIKGroup = SettingsCategory.AddGroup("Affect IK", LOCTEXT("AffectIK_Label", "Affect IK"), false, true);
+	AddPropertyRowToGroup(AffectIKGroup,GET_MEMBER_NAME_STRING_CHECKED(FTargetRootSettings, AffectIKHorizontal));
+	AddPropertyRowToGroup(AffectIKGroup,GET_MEMBER_NAME_STRING_CHECKED(FTargetRootSettings, AffectIKVertical));
+}
+
+void FRetargetGlobalSettingsCustomization::CustomizeDetails(IDetailLayoutBuilder& DetailBuilder)
+{
+	TArray<TWeakObjectPtr<UObject>> ObjectsBeingCustomized = DetailBuilder.GetSelectedObjects();
+	if (ObjectsBeingCustomized.Num() < 0)
+	{
+		return;
+	}
+	
+	GlobalSettingsObject = Cast<UIKRetargetGlobalSettings>(ObjectsBeingCustomized[0].Get());
+	if (!GlobalSettingsObject.IsValid())
+	{
+		return;
+	}
+
+	Controller = GlobalSettingsObject->EditorController;
+	if (!Controller.IsValid())
+	{
+		return;
+	}
+
+	const FString StructPropertyPath = GET_MEMBER_NAME_STRING_CHECKED(URetargetRootSettings, Settings);
+
+	auto GetPropertyHandle = [&DetailBuilder, &StructPropertyPath](const FString& PropertyName)
+	{
+		const FString EnableWarpingPropertyName = GET_MEMBER_NAME_STRING_CHECKED(FRetargetGlobalSettings, bIKWarping);
+		const FString PropertyPath = FString::Printf(TEXT("%s.%s"), *StructPropertyPath, *PropertyName);
+		return DetailBuilder.GetProperty(*PropertyPath);
+	};
+
+	DetailBuilder.HideCategory("Settings");
+
+	// replace the "Settings" struct property
+	IDetailCategoryBuilder& SettingsCategory = DetailBuilder.EditCategory(FName("Global Settings"));
+
+	// phases group
+	IDetailGroup& AlphaGroup = SettingsCategory.AddGroup("Phases", LOCTEXT("Phases_Label", "Global Retarget Phases"), false, true);
+	AlphaGroup.AddPropertyRow(GetPropertyHandle(GET_MEMBER_NAME_STRING_CHECKED(FRetargetGlobalSettings, bEnableRoot)));
+	AlphaGroup.AddPropertyRow(GetPropertyHandle(GET_MEMBER_NAME_STRING_CHECKED(FRetargetGlobalSettings, bEnableFK)));
+	AlphaGroup.AddPropertyRow(GetPropertyHandle(GET_MEMBER_NAME_STRING_CHECKED(FRetargetGlobalSettings, bEnableIK)));
+
+	// stride warping group
+	IDetailGroup& WarpingGroup = SettingsCategory.AddGroup("Warping", LOCTEXT("Warping_Label", "Stride Warping"), false, true);
+
+	// header with "enable" checkbox that disables whole group
+	const TSharedRef<IPropertyHandle> EnableWarpingPropertyHandle = GetPropertyHandle(GET_MEMBER_NAME_STRING_CHECKED(FRetargetGlobalSettings, bIKWarping));
+	WarpingGroup.HeaderProperty(EnableWarpingPropertyHandle);
+	
+	const TAttribute<bool> EditCondition = TAttribute<bool>::Create([EnableWarpingPropertyHandle]()
+	{
+		bool bIsEnabledPropertyTrue = true;
+		EnableWarpingPropertyHandle->GetValue(bIsEnabledPropertyTrue);
+		return bIsEnabledPropertyTrue;
+	});
+	
+	// add all the warping parameters
+	WarpingGroup.AddPropertyRow(GetPropertyHandle(GET_MEMBER_NAME_STRING_CHECKED(FRetargetGlobalSettings, ForwardDirection))).EditCondition(EditCondition,nullptr);
+	WarpingGroup.AddPropertyRow(GetPropertyHandle(GET_MEMBER_NAME_STRING_CHECKED(FRetargetGlobalSettings, DirectionSource))).EditCondition(EditCondition,nullptr);
+
+	// get target chain options for combobox
+	TargetChainOptions.Reset();
+	TargetChainOptions.Add(MakeShareable(new FString(TEXT("None"))));
+	if (const UIKRigDefinition* TargetIKRig = Controller->AssetController->GetIKRig(ERetargetSourceOrTarget::Target))
+	{
+		const TArray<FBoneChain>& Chains = TargetIKRig->GetRetargetChains();
+		for (const FBoneChain& BoneChain : Chains)
+		{
+			TargetChainOptions.Add(MakeShareable(new FString(BoneChain.ChainName.ToString())));
+		}
+	}
+	
+	// add row to target chain to use for body direction
+	const TSharedRef<IPropertyHandle> DirectionChainPropertyHandle = GetPropertyHandle(GET_MEMBER_NAME_STRING_CHECKED(FRetargetGlobalSettings, DirectionChain));
+	WarpingGroup.AddPropertyRow(DirectionChainPropertyHandle).CustomWidget()
+	.IsEnabled(GlobalSettingsObject->Settings.DirectionSource == EWarpingDirectionSource::Chain)
+	.NameContent()
+	[
+		SNew(SBox)
+		.Content()
+		[
+			SNew(STextBlock)
+			.Text(LOCTEXT( "DirectionChainName_Label", "Direction Chain" ))
+			.Font(IDetailLayoutBuilder::GetDetailFont())
+		]
+	]
+	.ValueContent()
+	.MinDesiredWidth(166.0f)
+	[
+		SNew(SSearchableComboBox)
+		.ToolTipText(LOCTEXT("DirectionChainOptionsToolTip", "Select target chain (usually the Spine) to help define the forward direction of the character. This is especially useful for creatures with a horizontal body, like quadrupeds."))
+		.OptionsSource(&TargetChainOptions)
+		.OnGenerateWidget_Lambda([](TSharedPtr<FString> InItem)
+		{
+			return SNew(STextBlock).Text(FText::FromString(*InItem.Get()));
+		})
+		.OnSelectionChanged_Lambda([this](TSharedPtr<FString> InString, ESelectInfo::Type SelectInfo)
+		{
+			const FName TargetChainName = FName(*InString.Get());
+			GlobalSettingsObject->Settings.DirectionChain = TargetChainName;
+		})
+		[
+			SNew(STextBlock)
+			.Font(IPropertyTypeCustomizationUtils::GetRegularFont())
+			.Text_Lambda([this]()
+			{
+				return FText::FromName(GlobalSettingsObject->Settings.DirectionChain);
+			})
+		]
+	];
+	
+	WarpingGroup.AddPropertyRow(GetPropertyHandle(GET_MEMBER_NAME_STRING_CHECKED(FRetargetGlobalSettings, WarpForwards))).EditCondition(EditCondition,nullptr);
+	WarpingGroup.AddPropertyRow(GetPropertyHandle(GET_MEMBER_NAME_STRING_CHECKED(FRetargetGlobalSettings, WarpSplay))).EditCondition(EditCondition,nullptr);
+	WarpingGroup.AddPropertyRow(GetPropertyHandle(GET_MEMBER_NAME_STRING_CHECKED(FRetargetGlobalSettings, SidewaysOffset))).EditCondition(EditCondition,nullptr);
+}
+
 #undef LOCTEXT_NAMESPACE
diff --git a/Engine/Plugins/Animation/IKRig/Source/IKRigEditor/Private/RetargetEditor/IKRetargetEditPoseMode.cpp b/Engine/Plugins/Animation/IKRig/Source/IKRigEditor/Private/RetargetEditor/IKRetargetEditPoseMode.cpp
--- a/Engine/Plugins/Animation/IKRig/Source/IKRigEditor/Private/RetargetEditor/IKRetargetEditPoseMode.cpp
+++ b/Engine/Plugins/Animation/IKRig/Source/IKRigEditor/Private/RetargetEditor/IKRetargetEditPoseMode.cpp
@@ -55,7 +55,6 @@ void FIKRetargetEditPoseMode::Render(const FSceneView* View, FViewport* Viewport
 		return;
 	}
 
-	// render the skeleton
 	RenderSkeleton(PDI, Controller);
 }
 
@@ -90,7 +89,7 @@ void FIKRetargetEditPoseMode::RenderSkeleton(
 	DrawConfig.BoneDrawMode = (EBoneDrawMode::Type)ConfigOption->DefaultBoneDrawSelection;
 	DrawConfig.BoneDrawSize = BoneRadius;
 	DrawConfig.bAddHitProxy = true;
-	DrawConfig.bForceDraw = true;
+	DrawConfig.bForceDraw = false;
 	DrawConfig.DefaultBoneColor = GetMutableDefault<UPersonaOptions>()->DefaultBoneColor;
 	DrawConfig.AffectedBoneColor = GetMutableDefault<UPersonaOptions>()->AffectedBoneColor;
 	DrawConfig.SelectedBoneColor = GetMutableDefault<UPersonaOptions>()->SelectedBoneColor;
@@ -125,31 +124,6 @@ void FIKRetargetEditPoseMode::RenderSkeleton(
 	);
 }
 
-void FIKRetargetEditPoseMode::GetSelectedAndAffectedBones(
-	const FIKRetargetEditorController* Controller,
-	const FRetargetSkeleton& Skeleton,
-	TSet<int32>& OutAffectedBones,
-	TSet<int32>& OutSelectedBones) const
-{
-	// record selected bone indices
-	const TArray<FName>& SelectedBones = Controller->GetSelectedBones();
-	for (const FName& SelectedBone : SelectedBones)
-	{
-		OutSelectedBones.Add(Skeleton.FindBoneIndexByName(SelectedBone));
-	}
-
-	// "affected bones" are the selected bones AND their children, recursively
-	for (int32 SelectedBone : OutSelectedBones)
-	{
-		const int32 EndOfBranch = Skeleton.GetCachedEndOfBranchIndex(SelectedBone);
-		OutAffectedBones.Add(SelectedBone);
-		for (int32 BoneIndex=SelectedBone; BoneIndex<=EndOfBranch; ++BoneIndex)
-		{
-			OutAffectedBones.Add(BoneIndex);
-		}
-	}
-}
-
 bool FIKRetargetEditPoseMode::AllowWidgetMove()
 {
 	return false;
@@ -248,7 +222,7 @@ bool FIKRetargetEditPoseMode::HandleClick(FEditorViewportClient* InViewportClien
 		HIKRetargetEditorBoneProxy* BoneProxy = static_cast<HIKRetargetEditorBoneProxy*>(HitProxy);
 		const TArray<FName> BoneNames{BoneProxy->BoneName};
 		constexpr bool bFromHierarchy = false;
-		const EBoneSelectionEdit EditMode = Click.IsControlDown() || Click.IsShiftDown() ? EBoneSelectionEdit::Add : EBoneSelectionEdit::Replace;
+		const ESelectionEdit EditMode = Click.IsControlDown() || Click.IsShiftDown() ? ESelectionEdit::Add : ESelectionEdit::Replace;
 		Controller->EditBoneSelection(BoneNames, EditMode, bFromHierarchy);
 		return true;
 	}
@@ -300,6 +274,7 @@ bool FIKRetargetEditPoseMode::StartTracking(FEditorViewportClient* InViewportCli
 		Controller->AssetController->GetAsset()->Modify();
 		TrackingState = FIKRetargetTrackingState::TranslatingRoot;
 		UpdateWidgetTransform();
+		BoneEdit.AccumulatedPositionOffset = Controller->AssetController->GetCurrentRetargetPose(SourceOrTarget).GetRootTranslationDelta();
 		return true;
 	}
 
@@ -343,6 +318,12 @@ bool FIKRetargetEditPoseMode::InputDelta(
 	{
 		return false; 
 	}
+
+	const TObjectPtr<UIKRetargeterController> AssetController = Controller->AssetController;
+	if (AssetController.IsNull())
+	{
+		return false; 
+	}
 	
 	// rotating any bone
 	if (TrackingState == FIKRetargetTrackingState::RotatingBone)
@@ -368,7 +349,7 @@ bool FIKRetargetEditPoseMode::InputDelta(
 
 			// apply the new delta to the retarget pose
 			const FQuat TotalDeltaRotation = BoneEdit.PreviousDeltaRotation[SelectionIndex] * BoneLocalDelta;
-			Controller->AssetController->SetRotationOffsetForRetargetPoseBone(BoneName, TotalDeltaRotation, Controller->GetSourceOrTarget());
+			AssetController->SetRotationOffsetForRetargetPoseBone(BoneName, TotalDeltaRotation, Controller->GetSourceOrTarget());
 		}
 
 		return true;
@@ -382,8 +363,15 @@ bool FIKRetargetEditPoseMode::InputDelta(
 			return false;
 		}
 
+		// must scale drag vector by inverse of component scale or else it will create feedback cycle and go to outer space
+		float Scale;
+		FVector Offset;
+		GetEditedComponentScaleAndOffset(Scale,Offset);
+		Scale = FMath::Max(Scale, KINDA_SMALL_NUMBER);
+		BoneEdit.AccumulatedPositionOffset += InDrag * (1.0f / Scale);
+		
 		// apply translation delta to root
-		Controller->AssetController->AddTranslationOffsetToRetargetRootBone(InDrag, Controller->GetSourceOrTarget());
+		AssetController->GetCurrentRetargetPose(SourceOrTarget).SetRootTranslationDelta(BoneEdit.AccumulatedPositionOffset);
 		return true;
 	}
 	
diff --git a/Engine/Plugins/Animation/IKRig/Source/IKRigEditor/Private/RetargetEditor/IKRetargetEditor.cpp b/Engine/Plugins/Animation/IKRig/Source/IKRigEditor/Private/RetargetEditor/IKRetargetEditor.cpp
--- a/Engine/Plugins/Animation/IKRig/Source/IKRigEditor/Private/RetargetEditor/IKRetargetEditor.cpp
+++ b/Engine/Plugins/Animation/IKRig/Source/IKRigEditor/Private/RetargetEditor/IKRetargetEditor.cpp
@@ -91,6 +91,12 @@ void FIKRetargetEditor::InitAssetEditor(
 	RegenerateMenusAndToolbars();
 }
 
+void FIKRetargetEditor::OnClose()
+{
+	FPersonaAssetEditorToolkit::OnClose();
+	EditorController->Close();
+}
+
 void FIKRetargetEditor::RegisterTabSpawners(const TSharedRef<FTabManager>& InTabManager)
 {
 	WorkspaceMenuCategory = InTabManager->AddLocalWorkspaceMenuCategory(LOCTEXT("WorkspaceMenu_IKRigEditor", "IK Rig Editor"));
diff --git a/Engine/Plugins/Animation/IKRig/Source/IKRigEditor/Private/RetargetEditor/IKRetargetEditorController.cpp b/Engine/Plugins/Animation/IKRig/Source/IKRigEditor/Private/RetargetEditor/IKRetargetEditorController.cpp
--- a/Engine/Plugins/Animation/IKRig/Source/IKRigEditor/Private/RetargetEditor/IKRetargetEditorController.cpp
+++ b/Engine/Plugins/Animation/IKRig/Source/IKRigEditor/Private/RetargetEditor/IKRetargetEditorController.cpp
@@ -4,15 +4,11 @@
 
 #include "AnimPose.h"
 #include "AssetToolsModule.h"
-#include "ContentBrowserModule.h"
 #include "EditorModeManager.h"
 #include "IContentBrowserSingleton.h"
 #include "Widgets/Input/SButton.h"
 #include "Animation/DebugSkelMeshComponent.h"
-#include "Animation/AnimMontage.h"
-#include "Animation/AnimSequence.h"
 #include "Animation/PoseAsset.h"
-#include "Factories/PoseAssetFactory.h"
 #include "RetargetEditor/IKRetargetAnimInstance.h"
 #include "RetargetEditor/IKRetargetDefaultMode.h"
 #include "RetargetEditor/IKRetargetEditPoseMode.h"
@@ -47,22 +43,36 @@ void FIKRetargetEditorController::Initialize(TSharedPtr<FIKRetargetEditor> InEdi
 	BindToIKRigAsset(AssetController->GetAsset()->GetTargetIKRigWriteable());
 
 	// bind callback when retargeter needs reinitialized
-	AssetController->OnRetargeterNeedsInitialized().AddSP(this, &FIKRetargetEditorController::OnRetargeterNeedsInitialized);
+	ReInitDelegateHandle = AssetController->OnRetargeterNeedsInitialized().AddSP(this, &FIKRetargetEditorController::OnRetargeterNeedsInitialized);
 }
 
-void FIKRetargetEditorController::BindToIKRigAsset(UIKRigDefinition* InIKRig) const
+void FIKRetargetEditorController::Close() const
 {
+	AssetController->OnRetargeterNeedsInitialized().Remove(ReInitDelegateHandle);
+
+	if (UIKRigController* IKRigController = UIKRigController::GetIKRigController(BoundToIKRig))
+	{
+		IKRigController->OnIKRigNeedsInitialized().Remove(ReInitDelegateHandle);
+		IKRigController->OnRetargetChainRenamed().Remove(RenameChainDelegateHandle);
+		IKRigController->OnRetargetChainRemoved().Remove(RemoveChainDelegateHandle);
+	}
+}
+
+void FIKRetargetEditorController::BindToIKRigAsset(UIKRigDefinition* InIKRig)
+{
+	BoundToIKRig = InIKRig;
+	
 	if (!InIKRig)
 	{
 		return;
 	}
-
+	
 	UIKRigController* Controller = UIKRigController::GetIKRigController(InIKRig);
 	if (!Controller->OnIKRigNeedsInitialized().IsBoundToObject(this))
 	{
-		Controller->OnIKRigNeedsInitialized().AddSP(this, &FIKRetargetEditorController::OnIKRigNeedsInitialized);
-		Controller->OnRetargetChainRenamed().AddSP(this, &FIKRetargetEditorController::OnRetargetChainRenamed);
-		Controller->OnRetargetChainRemoved().AddSP(this, &FIKRetargetEditorController::OnRetargetChainRemoved);
+		ReInitIKDelegateHandle = Controller->OnIKRigNeedsInitialized().AddSP(this, &FIKRetargetEditorController::OnIKRigNeedsInitialized);
+		RenameChainDelegateHandle = Controller->OnRetargetChainRenamed().AddSP(this, &FIKRetargetEditorController::OnRetargetChainRenamed);
+		RemoveChainDelegateHandle = Controller->OnRetargetChainRemoved().AddSP(this, &FIKRetargetEditorController::OnRetargetChainRemoved);
 	}
 }
 
@@ -195,18 +205,13 @@ FName FIKRetargetEditorController::GetChainNameFromBone(const FName& BoneName, E
 	return Processor->GetChainNameForBone(BoneIndex, (int8)SourceOrTarget);
 }
 
-TObjectPtr<UIKRetargetBoneDetails> FIKRetargetEditorController::GetDetailsObjectForBone(const FName& BoneName)
+TObjectPtr<UIKRetargetBoneDetails> FIKRetargetEditorController::GetOrCreateBoneDetailsObject(const FName& BoneName)
 {
 	if (AllBoneDetails.Contains(BoneName))
 	{
 		return AllBoneDetails[BoneName];
 	}
 
-	return CreateBoneDetails(BoneName);
-}
-
-TObjectPtr<UIKRetargetBoneDetails> FIKRetargetEditorController::CreateBoneDetails(const FName& BoneName)
-{
 	// create and store a new one
 	UIKRetargetBoneDetails* NewBoneDetails = NewObject<UIKRetargetBoneDetails>(AssetController->GetAsset(), FName(BoneName), RF_Standalone | RF_Transient );
 	NewBoneDetails->SelectedBone = BoneName;
@@ -268,15 +273,6 @@ FTransform FIKRetargetEditorController::GetGlobalRetargetPoseOfBone(
 	return BoneTransform;
 }
 
-FTransform FIKRetargetEditorController::GetTargetBoneLocalTransform(
-	const UIKRetargetProcessor* RetargetProcessor,
-	const int32& TargetBoneIndex) const
-{
-	check(RetargetProcessor && RetargetProcessor->IsInitialized())
-
-	return RetargetProcessor->GetTargetBoneRetargetPoseLocalTransform(TargetBoneIndex);
-}
-
 void FIKRetargetEditorController::GetGlobalRetargetPoseOfImmediateChildren(
 	const FRetargetSkeleton& RetargetSkeleton,
 	const int32& BoneIndex,
@@ -418,6 +414,12 @@ void FIKRetargetEditorController::PlayAnimationAsset(UAnimationAsset* AssetToPla
 	}
 }
 
+void FIKRetargetEditorController::StopPlayback()
+{
+	SourceAnimInstance->SetPlaying(false);
+	SourceAnimInstance->SetAnimationAsset(nullptr);
+}
+
 void FIKRetargetEditorController::PausePlayback()
 {
 	if (UAnimationAsset* CurrentAnim = SourceAnimInstance->GetAnimationAsset())
@@ -514,7 +516,7 @@ UPrimitiveComponent* FIKRetargetEditorController::GetSelectedMesh()
 
 void FIKRetargetEditorController::EditBoneSelection(
 	const TArray<FName>& InBoneNames,
-	EBoneSelectionEdit EditMode,
+	ESelectionEdit EditMode,
 	const bool bFromHierarchyView)
 {
 	// must have a skeletal mesh
@@ -523,10 +525,14 @@ void FIKRetargetEditorController::EditBoneSelection(
 	{
 		return;
 	}
+
+	// deselect mesh
+	SetSelectedMesh(nullptr);
+	SetRootSelected(false);
 	
 	switch (EditMode)
 	{
-		case EBoneSelectionEdit::Add:
+		case ESelectionEdit::Add:
 		{
 			for (const FName& BoneName : InBoneNames)
 			{
@@ -535,7 +541,7 @@ void FIKRetargetEditorController::EditBoneSelection(
 			
 			break;
 		}
-		case EBoneSelectionEdit::Remove:
+		case ESelectionEdit::Remove:
 		{
 			for (const FName& BoneName : InBoneNames)
 			{
@@ -543,7 +549,7 @@ void FIKRetargetEditorController::EditBoneSelection(
 			}
 			break;
 		}
-		case EBoneSelectionEdit::Replace:
+		case ESelectionEdit::Replace:
 		{
 			SelectedBones = InBoneNames;
 			break;
@@ -567,9 +573,6 @@ void FIKRetargetEditorController::EditBoneSelection(
 			break;
 		}
 	}
-
-	// deselect mesh
-	SetSelectedMesh(nullptr);
 	
 	// apply selection to debug mesh component so rendering knows
 	DebugComponent->BonesOfInterest = SelectedBoneIndices;
@@ -587,25 +590,110 @@ void FIKRetargetEditorController::EditBoneSelection(
 	}
 	else
 	{
-		SelectedBoneDetails.Reset();
+		TArray<UObject*> SelectedBoneDetails;
 		for (const FName& SelectedBone : SelectedBones)
 		{
-			TObjectPtr<UIKRetargetBoneDetails> BoneDetails = GetDetailsObjectForBone(SelectedBone);
+			TObjectPtr<UIKRetargetBoneDetails> BoneDetails = GetOrCreateBoneDetailsObject(SelectedBone);
 			SelectedBoneDetails.Add(BoneDetails);
 		}
 		SetDetailsObjects(SelectedBoneDetails);
 	}
 }
 
+void FIKRetargetEditorController::EditChainSelection(
+	const TArray<FName>& InChainNames,
+	ESelectionEdit EditMode,
+	const bool bFromChainsView)
+{
+	// deselect others
+	SetSelectedMesh(nullptr);
+	SetRootSelected(false);
+	
+	// update selection set based on edit mode
+	switch (EditMode)
+	{
+	case ESelectionEdit::Add:
+		{
+			for (const FName& ChainName : InChainNames)
+			{
+				SelectedChains.AddUnique(ChainName);
+			}
+			
+			break;
+		}
+	case ESelectionEdit::Remove:
+		{
+			for (const FName& ChainName : InChainNames)
+			{
+				SelectedChains.Remove(ChainName);
+			}
+			break;
+		}
+	case ESelectionEdit::Replace:
+		{
+			SelectedChains = InChainNames;
+			break;
+		}
+	default:
+		checkNoEntry();
+	}
+
+	// update chains view with selected chains
+	if (ChainsView)
+	{
+		if (!bFromChainsView)
+		{
+			ChainsView->RefreshView();
+		}
+	}
+
+	// get selected chain UObjects to show in details view
+	TArray<UObject*> SelectedChainSettings;
+	const TArray<TObjectPtr<URetargetChainSettings>>& AllChainSettings = AssetController->GetAsset()->GetAllChainSettings();
+	for (const TObjectPtr<URetargetChainSettings>& ChainSettings : AllChainSettings)
+	{
+		if (ChainSettings->EditorController.Get() != this)
+		{
+			ChainSettings->EditorController = SharedThis(this);	
+		}
+		
+		if (SelectedChains.Contains(ChainSettings.Get()->TargetChain))
+		{
+			SelectedChainSettings.Add(ChainSettings);
+		}
+	}
+
+	// no chains selected, then show asset settings in the details view
+	SelectedChainSettings.IsEmpty() ? SetDetailsObject(AssetController->GetAsset()) : SetDetailsObjects(SelectedChainSettings);
+}
+
+void FIKRetargetEditorController::SetRootSelected(const bool bIsSelected)
+{
+	bIsRootSelected = bIsSelected;
+
+	if (bIsSelected)
+	{
+		URetargetRootSettings* RootSettings = AssetController->GetAsset()->GetRootSettingsUObject();
+		if (RootSettings->EditorController.Get() != this)
+		{
+			RootSettings->EditorController = SharedThis(this);	
+		}
+		
+		SetDetailsObject(RootSettings);
+	}
+}
+
 void FIKRetargetEditorController::ClearSelection(const bool bKeepBoneSelection)
 {
-	// clear mesh selection
+	// clear root and mesh selection
+	SetRootSelected(false);
 	SetSelectedMesh(nullptr);
 	
 	// deselect all chains
 	if (ChainsView.IsValid())
 	{
 		ChainsView->ClearSelection();
+		SelectedChains.Empty();
 	}
 
 	// clear bone selection
@@ -613,7 +701,7 @@ void FIKRetargetEditorController::ClearSelection(const bool bKeepBoneSelection)
 	{
 		const TArray<FName> Empty;
 		constexpr bool bFromHierarchy = false;
-		EditBoneSelection(Empty, EBoneSelectionEdit::Replace, bFromHierarchy);
+		EditBoneSelection(Empty, ESelectionEdit::Replace, bFromHierarchy);
 
 		// show global details
 		SetDetailsObject(AssetController->GetAsset());
@@ -652,7 +740,7 @@ void FIKRetargetEditorController::SetRetargeterMode(ERetargeterOutputMode Mode)
 			OutputMode = ERetargeterOutputMode::RunRetarget;
 			SourceAnimInstance->SetRetargetMode(ERetargeterOutputMode::RunRetarget);
 			TargetAnimInstance->SetRetargetMode(ERetargeterOutputMode::RunRetarget);
-			// must reinitialize after editing the retarget pose
+			// force a reinitialization in case retarget pose was edited
 			AssetController->BroadcastNeedsReinitialized();
 			ResumePlayback();
 			break;
diff --git a/Engine/Plugins/Animation/IKRig/Source/IKRigEditor/Private/RetargetEditor/IKRetargetHitProxies.cpp b/Engine/Plugins/Animation/IKRig/Source/IKRigEditor/Private/RetargetEditor/IKRetargetHitProxies.cpp
--- a/Engine/Plugins/Animation/IKRig/Source/IKRigEditor/Private/RetargetEditor/IKRetargetHitProxies.cpp
+++ b/Engine/Plugins/Animation/IKRig/Source/IKRigEditor/Private/RetargetEditor/IKRetargetHitProxies.cpp
@@ -2,4 +2,6 @@
 
 #include "RetargetEditor/IKRetargetHitProxies.h"
 
-IMPLEMENT_HIT_PROXY(HIKRetargetEditorBoneProxy, HHitProxy);
\ No newline at end of file
+IMPLEMENT_HIT_PROXY(HIKRetargetEditorBoneProxy, HHitProxy);
+IMPLEMENT_HIT_PROXY(HIKRetargetEditorChainProxy, HHitProxy);
+IMPLEMENT_HIT_PROXY(HIKRetargetEditorRootProxy, HHitProxy);
diff --git a/Engine/Plugins/Animation/IKRig/Source/IKRigEditor/Private/RetargetEditor/IKRetargeterController.cpp b/Engine/Plugins/Animation/IKRig/Source/IKRigEditor/Private/RetargetEditor/IKRetargeterController.cpp
--- a/Engine/Plugins/Animation/IKRig/Source/IKRigEditor/Private/RetargetEditor/IKRetargeterController.cpp
+++ b/Engine/Plugins/Animation/IKRig/Source/IKRigEditor/Private/RetargetEditor/IKRetargeterController.cpp
@@ -8,6 +8,7 @@
 #include "Engine/SkeletalMesh.h"
 #include "RetargetEditor/IKRetargetEditorController.h"
 #include "Retargeter/IKRetargeter.h"
+#include "RigEditor/IKRigController.h"
 
 #define LOCTEXT_NAMESPACE "IKRetargeterController"
 
@@ -65,6 +66,11 @@ const UIKRigDefinition* UIKRetargeterController::GetIKRig(const ERetargetSourceO
 	return SourceOrTarget == ERetargetSourceOrTarget::Source ? Asset->GetSourceIKRig() : Asset->GetTargetIKRig();
 }
 
+UIKRigDefinition* UIKRetargeterController::GetIKRigWriteable(const ERetargetSourceOrTarget& SourceOrTarget) const
+{
+	return SourceOrTarget == ERetargetSourceOrTarget::Source ? Asset->GetSourceIKRigWriteable() : Asset->GetTargetIKRigWriteable();
+}
+
 void UIKRetargeterController::OnIKRigChanged(const ERetargetSourceOrTarget& SourceOrTarget) const
 {
 	if (const UIKRigDefinition* IKRig = GetIKRig(SourceOrTarget))
@@ -318,6 +324,24 @@ const TArray<TObjectPtr<URetargetChainSettings>>& UIKRetargeterController::GetCh
 	return Asset->ChainSettings;
 }
 
+bool UIKRetargeterController::IsChainRunningIK(const TObjectPtr<URetargetChainSettings> ChainSettings) const
+{
+	if (ChainSettings.IsNull())
+	{
+		return false;
+	}
+	
+	UIKRigDefinition* TargetIKRig = GetIKRigWriteable(ERetargetSourceOrTarget::Target);
+	if (!TargetIKRig)
+	{
+		return false;
+	}
+	
+	const UIKRigController* IKRigController = UIKRigController::GetIKRigController(TargetIKRig);
+	const FName ChainGoalName = IKRigController->GetRetargetChainGoal(ChainSettings->TargetChain);
+	return IKRigController->IsGoalConnectedToAnySolver(ChainGoalName);
+}
+
 void UIKRetargeterController::AddRetargetPose(
 	const FName& NewPoseName,
 	const FIKRetargetPose* ToDuplicate,
@@ -494,18 +518,6 @@ FQuat UIKRetargeterController::GetRotationOffsetForRetargetPoseBone(
 	return BoneOffsets[BoneName];
 }
 
-void UIKRetargeterController::SetTranslationOffsetOnRetargetRootBone(
-	const FVector& TranslationOffset,
-	const ERetargetSourceOrTarget& SourceOrTarget) const
-{
-	GetCurrentRetargetPose(SourceOrTarget).SetRootTranslationDelta(TranslationOffset);
-}
-
-const FVector& UIKRetargeterController::GetTranslationOffsetOnRetargetRootBone(const ERetargetSourceOrTarget& SourceOrTarget) const
-{
-	return GetCurrentRetargetPose(SourceOrTarget).RootTranslationOffset;
-}
-
 void UIKRetargeterController::AddTranslationOffsetToRetargetRootBone(
 	const FVector& TranslationOffset,
 	const ERetargetSourceOrTarget& SourceOrTarget) const
