Added profiles to the IK retargeter

diff --git a/Engine/Plugins/Animation/IKRig/Source/IKRigEditor/Private/RetargetEditor/SIKRetargetChainMapList.cpp b/Engine/Plugins/Animation/IKRig/Source/IKRigEditor/Private/RetargetEditor/SIKRetargetChainMapList.cpp
--- a/Engine/Plugins/Animation/IKRig/Source/IKRigEditor/Private/RetargetEditor/SIKRetargetChainMapList.cpp
+++ b/Engine/Plugins/Animation/IKRig/Source/IKRigEditor/Private/RetargetEditor/SIKRetargetChainMapList.cpp
@@ -14,11 +14,15 @@
 #include "RigEditor/IKRigEditorStyle.h"
 #include "Widgets/Input/SComboBox.h"
 #include "SSearchableComboBox.h"
+#include "Misc/TextFilterExpressionEvaluator.h"
+#include "Widgets/Input/SSearchBox.h"
 
 #define LOCTEXT_NAMESPACE "SIKRigRetargetChains"
 
 static const FName ColumnId_TargetChainLabel( "Target Bone Chain" );
 static const FName ColumnId_SourceChainLabel( "Source Bone Chain" );
+static const FName ColumnId_IKGoalNameLabel( "Target IK Goal" );
+static const FName ColumnId_ResetLabel( "Reset" );
 
 TSharedRef<ITableRow> FRetargetChainMapElement::MakeListRowWidget(
 	const TSharedRef<STableViewBase>& InOwnerTable,
@@ -41,7 +45,7 @@ void SIKRetargetChainMapRow::Construct(
 	// NOTE: cannot just use FName because "None" is considered a null entry and removed from ComboBox.
 	SourceChainOptions.Reset();
 	SourceChainOptions.Add(MakeShareable(new FString(TEXT("None"))));
-	const UIKRigDefinition* SourceIKRig = ChainMapList.Pin()->EditorController.Pin()->AssetController->GetAsset()->GetSourceIKRig();
+	const UIKRigDefinition* SourceIKRig = ChainMapList.Pin()->EditorController.Pin()->AssetController->GetIKRig(ERetargetSourceOrTarget::Source);
 	if (SourceIKRig)
 	{
 		const TArray<FBoneChain>& Chains = SourceIKRig->GetRetargetChains();
@@ -54,9 +58,51 @@ void SIKRetargetChainMapRow::Construct(
 	SMultiColumnTableRow< FRetargetChainMapElementPtr >::Construct( FSuperRowType::FArguments(), InOwnerTableView );
 }
 
-TSharedRef< SWidget > SIKRetargetChainMapRow::GenerateWidgetForColumn( const FName& ColumnName )
+TSharedRef< SWidget > SIKRetargetChainMapRow::GenerateWidgetForColumn(const FName& ColumnName)
 {
 	if (ColumnName == ColumnId_TargetChainLabel)
+	{
+		TSharedRef<SWidget> NewWidget =
+			SNew(SHorizontalBox)
+			+ SHorizontalBox::Slot()
+			.AutoWidth()
+			.HAlign(HAlign_Left)
+			.VAlign(VAlign_Center)
+			.Padding(3.0f, 1.0f)
+			[
+				SNew(STextBlock)
+				.Text(FText::FromName(ChainMapElement.Pin()->ChainMap->TargetChain))
+				.Font(FAppStyle::GetFontStyle(TEXT("BoldFont")))
+			];
+		return NewWidget;
+	}
+
+	if (ColumnName == ColumnId_SourceChainLabel)
+	{
+		TSharedRef<SWidget> NewWidget =
+		SNew(SHorizontalBox)
+		+ SHorizontalBox::Slot()
+		.AutoWidth()
+		.HAlign(HAlign_Left)
+		.VAlign(VAlign_Center)
+		.Padding(3.0f, 1.0f)
+		[
+			SNew(SSearchableComboBox)
+			.OptionsSource(&SourceChainOptions)
+			.OnGenerateWidget_Lambda([](TSharedPtr<FString> InItem)
+			{
+				return SNew(STextBlock).Text(FText::FromString(*InItem.Get()));
+			})
+			.OnSelectionChanged(this, &SIKRetargetChainMapRow::OnSourceChainComboSelectionChanged)
+			[
+				SNew(STextBlock)
+				.Text(this, &SIKRetargetChainMapRow::GetSourceChainName)
+			]
+		];
+		return NewWidget;
+	}
+
+	if (ColumnName == ColumnId_IKGoalNameLabel)
 	{
 		TSharedRef<SWidget> NewWidget =
 		SNew(SHorizontalBox)
@@ -67,12 +113,13 @@ TSharedRef< SWidget > SIKRetargetChainMapRow::GenerateWidgetForColumn( const FNa
 		.Padding(3.0f, 1.0f)
 		[
 			SNew(STextBlock)
-			.Text(FText::FromName(ChainMapElement.Pin()->ChainMap->TargetChain))
+			.Text(GetTargetIKGoalName())
 			.Font(FAppStyle::GetFontStyle(TEXT("BoldFont")))
 		];
 		return NewWidget;
 	}
-	else
+
+	if (ColumnName == ColumnId_ResetLabel)
 	{
 		TSharedRef<SWidget> NewWidget =
 		SNew(SHorizontalBox)
@@ -82,20 +129,22 @@ TSharedRef< SWidget > SIKRetargetChainMapRow::GenerateWidgetForColumn( const FNa
 		.VAlign(VAlign_Center)
 		.Padding(3.0f, 1.0f)
 		[
-			SNew(SSearchableComboBox)
-			.OptionsSource(&SourceChainOptions)
-			.OnGenerateWidget_Lambda([](TSharedPtr<FString> InItem)
-			{
-				return SNew(STextBlock).Text(FText::FromString(*InItem.Get()));
-			})
-			.OnSelectionChanged(this, &SIKRetargetChainMapRow::OnSourceChainComboSelectionChanged)
+			SNew(SButton)
+			//.OnClicked(this, &SIKRetargetChainMapRow::OnResetToDefaultClicked) TODO 
+			//.Visibility(this, &SIKRetargetChainMapRow::GetResetToDefaultVisibility) TODO 
+			.ToolTipText(LOCTEXT("ResetToDefaultToolTip", "Reset to Default"))
+			.ButtonStyle(FAppStyle::Get(), "NoBorder")
+			.Content()
 			[
-				SNew(STextBlock)
-				.Text(this, &SIKRetargetChainMapRow::GetSourceChainName)
+				SNew(SImage)
+				.Image(FAppStyle::GetBrush("PropertyWindow.DiffersFromDefault"))
 			]
 		];
 		return NewWidget;
 	}
+	
+	checkNoEntry();
+	return SNullWidget::NullWidget;
 }
 
 void SIKRetargetChainMapRow::OnSourceChainComboSelectionChanged(TSharedPtr<FString> InName, ESelectInfo::Type SelectInfo)
@@ -121,12 +170,42 @@ FText SIKRetargetChainMapRow::GetSourceChainName() const
 	return FText::FromName(ChainMapElement.Pin()->ChainMap->SourceChain);
 }
 
+FText SIKRetargetChainMapRow::GetTargetIKGoalName() const
+{
+	UIKRetargeterController* RetargeterController = ChainMapList.Pin()->GetRetargetController();
+	if (!RetargeterController)
+	{
+		return FText(); 
+	}
+
+	const UIKRigDefinition* IKRig = RetargeterController->GetIKRig(ERetargetSourceOrTarget::Target);
+	if (!IKRig)
+	{
+		return FText(); 
+	}
+
+	const FBoneChain* Chain = IKRig->GetRetargetChainByName(ChainMapElement.Pin()->ChainMap->TargetChain);
+	if (!Chain)
+	{
+		return FText(); 
+	}
+
+	if (Chain->IKGoalName == NAME_None)
+	{
+		return FText::FromString("");
+	}
+	
+	return FText::FromName(Chain->IKGoalName);
+}
+
 void SIKRetargetChainMapList::Construct(
 	const FArguments& InArgs,
 	TSharedRef<FIKRetargetEditorController> InEditorController)
 {
 	EditorController = InEditorController;
 	EditorController.Pin()->SetChainsView(SharedThis(this));
+
+	TextFilter = MakeShareable(new FTextFilterExpressionEvaluator(ETextFilterExpressionEvaluatorMode::BasicString));
 	
 	ChildSlot
     [
@@ -136,6 +215,19 @@ void SIKRetargetChainMapList::Construct(
 		[
 			SNew(SHorizontalBox)
 
+			+ SHorizontalBox::Slot()
+			.AutoWidth()
+			.Padding(5, 0)
+			.HAlign(HAlign_Left)
+			.VAlign(VAlign_Center)
+			[
+				SNew(SPositiveActionButton)
+				.Icon(FAppStyle::Get().GetBrush("Icons.Settings"))
+				.Text(LOCTEXT("EditRootButtonLabel", "Global Settings"))
+				.ToolTipText(LOCTEXT("EditGlobalButtonToolTip", "Edit the global retarget settings."))
+				.OnClicked(this, &SIKRetargetChainMapList::OnGlobalSettingsButtonClicked)
+			]
+			
 			+ SHorizontalBox::Slot()
 			.AutoWidth()
 			.Padding(5, 0)
@@ -146,7 +238,7 @@ void SIKRetargetChainMapList::Construct(
 				.Icon(FAppStyle::Get().GetBrush("Icons.Settings"))
 				.Text(LOCTEXT("EditRootButtonLabel", "Root Settings"))
 				.ToolTipText(LOCTEXT("EditRootButtonToolTip", "Edit the root retarget settings."))
-				.OnClicked(this, &SIKRetargetChainMapList::OnEditSettingsButtonClicked)
+				.OnClicked(this, &SIKRetargetChainMapList::OnRootSettingsButtonClicked)
 			]
 
 			+ SHorizontalBox::Slot()
@@ -193,17 +285,52 @@ void SIKRetargetChainMapList::Construct(
 		]
 
 		+SVerticalBox::Slot()
+		.Padding(2.0f)
 		.AutoHeight()
-		.VAlign(VAlign_Center)
-		.HAlign(HAlign_Fill)
-		.Padding(5)
 		[
-			SNew(SPositiveActionButton)
-			.Visibility(this, &SIKRetargetChainMapList::IsAutoMapButtonVisible)
-			.Icon(FAppStyle::Get().GetBrush("Icons.Refresh"))
-			.Text(LOCTEXT("AutoMapButtonLabel", "Auto-Map Chains"))
-			.ToolTipText(LOCTEXT("AutoMapButtonToolTip", "Automatically assign source chains based on fuzzy string match"))
-			.OnClicked(this, &SIKRetargetChainMapList::OnAutoMapButtonClicked)
+			SNew(SHorizontalBox)
+
+			+SHorizontalBox::Slot()
+			.AutoWidth()
+			.VAlign(VAlign_Center)
+			.Padding(FMargin(6.f, 0.0))
+			[
+				SNew(SPositiveActionButton)
+				.Visibility(this, &SIKRetargetChainMapList::IsAutoMapButtonVisible)
+				.Icon(FAppStyle::Get().GetBrush("Icons.Refresh"))
+				.Text(LOCTEXT("AutoMapButtonLabel", "Auto-Map Chains"))
+				.ToolTipText(LOCTEXT("AutoMapButtonToolTip", "Automatically assign source chains based on fuzzy string match"))
+				.OnClicked(this, &SIKRetargetChainMapList::OnAutoMapButtonClicked)
+			]
+			
+			+SHorizontalBox::Slot()
+			.FillWidth(1.0f)
+			[
+				SNew(SSearchBox)
+				.SelectAllTextWhenFocused(true)
+				.OnTextChanged( this, &SIKRetargetChainMapList::OnFilterTextChanged )
+				.HintText( LOCTEXT( "SearchBoxHint", "Filter Chain List...") )
+			]
+
+			+SHorizontalBox::Slot()
+			.AutoWidth()
+			.Padding(FMargin(6.f, 0.0))
+			.VAlign(VAlign_Center)
+			[
+				SNew(SComboButton)
+				.ComboButtonStyle(&FAppStyle::Get().GetWidgetStyle<FComboButtonStyle>("SimpleComboButton"))
+				.ForegroundColor(FSlateColor::UseStyle())
+				.ContentPadding(2.0f)
+				.OnGetMenuContent( this, &SIKRetargetChainMapList::CreateFilterMenuWidget )
+				.ToolTipText( LOCTEXT("ChainMapFilterToolTip", "Filter list of chain mappings."))
+				.HasDownArrow(true)
+				.ButtonContent()
+				[
+					SNew(SImage)
+					.Image(FAppStyle::Get().GetBrush("Icons.Settings"))
+					.ColorAndOpacity(FSlateColor::UseForeground())
+				]
+			]
 		]
 	    
         +SVerticalBox::Slot()
@@ -213,21 +340,30 @@ void SIKRetargetChainMapList::Construct(
 			.IsEnabled(this, &SIKRetargetChainMapList::IsChainMapEnabled)
 			.ListItemsSource( &ListViewItems )
 			.OnGenerateRow( this, &SIKRetargetChainMapList::MakeListRowWidget )
-			.OnMouseButtonClick(this, &SIKRetargetChainMapList::OnItemClicked)
-			.OnSelectionChanged_Lambda([this] (TSharedPtr<FRetargetChainMapElement> NewValue, ESelectInfo::Type SelectInfo)
+			.OnMouseButtonClick_Lambda([this](TSharedPtr<FRetargetChainMapElement> Item)
 			{
-				OnSelectionChanged();
+				OnItemClicked(Item);
 			})
 
 			.ItemHeight( 22.0f )
 			.HeaderRow
 			(
 				SNew( SHeaderRow )
-				+ SHeaderRow::Column( ColumnId_TargetChainLabel )
-				.DefaultLabel( LOCTEXT( "TargetColumnLabel", "Target Chain" ) )
+				+ SHeaderRow::Column(ColumnId_TargetChainLabel)
+				.DefaultLabel(LOCTEXT("TargetColumnLabel", "Target Chain"))
+				.DefaultTooltip(LOCTEXT("TargetChainToolTip", "The chain on the target skeleton to copy animation TO."))
+				
+				+ SHeaderRow::Column(ColumnId_IKGoalNameLabel)
+				.DefaultLabel(LOCTEXT("IKColumnLabel", "Target IK Goal"))
+				.DefaultTooltip(LOCTEXT("IKGoalToolTip", "The IK Goal assigned to the target chain (if any). Note, this goal should be on the LAST bone in the chain."))
 
-				+ SHeaderRow::Column( ColumnId_SourceChainLabel )
-				.DefaultLabel( LOCTEXT( "SourceColumnLabel", "Source Chain" ) )
+				+ SHeaderRow::Column(ColumnId_SourceChainLabel)
+				.DefaultLabel(LOCTEXT("SourceColumnLabel", "Source Chain"))
+				.DefaultTooltip(LOCTEXT("SourceChainToolTip", "The chain on the source skeleton to copy animation FROM."))
+				
+				+ SHeaderRow::Column(ColumnId_ResetLabel)
+                .DefaultLabel(LOCTEXT("ResetColumnLabel", "Reset"))
+                .DefaultTooltip(LOCTEXT("ResetChainToolTip", "Reset the FK and IK settings for this chain."))
 			)
         ]
     ];
@@ -292,37 +428,172 @@ bool SIKRetargetChainMapList::IsChainMapEnabled() const
 
 void SIKRetargetChainMapList::RefreshView()
 {
-	UIKRetargeterController* RetargeterController = GetRetargetController();
-	if (!RetargeterController)
+	UIKRetargeterController* AssetController = GetRetargetController();
+	if (!AssetController)
 	{
 		return; 
 	}
-	
-	TArray<TSharedPtr<FRetargetChainMapElement>> PreviouslySelectedItems = ListView->GetSelectedItems();
-	TArray<TObjectPtr<URetargetChainSettings>> SelectedChainMaps;
-	for (TSharedPtr<FRetargetChainMapElement> SelectedItem : PreviouslySelectedItems)
+
+	auto FilterString = [this](const FString& StringToTest) ->bool
 	{
-		SelectedChainMaps.Add(SelectedItem->ChainMap.Get());
-	}
-	
-	// refresh list of chains
-	TArray< TSharedPtr<FRetargetChainMapElement> > SelectedItems;
+		return TextFilter->TestTextFilter(FBasicStringFilterExpressionContext(StringToTest));
+	};
+
+	auto DoesChainHaveIK = [AssetController](const FName TargetChainName) ->bool
+	{
+		const UIKRigDefinition* IKRig = AssetController->GetIKRig(ERetargetSourceOrTarget::Target);
+		if (!IKRig)
+		{
+			return false;
+		}
+
+		const FBoneChain* Chain = IKRig->GetRetargetChainByName(TargetChainName);
+		if (!Chain)
+		{
+			return false;
+		}
+		
+		return Chain->IKGoalName != NAME_None;
+	};
+
+	const TArray<FName>& SelectedChains = EditorController.Pin()->GetSelectedChains();
+	const FName LiteralNone = FName("None");
+
+	// refresh items
 	ListViewItems.Reset();
-	const TArray<TObjectPtr<URetargetChainSettings>>& ChainMappings = RetargeterController->GetChainMappings();
+	const TArray<TObjectPtr<URetargetChainSettings>>& ChainMappings = AssetController->GetChainMappings();
 	for (const TObjectPtr<URetargetChainSettings> ChainMap : ChainMappings)
 	{
-		TSharedPtr<FRetargetChainMapElement> ChainItem = FRetargetChainMapElement::Make(ChainMap);
-		ListViewItems.Add(ChainItem);
+		// apply text filter to items
+		if (!(TextFilter->GetFilterText().IsEmpty() ||
+			FilterString(ChainMap->SourceChain.ToString()) ||
+			FilterString(ChainMap->TargetChain.ToString())))
+		{
+			continue;
+		}
+		
+		// apply "only IK" filter
+		if (ChainFilterOptions.bHideChainsWithoutIK && !DoesChainHaveIK(ChainMap->TargetChain))
+		{
+			continue;
+		}
 
-		if (SelectedChainMaps.Contains(ChainMap))
+		// apply "hide mapped chains" filter
+		if (ChainFilterOptions.bHideMappedChains && ChainMap->SourceChain != LiteralNone)
 		{
-			ListView->SetItemSelection(ChainItem, true, ESelectInfo::Direct);
+			continue;
 		}
+		
+		// apply "hide un-mapped chains" filter
+		if (ChainFilterOptions.bHideUnmappedChains && ChainMap->SourceChain == LiteralNone)
+		{
+			continue;
+		}
+		
+		// create an item for this chain
+		TSharedPtr<FRetargetChainMapElement> ChainItem = FRetargetChainMapElement::Make(ChainMap);
+		ListViewItems.Add(ChainItem);
+
+		// select/deselect it
+		const bool bIsSelected = SelectedChains.Contains(ChainMap->TargetChain);
+		ListView->SetItemSelection(ChainItem, bIsSelected, ESelectInfo::Direct);
 	}
 	
 	ListView->RequestListRefresh();
 }
 
+TSharedRef<SWidget> SIKRetargetChainMapList::CreateFilterMenuWidget()
+{
+	const FUIAction FilterHideMappedAction = FUIAction(
+		FExecuteAction::CreateLambda([this]
+		{
+			ChainFilterOptions.bHideMappedChains = !ChainFilterOptions.bHideMappedChains;
+			RefreshView();
+		}),
+		FCanExecuteAction(),
+		FIsActionChecked::CreateLambda([this]()
+		{
+			return ChainFilterOptions.bHideMappedChains;
+		}));
+
+	const FUIAction FilterOnlyUnMappedAction = FUIAction(
+		FExecuteAction::CreateLambda([this]
+		{
+			ChainFilterOptions.bHideUnmappedChains = !ChainFilterOptions.bHideUnmappedChains;
+			RefreshView();
+		}),
+		FCanExecuteAction(),
+		FIsActionChecked::CreateLambda([this]()
+		{
+			return ChainFilterOptions.bHideUnmappedChains;
+		}));
+
+	const FUIAction FilterIKChainAction = FUIAction(
+		FExecuteAction::CreateLambda([this]
+		{
+			ChainFilterOptions.bHideChainsWithoutIK = !ChainFilterOptions.bHideChainsWithoutIK;
+			RefreshView();
+		}),
+		FCanExecuteAction(),
+		FIsActionChecked::CreateLambda([this]()
+		{
+			return ChainFilterOptions.bHideChainsWithoutIK;
+		}));
+	
+	static constexpr bool CloseAfterSelection = true;
+	FMenuBuilder MenuBuilder(CloseAfterSelection, CommandList);
+
+	MenuBuilder.BeginSection("Chain Map Filters", LOCTEXT("ChainFiltersSection", "Filter Chain Mappings"));
+
+	MenuBuilder.AddMenuEntry(
+		LOCTEXT("HideMappedLabel", "Hide Mapped Chains"),
+		LOCTEXT("HideMappedTooltip", "Hide chains mapped to a source chain."),
+		FSlateIcon(),
+		FilterHideMappedAction,
+		NAME_None,
+		EUserInterfaceActionType::Check);
+
+	MenuBuilder.AddMenuEntry(
+		LOCTEXT("HideUnMappedLabel", "Hide Unmapped Chains"),
+		LOCTEXT("HideUnMappedTooltip", "Hide chains not mapped to a source chain."),
+		FSlateIcon(),
+		FilterOnlyUnMappedAction,
+		NAME_None,
+		EUserInterfaceActionType::Check);
+
+	MenuBuilder.AddMenuEntry(
+		LOCTEXT("HideNonIKLabel", "Hide Chains Without IK"),
+		LOCTEXT("HideNonIKTooltip", "Hide chains not using IK."),
+		FSlateIcon(),
+		FilterIKChainAction,
+		NAME_None,
+		EUserInterfaceActionType::Check);
+	
+	MenuBuilder.EndSection();
+
+	MenuBuilder.BeginSection("Clear", LOCTEXT("ClearMapFiltersSection", "Clear"));
+	
+	MenuBuilder.AddMenuEntry(
+		LOCTEXT("ClearMapFilterLabel", "Clear Filters"),
+		LOCTEXT("ClearMapFilterTooltip", "Clear all filters to show all chain mappings."),
+		FSlateIcon(),
+		FUIAction(FExecuteAction::CreateLambda([this]
+		{
+			ChainFilterOptions = FChainMapFilterOptions();
+			RefreshView();
+		})));
+
+	MenuBuilder.EndSection();
+	
+	return MenuBuilder.MakeWidget();
+}
+
+void SIKRetargetChainMapList::OnFilterTextChanged(const FText& SearchText)
+{
+	TextFilter->SetFilterText(SearchText);
+	RefreshView();
+}
+
 TSharedRef<ITableRow> SIKRetargetChainMapList::MakeListRowWidget(
 	TSharedPtr<FRetargetChainMapElement> InElement,
     const TSharedRef<STableViewBase>& OwnerTable)
@@ -333,38 +604,19 @@ TSharedRef<ITableRow> SIKRetargetChainMapList::MakeListRowWidget(
 		SharedThis(this));
 }
 
-void SIKRetargetChainMapList::OnItemClicked(TSharedPtr<FRetargetChainMapElement> InItem)
-{
-	OnSelectionChanged();
-}
-
-void SIKRetargetChainMapList::OnSelectionChanged()
+void SIKRetargetChainMapList::OnItemClicked(TSharedPtr<FRetargetChainMapElement> InItem) const 
 {
-	const TSharedPtr<FIKRetargetEditorController> Controller = EditorController.Pin();
-	if (!Controller.IsValid())
-	{
-		return;
-	}
-	
-	// get selected chain settings
-	TArray<UObject*> SelectedChainSettings;
+	// get list of selected chains
+	TArray<FName> SelectedChains;
 	TArray<TSharedPtr<FRetargetChainMapElement>> SelectedItems = ListView.Get()->GetSelectedItems();
 	for (const TSharedPtr<FRetargetChainMapElement>& Item : SelectedItems)
 	{
-		SelectedChainSettings.Add(Item->ChainMap.Get());
-	}
-
-	// selection cleared
-	if (SelectedChainSettings.IsEmpty())
-	{
-		// show asset settings in the details view
-		Controller->SetDetailsObject(Controller->AssetController->GetAsset());
-	}
-	else
-	{
-		// show chain settings in the details view
-		Controller->SetDetailsObjects(SelectedChainSettings);
+		SelectedChains.Add(Item.Get()->ChainMap->TargetChain);
 	}
+	
+	// replace the chain selection
+	constexpr bool bEditFromChainsView = true;
+	EditorController.Pin()->EditChainSelection(SelectedChains, ESelectionEdit::Replace, bEditFromChainsView);
 }
 
 EVisibility SIKRetargetChainMapList::IsAutoMapButtonVisible() const
@@ -392,7 +644,32 @@ FReply SIKRetargetChainMapList::OnAutoMapButtonClicked() const
 	return FReply::Handled();
 }
 
-FReply SIKRetargetChainMapList::OnEditSettingsButtonClicked() const
+FReply SIKRetargetChainMapList::OnGlobalSettingsButtonClicked() const
+{
+	const TSharedPtr<FIKRetargetEditorController> Controller = EditorController.Pin();
+	if (!Controller.IsValid())
+	{
+		return FReply::Unhandled();
+	}
+
+	UIKRetargeterController* RetargeterController = GetRetargetController();
+	if (!RetargeterController)
+	{
+		return FReply::Unhandled();
+	}
+
+	UIKRetargetGlobalSettings* GlobalSettings = Controller->AssetController->GetAsset()->GetGlobalSettingsUObject();
+	if (GlobalSettings->EditorController != Controller)
+	{
+		GlobalSettings->EditorController = Controller;
+	}
+	
+	Controller->SetDetailsObject(GlobalSettings);
+	
+	return FReply::Handled();
+}
+
+FReply SIKRetargetChainMapList::OnRootSettingsButtonClicked() const
 {
 	const TSharedPtr<FIKRetargetEditorController> Controller = EditorController.Pin();
 	if (!Controller.IsValid())
@@ -406,7 +683,7 @@ FReply SIKRetargetChainMapList::OnEditSettingsButtonClicked() const
 		return FReply::Unhandled();
 	}
 	
-	Controller->SetDetailsObject(RetargeterController->GetAsset()->GetRetargetRootSettings());
+	Controller->SetRootSelected(true);
 	return FReply::Handled();
 }
 
diff --git a/Engine/Plugins/Animation/IKRig/Source/IKRigEditor/Private/RetargetEditor/SIKRetargetHierarchy.cpp b/Engine/Plugins/Animation/IKRig/Source/IKRigEditor/Private/RetargetEditor/SIKRetargetHierarchy.cpp
--- a/Engine/Plugins/Animation/IKRig/Source/IKRigEditor/Private/RetargetEditor/SIKRetargetHierarchy.cpp
+++ b/Engine/Plugins/Animation/IKRig/Source/IKRigEditor/Private/RetargetEditor/SIKRetargetHierarchy.cpp
@@ -10,6 +10,9 @@
 #include "Widgets/Layout/SScrollBox.h"
 #include "Widgets/Layout/SSeparator.h"
 #include "Engine/SkeletalMesh.h"
+#include "SPositiveActionButton.h"
+#include "Misc/TextFilterExpressionEvaluator.h"
+#include "Widgets/Input/SSearchBox.h"
 
 #define LOCTEXT_NAMESPACE "SIKRetargetHierarchy"
 
@@ -18,19 +21,21 @@ static const FName ChainColumnName(TEXT("RetargetChainName"));
 
 FIKRetargetHierarchyElement::FIKRetargetHierarchyElement(
 	const FName& InName,
+	const FName& InChainName,
 	const TSharedRef<FIKRetargetEditorController>& InEditorController)
 	: Key(FText::FromName(InName)),
 	Name(InName),
+	ChainName(InChainName),
 	EditorController(InEditorController)
 {}
 
 TSharedRef<SWidget> SIKRetargetHierarchyRow::GenerateWidgetForColumn(const FName& ColumnName)
 {
-	FName BoneName = WeakTreeElement.Pin()->Name;
+	const FName BoneName = WeakTreeElement.Pin()->Name;
+	const FName ChainName = WeakTreeElement.Pin()->ChainName;
 	ERetargetSourceOrTarget CurrentSkeleton = EditorController.Pin()->GetSourceOrTarget();
 	const bool bIsBoneRetargeted = EditorController.Pin()->IsBoneRetargeted(BoneName, CurrentSkeleton);
-	const FText ChainName = FText::FromName(EditorController.Pin()->GetChainNameFromBone(BoneName, CurrentSkeleton));
-
+	
 	// determine icon based on if bone is retargeted
 	const FSlateBrush* Brush;
 	if (bIsBoneRetargeted)
@@ -101,7 +106,7 @@ TSharedRef<SWidget> SIKRetargetHierarchyRow::GenerateWidgetForColumn(const FName
 		.AutoWidth()
 		[
 			SNew(STextBlock)
-			.Text(ChainName)
+			.Text(FText::FromName(ChainName))
 			.Font(TextFont)
 			.ColorAndOpacity(TextColor)
 		];
@@ -119,6 +124,8 @@ void SIKRetargetHierarchy::Construct(
 	EditorController = InEditorController;
 	EditorController.Pin()->SetHierarchyView(SharedThis(this));
 	CommandList = MakeShared<FUICommandList>();
+
+	TextFilter = MakeShareable(new FTextFilterExpressionEvaluator(ETextFilterExpressionEvaluatorMode::BasicString));
 	
 	ChildSlot
 	[
@@ -169,6 +176,41 @@ void SIKRetargetHierarchy::Construct(
 			SNew(SSeparator)
 		]
 
+		+SVerticalBox::Slot()
+		.Padding(2.0f)
+		.AutoHeight()
+		[
+			SNew(SHorizontalBox)
+			+SHorizontalBox::Slot()
+			.FillWidth(1.0f)
+			[
+				SNew(SSearchBox)
+				.SelectAllTextWhenFocused(true)
+				.OnTextChanged( this, &SIKRetargetHierarchy::OnFilterTextChanged )
+				.HintText( LOCTEXT( "SearchBoxHint", "Filter Hierarchy Tree...") )
+			]
+
+			+SHorizontalBox::Slot()
+			.AutoWidth()
+			.Padding(FMargin(6.f, 0.0))
+			.VAlign(VAlign_Center)
+			[
+				SNew(SComboButton)
+				.ComboButtonStyle(&FAppStyle::Get().GetWidgetStyle<FComboButtonStyle>("SimpleComboButton"))
+				.ForegroundColor(FSlateColor::UseStyle())
+				.ContentPadding(2.0f)
+				.OnGetMenuContent( this, &SIKRetargetHierarchy::CreateFilterMenuWidget )
+				.ToolTipText( LOCTEXT( "HierarchyFilterToolTip", "Options to filter the hierarchy tree.") )
+				.HasDownArrow(true)
+				.ButtonContent()
+				[
+					SNew(SImage)
+					.Image(FAppStyle::Get().GetBrush("Icons.Settings"))
+					.ColorAndOpacity(FSlateColor::UseForeground())
+				]
+			]
+		]
+
 		+SVerticalBox::Slot()
 		.Padding(2.0f)
 		[
@@ -189,7 +231,6 @@ void SIKRetargetHierarchy::Construct(
 				})
 				.OnGetChildren(this, &SIKRetargetHierarchy::HandleGetChildrenForTree)
 				.OnSelectionChanged(this, &SIKRetargetHierarchy::OnSelectionChanged)
-				.OnMouseButtonClick(this, &SIKRetargetHierarchy::OnItemClicked)
 				.OnMouseButtonDoubleClick(this, &SIKRetargetHierarchy::OnItemDoubleClicked)
 				.OnSetExpansionRecursive(this, &SIKRetargetHierarchy::OnSetExpansionRecursive)
 				.HighlightParentNodesForSelection(false)
@@ -243,6 +284,98 @@ void SIKRetargetHierarchy::ShowItemAfterSelection(FName ItemName)
 	TreeView->RequestScrollIntoView(ItemToShow);
 }
 
+TSharedRef<SWidget> SIKRetargetHierarchy::CreateFilterMenuWidget()
+{
+	const FUIAction FilterHideNotInChainAction = FUIAction(
+		FExecuteAction::CreateLambda([this]
+		{
+			FilterOptions.bHideBonesNotInChain = !FilterOptions.bHideBonesNotInChain;
+			RefreshTreeView();
+		}),
+		FCanExecuteAction(),
+		FIsActionChecked::CreateLambda([this]()
+		{
+			return FilterOptions.bHideBonesNotInChain;
+		}));
+
+	const FUIAction FilterHideNotRetargetedAction = FUIAction(
+		FExecuteAction::CreateLambda([this]
+		{
+			FilterOptions.bHideNotRetargetedBones = !FilterOptions.bHideNotRetargetedBones;
+			RefreshTreeView();
+		}),
+		FCanExecuteAction(),
+		FIsActionChecked::CreateLambda([this]()
+		{
+			return FilterOptions.bHideNotRetargetedBones;
+		}));
+
+	const FUIAction FilterHideRetargetedAction = FUIAction(
+		FExecuteAction::CreateLambda([this]
+		{
+			FilterOptions.bHideRetargetedBones = !FilterOptions.bHideRetargetedBones;
+			RefreshTreeView();
+		}),
+		FCanExecuteAction(),
+		FIsActionChecked::CreateLambda([this]()
+		{
+			return FilterOptions.bHideRetargetedBones;
+		}));
+	
+	static constexpr bool CloseAfterSelection = true;
+	FMenuBuilder MenuBuilder(CloseAfterSelection, CommandList);
+
+	MenuBuilder.BeginSection("Hierarchy Filters", LOCTEXT("HierarchyFiltersSection", "Hierarchy Filters"));
+
+	MenuBuilder.AddMenuEntry(
+		LOCTEXT("HideNotInChainLabel", "Hide Bones Not in a Chain"),
+		LOCTEXT("HideNotInChainTooltip", "Hide bones that are not part of a retarget chain."),
+		FSlateIcon(),
+		FilterHideNotInChainAction,
+		NAME_None,
+		EUserInterfaceActionType::Check);
+
+	MenuBuilder.AddMenuEntry(
+		LOCTEXT("HideUnMappedLabel", "Hide Bones Not Retargeted"),
+		LOCTEXT("HideUnMappedTooltip", "Hide bones that are not being affected by the retargeting; either by not being in a retarget chain or being in a chain that is not mapped to anything on the source."),
+		FSlateIcon(),
+		FilterHideNotRetargetedAction,
+		NAME_None,
+		EUserInterfaceActionType::Check);
+
+	MenuBuilder.AddMenuEntry(
+		LOCTEXT("HideNonIKLabel", "Hide Retargeted Bones"),
+		LOCTEXT("HideNonIKTooltip", "Hide all bones that are part of a mapped retarget chain."),
+		FSlateIcon(),
+		FilterHideRetargetedAction,
+		NAME_None,
+		EUserInterfaceActionType::Check);
+	
+	MenuBuilder.EndSection();
+
+	MenuBuilder.BeginSection("Clear", LOCTEXT("ClearMapFiltersSection", "Clear"));
+	
+	MenuBuilder.AddMenuEntry(
+		LOCTEXT("ClearMapFilterLabel", "Clear Filters"),
+		LOCTEXT("ClearMapFilterTooltip", "Clear all filters to show all bones."),
+		FSlateIcon(),
+		FUIAction(FExecuteAction::CreateLambda([this]
+		{
+			FilterOptions = FRetargetHierarchyFilterOptions();
+			RefreshTreeView();
+		})));
+
+	MenuBuilder.EndSection();
+	
+	return MenuBuilder.MakeWidget();
+}
+
+void SIKRetargetHierarchy::OnFilterTextChanged(const FText& SearchText)
+{
+	TextFilter->SetFilterText(SearchText);
+	RefreshTreeView(false);
+}
+
 void SIKRetargetHierarchy::RefreshTreeView(bool IsInitialSetup)
 {
 	const TSharedPtr<FIKRetargetEditorController> Controller = EditorController.Pin();
@@ -288,19 +421,54 @@ void SIKRetargetHierarchy::RefreshTreeView(bool IsInitialSetup)
 		BoneNames[BoneIndex] = Skeleton.GetBoneName(BoneIndex);
 		ParentIndices[BoneIndex] = Skeleton.GetParentIndex(BoneIndex);
 	}
+
+	auto FilterString = [this](const FString& StringToTest) ->bool
+	{
+		return TextFilter->TestTextFilter(FBasicStringFilterExpressionContext(StringToTest));
+	};
 	
 	// record bone element indices
 	TMap<FName, int32> BoneTreeElementIndices;
 
 	// create all bone elements
-	// TODO, add filtering of retargeted bones, by name, etc..
 	for (const FName BoneName : BoneNames)
 	{
 		// create "Bone" tree element for this bone
-		TSharedPtr<FIKRetargetHierarchyElement> BoneElement = MakeShared<FIKRetargetHierarchyElement>(BoneName, ControllerRef);
-		BoneElement.Get()->Name = BoneName;
+		FName ChainName = Controller->GetChainNameFromBone(BoneName, Controller->GetSourceOrTarget());
+		TSharedPtr<FIKRetargetHierarchyElement> BoneElement = MakeShared<FIKRetargetHierarchyElement>(BoneName, ChainName, ControllerRef);
 		const int32 BoneElementIndex = AllElements.Add(BoneElement);
 		BoneTreeElementIndices.Add(BoneName, BoneElementIndex);
+		
+		// apply text filter to items
+		if (!(TextFilter->GetFilterText().IsEmpty() ||
+			FilterString(BoneName.ToString()) ||
+			FilterString(ChainName.ToString()))
+			)
+		{
+			BoneElement->bIsHidden = true;
+		}
+
+		// apply filter "hide bones not in chain"
+		if (FilterOptions.bHideBonesNotInChain && BoneElement->ChainName == NAME_None)
+		{
+			BoneElement->bIsHidden = true;
+		}
+
+		// apply filter "hide bones not retargeted"
+		if (FilterOptions.bHideNotRetargetedBones || FilterOptions.bHideRetargetedBones)
+		{
+			const bool bIsRetargeted = Controller->IsBoneRetargeted(BoneName, Controller->GetSourceOrTarget());
+
+			if (FilterOptions.bHideRetargetedBones && bIsRetargeted)
+			{
+				BoneElement->bIsHidden = true;
+			}
+			
+			if (FilterOptions.bHideNotRetargetedBones && !bIsRetargeted)
+			{
+				BoneElement->bIsHidden = true;
+			}
+		}
 	}
 
 	// store children/parent pointers on all bone elements
@@ -308,22 +476,45 @@ void SIKRetargetHierarchy::RefreshTreeView(bool IsInitialSetup)
 	{
 		const FName BoneName = BoneNames[BoneIndex];
 		const TSharedPtr<FIKRetargetHierarchyElement> BoneTreeElement = AllElements[BoneTreeElementIndices[BoneName]];
-		const int32 ParentIndex = ParentIndices[BoneIndex];
-		if (ParentIndex < 0)
+
+		if (BoneTreeElement->bIsHidden)
 		{
-			// store the root element
-			RootElements.Add(BoneTreeElement);
-			// has no parent, so skip storing parent pointer
 			continue;
 		}
+		
+		// find first parent that is not filtered
+		const TSharedPtr<FIKRetargetHierarchyElement>* FirstNonHiddenParentElement = nullptr;
+		int32 ParentBoneIndex = ParentIndices[BoneIndex];
+		while (true)
+		{
+			if (ParentBoneIndex < 0)
+			{
+				break;
+			}
+
+			const FName ParentBoneName = BoneNames[ParentBoneIndex];
+			TSharedPtr<FIKRetargetHierarchyElement> ParentBoneTreeElement = AllElements[BoneTreeElementIndices[ParentBoneName]];
+			if (!ParentBoneTreeElement->bIsHidden)
+			{
+				FirstNonHiddenParentElement = &ParentBoneTreeElement;
+				break;
+			}
+			
+			ParentBoneIndex = ParentIndices[ParentBoneIndex];
+		}
 
-		// get parent tree element
-		const FName ParentBoneName = BoneNames[ParentIndex];
-		const TSharedPtr<FIKRetargetHierarchyElement> ParentBoneTreeElement = AllElements[BoneTreeElementIndices[ParentBoneName]];
-		// store pointer to child on parent
-		ParentBoneTreeElement->Children.Add(BoneTreeElement);
-		// store pointer to parent on child
-		BoneTreeElement->Parent = ParentBoneTreeElement;
+		if (FirstNonHiddenParentElement)
+		{
+			// store pointer to child on parent
+			FirstNonHiddenParentElement->Get()->Children.Add(BoneTreeElement);
+			// store pointer to parent on child
+			BoneTreeElement->Parent = *FirstNonHiddenParentElement;
+		}
+		else
+		{
+			// has no parent, store a root element
+			RootElements.Add(BoneTreeElement);
+		}
 	}
 
 	// expand all elements upon the initial construction of the tree
@@ -350,7 +541,7 @@ void SIKRetargetHierarchy::HandleGetChildrenForTree(
 	TSharedPtr<FIKRetargetHierarchyElement> InItem,
 	TArray<TSharedPtr<FIKRetargetHierarchyElement>>& OutChildren)
 {
-	OutChildren = InItem.Get()->Children;
+	OutChildren = InItem->Children;
 }
 
 void SIKRetargetHierarchy::OnSelectionChanged(
@@ -376,12 +567,7 @@ void SIKRetargetHierarchy::OnSelectionChanged(
 	}
 
 	constexpr bool bFromHierarchy = true;
-	Controller->EditBoneSelection(SelectedBoneNames, EBoneSelectionEdit::Replace, bFromHierarchy);
-}
-
-void SIKRetargetHierarchy::OnItemClicked(TSharedPtr<FIKRetargetHierarchyElement> InItem)
-{
-	// TODO show bone details
+	Controller->EditBoneSelection(SelectedBoneNames, ESelectionEdit::Replace, bFromHierarchy);
 }
 
 void SIKRetargetHierarchy::OnItemDoubleClicked(TSharedPtr<FIKRetargetHierarchyElement> InItem)
diff --git a/Engine/Plugins/Animation/IKRig/Source/IKRigEditor/Private/RigEditor/IKRigController.cpp b/Engine/Plugins/Animation/IKRig/Source/IKRigEditor/Private/RigEditor/IKRigController.cpp
--- a/Engine/Plugins/Animation/IKRig/Source/IKRigEditor/Private/RigEditor/IKRigController.cpp
+++ b/Engine/Plugins/Animation/IKRig/Source/IKRigEditor/Private/RigEditor/IKRigController.cpp
@@ -152,14 +152,15 @@ bool UIKRigController::DoesBoneHaveSettings(const FName& BoneName) const
 	return false;
 }
 
-void UIKRigController::AddRetargetChain(const FName& ChainName, const FName& StartBone, const FName& EndBone) const
+void UIKRigController::AddRetargetChain(const FBoneChain& BoneChain) const
 {
-	FName UniqueChainName = GetUniqueRetargetChainName(ChainName);
+	FBoneChain ChainToAdd = BoneChain;
+	ChainToAdd.ChainName = GetUniqueRetargetChainName(BoneChain.ChainName);
 
 	FScopedTransaction Transaction(LOCTEXT("AddRetargetChain_Label", "Add Retarget Chain"));
 	Asset->Modify();
 	
-	Asset->RetargetDefinition.BoneChains.Emplace(UniqueChainName, StartBone, EndBone);
+	Asset->RetargetDefinition.BoneChains.Emplace(ChainToAdd);
 	BroadcastNeedsReinitialized();
 }
 
@@ -291,6 +292,11 @@ FName UIKRigController::GetRetargetChainEndBone(const FName& ChainName) const
 	return Chain->EndBone.BoneName;
 }
 
+const FBoneChain* UIKRigController::GetRetargetChainByName(const FName& ChainName) const
+{
+	return Asset->GetRetargetChainByName(ChainName);
+}
+
 const TArray<FBoneChain>& UIKRigController::GetRetargetChains() const
 {
 	check(Asset)
@@ -863,6 +869,20 @@ FName UIKRigController::GetBoneForGoal(const FName& GoalName) const
 	return NAME_None;
 }
 
+const UIKRigEffectorGoal* UIKRigController::GetGoalForBone(const FName& BoneName) const
+{
+	const TArray<UIKRigEffectorGoal*>& AllGoals = GetAllGoals();
+	for (const UIKRigEffectorGoal* Goal : AllGoals)
+	{
+		if (Goal->BoneName == BoneName)
+		{
+			return Goal;
+		}
+	}
+
+	return nullptr;
+}
+
 bool UIKRigController::ConnectGoalToSolver(const UIKRigEffectorGoal& Goal, int32 SolverIndex) const
 {
 	// can't add goal that is not present in the core
@@ -914,6 +934,19 @@ bool UIKRigController::IsGoalConnectedToSolver(const FName& GoalName, int32 Solv
 	return Asset->Solvers[SolverIndex]->IsGoalConnected(GoalName);
 }
 
+bool UIKRigController::IsGoalConnectedToAnySolver(const FName& GoalName) const
+{
+	for (const TObjectPtr<UIKRigSolver> Solver : Asset->Solvers)
+	{
+		if (Solver->IsGoalConnected(GoalName))
+		{
+			return true;
+		}
+	}
+
+	return false;
+}
+
 const TArray<UIKRigEffectorGoal*>& UIKRigController::GetAllGoals() const
 {
 	return Asset->Goals;
diff --git a/Engine/Plugins/Animation/IKRig/Source/IKRigEditor/Private/RigEditor/IKRigEditMode.cpp b/Engine/Plugins/Animation/IKRig/Source/IKRigEditor/Private/RigEditor/IKRigEditMode.cpp
--- a/Engine/Plugins/Animation/IKRig/Source/IKRigEditor/Private/RigEditor/IKRigEditMode.cpp
+++ b/Engine/Plugins/Animation/IKRig/Source/IKRigEditor/Private/RigEditor/IKRigEditMode.cpp
@@ -124,21 +124,15 @@ void FIKRigEditMode::RenderGoals(FPrimitiveDrawInterface* PDI)
 
 void FIKRigEditMode::RenderBones(FPrimitiveDrawInterface* PDI)
 {
-	// editor configured and initialized?
+	// get the controller
 	const TSharedPtr<FIKRigEditorController> Controller = EditorController.Pin();
 	if (!Controller.IsValid())
 	{
 		return; 
 	}
 
-	// anim instance initialized?
-	if (!Controller->AnimInstance.IsValid())
-	{
-		return;
-	}
-
 	// IKRig processor initialized and running?
-	UIKRigProcessor* CurrentProcessor = Controller->GetIKRigProcessor();
+	const UIKRigProcessor* CurrentProcessor = Controller->GetIKRigProcessor();
 	if (!IsValid(CurrentProcessor))
 	{
 		return;
@@ -147,105 +141,90 @@ void FIKRigEditMode::RenderBones(FPrimitiveDrawInterface* PDI)
 	{
 		return;
 	}
+	
+	const UDebugSkelMeshComponent* MeshComponent = Controller->SkelMeshComponent;
+	const FTransform ComponentTransform = MeshComponent->GetComponentTransform();
+	const FReferenceSkeleton& RefSkeleton = MeshComponent->GetReferenceSkeleton();
+	const int32 NumBones = RefSkeleton.GetNum();
 
-	// determine size to draw bones from preference stored on the asset
-	const float BoneRadius = FMath::Max(0.01f, Controller->AssetController->GetAsset()->BoneSize);
-
-	// get affected / selected bones
-	TSet<int32> OutAffectedBones;
-	TSet<int32> OutSelectedBones;
-	GetAffectedBones(Controller.Get(), CurrentProcessor,OutAffectedBones,OutSelectedBones);
-
-	// use colors from user preferences
-	const UPersonaOptions* PersonaOptions = GetDefault<UPersonaOptions>();
-
-	// draw bones
-	TArray<int32> ChildrenIndices;
-	const FIKRigSkeleton& Skeleton = CurrentProcessor->GetSkeleton();
-	const TArray<FTransform>& BoneTransforms = Skeleton.CurrentPoseGlobal;
-	for (int32 BoneIndex=0; BoneIndex<Skeleton.BoneNames.Num(); ++BoneIndex)
+	// get world transforms of bones
+	TArray<FBoneIndexType> RequiredBones;
+	RequiredBones.AddUninitialized(NumBones);
+	TArray<FTransform> WorldTransforms;
+	WorldTransforms.AddUninitialized(NumBones);
+	for (int32 Index=0; Index<NumBones; ++Index)
 	{
-		// selected bones are drawn with different color
-		const bool bIsSelected = OutSelectedBones.Contains(BoneIndex);
-		const bool bIsAffected = OutAffectedBones.Contains(BoneIndex);
-		FLinearColor LineColor = bIsAffected ? PersonaOptions->AffectedBoneColor : PersonaOptions->DefaultBoneColor;
-		LineColor = bIsSelected ? PersonaOptions->SelectedBoneColor : LineColor;
-
-		// only draw axes on affected/selected bones
-		const bool bDrawAxes = bIsSelected || bIsAffected;
-		
-		// draw line from bone to each child
-		FTransform BoneTransform = BoneTransforms[BoneIndex];
-		TArray<FVector> ChildPoints;
-		Skeleton.GetChildIndices(BoneIndex,ChildrenIndices);
-		for (int32 ChildIndex : ChildrenIndices)
-		{
-			ChildPoints.Add(BoneTransforms[ChildIndex].GetLocation());
-		}
-
-		// draw cone to parent with different color if child is selected
-		TArray<FLinearColor> ChildColors;
-		for (int32 ChildIndex : ChildrenIndices)
-		{
-			FLinearColor ChildLineColor = LineColor;
-			if (!bIsSelected && OutSelectedBones.Contains(ChildIndex))
-			{
-				ChildLineColor = PersonaOptions->ParentOfSelectedBoneColor;
-			}
-			ChildColors.Add(ChildLineColor);
-		}
-
-		PDI->SetHitProxy(new HIKRigEditorBoneProxy(Skeleton.BoneNames[BoneIndex]));
-		SkeletalDebugRendering::DrawWireBoneAdvanced(
-			PDI,
-			BoneTransform,
-			ChildPoints,
-			ChildColors,
-			LineColor,
-			SDPG_Foreground,
-			BoneRadius,
-			bDrawAxes);
-		PDI->SetHitProxy(nullptr);
+		RequiredBones[Index] = Index;
+		WorldTransforms[Index] = MeshComponent->GetBoneTransform(Index, ComponentTransform);
+	}
+	
+	UPersonaOptions* ConfigOption = UPersonaOptions::StaticClass()->GetDefaultObject<UPersonaOptions>();
+	
+	FSkelDebugDrawConfig DrawConfig;
+	DrawConfig.BoneDrawMode = (EBoneDrawMode::Type)ConfigOption->DefaultBoneDrawSelection;
+	DrawConfig.BoneDrawSize = Controller->AssetController->GetAsset()->BoneSize;
+	DrawConfig.bAddHitProxy = true;
+	DrawConfig.bForceDraw = false;
+	DrawConfig.DefaultBoneColor = GetMutableDefault<UPersonaOptions>()->DefaultBoneColor;
+	DrawConfig.AffectedBoneColor = GetMutableDefault<UPersonaOptions>()->AffectedBoneColor;
+	DrawConfig.SelectedBoneColor = GetMutableDefault<UPersonaOptions>()->SelectedBoneColor;
+	DrawConfig.ParentOfSelectedBoneColor = GetMutableDefault<UPersonaOptions>()->ParentOfSelectedBoneColor;
+
+	TArray<HHitProxy*> HitProxies;
+	HitProxies.AddUninitialized(NumBones);
+	for (int32 Index = 0; Index < NumBones; ++Index)
+	{
+		HitProxies[Index] = new HIKRigEditorBoneProxy(RefSkeleton.GetBoneName(Index));
 	}
+	
+	// get selected bones
+	TArray<TSharedPtr<FIKRigTreeElement>> SelectedBoneItems;
+	Controller->GetSelectedBones(SelectedBoneItems);
+	TArray<int32> SelectedBones;
+	for (TSharedPtr<FIKRigTreeElement> SelectedBone: SelectedBoneItems)
+	{
+		const int32 BoneIndex = RefSkeleton.FindBoneIndex(SelectedBone->BoneName);
+		SelectedBones.Add(BoneIndex);
+	}
+
+	// get bone colors
+	TArray<FLinearColor> BoneColors;
+	GetBoneColors(Controller.Get(), CurrentProcessor, RefSkeleton, BoneColors);
+
+	SkeletalDebugRendering::DrawBones(
+		PDI,
+		ComponentTransform.GetLocation(),
+		RequiredBones,
+		RefSkeleton,
+		WorldTransforms,
+		SelectedBones,
+		BoneColors,
+		HitProxies,
+		DrawConfig
+	);
 }
 
-void FIKRigEditMode::GetAffectedBones(
+void FIKRigEditMode::GetBoneColors(
 	FIKRigEditorController* Controller,
-	UIKRigProcessor* Processor,
-	TSet<int32>& OutAffectedBones,
-	TSet<int32>& OutSelectedBones) const
+	const UIKRigProcessor* Processor,
+	const FReferenceSkeleton& RefSkeleton,
+	TArray<FLinearColor>& OutBoneColors) const
 {
-	OutAffectedBones.Reset();
-	OutSelectedBones.Reset();
+	const FLinearColor DefaultColor = GetMutableDefault<UPersonaOptions>()->DefaultBoneColor;
+	const FLinearColor HighlightedColor = FLinearColor::Blue;
 
-	const FIKRigSkeleton& Skeleton = Processor->GetSkeleton();
+	// set all to default color
+	OutBoneColors.SetNum(RefSkeleton.GetNum());
+	for (int32 Index=0; Index<RefSkeleton.GetNum(); ++Index)
+	{
+		OutBoneColors[Index] = DefaultColor;
+	}
 	
+	// highlight bones of the last selected UI element (could be solver or retarget chain) 
 	switch (Controller->GetLastSelectedType())
 	{
 		case EIKRigSelectionType::Hierarchy:
-		{
-			// get selected bones
-			TArray<TSharedPtr<FIKRigTreeElement>> SelectedBoneItems;
-			Controller->GetSelectedBones(SelectedBoneItems);
-
-			// record indices of all selected bones
-			for (TSharedPtr<FIKRigTreeElement> SelectedBone: SelectedBoneItems)
-			{
-				int32 BoneIndex = Skeleton.GetBoneIndexFromName(SelectedBone->BoneName);
-				OutSelectedBones.Add(BoneIndex);
-			}
-
-			// "affected bones" are the selected bones AND their children, recursively
-			for (int32 SelectedBone : OutSelectedBones)
-			{
-				const int32 EndOfBranch = Skeleton.GetCachedEndOfBranchIndex(SelectedBone);
-				OutAffectedBones.Add(SelectedBone);
-				for (int32 BoneIndex=SelectedBone; BoneIndex<=EndOfBranch; ++BoneIndex)
-				{
-					OutAffectedBones.Add(BoneIndex);
-				}
-			}
-		}
+		// handled by SkeletalDebugRendering::DrawBones()
 		break;
 		
 		case EIKRigSelectionType::SolverStack:
@@ -259,16 +238,16 @@ void FIKRigEditMode::GetAffectedBones(
 			}
 
 			// record which bones in the skeleton are affected by this solver
-			UIKRigController* AssetController = Controller->AssetController;
-			UIKRigSolver* SelectedSolver = AssetController->GetSolver(SelectedSolvers[0].Get()->IndexInStack);
-			if(SelectedSolver)
+			const FIKRigSkeleton& Skeleton = Processor->GetSkeleton();
+			const UIKRigController* AssetController = Controller->AssetController;
+			if(const UIKRigSolver* SelectedSolver = AssetController->GetSolver(SelectedSolvers[0].Get()->IndexInStack))
 			{
-				for (int32 BoneIndex=0; BoneIndex<Skeleton.BoneNames.Num(); ++BoneIndex)
+				for (int32 BoneIndex=0; BoneIndex < Skeleton.BoneNames.Num(); ++BoneIndex)
 				{
 					const FName& BoneName = Skeleton.BoneNames[BoneIndex];
 					if (SelectedSolver->IsBoneAffectedBySolver(BoneName, Skeleton))
 					{
-						OutAffectedBones.Add(BoneIndex);
+						OutBoneColors[BoneIndex] = HighlightedColor;
 					}
 				}
 			}
@@ -277,12 +256,23 @@ void FIKRigEditMode::GetAffectedBones(
 		
 		case EIKRigSelectionType::RetargetChains:
 		{
-			const FName SelectedChainName = Controller->GetSelectedChain();
-			if (SelectedChainName == NAME_None)
+			const TArray<FName> SelectedChainNames = Controller->GetSelectedChains();
+			if (SelectedChainNames.IsEmpty())
 			{
 				return;
 			}
-			Controller->AssetController->ValidateChain(SelectedChainName, OutSelectedBones);
+
+			for (const FName& SelectedChainName : SelectedChainNames)
+			{
+				TSet<int32> OutChainBoneIndices;
+				if (Controller->AssetController->ValidateChain(SelectedChainName, OutChainBoneIndices))
+				{
+					for (const int32 IndexOfBoneInChain : OutChainBoneIndices)
+					{
+						OutBoneColors[IndexOfBoneInChain] = HighlightedColor;
+					}
+				}
+			}
 		}
 		break;
 		
diff --git a/Engine/Plugins/Animation/IKRig/Source/IKRigEditor/Private/RigEditor/IKRigEditorController.cpp b/Engine/Plugins/Animation/IKRig/Source/IKRigEditor/Private/RigEditor/IKRigEditorController.cpp
--- a/Engine/Plugins/Animation/IKRig/Source/IKRigEditor/Private/RigEditor/IKRigEditorController.cpp
+++ b/Engine/Plugins/Animation/IKRig/Source/IKRigEditor/Private/RigEditor/IKRigEditorController.cpp
@@ -179,6 +179,114 @@ void UIKRigBoneDetails::OnPasteFromClipboard(ESlateTransformComponent::Type Comp
 
 #endif
 
+void FRetargetChainAnalyzer::AssignBestGuessName(FBoneChain& Chain, const FIKRigSkeleton& IKRigSkeleton)
+{
+	// a map of common names used in bone hierarchies, used to derive a best guess name for a retarget chain
+	static TMap<FString, TArray<FString>> ChainToBonesMap;
+	ChainToBonesMap.Add(FString("Head"), {"head"});
+	ChainToBonesMap.Add(FString("Neck"), {"neck"});
+	ChainToBonesMap.Add(FString("Leg"), {"leg", "hip", "thigh", "calf", "knee", "foot", "ankle", "toe"});
+	ChainToBonesMap.Add(FString("Arm"), {"arm", "clavicle", "shoulder", "elbow", "wrist", "hand"});
+	ChainToBonesMap.Add(FString("Spine"), {"spine"});
+	
+	ChainToBonesMap.Add(FString("Jaw"),{"jaw"});
+	ChainToBonesMap.Add(FString("Tail"),{"tail", "tentacle"});
+	
+	ChainToBonesMap.Add(FString("Thumb"), {"thumb"});
+	ChainToBonesMap.Add(FString("Index"), {"index"});
+	ChainToBonesMap.Add(FString("Middle"), {"middle"});
+	ChainToBonesMap.Add(FString("Ring"), {"ring"});
+	ChainToBonesMap.Add(FString("Pinky"), {"pinky"});
+
+	TArray<int32> BonesInChainIndices;
+	const bool bIsChainValid = IKRigSkeleton.GetBonesInChain(Chain, BonesInChainIndices);
+	if (!bIsChainValid)
+	{
+		Chain.ChainName = GetDefaultChainName();
+		return;
+	}
+
+	// initialize map of chain names with initial scores of 0 for each chain
+	TArray<FString> ChainNames;
+	ChainToBonesMap.GetKeys(ChainNames);
+	TMap<FString, float> ChainScores;
+	for (const FString& ChainName : ChainNames)
+	{
+		ChainScores.Add(*ChainName, 0.f);
+	}
+
+	// run text filter on the predefined bone names and record score for each predefined chain name
+	for (const int32 ChainBoneIndex : BonesInChainIndices)
+	{
+		const FName& ChainBoneName = IKRigSkeleton.GetBoneNameFromIndex(ChainBoneIndex);
+		const FString ChainBoneNameStr = ChainBoneName.ToString().ToLower();
+		for (const FString& ChainName : ChainNames)
+		{
+			for (const FString& BoneNameToTry : ChainToBonesMap[ChainName])
+			{
+				TextFilter->SetFilterText(FText::FromString(BoneNameToTry));
+				if (TextFilter->TestTextFilter(FBasicStringFilterExpressionContext(ChainBoneNameStr)))
+				{
+					++ChainScores[ChainName];	
+				}
+			}
+		}
+	}
+
+	// find the chain name with the highest score
+	float HighestScore = 0.f;
+	FString RecommendedChainName = GetDefaultChainName().ToString();
+	for (const FString& ChainName : ChainNames)
+	{
+		if (ChainScores[ChainName] > HighestScore)
+		{
+			HighestScore = ChainScores[ChainName];
+			RecommendedChainName = ChainName;
+		}
+	}
+
+	// now determine "sidedness" of the chain, and add a suffix if its on the left or right side
+	const EChainSide ChainSide = GetSideOfChain(BonesInChainIndices, IKRigSkeleton);
+	switch (ChainSide)
+	{
+	case EChainSide::Left:
+		RecommendedChainName = "Left" + RecommendedChainName;
+		break;
+	case EChainSide::Right:
+		RecommendedChainName = "Right" + RecommendedChainName;
+		break;
+	default:
+		break;
+	}
+	
+	Chain.ChainName = FName(RecommendedChainName);
+}
+
+FName FRetargetChainAnalyzer::GetDefaultChainName()
+{
+	static FText NewChainText = LOCTEXT("NewRetargetChainLabel", "NewRetargetChain");
+	return FName(*NewChainText.ToString());
+}
+
+EChainSide FRetargetChainAnalyzer::GetSideOfChain(const TArray<int32>& BoneIndices, const FIKRigSkeleton& IKRigSkeleton)
+{
+	// determine "sidedness" of the chain based on the location of the bones (left, right or center of YZ plane)
+	float AverageXPositionOfChain = 0.f;
+	for (const int32 BoneIndex : BoneIndices)
+	{
+		AverageXPositionOfChain += IKRigSkeleton.RefPoseGlobal[BoneIndex].GetLocation().X;
+	}
+	AverageXPositionOfChain /= static_cast<float>(BoneIndices.Num());
+
+	constexpr float CenterThresholdDistance = 1.0f;
+	if (FMath::Abs(AverageXPositionOfChain) < CenterThresholdDistance)
+	{
+		return EChainSide::Center;
+	}
+
+	return AverageXPositionOfChain > 0 ? EChainSide::Left :  EChainSide::Right;
+}
+
 void FIKRigEditorController::Initialize(TSharedPtr<FIKRigEditorToolkit> Toolkit, UIKRigDefinition* IKRigAsset)
 {
 	EditorToolkit = Toolkit;
@@ -188,13 +296,19 @@ void FIKRigEditorController::Initialize(TSharedPtr<FIKRigEditorToolkit> Toolkit,
 	// register callback to be informed when rig asset is modified by editor
 	if (!AssetController->OnIKRigNeedsInitialized().IsBoundToObject(this))
 	{
-		AssetController->OnIKRigNeedsInitialized().AddSP(this, &FIKRigEditorController::OnIKRigNeedsInitialized);
+		ReinitializeDelegateHandle = AssetController->OnIKRigNeedsInitialized().AddSP(
+			this, &FIKRigEditorController::OnIKRigNeedsInitialized);
 
 		// Initialize editor's instances at first initialization
 		InitializeSolvers();
 	}
 }
 
+void FIKRigEditorController::Close() const
+{
+	AssetController->OnIKRigNeedsInitialized().Remove(ReinitializeDelegateHandle);
+}
+
 UIKRigProcessor* FIKRigEditorController::GetIKRigProcessor() const
 {
 	if (AnimInstance.IsValid())
@@ -311,6 +425,9 @@ void FIKRigEditorController::AddNewGoals(const TArray<FName>& GoalNames, const T
 			AssetController->ConnectGoalToSolver(*NewGoal, SolverElement->IndexInStack);	
 		}
 
+		// ask user if they want to assign this goal to a chain (if there is one on this bone)
+		PromptToAddGoalToChain(NewGoal);
+
 		LastCreatedGoalName = GoalName;
 	}
 	
@@ -433,14 +550,14 @@ bool FIKRigEditorController::IsGoalSelected(const FName& GoalName) const
 	return SkeletonView->IsGoalSelected(GoalName);
 }
 
-FName FIKRigEditorController::GetSelectedChain() const
+TArray<FName> FIKRigEditorController::GetSelectedChains() const
 {
 	if (!RetargetingView.IsValid())
 	{
-		return NAME_None;
+		return TArray<FName>();
 	}
 
-	return RetargetingView->GetSelectedChain();
+	return RetargetingView->GetSelectedChains();
 }
 
 bool FIKRigEditorController::DoesSkeletonHaveSelectedItems() const
@@ -452,6 +569,46 @@ bool FIKRigEditorController::DoesSkeletonHaveSelectedItems() const
 	return SkeletonView->HasSelectedItems();
 }
 
+void FIKRigEditorController::GetChainsSelectedInSkeletonView(TArray<FBoneChain>& InOutChains)
+{
+	if (!SkeletonView.IsValid())
+	{
+		return;
+	}
+
+	SkeletonView->GetSelectedBoneChains(InOutChains);
+}
+
+void FIKRigEditorController::CreateNewRetargetChains()
+{
+	// get selected chains from hierarchy view
+	TArray<FBoneChain> SelectedBoneChains;
+	if (SkeletonView.IsValid())
+	{
+		SkeletonView->GetSelectedBoneChains(SelectedBoneChains);
+	}
+
+	const FIKRigSkeleton& IKRigSkeleton = AssetController->GetIKRigSkeleton();
+	
+	if (!SelectedBoneChains.IsEmpty())
+	{
+		// create a chain for each selected chain in hierarchy
+		for (FBoneChain& BoneChain : SelectedBoneChains)
+		{
+			ChainAnalyzer.AssignBestGuessName(BoneChain, IKRigSkeleton);
+			PromptToAddNewRetargetChain(BoneChain);
+		}
+	}
+	else
+	{
+		// create an empty chain
+		const FBoneChain Chain(FRetargetChainAnalyzer::GetDefaultChainName(), NAME_None, NAME_None);
+		PromptToAddNewRetargetChain(Chain);
+	}
+	
+	RefreshAllViews();
+}
+
 bool FIKRigEditorController::PromptToAddSolver() const
 {
 	if (AssetController->GetNumSolvers() > 0)
@@ -498,7 +655,9 @@ bool FIKRigEditorController::PromptToAddSolver() const
 	})
 	.Content()
 	[
-		SNew(STextBlock).Text_Lambda([&SelectedSolver]()
+		SNew(STextBlock)
+		.MinDesiredWidth(200)
+		.Text_Lambda([&SelectedSolver]()
 			{
 			return SelectedSolver->NiceName;
 		})
@@ -528,6 +687,43 @@ bool FIKRigEditorController::PromptToAddSolver() const
 	return true;
 }
 
+void FIKRigEditorController::PromptToAddGoalToChain(UIKRigEffectorGoal* NewGoal) const
+{
+	const TArray<FBoneChain>& AllRetargetChains = AssetController->GetRetargetChains();
+	FName ChainToAddGoalTo = NAME_None;
+	for (const FBoneChain& Chain : AllRetargetChains)
+	{
+		if (Chain.EndBone == NewGoal->BoneName)
+		{
+			ChainToAddGoalTo = Chain.ChainName;
+		}
+	}
+
+	if (ChainToAddGoalTo == NAME_None)
+	{
+		return;
+	}
+
+	TSharedRef<SCustomDialog> AddGoalToChainDialog = SNew(SCustomDialog)
+		.Title(FText(LOCTEXT("AddGoalToNewChainTitle", "Add Goal to Retarget Chain")))
+		.Content()
+		[
+			SNew(STextBlock)
+			.Text(FText::Format(LOCTEXT("AddNewGoalToChainLabel", "Add new goal, {0} to retarget chain, {1}?"), FText::FromName(NewGoal->GoalName), FText::FromName(ChainToAddGoalTo)))
+		]
+		.Buttons({
+			SCustomDialog::FButton(LOCTEXT("OK", "OK")),
+			SCustomDialog::FButton(LOCTEXT("Cancel", "Cancel"))
+	});
+
+	if (AddGoalToChainDialog->ShowModal() != 0)
+	{
+		return; // cancel button pressed, or window closed
+	}
+
+	AssetController->SetRetargetChainGoal(ChainToAddGoalTo, NewGoal->GoalName);
+}
+
 bool FIKRigEditorController::IsElementConnectedToSolver(TSharedRef<FIKRigTreeElement> TreeElement, int32 SolverIndex)
 {
 	if (!AssetController->GetSolverArray().IsValidIndex(SolverIndex))
@@ -557,7 +753,7 @@ bool FIKRigEditorController::IsElementConnectedToSolver(TSharedRef<FIKRigTreeEle
 	if (TreeElement->ElementType == IKRigTreeElementType::SOLVERGOAL)
 	{
 		// is this an effector for this solver?
-		return TreeElement->SolverGoalIndex == SolverIndex;
+		return TreeElement->EffectorIndex == SolverIndex;
 	}
 
 	checkNoEntry();
@@ -674,7 +870,7 @@ void FIKRigEditorController::ShowDetailsForElements(const TArray<TSharedPtr<FIKR
 		break;
 		
 	case IKRigTreeElementType::SOLVERGOAL:
-		ShowDetailsForGoalSettings(LastItem->SolverGoalName, LastItem->SolverGoalIndex);
+		ShowDetailsForGoalSettings(LastItem->EffectorGoalName, LastItem->EffectorIndex);
 		break;
 		
 	case IKRigTreeElementType::BONE_SETTINGS:
@@ -724,18 +920,17 @@ void FIKRigEditorController::SetDetailsView(const TSharedPtr<IDetailsView>& InDe
 	ShowEmptyDetails();
 }
 
-void FIKRigEditorController::AddNewRetargetChain(const FName ChainName, const FName StartBone, const FName EndBone)
+void FIKRigEditorController::PromptToAddNewRetargetChain(const FBoneChain& BoneChain) const
 {
-	FIKRigRetargetChainSettings Settings(ChainName, StartBone, EndBone);
-	const TSharedPtr<FStructOnScope> StructToDisplay = MakeShareable(new FStructOnScope(FIKRigRetargetChainSettings::StaticStruct(), (uint8*)&Settings));
+	const TSharedPtr<FStructOnScope> StructToDisplay = MakeShareable(new FStructOnScope(FIKRigRetargetChainSettings::StaticStruct(), (uint8*)&BoneChain));
 	TSharedRef<SKismetInspector> KismetInspector = SNew(SKismetInspector);
 	KismetInspector->ShowSingleStruct(StructToDisplay);
 
 	SGenericDialogWidget::FArguments DialogArguments;
-	DialogArguments.OnOkPressed_Lambda([&Settings, this] ()
+	DialogArguments.OnOkPressed_Lambda([&BoneChain, this] ()
 	{
 		// add the retarget chain
-		AssetController->AddRetargetChain(Settings.ChainName, Settings.StartBone, Settings.EndBone);
+		AssetController->AddRetargetChain(BoneChain);
 		RefreshAllViews();
 	});
 
@@ -746,6 +941,40 @@ void FIKRigEditorController::AddNewRetargetChain(const FName ChainName, const FN
 		true);
 }
 
+void FIKRigEditorController::PromptToAddGoalToNewlyMirroredChain(const FBoneChain& OldChain, const FBoneChain& MirroredChain) const
+{
+	const TSharedRef<SCustomDialog> AddGoalToChainDialog = SNew(SCustomDialog)
+		.Title(FText(LOCTEXT("AddGoalToMirroredChainLabel", "Mirror Goal")))
+		.Content()
+		[
+			SNew(STextBlock)
+			.Text(FText::Format(LOCTEXT("AddGoalToMirroredChainLabel", "Mirror goal, {0} on chain, {1}?"), FText::FromName(OldChain.IKGoalName), FText::FromName(OldChain.IKGoalName)))
+		]
+		.Buttons({
+			SCustomDialog::FButton(LOCTEXT("Ok", "OK")),
+			SCustomDialog::FButton(LOCTEXT("Cancel", "Cancel"))
+	});
+
+	if (AddGoalToChainDialog->ShowModal() != 0)
+	{
+		return; // cancel button pressed, or window closed
+	}
+
+	UIKRigEffectorGoal* OldGoal = AssetController->GetGoal(OldChain.IKGoalName);
+	const FName NewGoalName = FName(OldChain.IKGoalName.ToString() + "_Mirrored");
+	UIKRigEffectorGoal* NewGoal = AssetController->AddNewGoal(NewGoalName, MirroredChain.EndBone.BoneName);
+
+	// connect to the same solvers
+	const TArray<UIKRigSolver*>& AllSolvers =  AssetController->GetSolverArray();
+	for (int32 SolverIndex=0; SolverIndex<AllSolvers.Num(); ++SolverIndex)
+	{
+		if (AssetController->IsGoalConnectedToSolver(OldGoal->GoalName, SolverIndex))
+		{
+			AssetController->ConnectGoalToSolver(*NewGoal, SolverIndex);
+		}
+	}
+}
+
 void FIKRigEditorController::PlayAnimationAsset(UAnimationAsset* AssetToPlay)
 {
 	if (AssetToPlay && AnimInstance.IsValid())
diff --git a/Engine/Plugins/Animation/IKRig/Source/IKRigEditor/Private/RigEditor/IKRigToolkit.cpp b/Engine/Plugins/Animation/IKRig/Source/IKRigEditor/Private/RigEditor/IKRigToolkit.cpp
--- a/Engine/Plugins/Animation/IKRig/Source/IKRigEditor/Private/RigEditor/IKRigToolkit.cpp
+++ b/Engine/Plugins/Animation/IKRig/Source/IKRigEditor/Private/RigEditor/IKRigToolkit.cpp
@@ -86,6 +86,12 @@ void FIKRigEditorToolkit::InitAssetEditor(
 	RegenerateMenusAndToolbars();
 }
 
+void FIKRigEditorToolkit::OnClose()
+{
+	FPersonaAssetEditorToolkit::OnClose();
+	EditorController->Close();
+}
+
 void FIKRigEditorToolkit::RegisterTabSpawners(const TSharedRef<FTabManager>& InTabManager)
 {
 	WorkspaceMenuCategory = InTabManager->AddLocalWorkspaceMenuCategory(LOCTEXT("WorkspaceMenu_IKRigEditor", "IK Rig Editor"));
diff --git a/Engine/Plugins/Animation/IKRig/Source/IKRigEditor/Private/RigEditor/SIKRigRetargetChainList.cpp b/Engine/Plugins/Animation/IKRig/Source/IKRigEditor/Private/RigEditor/SIKRigRetargetChainList.cpp
--- a/Engine/Plugins/Animation/IKRig/Source/IKRigEditor/Private/RigEditor/SIKRigRetargetChainList.cpp
+++ b/Engine/Plugins/Animation/IKRig/Source/IKRigEditor/Private/RigEditor/SIKRigRetargetChainList.cpp
@@ -17,6 +17,7 @@
 #include "SSearchableComboBox.h"
 #include "BoneSelectionWidget.h"
 #include "Engine/SkeletalMesh.h"
+#include "Widgets/Input/SSearchBox.h"
 
 #define LOCTEXT_NAMESPACE "SIKRigRetargetChains"
 
@@ -291,6 +292,8 @@ void SIKRigRetargetChainList::Construct(const FArguments& InArgs, TSharedRef<FIK
 {
 	EditorController = InEditorController;
 	EditorController.Pin()->SetRetargetingView(SharedThis(this));
+
+	TextFilter = MakeShareable(new FTextFilterExpressionEvaluator(ETextFilterExpressionEvaluatorMode::BasicString));
 	
 	CommandList = MakeShared<FUICommandList>();
 
@@ -325,31 +328,61 @@ void SIKRigRetargetChainList::Construct(const FArguments& InArgs, TSharedRef<FIK
 				.IsReadOnly(true)
 			]
         ]
-        
-		+ SVerticalBox::Slot()
+
+        +SVerticalBox::Slot()
+		.Padding(2.0f)
 		.AutoHeight()
-		.VAlign(VAlign_Top)
-		.HAlign(HAlign_Left)
 		[
-			SNew(SPositiveActionButton)
-			.Icon(FAppStyle::Get().GetBrush("Icons.Plus"))
-			.Text(LOCTEXT("AddNewChainLabel", "Add New Chain"))
-			.ToolTipText(LOCTEXT("AddNewChainToolTip", "Add a new retarget bone chain."))
-			.OnClicked_Lambda([this]()
-			{
-				const UIKRigController* Controller = EditorController.Pin()->AssetController;
-				static FText NewChainText = LOCTEXT("NewRetargetChainLabel", "NewRetargetChain");
-				static FName NewChainName = FName(*NewChainText.ToString());
-				Controller->AddRetargetChain(NewChainName, NAME_None, NAME_None);
-				RefreshView();
-				return FReply::Handled();
-			})
+			SNew(SHorizontalBox)
+			+SHorizontalBox::Slot()
+			.AutoWidth()
+			.VAlign(VAlign_Center)
+			.Padding(FMargin(6.f, 0.0))
+			[
+				SNew(SPositiveActionButton)
+				.Icon(FAppStyle::Get().GetBrush("Icons.Plus"))
+				.Text(LOCTEXT("AddNewChainLabel", "Add New Chain"))
+				.ToolTipText(LOCTEXT("AddNewChainToolTip", "Add a new retarget bone chain."))
+				.OnClicked_Lambda([this]()
+				{
+					EditorController.Pin()->CreateNewRetargetChains();
+					return FReply::Handled();
+				})
+			]
+
+			+SHorizontalBox::Slot()
+			.FillWidth(1.0f)
+			[
+				SNew(SSearchBox)
+				.SelectAllTextWhenFocused(true)
+				.OnTextChanged( this, &SIKRigRetargetChainList::OnFilterTextChanged )
+				.HintText( LOCTEXT( "SearchBoxHint", "Filter Chain List...") )
+			]
+
+			+SHorizontalBox::Slot()
+			.AutoWidth()
+			.Padding(FMargin(6.f, 0.0))
+			.VAlign(VAlign_Center)
+			[
+				SNew(SComboButton)
+				.ComboButtonStyle(&FAppStyle::Get().GetWidgetStyle<FComboButtonStyle>("SimpleComboButton"))
+				.ForegroundColor(FSlateColor::UseStyle())
+				.ContentPadding(2.0f)
+				.OnGetMenuContent(this, &SIKRigRetargetChainList::CreateFilterMenuWidget)
+				.HasDownArrow(true)
+				.ButtonContent()
+				[
+					SNew(SImage)
+					.Image(FAppStyle::Get().GetBrush("Icons.Settings"))
+					.ColorAndOpacity(FSlateColor::UseForeground())
+				]
+			]
 		]
 
         +SVerticalBox::Slot()
 		[
 			SAssignNew(ListView, SRetargetChainListViewType )
-			.SelectionMode(ESelectionMode::Single)
+			.SelectionMode(ESelectionMode::Multi)
 			.IsEnabled(this, &SIKRigRetargetChainList::IsAddChainEnabled)
 			.ListItemsSource( &ListViewItems )
 			.OnGenerateRow( this, &SIKRigRetargetChainList::MakeListRowWidget )
@@ -380,14 +413,113 @@ void SIKRigRetargetChainList::Construct(const FArguments& InArgs, TSharedRef<FIK
 	RefreshView();
 }
 
-FName SIKRigRetargetChainList::GetSelectedChain()
+TArray<FName> SIKRigRetargetChainList::GetSelectedChains() const
 {
 	TArray<TSharedPtr<FRetargetChainElement>> SelectedItems = ListView->GetSelectedItems();
 	if (SelectedItems.IsEmpty())
 	{
-		return NAME_None;
+		return TArray<FName>();
+	}
+
+	TArray<FName> SelectedChainNames;
+	for (const TSharedPtr<FRetargetChainElement>& SelectedItem : SelectedItems)
+	{
+		SelectedChainNames.Add(SelectedItem->ChainName);
 	}
-	return SelectedItems[0].Get()->ChainName;
+	
+	return SelectedChainNames;
+}
+
+void SIKRigRetargetChainList::OnFilterTextChanged(const FText& SearchText)
+{
+	TextFilter->SetFilterText(SearchText);
+	RefreshView();
+}
+
+TSharedRef<SWidget> SIKRigRetargetChainList::CreateFilterMenuWidget()
+{
+	const FUIAction FilterSingleBoneAction = FUIAction(
+		FExecuteAction::CreateLambda([this]
+		{
+			ChainFilterOptions.bHideSingleBoneChains = !ChainFilterOptions.bHideSingleBoneChains;
+			RefreshView();
+		}),
+		FCanExecuteAction(),
+		FIsActionChecked::CreateLambda([this]()
+		{
+			return ChainFilterOptions.bHideSingleBoneChains;
+		}));
+
+	const FUIAction FilterIKChainAction = FUIAction(
+		FExecuteAction::CreateLambda([this]
+		{
+			ChainFilterOptions.bShowOnlyIKChains = !ChainFilterOptions.bShowOnlyIKChains;
+			RefreshView();
+		}),
+		FCanExecuteAction(),
+		FIsActionChecked::CreateLambda([this]()
+		{
+			return ChainFilterOptions.bShowOnlyIKChains;
+		}));
+
+	const FUIAction MissingBoneChainAction = FUIAction(
+		FExecuteAction::CreateLambda([this]
+		{
+			ChainFilterOptions.bShowOnlyMissingBoneChains = !ChainFilterOptions.bShowOnlyMissingBoneChains;
+			RefreshView();
+		}),
+		FCanExecuteAction(),
+		FIsActionChecked::CreateLambda([this]()
+		{
+			return ChainFilterOptions.bShowOnlyMissingBoneChains;
+		}));
+	
+	static constexpr bool CloseAfterSelection = true;
+	FMenuBuilder MenuBuilder(CloseAfterSelection, CommandList);
+
+	MenuBuilder.BeginSection("Chain Filters", LOCTEXT("ChainFiltersSection", "Filter"));
+
+	MenuBuilder.AddMenuEntry(
+		LOCTEXT("SingleBoneLabel", "Hide Single Bone Chains"),
+		LOCTEXT("SingleBoneTooltip", "Show only chains that contain multiple bones."),
+		FSlateIcon(),
+		FilterSingleBoneAction,
+		NAME_None,
+		EUserInterfaceActionType::Check);
+
+	MenuBuilder.AddMenuEntry(
+		LOCTEXT("HasIKLabel", "Show Only IK Chains"),
+		LOCTEXT("HasIKTooltip", "Show only chains that have an IK Goal."),
+		FSlateIcon(),
+		FilterIKChainAction,
+		NAME_None,
+		EUserInterfaceActionType::Check);
+
+	MenuBuilder.AddMenuEntry(
+		LOCTEXT("MissingBoneLabel", "Show Only Chains Missing Bones"),
+		LOCTEXT("MissingBoneTooltip", "Show only chains that are missing either a Start or End bone."),
+		FSlateIcon(),
+		MissingBoneChainAction,
+		NAME_None,
+		EUserInterfaceActionType::Check);
+	
+	MenuBuilder.EndSection();
+
+	MenuBuilder.BeginSection("Clear", LOCTEXT("ClearFiltersSection", "Clear"));
+	
+	MenuBuilder.AddMenuEntry(
+		LOCTEXT("ClearFilterLabel", "Clear Filters"),
+		LOCTEXT("ClearFilterTooltip", "Clear all filters to show all chains."),
+		FSlateIcon(),
+		FUIAction(FExecuteAction::CreateLambda([this]
+		{
+			ChainFilterOptions = FChainFilterOptions();
+			RefreshView();
+		})));
+
+	MenuBuilder.EndSection();
+	
+	return MenuBuilder.MakeWidget();
 }
 
 bool SIKRigRetargetChainList::IsAddChainEnabled() const
@@ -416,6 +548,11 @@ void SIKRigRetargetChainList::RefreshView()
 	{
 		return;
 	}
+
+	auto FilterString = [this](const FString& StringToTest) ->bool
+	{
+		return TextFilter->TestTextFilter(FBasicStringFilterExpressionContext(StringToTest));
+	};
 	
 	// refresh retarget root
 	RetargetRootTextBox.Get()->SetText(FText::FromName(Controller->AssetController->GetRetargetRoot()));
@@ -425,6 +562,37 @@ void SIKRigRetargetChainList::RefreshView()
 	const TArray<FBoneChain>& Chains = Controller->AssetController->GetRetargetChains();
 	for (const FBoneChain& Chain : Chains)
 	{
+		// apply text filter to items
+		if (!(TextFilter->GetFilterText().IsEmpty() ||
+			FilterString(Chain.ChainName.ToString()) ||
+			FilterString(Chain.StartBone.BoneName.ToString()) ||
+			FilterString(Chain.EndBone.BoneName.ToString()) ||
+			FilterString(Chain.IKGoalName.ToString())))
+		{
+			continue;
+		}
+
+		// apply single-bone filter
+		if (ChainFilterOptions.bHideSingleBoneChains &&
+			Chain.StartBone == Chain.EndBone)
+		{
+			continue;
+		}
+
+		// apply missing-bone filter
+		if (ChainFilterOptions.bShowOnlyMissingBoneChains &&
+			(Chain.StartBone.BoneName != NAME_None && Chain.EndBone.BoneName != NAME_None))
+		{
+			continue;
+		}
+		
+		// apply IK filter
+		if (ChainFilterOptions.bShowOnlyIKChains &&
+			Chain.IKGoalName == NAME_None)
+		{
+			continue;
+		}
+		
 		TSharedPtr<FRetargetChainElement> ChainItem = FRetargetChainElement::Make(Chain.ChainName);
 		ListViewItems.Add(ChainItem);
 	}
@@ -490,10 +658,19 @@ TSharedPtr<SWidget> SIKRigRetargetChainList::CreateContextMenu()
 
 	MenuBuilder.BeginSection("Chains", LOCTEXT("ChainsSection", "Chains"));
 
-	const FUIAction Action = FUIAction( FExecuteAction::CreateSP(this, &SIKRigRetargetChainList::SortChainList));
-	static const FText Label = LOCTEXT("SortChainsLabel", "Sort Chains");
-	static const FText Tooltip = LOCTEXT("SortChainsTooltip", "Sort chain list in hierarchical order. This does not affect the retargeting behavior.");
-	MenuBuilder.AddMenuEntry(Label, Tooltip, FSlateIcon(), Action);
+	MenuBuilder.AddMenuEntry(
+		LOCTEXT("MirrorChainsLabel", "Mirror Chain"),
+		LOCTEXT("MirrorChainsTooltip", "Create a new, duplicate chain on the opposite side of the skeleton."),
+		FSlateIcon(),
+		FUIAction( FExecuteAction::CreateSP(this, &SIKRigRetargetChainList::MirrorSelectedChains)));
+
+	MenuBuilder.AddSeparator();
+	
+	MenuBuilder.AddMenuEntry(
+		LOCTEXT("SortChainsLabel", "Sort Chains"),
+		LOCTEXT("SortChainsTooltip", "Sort chain list in hierarchical order. This does not affect the retargeting behavior."),
+		FSlateIcon(),
+		FUIAction( FExecuteAction::CreateSP(this, &SIKRigRetargetChainList::SortChainList)));
 
 	MenuBuilder.EndSection();
 	
@@ -515,4 +692,65 @@ void SIKRigRetargetChainList::SortChainList()
 	}
 }
 
+void SIKRigRetargetChainList::MirrorSelectedChains()
+{
+	const TArray<FName> SelectedChainNames = GetSelectedChains();
+	if (SelectedChainNames.IsEmpty())
+	{
+		return;
+	}
+
+	const FIKRigEditorController& Controller = *EditorController.Pin();
+	const UIKRigController* AssetController = Controller.AssetController;
+	const FIKRigSkeleton& IKRigSkeleton = AssetController->GetIKRigSkeleton();
+
+	for (const FName& SelectedChainName : SelectedChainNames)
+	{
+		const FBoneChain* Chain = AssetController->GetRetargetChainByName(SelectedChainName);
+		if (!Chain)
+		{
+			continue;
+		}
+		
+		TArray<int32> BonesInChainIndices;
+		const bool bIsChainValid = IKRigSkeleton.GetBonesInChain(*Chain, BonesInChainIndices);
+		if (!bIsChainValid)
+		{
+			continue;
+		}
+
+		const EChainSide ChainSide = Controller.ChainAnalyzer.GetSideOfChain(BonesInChainIndices, IKRigSkeleton);
+		if (ChainSide == EChainSide::Center)
+		{
+			continue;
+		}
+
+		TArray<int32> MirroredIndices;
+		if (!IKRigSkeleton.GetMirroredBoneIndices(BonesInChainIndices, MirroredIndices))
+		{
+			continue;
+		}
+
+		FBoneChain MirroredChain = *Chain;
+		MirroredChain.ChainName = FName(MirroredChain.ChainName.ToString() + "_Mirrored");
+		MirroredChain.StartBone = IKRigSkeleton.BoneNames[MirroredIndices[0]];
+		MirroredChain.EndBone = IKRigSkeleton.BoneNames[MirroredIndices.Last()];
+		const UIKRigEffectorGoal* GoalOnMirroredBone = AssetController->GetGoalForBone(MirroredChain.EndBone.BoneName);
+		if (GoalOnMirroredBone)
+		{
+			MirroredChain.IKGoalName = GoalOnMirroredBone->GoalName;
+		}
+
+		Controller.PromptToAddNewRetargetChain(MirroredChain);
+
+		// old bone DOES have a goal, but the mirrored end bone does NOT have a goal,
+		// so lets ask the user if they want to add one...
+		const UIKRigEffectorGoal* GoalOnOldBone = AssetController->GetGoalForBone(Chain->EndBone.BoneName);
+		if (GoalOnOldBone && !GoalOnMirroredBone)
+		{
+			Controller.PromptToAddGoalToNewlyMirroredChain(*Chain, MirroredChain);
+		}
+	}
+}
+
 #undef LOCTEXT_NAMESPACE
diff --git a/Engine/Plugins/Animation/IKRig/Source/IKRigEditor/Private/RigEditor/SIKRigSkeleton.cpp b/Engine/Plugins/Animation/IKRig/Source/IKRigEditor/Private/RigEditor/SIKRigSkeleton.cpp
--- a/Engine/Plugins/Animation/IKRig/Source/IKRigEditor/Private/RigEditor/SIKRigSkeleton.cpp
+++ b/Engine/Plugins/Animation/IKRig/Source/IKRigEditor/Private/RigEditor/SIKRigSkeleton.cpp
@@ -17,6 +17,7 @@
 #include "Framework/MultiBox/MultiBoxBuilder.h"
 #include "Framework/Commands/UICommandList.h"
 #include "Preferences/PersonaOptions.h"
+#include "Widgets/Input/SSearchBox.h"
 
 #define LOCTEXT_NAMESPACE "SIKRigSkeleton"
 
@@ -72,9 +73,9 @@ TWeakObjectPtr< UObject > FIKRigTreeElement::GetObject() const
 		return Controller->AssetController->GetGoal(GoalName);
 		break;
 	case IKRigTreeElementType::SOLVERGOAL:
-		if (const UIKRigSolver* SolverWithEffector = Controller->AssetController->GetSolver(SolverGoalIndex))
+		if (const UIKRigSolver* SolverWithEffector = Controller->AssetController->GetSolver(EffectorIndex))
 		{
-			return SolverWithEffector->GetGoalSettings(SolverGoalName);
+			return SolverWithEffector->GetGoalSettings(EffectorGoalName);
 		}
 		break;
 	default:
@@ -281,6 +282,7 @@ void SIKRigSkeleton::Construct(const FArguments& InArgs, TSharedRef<FIKRigEditor
 {
 	EditorController = InEditorController;
 	EditorController.Pin()->SetSkeletonsView(SharedThis(this));
+	TextFilter = MakeShareable(new FTextFilterExpressionEvaluator(ETextFilterExpressionEvaluatorMode::BasicString));
 	CommandList = MakeShared<FUICommandList>();
 	BindCommands();
 	
@@ -288,6 +290,51 @@ void SIKRigSkeleton::Construct(const FArguments& InArgs, TSharedRef<FIKRigEditor
     [
         SNew(SVerticalBox)
 
+        +SVerticalBox::Slot()
+		.Padding(2.0f)
+		.AutoHeight()
+		[
+			SNew(SHorizontalBox)
+			+SHorizontalBox::Slot()
+			.AutoWidth()
+			.VAlign(VAlign_Center)
+			.Padding(FMargin(6.f, 0.0))
+			[
+				SNew(SPositiveActionButton)
+				.OnGetMenuContent(this, &SIKRigSkeleton::CreateAddNewMenu)
+				.Icon(FAppStyle::Get().GetBrush("Icons.Plus"))
+			]
+
+			+SHorizontalBox::Slot()
+			.FillWidth(1.0f)
+			[
+				SNew(SSearchBox)
+				.SelectAllTextWhenFocused(true)
+				.OnTextChanged( this, &SIKRigSkeleton::OnFilterTextChanged )
+				.HintText(LOCTEXT( "SearchBoxHint", "Filter Hierarchy Tree..."))
+			]
+
+			+SHorizontalBox::Slot()
+			.AutoWidth()
+			.Padding(FMargin(6.f, 0.0))
+			.VAlign(VAlign_Center)
+			[
+				SNew(SComboButton)
+				.ComboButtonStyle(&FAppStyle::Get().GetWidgetStyle<FComboButtonStyle>("SimpleComboButton"))
+				.ForegroundColor(FSlateColor::UseStyle())
+				.ContentPadding(2.0f)
+				.OnGetMenuContent( this, &SIKRigSkeleton::CreateFilterMenuWidget )
+				.ToolTipText(LOCTEXT( "FilterMenuLabel", "Filter hierarchy tree options."))
+				.HasDownArrow(true)
+				.ButtonContent()
+				[
+					SNew(SImage)
+					.Image(FAppStyle::Get().GetBrush("Icons.Settings"))
+					.ColorAndOpacity(FSlateColor::UseForeground())
+				]
+			]
+		]
+
         +SVerticalBox::Slot()
         .Padding(0.0f, 0.0f)
         [
@@ -418,7 +465,7 @@ void SIKRigSkeleton::ReplaceItemInSelection(const FText& OldName, const FText& N
 	}
 }
 
-void SIKRigSkeleton::GetSelectedBoneChains(TArray<FIKRigSkeletonChain>& OutChains)
+void SIKRigSkeleton::GetSelectedBoneChains(TArray<FBoneChain>& OutChains)
 {
 	const TSharedPtr<FIKRigEditorController> Controller = EditorController.Pin();
 	if (!Controller.IsValid())
@@ -426,12 +473,12 @@ void SIKRigSkeleton::GetSelectedBoneChains(TArray<FIKRigSkeletonChain>& OutChain
 		return; 
 	}
 
+	const FIKRigSkeleton& Skeleton = Controller->AssetController->GetIKRigSkeleton();
+
 	// get selected bones
 	TArray<TSharedPtr<FIKRigTreeElement>> SelectedBoneItems;
 	GetSelectedBones(SelectedBoneItems);
 
-	const FIKRigSkeleton& Skeleton = Controller->AssetController->GetIKRigSkeleton();
-
 	// get selected bone indices
 	TArray<int32> SelectedBones;
 	for (const TSharedPtr<FIKRigTreeElement>& BoneItem : SelectedBoneItems)
@@ -440,8 +487,18 @@ void SIKRigSkeleton::GetSelectedBoneChains(TArray<FIKRigSkeletonChain>& OutChain
 		const int32 BoneIndex = Skeleton.GetBoneIndexFromName(BoneName);
 		SelectedBones.Add(BoneIndex);
 	}
+
+	// get bones in chains
+	Skeleton.GetChainsInList(SelectedBones, OutChains);
 	
-	return Skeleton.GetChainsInList(SelectedBones, OutChains);
+	// add goals (if there are any on the end bone of the chain)
+	for (FBoneChain& Chain : OutChains)
+	{
+		if (const UIKRigEffectorGoal* Goal = Controller->AssetController->GetGoalForBone(Chain.EndBone.BoneName))
+		{
+			Chain.IKGoalName = Goal->GoalName;
+		}
+	}
 }
 
 TArray<TSharedPtr<FIKRigTreeElement>> SIKRigSkeleton::GetSelectedItems() const
@@ -524,7 +581,12 @@ void SIKRigSkeleton::FillContextMenu(FMenuBuilder& MenuBuilder)
 	const TArray<TSharedPtr<FIKRigTreeElement>> SelectedItems = TreeView->GetSelectedItems();
 	if (SelectedItems.IsEmpty())
 	{
-		return;
+		MenuBuilder.AddWidget(
+			SNew(SBox)
+			.HAlign(HAlign_Center)
+			[
+				SNew(STextBlock).Text(LOCTEXT("InvalidSelection", "Nothing selected."))
+			], FText(), false);
 	}
 
 	MenuBuilder.BeginSection("AddRemoveGoals", LOCTEXT("AddRemoveGoalOperations", "Goals"));
@@ -624,7 +686,7 @@ void SIKRigSkeleton::HandleDeleteElement()
 				Controller->AssetController->RemoveGoal(SelectedItem->GoalName);
 				break;
 			case IKRigTreeElementType::SOLVERGOAL:
-				Controller->AssetController->DisconnectGoalFromSolver(SelectedItem->SolverGoalName, SelectedItem->SolverGoalIndex);
+				Controller->AssetController->DisconnectGoalFromSolver(SelectedItem->EffectorGoalName, SelectedItem->EffectorIndex);
 				break;
 			case IKRigTreeElementType::BONE_SETTINGS:
 				Controller->AssetController->RemoveBoneSetting(SelectedItem->BoneSettingBoneName, SelectedItem->BoneSettingsSolverIndex);
@@ -1108,15 +1170,8 @@ void SIKRigSkeleton::HandleNewRetargetChain()
 	{
 		return;
 	}
-	
-	TArray<FIKRigSkeletonChain> BoneChains;
-	GetSelectedBoneChains(BoneChains);
-	for (const FIKRigSkeletonChain& BoneChain : BoneChains)
-	{
-		Controller->AddNewRetargetChain(BoneChain.StartBone, BoneChain.StartBone, BoneChain.EndBone);
-	}
-	
-	Controller->RefreshAllViews();
+
+	Controller->CreateNewRetargetChains();
 }
 
 bool SIKRigSkeleton::CanAddNewRetargetChain()
@@ -1280,6 +1335,142 @@ bool SIKRigSkeleton::CanRenameGoal() const
 	return SelectedGoals.Num() == 1;
 }
 
+TSharedRef<SWidget> SIKRigSkeleton::CreateFilterMenuWidget()
+{
+	const FUIAction FilterShowAllAction = FUIAction(
+		FExecuteAction::CreateLambda([this]
+		{
+			FilterOptions.ResetShowOptions();
+			FilterOptions.bShowAll = true; 
+			RefreshTreeView();
+		}),
+		FCanExecuteAction(),
+		FIsActionChecked::CreateLambda([this]()
+		{
+			return FilterOptions.bShowAll;
+		}));
+	
+	const FUIAction FilterOnlyBonesAction = FUIAction(
+		FExecuteAction::CreateLambda([this]
+		{
+			FilterOptions.ResetShowOptions();
+			FilterOptions.bShowOnlyBones = true;
+			RefreshTreeView();
+		}),
+		FCanExecuteAction(),
+		FIsActionChecked::CreateLambda([this]()
+		{
+			return FilterOptions.bShowOnlyBones;
+		}));
+
+	const FUIAction FilterOnlyGoalsAction = FUIAction(
+		FExecuteAction::CreateLambda([this]
+		{
+			FilterOptions.ResetShowOptions();
+			FilterOptions.bShowOnlyGoals = true;
+			RefreshTreeView();
+		}),
+		FCanExecuteAction(),
+		FIsActionChecked::CreateLambda([this]()
+		{
+			return FilterOptions.bShowOnlyGoals;
+		}));
+
+	const FUIAction FilterOnlyBoneSettingsAction = FUIAction(
+		FExecuteAction::CreateLambda([this]
+		{
+			FilterOptions.ResetShowOptions();
+			FilterOptions.bShowOnlyBoneSettings = true;
+			RefreshTreeView();
+		}),
+		FCanExecuteAction(),
+		FIsActionChecked::CreateLambda([this]()
+		{
+			return FilterOptions.bShowOnlyBoneSettings;
+		}));
+
+	const FUIAction FilterUnaffectedBonesAction = FUIAction(
+		FExecuteAction::CreateLambda([this]
+		{
+			FilterOptions.bHideUnaffectedBones = !FilterOptions.bHideUnaffectedBones;
+			RefreshTreeView();
+		}),
+		FCanExecuteAction(),
+		FIsActionChecked::CreateLambda([this]()
+		{
+			return FilterOptions.bHideUnaffectedBones;
+		}));
+	
+	static constexpr bool CloseAfterSelection = true;
+	FMenuBuilder MenuBuilder(CloseAfterSelection, CommandList);
+
+	MenuBuilder.BeginSection("Hierarchy Filters", LOCTEXT("HierarchyFiltersSection", "Hierarchy Filters"));
+
+	MenuBuilder.AddMenuEntry(
+		LOCTEXT("OnlyBonesLabel", "Show All Types"),
+		LOCTEXT("OnlyBonesTooltip", "Show all object types in the hierarchy."),
+		FSlateIcon(),
+		FilterShowAllAction,
+		NAME_None,
+		EUserInterfaceActionType::RadioButton);
+
+	MenuBuilder.AddMenuEntry(
+		LOCTEXT("OnlyBonesLabel", "Show Only Bones"),
+		LOCTEXT("OnlyBonesTooltip", "Hide everything except bones."),
+		FSlateIcon(),
+		FilterOnlyBonesAction,
+		NAME_None,
+		EUserInterfaceActionType::RadioButton);
+
+	MenuBuilder.AddMenuEntry(
+		LOCTEXT("OnlyGoalsLabel", "Show Only Goals"),
+		LOCTEXT("OnlyGoalsTooltip", "Hide everything except IK goals."),
+		FSlateIcon(),
+		FilterOnlyGoalsAction,
+		NAME_None,
+		EUserInterfaceActionType::RadioButton);
+
+	MenuBuilder.AddMenuEntry(
+		LOCTEXT("OnlyBoneSettingsLabel", "Show Only Bone Settings"),
+		LOCTEXT("OnlyBoneSettingsTooltip", "Hide everything except bone settings."),
+		FSlateIcon(),
+		FilterOnlyBoneSettingsAction,
+		NAME_None,
+		EUserInterfaceActionType::RadioButton);
+
+	MenuBuilder.AddMenuEntry(
+		LOCTEXT("UnaffectedBonesLabel", "Hide Unaffected Bones"),
+		LOCTEXT("UnaffectedBonesTooltip", "Hide all bones that are unaffected by any solvers."),
+		FSlateIcon(),
+		FilterUnaffectedBonesAction,
+		NAME_None,
+		EUserInterfaceActionType::Check);
+	
+	MenuBuilder.EndSection();
+
+	MenuBuilder.BeginSection("Clear", LOCTEXT("ClearMapFiltersSection", "Clear"));
+	
+	MenuBuilder.AddMenuEntry(
+		LOCTEXT("ClearMapFilterLabel", "Clear Filters"),
+		LOCTEXT("ClearMapFilterTooltip", "Clear all filters to show all hierarchy elements."),
+		FSlateIcon(),
+		FUIAction(FExecuteAction::CreateLambda([this]
+		{
+			FilterOptions = FIKRigHierarchyFilterOptions();
+			RefreshTreeView();
+		})));
+
+	MenuBuilder.EndSection();
+	
+	return MenuBuilder.MakeWidget();
+}
+
+void SIKRigSkeleton::OnFilterTextChanged(const FText& SearchText)
+{
+	TextFilter->SetFilterText(SearchText);
+	RefreshTreeView();
+}
+
 void SIKRigSkeleton::RefreshTreeView(bool IsInitialSetup /*=false*/)
 {
 	const TSharedPtr<FIKRigEditorController> Controller = EditorController.Pin();
@@ -1314,30 +1505,77 @@ void SIKRigSkeleton::RefreshTreeView(bool IsInitialSetup /*=false*/)
 	// record bone element indices
 	TMap<FName, int32> BoneTreeElementIndices;
 
+	auto FilterString = [this](const FString& StringToTest) ->bool
+	{
+		return TextFilter->TestTextFilter(FBasicStringFilterExpressionContext(StringToTest));
+	};
+
 	// create all bone elements
-	for (const FName BoneName : Skeleton.BoneNames)
+	for (int32 BoneIndex=0; BoneIndex<Skeleton.BoneNames.Num(); ++BoneIndex)
 	{
 		// create "Bone" tree element for this bone
+		const FName BoneName = Skeleton.BoneNames[BoneIndex];
 		const FText BoneDisplayName = FText::FromName(BoneName);
 		TSharedPtr<FIKRigTreeElement> BoneElement = MakeShared<FIKRigTreeElement>(BoneDisplayName, IKRigTreeElementType::BONE, ControllerRef);
 		BoneElement.Get()->BoneName = BoneName;
 		const int32 BoneElementIndex = AllElements.Add(BoneElement);
 		BoneTreeElementIndices.Add(BoneName, BoneElementIndex);
 
+		// store pointer to parent (if there is one)
+		const int32 ParentIndex = Skeleton.ParentIndices[BoneIndex];
+		if (ParentIndex != INDEX_NONE)
+		{
+			// get parent tree element
+			const FName ParentBoneName = Skeleton.BoneNames[ParentIndex];
+			const TSharedPtr<FIKRigTreeElement> ParentBoneTreeElement = AllElements[BoneTreeElementIndices[ParentBoneName]];
+			BoneElement->UnFilteredParent = ParentBoneTreeElement;
+		}
+
+		// apply type filter to bones
+		if (FilterOptions.bShowOnlyBoneSettings || FilterOptions.bShowOnlyGoals)
+		{
+			BoneElement->bIsHidden = true;
+		}
+
+		if (FilterOptions.bHideUnaffectedBones)
+		{
+			
+		}
+		
+		// apply text filter to bones
+		if (!(TextFilter->GetFilterText().IsEmpty() || FilterString(BoneName.ToString())))
+		{
+			BoneElement->bIsHidden = true;
+		}
+
 		// create all "Bone Setting" tree elements for this bone
 		for (int32 SolverIndex=0; SolverIndex<Solvers.Num(); ++SolverIndex)
 		{
-			if (Solvers[SolverIndex]->GetBoneSetting(BoneName))
+			if (!Solvers[SolverIndex]->GetBoneSetting(BoneName))
 			{
-				const FText SolverDisplayName = FText::FromString(AssetController->GetSolverUniqueName(SolverIndex));
-				const FText BoneSettingDisplayName = FText::Format(LOCTEXT("BoneSettings", "{0} settings for {1}"), BoneDisplayName, SolverDisplayName);
-				TSharedPtr<FIKRigTreeElement> SettingsItem = MakeShared<FIKRigTreeElement>(BoneSettingDisplayName, IKRigTreeElementType::BONE_SETTINGS, ControllerRef);
-				SettingsItem->BoneSettingBoneName = BoneName;
-				SettingsItem->BoneSettingsSolverIndex = SolverIndex;
-				AllElements.Add(SettingsItem);
-				// store hierarchy pointers for item
-				BoneElement->Children.Add(SettingsItem);
-				SettingsItem->Parent = BoneElement;
+				continue;
+			}
+			
+			const FText SolverDisplayName = FText::FromString(AssetController->GetSolverUniqueName(SolverIndex));
+			const FText BoneSettingDisplayName = FText::Format(LOCTEXT("BoneSettings", "{0} settings for {1}"), BoneDisplayName, SolverDisplayName);
+			TSharedPtr<FIKRigTreeElement> SettingsItem = MakeShared<FIKRigTreeElement>(BoneSettingDisplayName, IKRigTreeElementType::BONE_SETTINGS, ControllerRef);
+			SettingsItem->BoneSettingBoneName = BoneName;
+			SettingsItem->BoneSettingsSolverIndex = SolverIndex;
+			AllElements.Add(SettingsItem);
+			
+			// store hierarchy pointer for item
+			SettingsItem->UnFilteredParent = BoneElement;
+
+			// apply type filter to bone settings
+			if (FilterOptions.bShowOnlyBones || FilterOptions.bShowOnlyGoals)
+			{
+				SettingsItem->bIsHidden = true;
+			}
+
+			// apply text filter to bone settings
+			if (!(TextFilter->GetFilterText().IsEmpty() || FilterString(BoneName.ToString())))
+			{
+				SettingsItem->bIsHidden = true;
 			}
 		}
 
@@ -1355,50 +1593,87 @@ void SIKRigSkeleton::RefreshTreeView(bool IsInitialSetup /*=false*/)
 			GoalItem->GoalName = Goal->GoalName;
 			AllElements.Add(GoalItem);
 
-			// store hierarchy pointers for goal
-			BoneElement->Children.Add(GoalItem);
-			GoalItem->Parent = BoneElement;
+			// apply type filter to goals
+			if (FilterOptions.bShowOnlyBones || FilterOptions.bShowOnlyBoneSettings)
+			{
+				GoalItem->bIsHidden = true;
+			}
+
+			// apply text filter to goals
+			if (!(TextFilter->GetFilterText().IsEmpty() ||
+				FilterString(Goal->GoalName.ToString()) ||
+				FilterString(Goal->BoneName.ToString())))
+			{
+				GoalItem->bIsHidden = true;
+			}
+
+			// store hierarchy pointer for goal
+			GoalItem->UnFilteredParent = BoneElement;
 
 			// add all solver settings connected to this goal
 			for (int32 SolverIndex=0; SolverIndex<Solvers.Num(); ++SolverIndex)
 			{
-				if (UObject* GoalSettings = AssetController->GetGoalSettingsForSolver(Goal->GoalName, SolverIndex))
+				UObject* GoalSettings = AssetController->GetGoalSettingsForSolver(Goal->GoalName, SolverIndex);
+				if (!GoalSettings)
 				{
-					// make new element for solver goal
-					const FText SolverDisplayName = FText::FromString(AssetController->GetSolverUniqueName(SolverIndex));
-					const FText SolverGoalDisplayName = FText::Format(LOCTEXT("GoalSettingsForSolver", "{0} settings for solver {1}"), FText::FromName(Goal->GoalName), SolverDisplayName);
-					TSharedPtr<FIKRigTreeElement> SolverGoalItem = MakeShared<FIKRigTreeElement>(SolverGoalDisplayName, IKRigTreeElementType::SOLVERGOAL, ControllerRef);
-					SolverGoalItem->SolverGoalIndex = SolverIndex;
-					SolverGoalItem->SolverGoalName = Goal->GoalName;
-					AllElements.Add(SolverGoalItem);
-					SolverGoalItem->Parent = GoalItem;
-					GoalItem->Children.Add(SolverGoalItem);
+					continue;
 				}
+				
+				// make new element for effector
+				const FText SolverDisplayName = FText::FromString(AssetController->GetSolverUniqueName(SolverIndex));
+				const FText EffectorDisplayName = FText::Format(LOCTEXT("GoalSettingsForSolver", "{0} settings for solver {1}"), FText::FromName(Goal->GoalName), SolverDisplayName);
+				TSharedPtr<FIKRigTreeElement> EffectorItem = MakeShared<FIKRigTreeElement>(EffectorDisplayName, IKRigTreeElementType::SOLVERGOAL, ControllerRef);
+				EffectorItem->EffectorIndex = SolverIndex;
+				EffectorItem->EffectorGoalName = Goal->GoalName;
+				AllElements.Add(EffectorItem);
+
+				// apply filter to effectors (treated as goals for filtering purposes)
+				EffectorItem->bIsHidden = GoalItem->bIsHidden;
+				
+				// store hierarchy pointer for effectors
+				EffectorItem->UnFilteredParent = GoalItem;
 			}
 		}
 	}
 
-	// store children/parent pointers on all bone elements
-	for (int32 BoneIndex=0; BoneIndex<Skeleton.BoneNames.Num(); ++BoneIndex)
+	// resolve parent/children pointers on all tree elements, taking into consideration the filter options
+	// (elements are parented to their nearest non-hidden/filtered parent element)
+	for (TSharedPtr<FIKRigTreeElement>& Element : AllElements)
 	{
-		const FName BoneName = Skeleton.BoneNames[BoneIndex];
-		const TSharedPtr<FIKRigTreeElement> BoneTreeElement = AllElements[BoneTreeElementIndices[BoneName]];
-		const int32 ParentIndex = Skeleton.ParentIndices[BoneIndex];
-		if (ParentIndex < 0)
+		if (Element->bIsHidden)
 		{
-			// store the root element
-			RootElements.Add(BoneTreeElement);
-			// has no parent, so skip storing parent pointer
 			continue;
 		}
+		
+		// find first parent that is not filtered
+		TSharedPtr<FIKRigTreeElement> ParentElement = Element->UnFilteredParent;
+		while (true)
+		{
+			if (!ParentElement.IsValid())
+			{
+				break;
+			}
 
-		// get parent tree element
-		const FName ParentBoneName = Skeleton.BoneNames[ParentIndex];
-		const TSharedPtr<FIKRigTreeElement> ParentBoneTreeElement = AllElements[BoneTreeElementIndices[ParentBoneName]];
-		// store pointer to child on parent
-		ParentBoneTreeElement->Children.Add(BoneTreeElement);
-		// store pointer to parent on child
-		BoneTreeElement->Parent = ParentBoneTreeElement;
+			if (!ParentElement->bIsHidden)
+			{
+				break;
+			}
+
+			ParentElement = ParentElement->UnFilteredParent;
+		}
+
+		if (ParentElement.IsValid())
+		{
+			// store pointer to child on parent
+			ParentElement->Children.Add(Element);
+			// store pointer to parent on child
+			Element->Parent = ParentElement;
+		}
+		else
+		{
+			// has no parent, store a root element
+			RootElements.Add(Element);
+		}
 	}
 
 	// expand all elements upon the initial construction of the tree
@@ -1450,9 +1725,17 @@ void SIKRigSkeleton::OnSelectionChanged(TSharedPtr<FIKRigTreeElement> Selection,
 	// NOTE: we may want to set the last selected item here
 }
 
+TSharedRef<SWidget> SIKRigSkeleton::CreateAddNewMenu()
+{
+	constexpr bool CloseAfterSelection = true;
+	FMenuBuilder MenuBuilder(CloseAfterSelection, CommandList);
+	FillContextMenu(MenuBuilder);
+	return MenuBuilder.MakeWidget();
+}
+
 TSharedPtr<SWidget> SIKRigSkeleton::CreateContextMenu()
 {
-	const bool CloseAfterSelection = true;
+	constexpr bool CloseAfterSelection = true;
 	FMenuBuilder MenuBuilder(CloseAfterSelection, CommandList);
 	FillContextMenu(MenuBuilder);
 	return MenuBuilder.MakeWidget();
diff --git a/Engine/Plugins/Animation/IKRig/Source/IKRigEditor/Public/IKRigDebugRendering.h b/Engine/Plugins/Animation/IKRig/Source/IKRigEditor/Public/IKRigDebugRendering.h
--- a/Engine/Plugins/Animation/IKRig/Source/IKRigEditor/Public/IKRigDebugRendering.h
+++ b/Engine/Plugins/Animation/IKRig/Source/IKRigEditor/Public/IKRigDebugRendering.h
@@ -7,34 +7,40 @@
 
 namespace IKRigDebugRendering
 {
-	static TArray BoxPoints =
+	// 8 vertices of cube
+	static TArray CubeVertices =
 	{
-		FVector(0.5f, 0.5f, 0.5f),
-		FVector(0.5f, -0.5f, 0.5f),
-		FVector(0.5f, -0.5f, 0.5f),
-		FVector(-0.5f, -0.5f, 0.5f),
-		FVector(-0.5f, -0.5f, 0.5f),
-		FVector(-0.5f, 0.5f, 0.5f),
-		FVector(-0.5f, 0.5f, 0.5f),
-		FVector(0.5f, 0.5f, 0.5f),
-
-		FVector(0.5f, 0.5f, -0.5f),
-		FVector(0.5f, -0.5f, -0.5f),
-		FVector(0.5f, -0.5f, -0.5f),
-		FVector(-0.5f, -0.5f, -0.5f),
-		FVector(-0.5f, -0.5f, -0.5f),
-		FVector(-0.5f, 0.5f, -0.5f),
 		FVector(-0.5f, 0.5f, -0.5f),
 		FVector(0.5f, 0.5f, -0.5f),
-
 		FVector(0.5f, 0.5f, 0.5f),
-		FVector(0.5f, 0.5f, -0.5f),
-		FVector(0.5f, -0.5f, 0.5f),
+		FVector(-0.5f, 0.5f, 0.5f),
+		
+		FVector(-0.5f, -0.5f, -0.5f),
 		FVector(0.5f, -0.5f, -0.5f),
+		FVector(0.5f, -0.5f, 0.5f),
 		FVector(-0.5f, -0.5f, 0.5f),
-		FVector(-0.5f, -0.5f, -0.5f),
-		FVector(-0.5f, 0.5f, 0.5f),
-		FVector(-0.5f, 0.5f, -0.5f),
+	};
+
+	// 12 edges of cube
+	static TArray CubeEdges =
+	{
+		// top
+		0,1,
+		1,2,
+		2,3,
+		3,0,
+
+		// bottom
+		4,5,
+		5,6,
+		6,7,
+		7,4,
+
+		// sides
+		0,4,
+		1,5,
+		2,6,
+		3,7
 	};
 
 	static void DrawWireCube(
@@ -44,10 +50,15 @@ namespace IKRigDebugRendering
 		float Size,
 		float Thickness)
 	{
-		const float Scale = FMath::Clamp(Size, 0.1f, 1000.0f);
-		for (int32 PointIndex = 0; PointIndex < BoxPoints.Num() - 1; PointIndex += 2)
+		const float Scale = FMath::Clamp(Size, 0.01f, 1000.0f);
+		for (int32 EdgeIndex = 0; EdgeIndex < CubeEdges.Num() - 1; EdgeIndex += 2)
 		{
-			PDI->DrawLine(Transform.TransformPosition(BoxPoints[PointIndex] * Scale), Transform.TransformPosition(BoxPoints[PointIndex + 1] * Scale), Color, SDPG_Foreground, Thickness);
+			PDI->DrawLine(
+				Transform.TransformPosition(CubeVertices[CubeEdges[EdgeIndex]] * Scale),
+				Transform.TransformPosition(CubeVertices[CubeEdges[EdgeIndex + 1]] * Scale),
+				Color,
+				SDPG_Foreground,
+				Thickness);
 		}
 	}
 }
diff --git a/Engine/Plugins/Animation/IKRig/Source/IKRigEditor/Public/RetargetEditor/AssetTypeActions_IKRetargeter.h b/Engine/Plugins/Animation/IKRig/Source/IKRigEditor/Public/RetargetEditor/AssetTypeActions_IKRetargeter.h
--- a/Engine/Plugins/Animation/IKRig/Source/IKRigEditor/Public/RetargetEditor/AssetTypeActions_IKRetargeter.h
+++ b/Engine/Plugins/Animation/IKRig/Source/IKRigEditor/Public/RetargetEditor/AssetTypeActions_IKRetargeter.h
@@ -15,7 +15,7 @@ public:
 	// IAssetTypeActions Implementation
 	virtual FText GetName() const override { return NSLOCTEXT("AssetTypeActions", "AssetTypeActions_IKRetargeter", "IK Retargeter"); }
 	virtual void GetActions(const TArray<UObject*>& InObjects, struct FToolMenuSection& Section) override;
-	virtual FColor GetTypeColor() const override { return FColor(0, 40, 255); }
+	virtual FColor GetTypeColor() const override { return FColor(0, 128, 255); }
 	virtual UClass* GetSupportedClass() const override;
 	virtual void OpenAssetEditor(const TArray<UObject*>& InObjects, TSharedPtr<class IToolkitHost> EditWithinLevelEditor = TSharedPtr<IToolkitHost>()) override;
 	virtual uint32 GetCategories() override { return EAssetTypeCategories::Animation; }
diff --git a/Engine/Plugins/Animation/IKRig/Source/IKRigEditor/Public/RetargetEditor/IKRetargetDetails.h b/Engine/Plugins/Animation/IKRig/Source/IKRigEditor/Public/RetargetEditor/IKRetargetDetails.h
--- a/Engine/Plugins/Animation/IKRig/Source/IKRigEditor/Public/RetargetEditor/IKRetargetDetails.h
+++ b/Engine/Plugins/Animation/IKRig/Source/IKRigEditor/Public/RetargetEditor/IKRetargetDetails.h
@@ -115,3 +115,78 @@ private:
 
 	TArray<UIKRetargetBoneDetails*> Bones;
 };
+
+/** ------------------------------------- BEGIN CHAIN DETAILS CUSTOMIZATION -------------*/
+
+
+class FRetargetChainSettingsCustomization : public IDetailCustomization
+{
+public:
+
+	static TSharedRef<IDetailCustomization> MakeInstance()
+	{
+		return MakeShareable(new FRetargetChainSettingsCustomization);
+	}
+
+	/** IDetailCustomization interface */
+	virtual void CustomizeDetails(IDetailLayoutBuilder& DetailBuilder) override;
+
+private:
+
+	void AddSettingsSection(
+		const IDetailLayoutBuilder& DetailBuilder,
+		IDetailCategoryBuilder& SettingsCategory,
+		const FString& StructPropertyName,
+		const FName& GroupName,
+		const FText& LocalizedGroupName,
+		const UScriptStruct* SettingsClass,
+		const FString& EnabledPropertyName,
+		const bool& bIsSectionEnabled,
+		const FText& DisabledMessage) const;
+	
+	TArray<TWeakObjectPtr<URetargetChainSettings>> ChainSettingsObjects;
+	TSharedPtr<FIKRetargetEditorController> Controller;
+	TArray<TSharedPtr<FString>> SourceChainOptions;
+};
+
+/** ------------------------------------- BEGIN ROOT DETAILS CUSTOMIZATION -------------*/
+
+class FRetargetRootSettingsCustomization : public IDetailCustomization
+{
+public:
+
+	static TSharedRef<IDetailCustomization> MakeInstance()
+	{
+		return MakeShareable(new FRetargetRootSettingsCustomization);
+	}
+
+	/** IDetailCustomization interface */
+	virtual void CustomizeDetails(IDetailLayoutBuilder& DetailBuilder) override;
+
+private:
+
+	TWeakObjectPtr<URetargetRootSettings> RootSettingsObject;
+	TSharedPtr<FIKRetargetEditorController> Controller;
+};
+
+/** ------------------------------------- BEGIN GLOBAL DETAILS CUSTOMIZATION -------------*/
+
+class FRetargetGlobalSettingsCustomization : public IDetailCustomization
+{
+public:
+
+	static TSharedRef<IDetailCustomization> MakeInstance()
+	{
+		return MakeShareable(new FRetargetGlobalSettingsCustomization);
+	}
+
+	/** IDetailCustomization interface */
+	virtual void CustomizeDetails(IDetailLayoutBuilder& DetailBuilder) override;
+
+private:
+
+	TWeakObjectPtr<UIKRetargetGlobalSettings> GlobalSettingsObject;
+	TSharedPtr<FIKRetargetEditorController> Controller;
+
+	TArray<TSharedPtr<FString>> TargetChainOptions;
+};
\ No newline at end of file
diff --git a/Engine/Plugins/Animation/IKRig/Source/IKRigEditor/Public/RetargetEditor/IKRetargetEditPoseMode.h b/Engine/Plugins/Animation/IKRig/Source/IKRigEditor/Public/RetargetEditor/IKRetargetEditPoseMode.h
--- a/Engine/Plugins/Animation/IKRig/Source/IKRigEditor/Public/RetargetEditor/IKRetargetEditPoseMode.h
+++ b/Engine/Plugins/Animation/IKRig/Source/IKRigEditor/Public/RetargetEditor/IKRetargetEditPoseMode.h
@@ -21,7 +21,7 @@ enum class FIKRetargetTrackingState : int8
 	TranslatingRoot,
 };
 
-struct BoneEdit
+struct FBoneEdit
 {
 	FName Name;							// name of last selected bone
 	int32 Index;						// index of last selected bone
@@ -32,9 +32,12 @@ struct BoneEdit
 	
 	TArray<FQuat> PreviousDeltaRotation;		// the prev stored local offsets of all selected bones
 
+	// translation edits
+	FVector AccumulatedPositionOffset;	// accumulated offset since start of edit
+
 	void Reset()
 	{
-		*this = BoneEdit();
+		*this = FBoneEdit();
 	}
 };
 
@@ -78,18 +81,12 @@ private:
 
 	// get the scale and offset associated with the currently edited skeletal mesh component
 	void GetEditedComponentScaleAndOffset(float& OutScale, FVector& OutOffset) const;
-
+	
 	// render skeleton in viewport
 	void RenderSkeleton(
 		FPrimitiveDrawInterface* PDI,
 		const FIKRetargetEditorController* Controller);
 
-	void GetSelectedAndAffectedBones(
-		const FIKRetargetEditorController* Controller,
-		const FRetargetSkeleton& Skeleton,
-		TSet<int32>& OutAffectedBones,
-		TSet<int32>& OutSelectedBones) const;
-
 	UE::Widget::EWidgetMode CurrentWidgetMode;
 
 	bool IsRootSelected() const;
@@ -97,7 +94,7 @@ private:
 	void UpdateWidgetTransform();
 
 	// the bone(s) currently being edited
-	BoneEdit BoneEdit;
+	FBoneEdit BoneEdit;
 
 	// the skeleton currently being edited
 	ERetargetSourceOrTarget SourceOrTarget;
diff --git a/Engine/Plugins/Animation/IKRig/Source/IKRigEditor/Public/RetargetEditor/IKRetargetEditor.h b/Engine/Plugins/Animation/IKRig/Source/IKRigEditor/Public/RetargetEditor/IKRetargetEditor.h
--- a/Engine/Plugins/Animation/IKRig/Source/IKRigEditor/Public/RetargetEditor/IKRetargetEditor.h
+++ b/Engine/Plugins/Animation/IKRig/Source/IKRigEditor/Public/RetargetEditor/IKRetargetEditor.h
@@ -45,6 +45,7 @@ public:
 		UIKRetargeter* Asset);
 
 	/** FAssetEditorToolkit interface */
+	virtual void OnClose() override;
 	virtual void RegisterTabSpawners(const TSharedRef<FTabManager>& InTabManager) override;
 	virtual void UnregisterTabSpawners(const TSharedRef<FTabManager>& InTabManager) override;
 	virtual FName GetToolkitFName() const override;
