Added profiles to the IK retargeter

diff --git a/Engine/Plugins/Animation/IKRig/Source/IKRigEditor/Public/RetargetEditor/IKRetargetEditorController.h b/Engine/Plugins/Animation/IKRig/Source/IKRigEditor/Public/RetargetEditor/IKRetargetEditorController.h
--- a/Engine/Plugins/Animation/IKRig/Source/IKRigEditor/Public/RetargetEditor/IKRetargetEditorController.h
+++ b/Engine/Plugins/Animation/IKRig/Source/IKRigEditor/Public/RetargetEditor/IKRetargetEditorController.h
@@ -32,7 +32,7 @@ enum class ERetargeterOutputMode : uint8
 	EditRetargetPose,	// allow editing the retarget pose
 };
 
-enum class EBoneSelectionEdit : uint8
+enum class ESelectionEdit : uint8
 {
 	Add,	// add to selection set
 	Remove,	// remove from selection
@@ -48,8 +48,10 @@ public:
 
 	/** Initialize the editor */
 	void Initialize(TSharedPtr<FIKRetargetEditor> InEditor, UIKRetargeter* InAsset);
+	/** Close the editor */
+	void Close() const;
 	/** Bind callbacks to this IK Rig */
-	void BindToIKRigAsset(UIKRigDefinition* InIKRig) const;
+	void BindToIKRigAsset(UIKRigDefinition* InIKRig);
 	/** callback when IK Rig asset requires reinitialization */
 	void OnIKRigNeedsInitialized(UIKRigDefinition* ModifiedIKRig) const;
 	/** callback when IK Rig asset's retarget chain has been renamed */
@@ -119,8 +121,6 @@ public:
 		const int32& BoneIndex,
 		const float& Scale,
 		const FVector& Offset) const;
-
-	FTransform GetTargetBoneLocalTransform(const UIKRetargetProcessor* RetargetProcessor, const int32& TargetBoneIndex) const;
 	
 	/** get world space positions of all immediate children of bone (with component scale / offset applied) */
 	static void GetGlobalRetargetPoseOfImmediateChildren(
@@ -138,6 +138,7 @@ public:
 
 	/** Sequence Browser**/
 	void PlayAnimationAsset(UAnimationAsset* AssetToPlay);
+	void StopPlayback();
 	void PausePlayback();
 	void ResumePlayback();
 	UAnimationAsset* AnimThatWasPlaying = nullptr;
@@ -157,27 +158,46 @@ public:
 	ERetargetSourceOrTarget GetSourceOrTarget() const { return CurrentlyEditingSourceOrTarget; };
 	void SetSourceOrTargetMode(ERetargetSourceOrTarget SourceOrTarget);
 
-	/** bone selection management (viewport or hierarchy view) */
+	/** ------------------------- SELECTION -----------------------------*/
+
+	/** SELECTION - BONES (viewport or hierarchy view) */
 	void EditBoneSelection(
 		const TArray<FName>& InBoneNames,
-		EBoneSelectionEdit EditMode,
+		ESelectionEdit EditMode,
 		const bool bFromHierarchyView);
-	void ClearSelection(const bool bKeepBoneSelection=false);
 	const TArray<FName>& GetSelectedBones() const {return SelectedBones; };
-	/** mesh selection management (viewport view) */
+	/* END bone selection */
+	
+	/** SELECTION - MESHES (viewport view) */
 	void SetSelectedMesh(UPrimitiveComponent* InComponent);
 	UPrimitiveComponent* GetSelectedMesh();
 	void AddOffsetToMeshComponent(const FVector& Offset, USceneComponent* MeshComponent) const;
+	/* END mesh selection */
+
+	/** SELECTION - CHAINS (viewport or chains view) */
+	void EditChainSelection(
+		const TArray<FName>& InChainNames,
+		ESelectionEdit EditMode,
+		const bool bFromChainsView);
+	const TArray<FName>& GetSelectedChains() const {return SelectedChains; };
+	/* END chain selection */
+
+	void SetRootSelected(const bool bIsSelected);
+	bool GetRootSelected() const { return bIsRootSelected; };
+	
+	void ClearSelection(const bool bKeepBoneSelection=false);
+
+	/** ------------------------- END SELECTION -----------------------------*/
 
 	/** determine if bone in the specified skeleton is part of the retarget (in a mapped chain) */
 	bool IsBoneRetargeted(const FName& BoneName, ERetargetSourceOrTarget SourceOrTarget) const;
 	/** get the name of the chain that contains this bone */
 	FName GetChainNameFromBone(const FName& BoneName, ERetargetSourceOrTarget SourceOrTarget) const;
 
-	/** factory to get/create bone details object */
-	TObjectPtr<UIKRetargetBoneDetails> GetDetailsObjectForBone(const FName& BoneName);
+	/** factory to get/create bone details UObject */
+	TObjectPtr<UIKRetargetBoneDetails> GetOrCreateBoneDetailsObject(const FName& BoneName);
 
-	/* START RETARGET POSES */
+	/** ------------------------- RETARGET POSES -----------------------------*/
 	
 	/** go to retarget pose */
 	FReply HandleShowRetargetPose();
@@ -222,7 +242,7 @@ public:
 	TSharedPtr<SWindow> RenamePoseWindow;
 	TSharedPtr<SEditableTextBox> NewNameEditableText;
 	
-	/* END RETARGET POSES */
+	/** ------------------------- END RETARGET POSES -----------------------------*/
 
 	/** FGCObject interface */
 	virtual void AddReferencedObjects(FReferenceCollector& Collector) override;
@@ -251,12 +271,19 @@ private:
 	ERetargetSourceOrTarget CurrentlyEditingSourceOrTarget;
 
 	/** current selection set */
+	bool bIsRootSelected;
+	UPrimitiveComponent* SelectedMesh;
+	TArray<FName> SelectedChains;
 	TArray<FName> SelectedBones;
 	UPROPERTY()
 	TMap<FName,TObjectPtr<UIKRetargetBoneDetails>> AllBoneDetails;
-	TArray<UObject*> SelectedBoneDetails;
-	TObjectPtr<UIKRetargetBoneDetails> CreateBoneDetails(const FName& InBoneName);
 
-	/** currently selected mesh */
-	UPrimitiveComponent* SelectedMesh;
+
+	/** to remove delegates when editor is closed */
+	FDelegateHandle ReInitDelegateHandle;
+	UIKRigDefinition* BoundToIKRig;
+	FDelegateHandle ReInitIKDelegateHandle;
+	FDelegateHandle RenameChainDelegateHandle;
+	FDelegateHandle RemoveChainDelegateHandle;
 };
+
diff --git a/Engine/Plugins/Animation/IKRig/Source/IKRigEditor/Public/RetargetEditor/IKRetargetEditorStyle.h b/Engine/Plugins/Animation/IKRig/Source/IKRigEditor/Public/RetargetEditor/IKRetargetEditorStyle.h
--- a/Engine/Plugins/Animation/IKRig/Source/IKRigEditor/Public/RetargetEditor/IKRetargetEditorStyle.h
+++ b/Engine/Plugins/Animation/IKRig/Source/IKRigEditor/Public/RetargetEditor/IKRetargetEditorStyle.h
@@ -2,11 +2,9 @@
 
 #pragma once
 
-#include "Styling/AppStyle.h"
 #include "Styling/SlateStyle.h"
 #include "Styling/SlateStyleRegistry.h"
 #include "Styling/SlateStyleMacros.h"
-#include "Styling/SlateTypes.h"
 
 class FIKRetargetEditorStyle final	: public FSlateStyleSet
 {
@@ -15,13 +13,14 @@ public:
 	FIKRetargetEditorStyle() : FSlateStyleSet("IKRetargetEditorStyle")
 	{
 		const FVector2D Icon16x16(16.0f, 16.0f);
+		const FVector2D Icon64x64(64.0f, 64.0f);
 		
-		const FString IKRigPluginContentDir = FPaths::EnginePluginsDir() / TEXT("Animation/IKRig/Content");
-		SetContentRoot(IKRigPluginContentDir);
+		SetContentRoot(FPaths::EnginePluginsDir() / TEXT("Animation/IKRig/Content"));
 		Set("IKRetarget.Tree.Bone", new IMAGE_BRUSH("Slate/Bone_16x", Icon16x16));
+		Set("ClassIcon.IKRetargeter", new IMAGE_BRUSH_SVG("Slate/IKRigRetargeter", Icon16x16));
+		Set("ClassThumbnail.IKRetargeter", new IMAGE_BRUSH_SVG("Slate/IKRigRetargeter_64", Icon64x64));
 
-		const FString EngineEditorSlateDir = FPaths::EngineContentDir() / TEXT("Editor/Slate");
-		SetContentRoot(EngineEditorSlateDir);
+		SetContentRoot(FPaths::EngineContentDir() / TEXT("Editor/Slate"));
 		Set( "IKRetarget.Viewport.Border", new BOX_BRUSH( "Old/Window/ViewportDebugBorder", 0.8f, FLinearColor(1.0f,1.0f,1.0f,1.0f) ) );
 		
 		FSlateStyleRegistry::RegisterSlateStyle(*this);
@@ -38,5 +37,3 @@ public:
 		FSlateStyleRegistry::UnRegisterSlateStyle(*this);
 	}
 };
-
-#undef IMAGE_BRUSH
\ No newline at end of file
diff --git a/Engine/Plugins/Animation/IKRig/Source/IKRigEditor/Public/RetargetEditor/IKRetargetHitProxies.h b/Engine/Plugins/Animation/IKRig/Source/IKRigEditor/Public/RetargetEditor/IKRetargetHitProxies.h
--- a/Engine/Plugins/Animation/IKRig/Source/IKRigEditor/Public/RetargetEditor/IKRetargetHitProxies.h
+++ b/Engine/Plugins/Animation/IKRig/Source/IKRigEditor/Public/RetargetEditor/IKRetargetHitProxies.h
@@ -5,7 +5,7 @@
 #include "CoreMinimal.h"
 #include "HitProxies.h"
 
-// allow limb selection to edit retarget pose
+// allow bone selection to edit retarget pose
 struct HIKRetargetEditorBoneProxy : public HHitProxy
 {
 	DECLARE_HIT_PROXY();
@@ -26,3 +26,45 @@ struct HIKRetargetEditorBoneProxy : public HHitProxy
 		return true;
 	}
 };
+
+
+// select chains/goals to edit chain settings
+struct HIKRetargetEditorChainProxy : public HHitProxy
+{
+	DECLARE_HIT_PROXY();
+
+	FName TargetChainName;
+	
+	HIKRetargetEditorChainProxy(const FName& InTargetChainName)
+		: HHitProxy(HPP_World)
+		, TargetChainName(InTargetChainName) {}
+
+	virtual EMouseCursor::Type GetMouseCursor()
+	{
+		return EMouseCursor::Crosshairs;
+	}
+
+	virtual bool AlwaysAllowsTranslucentPrimitives() const override
+	{
+		return true;
+	}
+};
+
+// select root control to edit root settings
+struct HIKRetargetEditorRootProxy : public HHitProxy
+{
+	DECLARE_HIT_PROXY();
+	
+	HIKRetargetEditorRootProxy()
+		: HHitProxy(HPP_World){}
+
+	virtual EMouseCursor::Type GetMouseCursor()
+	{
+		return EMouseCursor::Crosshairs;
+	}
+
+	virtual bool AlwaysAllowsTranslucentPrimitives() const override
+	{
+		return true;
+	}
+};
\ No newline at end of file
diff --git a/Engine/Plugins/Animation/IKRig/Source/IKRigEditor/Public/RetargetEditor/IKRetargeterController.h b/Engine/Plugins/Animation/IKRig/Source/IKRigEditor/Public/RetargetEditor/IKRetargeterController.h
--- a/Engine/Plugins/Animation/IKRig/Source/IKRigEditor/Public/RetargetEditor/IKRetargeterController.h
+++ b/Engine/Plugins/Animation/IKRig/Source/IKRigEditor/Public/RetargetEditor/IKRetargeterController.h
@@ -38,6 +38,8 @@ public:
 	USkeletalMesh* GetPreviewMesh(const ERetargetSourceOrTarget& SourceOrTarget) const;
 	/** Get either source or target IK Rig */
 	const UIKRigDefinition* GetIKRig(const ERetargetSourceOrTarget& SourceOrTarget) const;
+	/** Get either source or target IK Rig */
+	UIKRigDefinition* GetIKRigWriteable(const ERetargetSourceOrTarget& SourceOrTarget) const;
 	/** Set the SOURCE or TARGET preview mesh based on the mesh in the corresponding IK Rig asset */
 	void OnIKRigChanged(const ERetargetSourceOrTarget& SourceOrTarget) const;
 
@@ -61,6 +63,8 @@ public:
 	void SetSourceChainForTargetChain(URetargetChainSettings* ChainMap, FName SourceChainToMapTo) const;
 	/** Get read-only access to the list of chain mappings */
 	const TArray<TObjectPtr<URetargetChainSettings>>& GetChainMappings() const;
+	/** Get whether the given chain has a valid IK goal assigned to it */
+	bool IsChainRunningIK(const TObjectPtr<URetargetChainSettings> ChainSettings) const;
 	/** END RETARGET CHAIN MAPPING */
 
 	/** RETARGET POSE EDITING
@@ -106,12 +110,6 @@ public:
 	FQuat GetRotationOffsetForRetargetPoseBone(
 		const FName& BoneName,
 		const ERetargetSourceOrTarget& SourceOrTarget) const;
-	/** Set the delta translation of the root bone (used in Edit Mode in the retarget editor) */
-	void SetTranslationOffsetOnRetargetRootBone(
-		const FVector& TranslationOffset,
-		const ERetargetSourceOrTarget& SourceOrTarget) const;
-	/** Get the delta translation of the root bone (used in Edit Mode in the retarget editor) */
-	const FVector& GetTranslationOffsetOnRetargetRootBone(const ERetargetSourceOrTarget& SourceOrTarget) const;
 	/** Add a delta translation to the root bone (used in Edit Mode in the retarget editor) */
 	void AddTranslationOffsetToRetargetRootBone(
 		const FVector& TranslationOffset,
diff --git a/Engine/Plugins/Animation/IKRig/Source/IKRigEditor/Public/RetargetEditor/SIKRetargetChainMapList.h b/Engine/Plugins/Animation/IKRig/Source/IKRigEditor/Public/RetargetEditor/SIKRetargetChainMapList.h
--- a/Engine/Plugins/Animation/IKRig/Source/IKRigEditor/Public/RetargetEditor/SIKRetargetChainMapList.h
+++ b/Engine/Plugins/Animation/IKRig/Source/IKRigEditor/Public/RetargetEditor/SIKRetargetChainMapList.h
@@ -12,6 +12,7 @@
 #include "Widgets/Views/SListView.h"
 #include "Framework/Commands/UICommandList.h"
 
+class FTextFilterExpressionEvaluator;
 class FIKRigEditorController;
 class SIKRetargetChainMapList;
 class FIKRetargetEditor;
@@ -57,10 +58,12 @@ public:
 	virtual TSharedRef<SWidget> GenerateWidgetForColumn(const FName& ColumnName) override;
 
 	void OnSourceChainComboSelectionChanged(TSharedPtr<FString> InName, ESelectInfo::Type SelectInfo);
-
-	FText GetSourceChainName() const;
 	
 private:
+	
+	FText GetSourceChainName() const;
+
+	FText GetTargetIKGoalName() const;
 
 	TArray<TSharedPtr<FString>> SourceChainOptions;
 
@@ -71,6 +74,13 @@ private:
 	friend SIKRetargetChainMapList;
 };
 
+struct FChainMapFilterOptions
+{
+	bool bHideUnmappedChains = false;
+	bool bHideMappedChains = false;
+	bool bHideChainsWithoutIK = false;
+};
+
 typedef SListView< TSharedPtr<FRetargetChainMapElement> > SRetargetChainMapListViewType;
 
 class SIKRetargetChainMapList : public SCompoundWidget, public FEditorUndoClient
@@ -105,18 +115,23 @@ private:
 	FText GetTargetRootBone() const;
 	bool IsChainMapEnabled() const;
 	/** when a chain is clicked on in the table view */
-	void OnItemClicked(TSharedPtr<FRetargetChainMapElement> InItem);
+	void OnItemClicked(TSharedPtr<FRetargetChainMapElement> InItem) const;
+	/** when edit global settings button clicked */
+	FReply OnGlobalSettingsButtonClicked() const;
 	/** when edit root settings button clicked */
-	FReply OnEditSettingsButtonClicked() const;
-
-	/** update selection / details view */
-	void OnSelectionChanged();
+	FReply OnRootSettingsButtonClicked() const;
 
 	/** auto-map chain button*/
 	EVisibility IsAutoMapButtonVisible() const;
 	FReply OnAutoMapButtonClicked() const;
 	/** END auto-map chain button*/
 
+	/** filtering the list with search box */
+	TSharedRef<SWidget> CreateFilterMenuWidget();
+	void OnFilterTextChanged(const FText& SearchText);
+	TSharedPtr<FTextFilterExpressionEvaluator> TextFilter;
+	FChainMapFilterOptions ChainFilterOptions;
+
 	/** list view generate row callback */
 	TSharedRef<ITableRow> MakeListRowWidget(TSharedPtr<FRetargetChainMapElement> InElement, const TSharedRef<STableViewBase>& OwnerTable);
 
diff --git a/Engine/Plugins/Animation/IKRig/Source/IKRigEditor/Public/RetargetEditor/SIKRetargetHierarchy.h b/Engine/Plugins/Animation/IKRig/Source/IKRigEditor/Public/RetargetEditor/SIKRetargetHierarchy.h
--- a/Engine/Plugins/Animation/IKRig/Source/IKRigEditor/Public/RetargetEditor/SIKRetargetHierarchy.h
+++ b/Engine/Plugins/Animation/IKRig/Source/IKRigEditor/Public/RetargetEditor/SIKRetargetHierarchy.h
@@ -10,18 +10,24 @@
 #include "DragAndDrop/DecoratedDragDropOp.h"
 #include "Widgets/SCompoundWidget.h"
 
+class FTextFilterExpressionEvaluator;
 class SIKRetargetPoseEditor;
 
 class FIKRetargetHierarchyElement : public TSharedFromThis<FIKRetargetHierarchyElement>
 {
 public:
 	
-	FIKRetargetHierarchyElement(const FName& InName, const TSharedRef<FIKRetargetEditorController>& InEditorController);
+	FIKRetargetHierarchyElement(
+		const FName& InName,
+		const FName& InChainName,
+		const TSharedRef<FIKRetargetEditorController>& InEditorController);
 
 	FText Key;
 	TSharedPtr<FIKRetargetHierarchyElement> Parent;
 	TArray<TSharedPtr<FIKRetargetHierarchyElement>> Children;
 	FName Name;
+	FName ChainName;
+	bool bIsHidden = false;
 
 private:
 	
@@ -59,6 +65,13 @@ private:
 	TWeakPtr<FIKRetargetEditorController> EditorController;
 };
 
+struct FRetargetHierarchyFilterOptions
+{
+	bool bHideBonesNotInChain;
+	bool bHideNotRetargetedBones;
+	bool bHideRetargetedBones;
+};
+
 typedef SBaseHierarchyTreeView<FIKRetargetHierarchyElement> SIKRetargetHierarchyTreeView;
 
 class SIKRetargetHierarchy : public SCompoundWidget, public FEditorUndoClient
@@ -88,12 +101,17 @@ private:
 	TSharedPtr<SIKRetargetHierarchyTreeView> TreeView;
 	TArray<TSharedPtr<FIKRetargetHierarchyElement>> RootElements;
 	TArray<TSharedPtr<FIKRetargetHierarchyElement>> AllElements;
+
+	/** filtering the tree with search box */
+	TSharedRef<SWidget> CreateFilterMenuWidget();
+	void OnFilterTextChanged(const FText& SearchText);
+	TSharedPtr<FTextFilterExpressionEvaluator> TextFilter;
+	FRetargetHierarchyFilterOptions FilterOptions;
 	
 	/** tree view callbacks */
 	void RefreshTreeView(bool IsInitialSetup=false);
 	void HandleGetChildrenForTree(TSharedPtr<FIKRetargetHierarchyElement> InElement, TArray<TSharedPtr<FIKRetargetHierarchyElement>>& OutChildren);
 	void OnSelectionChanged(TSharedPtr<FIKRetargetHierarchyElement> Selection, ESelectInfo::Type SelectInfo);
-	void OnItemClicked(TSharedPtr<FIKRetargetHierarchyElement> InElement);
 	void OnItemDoubleClicked(TSharedPtr<FIKRetargetHierarchyElement> InElement);
 	void OnSetExpansionRecursive(TSharedPtr<FIKRetargetHierarchyElement> InElement, bool bShouldBeExpanded);
 	void SetExpansionRecursive(TSharedPtr<FIKRetargetHierarchyElement> InElement, bool bTowardsParent, bool bShouldBeExpanded);
diff --git a/Engine/Plugins/Animation/IKRig/Source/IKRigEditor/Public/RigEditor/IKRigController.h b/Engine/Plugins/Animation/IKRig/Source/IKRigEditor/Public/RigEditor/IKRigController.h
--- a/Engine/Plugins/Animation/IKRig/Source/IKRigEditor/Public/RigEditor/IKRigController.h
+++ b/Engine/Plugins/Animation/IKRig/Source/IKRigEditor/Public/RigEditor/IKRigController.h
@@ -92,12 +92,16 @@ public:
 	bool SetGoalBone(const FName& GoalName, const FName& NewBoneName) const;
 	/** The the Bone associated with the given Goal. */
 	FName GetBoneForGoal(const FName& GoalName) const;
+	/** The the Goal associated with the given Bone (may be null. */
+	const UIKRigEffectorGoal* GetGoalForBone(const FName& BoneName) const;
 	/** Connect the given Goal to the given Solver. This creates an "Effector" with settings specific to this Solver.*/
 	bool ConnectGoalToSolver(const UIKRigEffectorGoal& Goal, int32 SolverIndex) const;
 	/** Disconnect the given Goal from the given Solver. This removes the Effector that associates the Goal with the Solver.*/
 	bool DisconnectGoalFromSolver(const FName& GoalToRemove, int32 SolverIndex) const;
 	/** Returns true if the given Goal is connected to the given Solver. False otherwise. */
 	bool IsGoalConnectedToSolver(const FName& GoalName, int32 SolverIndex) const;
+	/** Returns true if the given Goal is connected to ANY solver. False otherwise. */
+	bool IsGoalConnectedToAnySolver(const FName& GoalName) const;
 	/** Get the index of the given Goal in the list of Goals. */
 	int32 GetGoalIndex(const FName& InGoalName, const ENameCase CompareMethod = ENameCase::IgnoreCase) const;
 	/** Get the name of Goal at the given index. */
@@ -145,7 +149,7 @@ public:
 	 * 
 	 */
 	/** Add a Chain with the given Name and Start/End bones. Returns true if a new Chain was created. */
-	void AddRetargetChain(const FName& ChainName, const FName& StartBone, const FName& EndBone) const;
+	void AddRetargetChain(const FBoneChain& BoneChain) const;
 	/** Remove a Chain with the given name. Returns true if a Chain was removed. */
 	bool RemoveRetargetChain(const FName& ChainName) const;
 	/** Renamed the given Chain. Returns the new name (same as old if unsuccessful). */
@@ -162,6 +166,8 @@ public:
     FName GetRetargetChainStartBone(const FName& ChainName) const;
     /** Get the Start Bone name for the given Chain. */
     FName GetRetargetChainEndBone(const FName& ChainName) const;
+	/** Get read-only access to a single retarget chain with the given name */
+	const FBoneChain* GetRetargetChainByName(const FName& ChainName) const;
 	/** Get read-only access to the list of Chains. */
 	const TArray<FBoneChain>& GetRetargetChains() const;
 	/** Set the Root Bone of the retargeting (can only be one). */
diff --git a/Engine/Plugins/Animation/IKRig/Source/IKRigEditor/Public/RigEditor/IKRigEditMode.h b/Engine/Plugins/Animation/IKRig/Source/IKRigEditor/Public/RigEditor/IKRigEditMode.h
--- a/Engine/Plugins/Animation/IKRig/Source/IKRigEditor/Public/RigEditor/IKRigEditMode.h
+++ b/Engine/Plugins/Animation/IKRig/Source/IKRigEditor/Public/RigEditor/IKRigEditMode.h
@@ -55,11 +55,11 @@ private:
 	
 	void RenderGoals(FPrimitiveDrawInterface* PDI);
 	void RenderBones(FPrimitiveDrawInterface* PDI);
-	void GetAffectedBones(
+	void GetBoneColors(
 		FIKRigEditorController* Controller,
-		UIKRigProcessor* Processor,
-		TSet<int32>& OutAffectedBones,
-		TSet<int32>& OutSelectedBones) const;
+		const UIKRigProcessor* Processor,
+		const FReferenceSkeleton& RefSkeleton,
+		TArray<FLinearColor>& OutBoneColors) const;
 	
 	/** The hosting app */
 	TWeakPtr<FIKRigEditorController> EditorController;
diff --git a/Engine/Plugins/Animation/IKRig/Source/IKRigEditor/Public/RigEditor/IKRigEditorController.h b/Engine/Plugins/Animation/IKRig/Source/IKRigEditor/Public/RigEditor/IKRigEditorController.h
--- a/Engine/Plugins/Animation/IKRig/Source/IKRigEditor/Public/RigEditor/IKRigEditorController.h
+++ b/Engine/Plugins/Animation/IKRig/Source/IKRigEditor/Public/RigEditor/IKRigEditorController.h
@@ -79,6 +79,30 @@ private:
 	static bool ReferenceTransformRelative[3];
 };
 
+enum class EChainSide
+{
+	Left,
+	Right,
+	Center
+};
+
+/** use string matching and skeletal analysis to suggest a reasonable default name for retarget chains */
+struct FRetargetChainAnalyzer
+{
+	TSharedPtr<FTextFilterExpressionEvaluator> TextFilter;
+
+	FRetargetChainAnalyzer()
+	{
+		TextFilter = MakeShared<FTextFilterExpressionEvaluator>(ETextFilterExpressionEvaluatorMode::BasicString);
+	}
+	
+	void AssignBestGuessName(FBoneChain& Chain, const FIKRigSkeleton& IKRigSkeleton);
+	
+	static FName GetDefaultChainName();
+
+	static EChainSide GetSideOfChain(const TArray<int32>& BoneIndices, const FIKRigSkeleton& IKRigSkeleton);
+};
+
 /** a home for cross-widget communication to synchronize state across all tabs and viewport */
 class FIKRigEditorController : public TSharedFromThis<FIKRigEditorController>, FGCObject
 {
@@ -86,6 +110,8 @@ public:
 
 	/** initialize the editor controller to an instance of the IK Rig editor */
 	void Initialize(TSharedPtr<FIKRigEditorToolkit> Toolkit, UIKRigDefinition* Asset);
+	/** cleanup when editor closed */
+	void Close() const;
 
 	/** get the currently active processor running the IK Rig in the editor */
 	UIKRigProcessor* GetIKRigProcessor() const;
@@ -124,10 +150,19 @@ public:
 	void GetSelectedBones(TArray<TSharedPtr<FIKRigTreeElement>>& OutBoneItems) const;
 	/** returns true if Goal is currently selected */
 	bool IsGoalSelected(const FName& GoalName) const;
-	/** get name of the selected retargeting chain */
-	FName GetSelectedChain() const;
 	/** is anything selected in the skeleton view? */
 	bool DoesSkeletonHaveSelectedItems() const;
+	/** get name of the selected retargeting chain */
+	TArray<FName> GetSelectedChains() const;
+	/** get chains selected in skeleton view */
+	void GetChainsSelectedInSkeletonView(TArray<FBoneChain>& InOutChains);
+	
+	/** create new retarget chains from selected bones (or single empty chain if no selection) */
+	void CreateNewRetargetChains();
+	/** show user the new retarget chain they are about to create (provides option to edit name) */
+	void PromptToAddNewRetargetChain(const FBoneChain& BoneChain) const;
+	/** prompt user if they want to add a goal to a newly mirrored chain */
+	void PromptToAddGoalToNewlyMirroredChain(const FBoneChain& OldChain, const FBoneChain& MirroredChain) const;
 	
 	/** determine if the element is connected to the selected solver */
 	bool IsElementConnectedToSolver(TSharedRef<FIKRigTreeElement> TreeElement, int32 SolverIndex);
@@ -162,9 +197,6 @@ public:
 	/** set retargeting tab view */
 	void SetRetargetingView(const TSharedPtr<SIKRigRetargetChainList>& InRetargetingView){ RetargetingView = InRetargetingView; };
 
-	/** create a new retarget chain */
-	void AddNewRetargetChain(const FName ChainName, const FName StartBone, const FName EndBone);
-
 	/** play preview animation on running anim instance in editor (before IK) */
 	void PlayAnimationAsset(UAnimationAsset* AssetToPlay);
 	
@@ -181,6 +213,9 @@ public:
 	/** the persona toolkit */
 	TWeakPtr<FIKRigEditorToolkit> EditorToolkit;
 
+	/** used to analyze a retarget chain */
+	FRetargetChainAnalyzer ChainAnalyzer;
+
 	/** UI and viewport selection state */
 	bool bManipulatingGoals = false;
 
@@ -205,6 +240,9 @@ private:
 	/** right after importing a skeleton, we ask user what solver they want to use */
 	bool PromptToAddSolver() const;
 
+	/** right after creating a goal, we ask user if they want it assigned to a retarget chain */
+	void PromptToAddGoalToChain(UIKRigEffectorGoal* NewGoal) const;
+
 	/** Initializes editor's solvers instances */
 	void InitializeSolvers() const;
 
@@ -231,6 +269,9 @@ private:
 	UPROPERTY()
 	TObjectPtr<UIKRigBoneDetails> BoneDetails;
 
+	/** remove delegate when closing editor */
+	FDelegateHandle ReinitializeDelegateHandle;
+
 	friend struct FIKRigOutputLogTabSummoner;
 	friend class SIKRigAssetBrowser;
 };
@@ -239,4 +280,4 @@ struct FIKRigSolverTypeAndName
 {
 	FText NiceName;
 	TSubclassOf<UIKRigSolver> SolverType;
-};
+};
\ No newline at end of file
diff --git a/Engine/Plugins/Animation/IKRig/Source/IKRigEditor/Public/RigEditor/IKRigEditorStyle.h b/Engine/Plugins/Animation/IKRig/Source/IKRigEditor/Public/RigEditor/IKRigEditorStyle.h
--- a/Engine/Plugins/Animation/IKRig/Source/IKRigEditor/Public/RigEditor/IKRigEditorStyle.h
+++ b/Engine/Plugins/Animation/IKRig/Source/IKRigEditor/Public/RigEditor/IKRigEditorStyle.h
@@ -2,12 +2,9 @@
 
 #pragma once
 
-#include "Styling/AppStyle.h"
-#include "Styling/CoreStyle.h"
 #include "Styling/SlateStyle.h"
 #include "Styling/SlateStyleRegistry.h"
 #include "Styling/SlateStyleMacros.h"
-#include "Styling/SlateTypes.h"
 
 class FIKRigEditorStyle	final : public FSlateStyleSet
 {
@@ -15,29 +12,25 @@ public:
 	
 	FIKRigEditorStyle() : FSlateStyleSet("IKRigEditorStyle")
 	{
-		SetContentRoot(FPaths::EnginePluginsDir() / TEXT("Animation/IKRig/Content"));
-		const FVector2D Icon10x10(10.0f, 10.0f);
-		const FVector2D Icon14x14(14.0f, 14.0f);
 		const FVector2D Icon16x16(16.0f, 16.0f);
 		const FVector2D Icon20x20(20.0f, 20.0f);
-		const FVector2D Icon24x24(24.0f, 24.0f);
-		const FVector2D Icon32x32(32.0f, 32.0f);
 		const FVector2D Icon40x40(40.0f, 40.0f);
+		const FVector2D Icon64x64(64.0f, 64.0f);
 
-#define IKRIG_IMAGE_BRUSH(RelativePath, ...) FSlateImageBrush(RootToContentDir(RelativePath, TEXT(".png")), __VA_ARGS__)
-		Set("IKRig.Tree.Bone", new IKRIG_IMAGE_BRUSH("Slate/Bone_16x", Icon16x16));
-		Set("IKRig.Tree.BoneWithSettings", new IKRIG_IMAGE_BRUSH("Slate/BoneWithSettings_16x", Icon16x16));
-		Set("IKRig.Tree.Goal", new IKRIG_IMAGE_BRUSH("Slate/Goal_16x", Icon16x16));
-		Set("IKRig.Tree.Effector", new IKRIG_IMAGE_BRUSH("Slate/Effector_16x", Icon16x16));
-		
-		Set("IKRig.TabIcon", new IKRIG_IMAGE_BRUSH("Slate/Tab_16x", Icon16x16));
+		SetContentRoot(FPaths::EnginePluginsDir() / TEXT("Animation/IKRig/Content"));
 		
-		Set("IKRig.Solver", new IKRIG_IMAGE_BRUSH("Slate/Solver_16x", Icon16x16));
-		Set("IKRig.DragSolver", new IKRIG_IMAGE_BRUSH("Slate/DragSolver", FVector2D(6, 15)));
+		Set("IKRig.Tree.Bone", new IMAGE_BRUSH("Slate/Bone_16x", Icon16x16));
+		Set("IKRig.Tree.BoneWithSettings", new IMAGE_BRUSH("Slate/BoneWithSettings_16x", Icon16x16));
+		Set("IKRig.Tree.Goal", new IMAGE_BRUSH("Slate/Goal_16x", Icon16x16));
+		Set("IKRig.Tree.Effector", new IMAGE_BRUSH("Slate/Effector_16x", Icon16x16));
+		Set("IKRig.TabIcon", new IMAGE_BRUSH("Slate/Tab_16x", Icon16x16));
+		Set("IKRig.Solver", new IMAGE_BRUSH("Slate/Solver_16x", Icon16x16));
+		Set("IKRig.DragSolver", new IMAGE_BRUSH("Slate/DragSolver", FVector2D(6, 15)));
+		Set("IKRig.Reset", new IMAGE_BRUSH("Slate/Reset", Icon40x40));
+		Set("IKRig.Reset.Small", new IMAGE_BRUSH("Slate/Reset", Icon20x20));
 
-		Set("IKRig.Reset", new IKRIG_IMAGE_BRUSH("Slate/Reset", Icon40x40));
-		Set("IKRig.Reset.Small", new IKRIG_IMAGE_BRUSH("Slate/Reset", Icon20x20));
-#undef IKRIG_IMAGE_BRUSH
+		Set("ClassIcon.IKRigDefinition", new IMAGE_BRUSH_SVG("Slate/IKRig", Icon16x16));
+		Set("ClassThumbnail.IKRigDefinition", new IMAGE_BRUSH_SVG("Slate/IKRig_64", Icon64x64));
 		
 		FTextBlockStyle NormalText = FAppStyle::GetWidgetStyle<FTextBlockStyle>("SkeletonTree.NormalFont");
 		Set( "IKRig.Tree.NormalText", FTextBlockStyle(NormalText));
diff --git a/Engine/Plugins/Animation/IKRig/Source/IKRigEditor/Public/RigEditor/IKRigToolkit.h b/Engine/Plugins/Animation/IKRig/Source/IKRigEditor/Public/RigEditor/IKRigToolkit.h
--- a/Engine/Plugins/Animation/IKRig/Source/IKRigEditor/Public/RigEditor/IKRigToolkit.h
+++ b/Engine/Plugins/Animation/IKRig/Source/IKRigEditor/Public/RigEditor/IKRigToolkit.h
@@ -35,6 +35,7 @@ public:
 		UIKRigDefinition* IKRigAsset);
 
 	/** FAssetEditorToolkit interface */
+	virtual void OnClose() override;
 	virtual void RegisterTabSpawners(const TSharedRef<FTabManager>& InTabManager) override;
 	virtual void UnregisterTabSpawners(const TSharedRef<FTabManager>& InTabManager) override;
 	virtual FName GetToolkitFName() const override;
diff --git a/Engine/Plugins/Animation/IKRig/Source/IKRigEditor/Public/RigEditor/SIKRigRetargetChainList.h b/Engine/Plugins/Animation/IKRig/Source/IKRigEditor/Public/RigEditor/SIKRigRetargetChainList.h
--- a/Engine/Plugins/Animation/IKRig/Source/IKRigEditor/Public/RigEditor/SIKRigRetargetChainList.h
+++ b/Engine/Plugins/Animation/IKRig/Source/IKRigEditor/Public/RigEditor/SIKRigRetargetChainList.h
@@ -12,6 +12,7 @@
 #include "Widgets/Views/SListView.h"
 #include "Widgets/Input/SEditableTextBox.h"
 #include "Framework/Commands/UICommandList.h"
+#include "Misc/TextFilterExpressionEvaluator.h"
 
 #include "SIKRigRetargetChainList.generated.h"
 
@@ -107,6 +108,13 @@ struct FIKRigRetargetChainSettings
 	FName EndBone;
 };
 
+struct FChainFilterOptions
+{
+	bool bHideSingleBoneChains = false;
+	bool bShowOnlyIKChains = false;
+	bool bShowOnlyMissingBoneChains = false;
+};
+
 typedef SListView< TSharedPtr<FRetargetChainElement> > SRetargetChainListViewType;
 
 class SIKRigRetargetChainList : public SCompoundWidget, public FEditorUndoClient
@@ -117,7 +125,7 @@ public:
 
 	void Construct(const FArguments& InArgs, TSharedRef<FIKRigEditorController> InEditorController);
 
-	FName GetSelectedChain();
+	TArray<FName> GetSelectedChains() const;
 
 private:
 	
@@ -135,6 +143,12 @@ private:
 	TArray< TSharedPtr<FRetargetChainElement> > ListViewItems;
 	/** END list view */
 
+	/** filtering the list with search box */
+	TSharedRef<SWidget> CreateFilterMenuWidget();
+	void OnFilterTextChanged(const FText& SearchText);
+	TSharedPtr<FTextFilterExpressionEvaluator> TextFilter;
+	FChainFilterOptions ChainFilterOptions;
+
 	/** callback when "Add New Chain" clicked */
 	bool IsAddChainEnabled() const;
 	
@@ -157,6 +171,9 @@ private:
 	/** sort chain list callback */
 	void SortChainList();
 
+	/** mirror the selected chains (right-click menu callback) */
+	void MirrorSelectedChains();
+
 	friend SIKRigRetargetChainRow;
 	friend FIKRigEditorController;
 };
diff --git a/Engine/Plugins/Animation/IKRig/Source/IKRigEditor/Public/RigEditor/SIKRigSkeleton.h b/Engine/Plugins/Animation/IKRig/Source/IKRigEditor/Public/RigEditor/SIKRigSkeleton.h
--- a/Engine/Plugins/Animation/IKRig/Source/IKRigEditor/Public/RigEditor/SIKRigSkeleton.h
+++ b/Engine/Plugins/Animation/IKRig/Source/IKRigEditor/Public/RigEditor/SIKRigSkeleton.h
@@ -12,7 +12,6 @@
 
 #include "SIKRigSolverStack.h"
 
-struct FIKRigSkeletonChain;
 class FIKRigEditorController;
 class SIKRigSkeleton;
 class FIKRigEditorToolkit;
@@ -36,12 +35,17 @@ public:
 
 	FText Key;
 	IKRigTreeElementType ElementType;
+
+	/** determines if this item is filtered out of the view */
+	bool bIsHidden = false;
+	TSharedPtr<FIKRigTreeElement> UnFilteredParent;
+	
 	TSharedPtr<FIKRigTreeElement> Parent;
 	TArray<TSharedPtr<FIKRigTreeElement>> Children;
 
 	/** effector meta-data (if it is an effector) */
-	FName SolverGoalName = NAME_None;
-	int32 SolverGoalIndex = INDEX_NONE;
+	FName EffectorGoalName = NAME_None;
+	int32 EffectorIndex = INDEX_NONE;
 
 	/** bone setting meta-data (if it is a bone setting) */
 	FName BoneSettingBoneName = NAME_None;
@@ -101,6 +105,23 @@ public:
 	TWeakPtr<FIKRigTreeElement> Element;
 };
 
+struct FIKRigHierarchyFilterOptions
+{
+	bool bShowAll = true;
+	bool bShowOnlyBones = false;
+	bool bShowOnlyGoals = false;
+	bool bShowOnlyBoneSettings = false;
+	bool bHideUnaffectedBones = false;
+
+	void ResetShowOptions()
+	{
+		bShowAll = false;
+		bShowOnlyBones = false;
+		bShowOnlyGoals = false;
+		bShowOnlyBoneSettings = false;
+	}
+};
+
 typedef SBaseHierarchyTreeView<FIKRigTreeElement> SIKRigSkeletonTreeView;
 
 class SIKRigSkeleton : public SCompoundWidget, public FEditorUndoClient
@@ -125,7 +146,7 @@ public:
 		const FName& ItemName,
 		IKRigTreeElementType ItemType,
 		const bool bReplace);
-	void GetSelectedBoneChains(TArray<FIKRigSkeletonChain>& OutChains);
+	void GetSelectedBoneChains(TArray<FBoneChain>& OutChains);
 	TArray<TSharedPtr<FIKRigTreeElement>> GetSelectedItems() const;
 	bool HasSelectedItems() const;
 	/** END selection state queries */
@@ -207,12 +228,19 @@ private:
 	TSharedPtr<SIKRigSkeletonTreeView> TreeView;
 	TArray<TSharedPtr<FIKRigTreeElement>> RootElements;
 	TArray<TSharedPtr<FIKRigTreeElement>> AllElements;
+
+	/** filtering the tree with search box */
+	TSharedRef<SWidget> CreateFilterMenuWidget();
+	void OnFilterTextChanged(const FText& SearchText);
+	TSharedPtr<FTextFilterExpressionEvaluator> TextFilter;
+	FIKRigHierarchyFilterOptions FilterOptions;
 	
 	/** tree view callbacks */
 	void RefreshTreeView(bool IsInitialSetup=false);
 	TSharedRef<ITableRow> MakeTableRowWidget(TSharedPtr<FIKRigTreeElement> InItem, const TSharedRef<STableViewBase>& OwnerTable);
 	void HandleGetChildrenForTree(TSharedPtr<FIKRigTreeElement> InItem, TArray<TSharedPtr<FIKRigTreeElement>>& OutChildren);
 	void OnSelectionChanged(TSharedPtr<FIKRigTreeElement> Selection, ESelectInfo::Type SelectInfo);
+	TSharedRef< SWidget > CreateAddNewMenu();
 	TSharedPtr< SWidget > CreateContextMenu();
 	void OnItemClicked(TSharedPtr<FIKRigTreeElement> InItem);
 	void OnItemDoubleClicked(TSharedPtr<FIKRigTreeElement> InItem);
diff --git a/Engine/Source/Editor/Persona/Private/AnimationEditorViewportClient.cpp b/Engine/Source/Editor/Persona/Private/AnimationEditorViewportClient.cpp
--- a/Engine/Source/Editor/Persona/Private/AnimationEditorViewportClient.cpp
+++ b/Engine/Source/Editor/Persona/Private/AnimationEditorViewportClient.cpp
@@ -1402,7 +1402,7 @@ ELocalAxesMode::Type FAnimationViewportClient::GetLocalAxesMode() const
 
 void FAnimationViewportClient::SetBoneDrawSize(const float InBoneDrawSize)
 {
-	BoneDrawSize = FMath::Max(0.01f,InBoneDrawSize);
+	BoneDrawSize = FMath::Max(0.001f,InBoneDrawSize);
 
 	// optionally inform editors that may want to maintain bone size between sessions
 	if (OnSetBoneSize.IsBound())
