Fixing issue: ApplySettingsFromProfile implementation missing in non-editor builds.

diff --git a/Engine/Plugins/Animation/IKRig/Source/IKRig/Private/IKRigProcessor.cpp b/Engine/Plugins/Animation/IKRig/Source/IKRig/Private/IKRigProcessor.cpp
--- a/Engine/Plugins/Animation/IKRig/Source/IKRig/Private/IKRigProcessor.cpp
+++ b/Engine/Plugins/Animation/IKRig/Source/IKRig/Private/IKRigProcessor.cpp
@@ -320,8 +320,6 @@ void UIKRigProcessor::SetNeedsInitialized()
 	bTriedToInitialize = false;
 };
 
-#if WITH_EDITOR
-
 void UIKRigProcessor::CopyAllInputsFromSourceAssetAtRuntime(const UIKRigDefinition* SourceAsset)
 {
 	check(SourceAsset)
@@ -348,8 +346,6 @@ void UIKRigProcessor::CopyAllInputsFromSourceAssetAtRuntime(const UIKRigDefiniti
 	}
 }
 
-#endif
-
 const FIKRigGoalContainer& UIKRigProcessor::GetGoalContainer() const
 {
 	check(bInitialized);
diff --git a/Engine/Plugins/Animation/IKRig/Source/IKRig/Private/Retargeter/IKRetargetProcessor.cpp b/Engine/Plugins/Animation/IKRig/Source/IKRig/Private/Retargeter/IKRetargetProcessor.cpp
--- a/Engine/Plugins/Animation/IKRig/Source/IKRig/Private/Retargeter/IKRetargetProcessor.cpp
+++ b/Engine/Plugins/Animation/IKRig/Source/IKRig/Private/Retargeter/IKRetargetProcessor.cpp
@@ -1854,6 +1854,64 @@
 	}
 }
 
+bool UIKRetargetProcessor::IsBoneRetargeted(
+	const int32& BoneIndex,
+	const int8& SkeletonToCheck) const
+{
+	const bool bUseSourceSkeleton = SkeletonToCheck == 0;
+	const FRetargetSkeleton& Skeleton = bUseSourceSkeleton ? SourceSkeleton : TargetSkeleton;
+	const int32 RootBoneIndex = bUseSourceSkeleton ? GetSourceRetargetRoot() : GetTargetRetargetRoot();
+	return (Skeleton.ChainThatContainsBone[BoneIndex] != NAME_None) || BoneIndex==RootBoneIndex;
+}
+
+FName UIKRetargetProcessor::GetChainNameForBone(const int32& BoneIndex, const int8& SkeletonToCheck) const
+{
+	const bool bUseSourceSkeleton = SkeletonToCheck == 0;
+	const FRetargetSkeleton& Skeleton = bUseSourceSkeleton ? SourceSkeleton : TargetSkeleton;
+	const int32 RootBoneIndex = bUseSourceSkeleton ? GetSourceRetargetRoot() : GetTargetRetargetRoot();
+	if (RootBoneIndex == BoneIndex)
+	{
+		return FName("Retarget Root");
+	}
+
+	return Skeleton.ChainThatContainsBone[BoneIndex];
+}
+
+#endif
+
+void UIKRetargetProcessor::ApplySettingsFromAsset()
+{
+	// copy IK Rig settings
+	if (const UIKRigDefinition* TargetIKRig = RetargeterAsset->GetTargetIKRig())
+	{
+		IKRigProcessor->CopyAllInputsFromSourceAssetAtRuntime(TargetIKRig);
+	}
+	
+	// copy chain settings
+	const TArray<TObjectPtr<URetargetChainSettings>>& AllChainSettings = RetargeterAsset->GetAllChainSettings();
+	for (const TObjectPtr<URetargetChainSettings>& ChainSettings : AllChainSettings)
+	{
+		for (FRetargetChainPairFK& Chain : ChainPairsFK)
+		{
+			if (Chain.TargetBoneChainName == ChainSettings->TargetChain)
+			{
+				Chain.Settings.CopySettingsFromAsset(ChainSettings);
+			}
+		}
+		
+		for (FRetargetChainPairIK& Chain : ChainPairsIK)
+		{
+			if (Chain.TargetBoneChainName == ChainSettings->TargetChain)
+			{
+				Chain.Settings.CopySettingsFromAsset(ChainSettings);
+			}
+		}
+	}
+
+	// copy root settings
+	RootRetargeter.Settings.CopySettingsFromAsset(RetargeterAsset->GetRootSettingsUObject());
+}
+
 void UIKRetargetProcessor::ApplySettingsFromProfile(const FRetargetProfile& Profile)
 {
 	auto ApplyRetargetPoseFromProfile = [this, Profile](ERetargetSourceOrTarget SourceOrTarget)
@@ -1915,62 +1973,4 @@
 	}
 }
 
-void UIKRetargetProcessor::ApplySettingsFromAsset()
-{
-	// copy IK Rig settings
-	if (const UIKRigDefinition* TargetIKRig = RetargeterAsset->GetTargetIKRig())
-	{
-		IKRigProcessor->CopyAllInputsFromSourceAssetAtRuntime(TargetIKRig);
-	}
-	
-	// copy chain settings
-	const TArray<TObjectPtr<URetargetChainSettings>>& AllChainSettings = RetargeterAsset->GetAllChainSettings();
-	for (const TObjectPtr<URetargetChainSettings>& ChainSettings : AllChainSettings)
-	{
-		for (FRetargetChainPairFK& Chain : ChainPairsFK)
-		{
-			if (Chain.TargetBoneChainName == ChainSettings->TargetChain)
-			{
-				Chain.Settings.CopySettingsFromAsset(ChainSettings);
-			}
-		}
-		
-		for (FRetargetChainPairIK& Chain : ChainPairsIK)
-		{
-			if (Chain.TargetBoneChainName == ChainSettings->TargetChain)
-			{
-				Chain.Settings.CopySettingsFromAsset(ChainSettings);
-			}
-		}
-	}
-
-	// copy root settings
-	RootRetargeter.Settings.CopySettingsFromAsset(RetargeterAsset->GetRootSettingsUObject());
-}
-
-bool UIKRetargetProcessor::IsBoneRetargeted(
-	const int32& BoneIndex,
-	const int8& SkeletonToCheck) const
-{
-	const bool bUseSourceSkeleton = SkeletonToCheck == 0;
-	const FRetargetSkeleton& Skeleton = bUseSourceSkeleton ? SourceSkeleton : TargetSkeleton;
-	const int32 RootBoneIndex = bUseSourceSkeleton ? GetSourceRetargetRoot() : GetTargetRetargetRoot();
-	return (Skeleton.ChainThatContainsBone[BoneIndex] != NAME_None) || BoneIndex==RootBoneIndex;
-}
-
-FName UIKRetargetProcessor::GetChainNameForBone(const int32& BoneIndex, const int8& SkeletonToCheck) const
-{
-	const bool bUseSourceSkeleton = SkeletonToCheck == 0;
-	const FRetargetSkeleton& Skeleton = bUseSourceSkeleton ? SourceSkeleton : TargetSkeleton;
-	const int32 RootBoneIndex = bUseSourceSkeleton ? GetSourceRetargetRoot() : GetTargetRetargetRoot();
-	if (RootBoneIndex == BoneIndex)
-	{
-		return FName("Retarget Root");
-	}
-
-	return Skeleton.ChainThatContainsBone[BoneIndex];
-}
-
-#endif
-
 #undef LOCTEXT_NAMESPACE
diff --git a/Engine/Plugins/Animation/IKRig/Source/IKRig/Public/IKRigProcessor.h b/Engine/Plugins/Animation/IKRig/Source/IKRig/Public/IKRigProcessor.h
--- a/Engine/Plugins/Animation/IKRig/Source/IKRig/Public/IKRigProcessor.h
+++ b/Engine/Plugins/Animation/IKRig/Source/IKRig/Public/IKRigProcessor.h
@@ -106,10 +106,8 @@ public:
 	/** logging system */
 	FIKRigLogger Log;
 
-#if WITH_EDITOR
 	/** Used to propagate setting values from the source asset at runtime (settings that do not require re-initialization) */
 	void CopyAllInputsFromSourceAssetAtRuntime(const UIKRigDefinition* SourceAsset);
-#endif
 	
 private:
 
