Added pole vector matching settings for feature complete (not functional yet).

diff --git a/Engine/Plugins/Animation/IKRig/Source/IKRig/Public/Retargeter/IKRetargetSettings.h b/Engine/Plugins/Animation/IKRig/Source/IKRig/Public/Retargeter/IKRetargetSettings.h
--- a/Engine/Plugins/Animation/IKRig/Source/IKRig/Public/Retargeter/IKRetargetSettings.h
+++ b/Engine/Plugins/Animation/IKRig/Source/IKRig/Public/Retargeter/IKRetargetSettings.h
@@ -88,6 +88,16 @@ struct FTargetChainFKSettings
 	*  In Globally Scaled and Absolute modes, the translation offset is scaled by this parameter.*/
 	UPROPERTY(EditAnywhere, BlueprintReadWrite, Category = "FK Adjustments", meta = (UIMin = "0.0", UIMax = "1.0"))
 	float TranslationAlpha = 1.0f;
+
+	/** Range 0 to 1. Default 1. Matches the twist angle of this chain (along the Pole direction) to the source chain.
+	*  At 0, the chain's pole vector direction will be left alon
+	*  At 1, the root bone of the chain will be twist-rotated in the pole direction to match the orientation of the source chain.*/
+	UPROPERTY(EditAnywhere, BlueprintReadWrite, Category = "FK Adjustments", meta = (UIMin = "0.0", UIMax = "1.0"))
+	float PoleVectorMatching = 0.0f;
+
+	/** Range +/- 180. Default 0. An angular offset, in degrees, for the pole direction of the chain. */
+	UPROPERTY(EditAnywhere, BlueprintReadWrite, Category = "FK Adjustments", meta = (UIMin = "-180.0", UIMax = "180.0"))
+	float PoleVectorOffset = 0.0f;
 };
 
 USTRUCT(BlueprintType)
