Fixing localization label names in IK Rig plugin.

diff --git a/Engine/Plugins/Animation/IKRig/Source/IKRigEditor/Private/RetargetEditor/IKRetargetDetails.cpp b/Engine/Plugins/Animation/IKRig/Source/IKRigEditor/Private/RetargetEditor/IKRetargetDetails.cpp
--- a/Engine/Plugins/Animation/IKRig/Source/IKRigEditor/Private/RetargetEditor/IKRetargetDetails.cpp
+++ b/Engine/Plugins/Animation/IKRig/Source/IKRigEditor/Private/RetargetEditor/IKRetargetDetails.cpp
@@ -850,7 +850,7 @@ void FRetargetChainSettingsCustomization::CustomizeDetails(IDetailLayoutBuilder&
 		.Content()
 		[
 			SNew(STextBlock)
-			.Text(LOCTEXT( "ChainStatus_Label", "Source Chain" ))
+			.Text(LOCTEXT( "SoureChain_Label", "Source Chain" ))
 			.Font(IDetailLayoutBuilder::GetDetailFont())
 		]
 	]
diff --git a/Engine/Plugins/Animation/IKRig/Source/IKRigEditor/Private/RetargetEditor/SIKRetargetChainMapList.cpp b/Engine/Plugins/Animation/IKRig/Source/IKRigEditor/Private/RetargetEditor/SIKRetargetChainMapList.cpp
--- a/Engine/Plugins/Animation/IKRig/Source/IKRigEditor/Private/RetargetEditor/SIKRetargetChainMapList.cpp
+++ b/Engine/Plugins/Animation/IKRig/Source/IKRigEditor/Private/RetargetEditor/SIKRetargetChainMapList.cpp
@@ -223,7 +223,7 @@ void SIKRetargetChainMapList::Construct(
 			[
 				SNew(SPositiveActionButton)
 				.Icon(FAppStyle::Get().GetBrush("Icons.Settings"))
-				.Text(LOCTEXT("EditRootButtonLabel", "Global Settings"))
+				.Text(LOCTEXT("EditGlobalButtonLabel", "Global Settings"))
 				.ToolTipText(LOCTEXT("EditGlobalButtonToolTip", "Edit the global retarget settings."))
 				.OnClicked(this, &SIKRetargetChainMapList::OnGlobalSettingsButtonClicked)
 			]
@@ -543,7 +543,7 @@ TSharedRef<SWidget> SIKRetargetChainMapList::CreateFilterMenuWidget()
 	static constexpr bool CloseAfterSelection = true;
 	FMenuBuilder MenuBuilder(CloseAfterSelection, CommandList);
 
-	MenuBuilder.BeginSection("Chain Map Filters", LOCTEXT("ChainFiltersSection", "Filter Chain Mappings"));
+	MenuBuilder.BeginSection("Chain Map Filters", LOCTEXT("ChainMapFiltersSection", "Filter Chain Mappings"));
 
 	MenuBuilder.AddMenuEntry(
 		LOCTEXT("HideMappedLabel", "Hide Mapped Chains"),
diff --git a/Engine/Plugins/Animation/IKRig/Source/IKRigEditor/Private/RetargetEditor/SIKRetargetHierarchy.cpp b/Engine/Plugins/Animation/IKRig/Source/IKRigEditor/Private/RetargetEditor/SIKRetargetHierarchy.cpp
--- a/Engine/Plugins/Animation/IKRig/Source/IKRigEditor/Private/RetargetEditor/SIKRetargetHierarchy.cpp
+++ b/Engine/Plugins/Animation/IKRig/Source/IKRigEditor/Private/RetargetEditor/SIKRetargetHierarchy.cpp
@@ -356,8 +356,8 @@ TSharedRef<SWidget> SIKRetargetHierarchy::CreateFilterMenuWidget()
 	MenuBuilder.BeginSection("Clear", LOCTEXT("ClearMapFiltersSection", "Clear"));
 	
 	MenuBuilder.AddMenuEntry(
-		LOCTEXT("ClearMapFilterLabel", "Clear Filters"),
-		LOCTEXT("ClearMapFilterTooltip", "Clear all filters to show all bones."),
+		LOCTEXT("ClearIKRigHierarchyFilterLabel", "Clear Filters"),
+		LOCTEXT("ClearIKRigHierarchyFilterTooltip", "Clear all filters to show all bones."),
 		FSlateIcon(),
 		FUIAction(FExecuteAction::CreateLambda([this]
 		{
diff --git a/Engine/Plugins/Animation/IKRig/Source/IKRigEditor/Private/RigEditor/IKRigEditorController.cpp b/Engine/Plugins/Animation/IKRig/Source/IKRigEditor/Private/RigEditor/IKRigEditorController.cpp
--- a/Engine/Plugins/Animation/IKRig/Source/IKRigEditor/Private/RigEditor/IKRigEditorController.cpp
+++ b/Engine/Plugins/Animation/IKRig/Source/IKRigEditor/Private/RigEditor/IKRigEditorController.cpp
@@ -944,7 +944,7 @@ void FIKRigEditorController::PromptToAddNewRetargetChain(const FBoneChain& BoneC
 void FIKRigEditorController::PromptToAddGoalToNewlyMirroredChain(const FBoneChain& OldChain, const FBoneChain& MirroredChain) const
 {
 	const TSharedRef<SCustomDialog> AddGoalToChainDialog = SNew(SCustomDialog)
-		.Title(FText(LOCTEXT("AddGoalToMirroredChainLabel", "Mirror Goal")))
+		.Title(FText(LOCTEXT("AddGoalToMirroredChainTitleLabel", "Mirror Goal")))
 		.Content()
 		[
 			SNew(STextBlock)
diff --git a/Engine/Plugins/Animation/IKRig/Source/IKRigEditor/Private/RigEditor/SIKRigSkeleton.cpp b/Engine/Plugins/Animation/IKRig/Source/IKRigEditor/Private/RigEditor/SIKRigSkeleton.cpp
--- a/Engine/Plugins/Animation/IKRig/Source/IKRigEditor/Private/RigEditor/SIKRigSkeleton.cpp
+++ b/Engine/Plugins/Animation/IKRig/Source/IKRigEditor/Private/RigEditor/SIKRigSkeleton.cpp
@@ -1407,8 +1407,8 @@ TSharedRef<SWidget> SIKRigSkeleton::CreateFilterMenuWidget()
 	MenuBuilder.BeginSection("Hierarchy Filters", LOCTEXT("HierarchyFiltersSection", "Hierarchy Filters"));
 
 	MenuBuilder.AddMenuEntry(
-		LOCTEXT("OnlyBonesLabel", "Show All Types"),
-		LOCTEXT("OnlyBonesTooltip", "Show all object types in the hierarchy."),
+		LOCTEXT("ShowAllTypesLabel", "Show All Types"),
+		LOCTEXT("ShowAllTypesTooltip", "Show all object types in the hierarchy."),
 		FSlateIcon(),
 		FilterShowAllAction,
 		NAME_None,
@@ -1451,8 +1451,8 @@ TSharedRef<SWidget> SIKRigSkeleton::CreateFilterMenuWidget()
 	MenuBuilder.BeginSection("Clear", LOCTEXT("ClearMapFiltersSection", "Clear"));
 	
 	MenuBuilder.AddMenuEntry(
-		LOCTEXT("ClearMapFilterLabel", "Clear Filters"),
-		LOCTEXT("ClearMapFilterTooltip", "Clear all filters to show all hierarchy elements."),
+		LOCTEXT("ClearSkeletonFilterLabel", "Clear Filters"),
+		LOCTEXT("ClearSkeletonFilterTooltip", "Clear all filters to show all hierarchy elements."),
 		FSlateIcon(),
 		FUIAction(FExecuteAction::CreateLambda([this]
 		{
