Fixing issue, retarget pose node's profile input not exposed to details panel.

diff --git a/Engine/Plugins/Animation/IKRig/Source/IKRig/Public/AnimNodes/AnimNode_RetargetPoseFromMesh.h b/Engine/Plugins/Animation/IKRig/Source/IKRig/Public/AnimNodes/AnimNode_RetargetPoseFromMesh.h
--- a/Engine/Plugins/Animation/IKRig/Source/IKRig/Public/AnimNodes/AnimNode_RetargetPoseFromMesh.h
+++ b/Engine/Plugins/Animation/IKRig/Source/IKRig/Public/AnimNodes/AnimNode_RetargetPoseFromMesh.h
@@ -31,7 +31,7 @@ struct IKRIG_API FAnimNode_RetargetPoseFromMesh : public FAnimNode_Base
 	TObjectPtr<UIKRetargeter> IKRetargeterAsset = nullptr;
 
 	/** connect a custom retarget profile to modify the retargeter's settings at runtime.*/
-	UPROPERTY(BlueprintReadWrite, transient, Category=Settings, meta=(PinHiddenByDefault))
+	UPROPERTY(EditAnywhere, BlueprintReadWrite, transient, Category=Settings, meta=(PinHiddenByDefault))
 	FRetargetProfile CustomRetargetProfile;
 
 	/* Toggle whether to print warnings about missing or incorrectly configured retarget configurations. */
