Fixing uninitialized properties in FTargetChainSettings

diff --git a/Engine/Plugins/Animation/IKRig/Source/IKRig/Public/Retargeter/IKRetargetSettings.h b/Engine/Plugins/Animation/IKRig/Source/IKRig/Public/Retargeter/IKRetargetSettings.h
--- a/Engine/Plugins/Animation/IKRig/Source/IKRig/Public/Retargeter/IKRetargetSettings.h
+++ b/Engine/Plugins/Animation/IKRig/Source/IKRig/Public/Retargeter/IKRetargetSettings.h
@@ -125,15 +125,15 @@ struct FTargetChainIKSettings
 
 	/** Default 0, 0, 0. Apply a static global-space offset to IK goal position. */
 	UPROPERTY(EditAnywhere, BlueprintReadWrite, Category = "IK Adjustments")
-	FVector StaticOffset;
+	FVector StaticOffset = FVector::ZeroVector;
 
 	/** Default 0, 0, 0. Apply a static local-space offset to IK goal position. */
 	UPROPERTY(EditAnywhere, BlueprintReadWrite, Category = "IK Adjustments")
-	FVector StaticLocalOffset;
+	FVector StaticLocalOffset = FVector::ZeroVector;
 
 	/** Default 0, 0, 0. Apply a static local-space offset to IK goal rotation. */
 	UPROPERTY(EditAnywhere, BlueprintReadWrite, Category = "IK Adjustments")
-	FRotator StaticRotationOffset;
+	FRotator StaticRotationOffset = FRotator::ZeroRotator;
 	
 	/** Range 0 to 5. Default 1. Brings IK goal closer (0) or further (1+) from origin of chain.
 	*  At 0 the effector is placed at the origin of the chain (ie Shoulder, Hip etc).
