Fixing NonUnity build errors in IKRig plugin.

diff --git a/Engine/Plugins/Animation/IKRig/Source/IKRigEditor/Public/RigEditor/IKRigEditorStyle.h b/Engine/Plugins/Animation/IKRig/Source/IKRigEditor/Public/RigEditor/IKRigEditorStyle.h
--- a/Engine/Plugins/Animation/IKRig/Source/IKRigEditor/Public/RigEditor/IKRigEditorStyle.h
+++ b/Engine/Plugins/Animation/IKRig/Source/IKRigEditor/Public/RigEditor/IKRigEditorStyle.h
@@ -5,6 +5,8 @@
 #include "Styling/SlateStyle.h"
 #include "Styling/SlateStyleRegistry.h"
 #include "Styling/SlateStyleMacros.h"
+#include "Styling/AppStyle.h"
+#include "Styling/CoreStyle.h"
 
 class FIKRigEditorStyle	final : public FSlateStyleSet
 {
