Fixing static analysis warnings in IKRetargetDefaultMode.cpp

diff --git a/Engine/Plugins/Animation/IKRig/Source/IKRigEditor/Private/RetargetEditor/IKRetargetDefaultMode.cpp b/Engine/Plugins/Animation/IKRig/Source/IKRigEditor/Private/RetargetEditor/IKRetargetDefaultMode.cpp
--- a/Engine/Plugins/Animation/IKRig/Source/IKRigEditor/Private/RetargetEditor/IKRetargetDefaultMode.cpp
+++ b/Engine/Plugins/Animation/IKRig/Source/IKRigEditor/Private/RetargetEditor/IKRetargetDefaultMode.cpp
@@ -86,7 +86,7 @@ void FIKRetargetDefaultMode::Render(const FSceneView* View, FViewport* Viewport,
 	}
 	
 	const UIKRetargetProcessor* RetargetProcessor = Controller->GetRetargetProcessor();
-	if (!RetargetProcessor && RetargetProcessor->IsInitialized())
+	if (!(RetargetProcessor && RetargetProcessor->IsInitialized()))
 	{
 		return;
 	}
