Fix for IK Rig blocking graph's GatherDebugData
#rnx

diff --git a/Engine/Plugins/Animation/IKRig/Source/IKRig/Private/AnimNodes/AnimNode_IKRig.cpp b/Engine/Plugins/Animation/IKRig/Source/IKRig/Private/AnimNodes/AnimNode_IKRig.cpp
--- a/Engine/Plugins/Animation/IKRig/Source/IKRig/Private/AnimNodes/AnimNode_IKRig.cpp
+++ b/Engine/Plugins/Animation/IKRig/Source/IKRig/Private/AnimNodes/AnimNode_IKRig.cpp
@@ -160,6 +160,8 @@ void FAnimNode_IKRig::GatherDebugData(FNodeDebugData& DebugData)
 		
 		DebugData.AddDebugItem(FString::Printf(TEXT("Goal supplied by node pin: %s"), *Goal.ToString()));
 	}
+
+	Source.GatherDebugData(DebugData);
 }
 
 void FAnimNode_IKRig::Initialize_AnyThread(const FAnimationInitializeContext& Context)
