Fixing issue with base solver class in non-editor builds.

diff --git a/Engine/Plugins/Animation/IKRig/Source/IKRig/Private/Solvers/IKRig_BodyMover.cpp b/Engine/Plugins/Animation/IKRig/Source/IKRig/Private/Solvers/IKRig_BodyMover.cpp
--- a/Engine/Plugins/Animation/IKRig/Source/IKRig/Private/Solvers/IKRig_BodyMover.cpp
+++ b/Engine/Plugins/Animation/IKRig/Source/IKRig/Private/Solvers/IKRig_BodyMover.cpp
@@ -78,8 +78,6 @@ void UIKRig_BodyMover::Solve(FIKRigSkeleton& IKRigSkeleton, const FIKRigGoalCont
 	IKRigSkeleton.PropagateGlobalPoseBelowBone(BodyBoneIndex);
 }
 
-#if WITH_EDITOR
-
 void UIKRig_BodyMover::UpdateSolverSettings(UIKRigSolver* InSettings)
 {
 	if(UIKRig_BodyMover* Settings = Cast<UIKRig_BodyMover>(InSettings))
@@ -112,6 +110,21 @@ void UIKRig_BodyMover::UpdateSolverSettings(UIKRigSolver* InSettings)
 	}
 }
 
+void UIKRig_BodyMover::RemoveGoal(const FName& GoalName)
+{
+	// ensure goal even exists in this solver
+	const int32 GoalIndex = GetIndexOfGoal(GoalName);
+	if (GoalIndex == INDEX_NONE)
+	{
+		return;
+	}
+
+	// remove it
+	Effectors.RemoveAt(GoalIndex);
+}
+
+#if WITH_EDITOR
+
 FText UIKRig_BodyMover::GetNiceName() const
 {
 	return FText(LOCTEXT("SolverName", "Body Mover"));
@@ -142,19 +155,6 @@ void UIKRig_BodyMover::AddGoal(const UIKRigEffectorGoal& NewGoal)
 	Effectors.Add(NewEffector);
 }
 
-void UIKRig_BodyMover::RemoveGoal(const FName& GoalName)
-{
-	// ensure goal even exists in this solver
-	const int32 GoalIndex = GetIndexOfGoal(GoalName);
-	if (GoalIndex == INDEX_NONE)
-	{
-		return;
-	}
-
-	// remove it
-	Effectors.RemoveAt(GoalIndex);
-}
-
 void UIKRig_BodyMover::RenameGoal(const FName& OldName, const FName& NewName)
 {
 	// ensure goal even exists in this solver
diff --git a/Engine/Plugins/Animation/IKRig/Source/IKRig/Private/Solvers/IKRig_LimbSolver.cpp b/Engine/Plugins/Animation/IKRig/Source/IKRig/Private/Solvers/IKRig_LimbSolver.cpp
--- a/Engine/Plugins/Animation/IKRig/Source/IKRig/Private/Solvers/IKRig_LimbSolver.cpp
+++ b/Engine/Plugins/Animation/IKRig/Source/IKRig/Private/Solvers/IKRig_LimbSolver.cpp
@@ -128,8 +128,6 @@ void UIKRig_LimbSolver::Solve(FIKRigSkeleton& InOutRigSkeleton, const FIKRigGoal
 	}
 }
 
-#if WITH_EDITOR
-
 void UIKRig_LimbSolver::UpdateSolverSettings(UIKRigSolver* InSettings)
 {
 	if (UIKRig_LimbSolver* Settings = Cast<UIKRig_LimbSolver>(InSettings))
@@ -147,6 +145,18 @@ void UIKRig_LimbSolver::UpdateSolverSettings(UIKRigSolver* InSettings)
 	}
 }
 
+void UIKRig_LimbSolver::RemoveGoal(const FName& GoalName)
+{
+	if (Effector->GoalName == GoalName)
+	{
+		Effector->Modify();
+		Effector->GoalName = NAME_None;
+		Effector->BoneName = NAME_None;
+	}
+}
+
+#if WITH_EDITOR
+
 FText UIKRig_LimbSolver::GetNiceName() const
 {
 	return FText(LOCTEXT("SolverName", "Limb IK"));
@@ -187,16 +197,6 @@ void UIKRig_LimbSolver::AddGoal(const UIKRigEffectorGoal& NewGoal)
 	Effector->BoneName = NewGoal.BoneName;
 }
 
-void UIKRig_LimbSolver::RemoveGoal(const FName& GoalName)
-{
-	if (Effector->GoalName == GoalName)
-	{
-		Effector->Modify();
-		Effector->GoalName = NAME_None;
-		Effector->BoneName = NAME_None;
-	}
-}
-
 void UIKRig_LimbSolver::RenameGoal(const FName& OldName, const FName& NewName)
 {
 	if (Effector->GoalName == OldName)
diff --git a/Engine/Plugins/Animation/IKRig/Source/IKRig/Private/Solvers/IKRig_PoleSolver.cpp b/Engine/Plugins/Animation/IKRig/Source/IKRig/Private/Solvers/IKRig_PoleSolver.cpp
--- a/Engine/Plugins/Animation/IKRig/Source/IKRig/Private/Solvers/IKRig_PoleSolver.cpp
+++ b/Engine/Plugins/Animation/IKRig/Source/IKRig/Private/Solvers/IKRig_PoleSolver.cpp
@@ -141,8 +141,6 @@ void UIKRig_PoleSolver::Solve(FIKRigSkeleton& IKRigSkeleton, const FIKRigGoalCon
 	}
 }
 
-#if WITH_EDITOR
-
 void UIKRig_PoleSolver::UpdateSolverSettings(UIKRigSolver* InSettings)
 {
 	if (UIKRig_PoleSolver* Settings = Cast<UIKRig_PoleSolver>(InSettings))
@@ -151,6 +149,18 @@ void UIKRig_PoleSolver::UpdateSolverSettings(UIKRigSolver* InSettings)
 	}
 }
 
+void UIKRig_PoleSolver::RemoveGoal(const FName& GoalName)
+{
+	if (Effector->GoalName == GoalName)
+	{
+		Effector->Modify();
+		Effector->GoalName = NAME_None;
+		Effector->BoneName = NAME_None;
+	}
+}
+
+#if WITH_EDITOR
+
 FText UIKRig_PoleSolver::GetNiceName() const
 {
 	return FText(LOCTEXT("SolverName", "Pole Solver"));
@@ -192,16 +202,6 @@ void UIKRig_PoleSolver::AddGoal(const UIKRigEffectorGoal& NewGoal)
 	Effector->BoneName = NewGoal.BoneName;
 }
 
-void UIKRig_PoleSolver::RemoveGoal(const FName& GoalName)
-{
-	if (Effector->GoalName == GoalName)
-	{
-		Effector->Modify();
-		Effector->GoalName = NAME_None;
-		Effector->BoneName = NAME_None;
-	}
-}
-
 void UIKRig_PoleSolver::RenameGoal(const FName& OldName, const FName& NewName)
 {
 	if (Effector->GoalName == OldName)
diff --git a/Engine/Plugins/Animation/IKRig/Source/IKRig/Private/Solvers/IKRig_SetTransform.cpp b/Engine/Plugins/Animation/IKRig/Source/IKRig/Private/Solvers/IKRig_SetTransform.cpp
--- a/Engine/Plugins/Animation/IKRig/Source/IKRig/Private/Solvers/IKRig_SetTransform.cpp
+++ b/Engine/Plugins/Animation/IKRig/Source/IKRig/Private/Solvers/IKRig_SetTransform.cpp
@@ -55,8 +55,6 @@ void UIKRig_SetTransform::Solve(FIKRigSkeleton& IKRigSkeleton, const FIKRigGoalC
 	IKRigSkeleton.PropagateGlobalPoseBelowBone(BoneIndex);
 }
 
-#if WITH_EDITOR
-
 void UIKRig_SetTransform::UpdateSolverSettings(UIKRigSolver* InSettings)
 {
 	if (UIKRig_SetTransform* Settings = Cast<UIKRig_SetTransform>(InSettings))
@@ -67,6 +65,17 @@ void UIKRig_SetTransform::UpdateSolverSettings(UIKRigSolver* InSettings)
 	}
 }
 
+void UIKRig_SetTransform::RemoveGoal(const FName& GoalName)
+{
+	if (Goal == GoalName)
+	{
+		Goal = NAME_None;
+		RootBone = NAME_None;
+	}	
+}
+
+#if WITH_EDITOR
+
 FText UIKRig_SetTransform::GetNiceName() const
 {
 	return FText(LOCTEXT("SolverName", "Set Transform"));
@@ -88,15 +97,6 @@ void UIKRig_SetTransform::AddGoal(const UIKRigEffectorGoal& NewGoal)
 	RootBone = NewGoal.BoneName;
 }
 
-void UIKRig_SetTransform::RemoveGoal(const FName& GoalName)
-{
-	if (Goal == GoalName)
-	{
-		Goal = NAME_None;
-		RootBone = NAME_None;
-	}	
-}
-
 void UIKRig_SetTransform::RenameGoal(const FName& OldName, const FName& NewName)
 {
 	if (Goal == OldName)
diff --git a/Engine/Plugins/Animation/IKRig/Source/IKRig/Public/IKRigSolver.h b/Engine/Plugins/Animation/IKRig/Source/IKRig/Public/IKRigSolver.h
--- a/Engine/Plugins/Animation/IKRig/Source/IKRig/Public/IKRigSolver.h
+++ b/Engine/Plugins/Animation/IKRig/Source/IKRig/Public/IKRigSolver.h
@@ -28,12 +28,14 @@ public:
 	virtual void Solve(FIKRigSkeleton& IKRigSkeleton, const FIKRigGoalContainer& Goals) PURE_VIRTUAL("Solve");
 	//** END RUNTIME */
 
-#if WITH_EDITORONLY_DATA
-	/** callback whenever this solver is edited */
-	DECLARE_EVENT_OneParam(UIKRigSolver, FIKRigSolverModified, UIKRigSolver*);
-	FIKRigSolverModified& OnSolverModified(){ return IKRigSolverModified; };
-#endif
-
+	//** ROOT BONE (optional, implement if your solver requires a root bone) */
+	/** if solver requires a root bone, then override this to return it. */
+	virtual FName GetRootBone() const { return NAME_None; };
+	/** override to support telling outside systems which bones this solver has setting for.
+	* NOTE: This must be overriden on solvers that use bone settings.
+	* NOTE: Only ADD to the incoming set, do not remove from it. */
+	virtual void GetBonesWithSettings(TSet<FName>& OutBonesWithSettings) const {};
+	
 	/** get if this solver is enabled */
 	bool IsEnabled() const { return bIsEnabled; };
 	/** turn solver on/off (will be skipped during execution if disabled) */
@@ -45,6 +47,16 @@ public:
 	 * This is necessary because at runtime, the IKRigProcessor creates a copy of your solver class
 	 * and the copy must be notified of changes made to the class settings in the source asset.*/
 	virtual void UpdateSolverSettings(UIKRigSolver* InSettings){};
+
+	/** override to support REMOVING a goal from custom solver */
+	virtual void RemoveGoal(const FName& GoalName) PURE_VIRTUAL("RemoveGoal");
+
+#if WITH_EDITORONLY_DATA
+	
+	/** callback whenever this solver is edited */
+	DECLARE_EVENT_OneParam(UIKRigSolver, FIKRigSolverModified, UIKRigSolver*);
+	FIKRigSolverModified& OnSolverModified(){ return IKRigSolverModified; };
+	
 	/** override to give your solver a nice name to display in the UI */
 	virtual FText GetNiceName() const { checkNoEntry() return FText::GetEmpty(); };
 	/** override to provide warning to user during setup of any missing components. return false if no warnings. */
@@ -54,8 +66,6 @@ public:
 	//** GOALS */
 	/** override to support ADDING a new goal to custom solver */
 	virtual void AddGoal(const UIKRigEffectorGoal& NewGoal) PURE_VIRTUAL("AddGoal");
-	/** override to support REMOVING a goal from custom solver */
-	virtual void RemoveGoal(const FName& GoalName) PURE_VIRTUAL("RemoveGoal");
 	/** override to support RENAMING an existing goal */
 	virtual void RenameGoal(const FName& OldName, const FName& NewName) PURE_VIRTUAL("RenameGoal");
 	/** override to support CHANGING BONE for an existing goal */
@@ -66,9 +76,6 @@ public:
 	virtual UObject* GetGoalSettings(const FName& GoalName) const {return nullptr;};
 	//** END GOALS */
 
-	//** ROOT BONE (optional, implement if your solver requires a root bone) */
-	/** if solver requires a root bone, then override this to return it. */
-	virtual FName GetRootBone() const { return NAME_None; };
 	//** ROOT BONE (optional, implement if your solver requires a root bone) */
 	/** override to support SETTING ROOT BONE for the solver */
 	virtual void SetRootBone(const FName& RootBoneName){};
@@ -89,11 +96,6 @@ public:
 	/** override to support supplying per-bone settings to outside systems for editing/UI
 	 ** NOTE: This must be overriden on solvers that use bone settings.*/
 	virtual UObject* GetBoneSetting(const FName& BoneName) const { ensure(!UsesBoneSettings()); return nullptr; };
-	/** override to support telling outside systems which bones this solver has setting for.
-	* NOTE: This must be overriden on solvers that use bone settings.
-	* NOTE: Only ADD to the incoming set, do not remove from it. */
-	virtual void GetBonesWithSettings(TSet<FName>& OutBonesWithSettings) const {};
-	//** END BONE SETTINGS */
 	
 	/** override to tell systems if this solver supports per-bone settings */
 	virtual bool UsesBoneSettings() const { return false;};
@@ -103,20 +105,17 @@ public:
 	virtual bool IsBoneAffectedBySolver(const FName& BoneName, const FIKRigSkeleton& IKRigSkeleton) const { return false; };
 	//** END ROOT BONE */
 
-#if WITH_EDITORONLY_DATA
 	/** UObject interface */
 	virtual void PostLoad() override;
 	virtual void PostTransacted(const FTransactionObjectEvent& TransactionEvent) override;
 	/** END UObject interface */
-#endif
 
 private:
 
-#if WITH_EDITORONLY_DATA
 	/** Register callbacks to update IK Rig when a solver is modified */
 	FIKRigSolverModified IKRigSolverModified;
 #endif
-
+	
 	UPROPERTY()
 	bool bIsEnabled = true;
-};
+};
\ No newline at end of file
diff --git a/Engine/Plugins/Animation/IKRig/Source/IKRig/Public/Solvers/IKRig_BodyMover.h b/Engine/Plugins/Animation/IKRig/Source/IKRig/Public/Solvers/IKRig_BodyMover.h
--- a/Engine/Plugins/Animation/IKRig/Source/IKRig/Public/Solvers/IKRig_BodyMover.h
+++ b/Engine/Plugins/Animation/IKRig/Source/IKRig/Public/Solvers/IKRig_BodyMover.h
@@ -89,14 +89,15 @@ public:
 	virtual void Solve(FIKRigSkeleton& IKRigSkeleton, const FIKRigGoalContainer& Goals) override;
 
 	virtual FName GetRootBone() const override { return RootBone; };
-#if WITH_EDITOR
+	virtual void RemoveGoal(const FName& GoalName) override;
 	virtual void UpdateSolverSettings(UIKRigSolver* InSettings) override;
+
+#if WITH_EDITOR
 	virtual FText GetNiceName() const override;
 	virtual bool GetWarningMessage(FText& OutWarningMessage) const override;
 	virtual bool IsBoneAffectedBySolver(const FName& BoneName, const FIKRigSkeleton& IKRigSkeleton) const override;
 	// goals
 	virtual void AddGoal(const UIKRigEffectorGoal& NewGoal) override;
-	virtual void RemoveGoal(const FName& GoalName) override;
 	virtual void RenameGoal(const FName& OldName, const FName& NewName) override;
 	virtual void SetGoalBone(const FName& GoalName, const FName& NewBoneName) override;
 	virtual bool IsGoalConnected(const FName& GoalName) const override;
diff --git a/Engine/Plugins/Animation/IKRig/Source/IKRig/Public/Solvers/IKRig_PoleSolver.h b/Engine/Plugins/Animation/IKRig/Source/IKRig/Public/Solvers/IKRig_PoleSolver.h
--- a/Engine/Plugins/Animation/IKRig/Source/IKRig/Public/Solvers/IKRig_PoleSolver.h
+++ b/Engine/Plugins/Animation/IKRig/Source/IKRig/Public/Solvers/IKRig_PoleSolver.h
@@ -45,14 +45,16 @@ public:
 	/** UIKRigSolver interface */
 	virtual void Initialize(const FIKRigSkeleton& IKRigSkeleton) override;
 	virtual void Solve(FIKRigSkeleton& IKRigSkeleton, const FIKRigGoalContainer& Goals) override;
-#if WITH_EDITOR
+	
 	virtual void UpdateSolverSettings(UIKRigSolver* InSettings) override;
+	virtual void RemoveGoal(const FName& GoalName) override;
+
+#if WITH_EDITOR
 	virtual FText GetNiceName() const override;
 	virtual bool GetWarningMessage(FText& OutWarningMessage) const override;
 	virtual bool IsBoneAffectedBySolver(const FName& BoneName, const FIKRigSkeleton& IKRigSkeleton) const override;
 	// goals
 	virtual void AddGoal(const UIKRigEffectorGoal& NewGoal) override;
-	virtual void RemoveGoal(const FName& GoalName) override;
 	virtual void RenameGoal(const FName& OldName, const FName& NewName) override;
 	virtual void SetGoalBone(const FName& GoalName, const FName& NewBoneName) override;
 	virtual bool IsGoalConnected(const FName& GoalName) const override;
diff --git a/Engine/Plugins/Animation/IKRig/Source/IKRig/Public/Solvers/IKRig_SetTransform.h b/Engine/Plugins/Animation/IKRig/Source/IKRig/Public/Solvers/IKRig_SetTransform.h
--- a/Engine/Plugins/Animation/IKRig/Source/IKRig/Public/Solvers/IKRig_SetTransform.h
+++ b/Engine/Plugins/Animation/IKRig/Source/IKRig/Public/Solvers/IKRig_SetTransform.h
@@ -48,12 +48,12 @@ public:
 	/** UIKRigSolver interface */
 	virtual void Initialize(const FIKRigSkeleton& IKRigSkeleton) override;
 	virtual void Solve(FIKRigSkeleton& IKRigSkeleton, const FIKRigGoalContainer& Goals) override;
-#if WITH_EDITOR
 	virtual void UpdateSolverSettings(UIKRigSolver* InSettings) override;
+	virtual void RemoveGoal(const FName& GoalName) override;
+#if WITH_EDITOR
 	virtual FText GetNiceName() const override;
 	virtual bool GetWarningMessage(FText& OutWarningMessage) const override;
 	virtual void AddGoal(const UIKRigEffectorGoal& NewGoal) override;
-	virtual void RemoveGoal(const FName& GoalName) override;
 	virtual void RenameGoal(const FName& OldName, const FName& NewName) override;
 	virtual void SetGoalBone(const FName& GoalName, const FName& NewBoneName) override;
 	virtual bool IsGoalConnected(const FName& GoalName) const override;
