Fixing IKRig build issues (missing file in original changelist 21574836)

diff --git a/Engine/Plugins/Animation/IKRig/Source/IKRig/Private/Solvers/IKRig_PBIKSolver.cpp b/Engine/Plugins/Animation/IKRig/Source/IKRig/Private/Solvers/IKRig_PBIKSolver.cpp
--- a/Engine/Plugins/Animation/IKRig/Source/IKRig/Private/Solvers/IKRig_PBIKSolver.cpp
+++ b/Engine/Plugins/Animation/IKRig/Source/IKRig/Private/Solvers/IKRig_PBIKSolver.cpp
@@ -142,8 +142,6 @@ void UIKRigPBIKSolver::GetBonesWithSettings(TSet<FName>& OutBonesWithSettings) c
 	}
 }
 
-#if WITH_EDITOR
-
 void UIKRigPBIKSolver::UpdateSolverSettings(UIKRigSolver* InSettings)
 {
 	if(UIKRigPBIKSolver* Settings = Cast<UIKRigPBIKSolver>(InSettings))
@@ -183,6 +181,21 @@ void UIKRigPBIKSolver::UpdateSolverSettings(UIKRigSolver* InSettings)
 	}
 }
 
+void UIKRigPBIKSolver::RemoveGoal(const FName& GoalName)
+{
+	// ensure goal even exists in this solver
+	const int32 GoalIndex = GetIndexOfGoal(GoalName);
+	if (GoalIndex == INDEX_NONE)
+	{
+		return;
+	}
+
+	// remove it
+	Effectors.RemoveAt(GoalIndex);
+}
+
+#if WITH_EDITOR
+
 FText UIKRigPBIKSolver::GetNiceName() const
 {
 	return FText(LOCTEXT("SolverName", "Full Body IK"));
@@ -213,19 +226,6 @@ void UIKRigPBIKSolver::AddGoal(const UIKRigEffectorGoal& NewGoal)
 	Effectors.Add(NewEffector);
 }
 
-void UIKRigPBIKSolver::RemoveGoal(const FName& GoalName)
-{
-	// ensure goal even exists in this solver
-	const int32 GoalIndex = GetIndexOfGoal(GoalName);
-	if (GoalIndex == INDEX_NONE)
-	{
-		return;
-	}
-
-	// remove it
-	Effectors.RemoveAt(GoalIndex);
-}
-
 void UIKRigPBIKSolver::RenameGoal(const FName& OldName, const FName& NewName)
 {
 	// ensure goal even exists in this solver
@@ -352,19 +352,6 @@ bool UIKRigPBIKSolver::IsBoneAffectedBySolver(const FName& BoneName, const FIKRi
 	return false;
 }
 
-int32 UIKRigPBIKSolver::GetIndexOfGoal(const FName& OldName) const
-{
-	for (int32 i=0; i<Effectors.Num(); ++i)
-	{
-		if (Effectors[i]->GoalName == OldName)
-		{
-			return i;
-		}
-	}
-
-	return INDEX_NONE;
-}
-
 void UIKRigPBIKSolver::PostLoad()
 {
 	Super::PostLoad();
@@ -387,4 +374,17 @@ void UIKRigPBIKSolver::PostLoad()
 
 #endif
 
+int32 UIKRigPBIKSolver::GetIndexOfGoal(const FName& OldName) const
+{
+	for (int32 i=0; i<Effectors.Num(); ++i)
+	{
+		if (Effectors[i]->GoalName == OldName)
+		{
+			return i;
+		}
+	}
+
+	return INDEX_NONE;
+}
+
 #undef LOCTEXT_NAMESPACE
