Fix for skeletal debug draw bones hit proxy memory leak

diff --git a/Engine/Plugins/Animation/IKRig/Source/IKRigEditor/Private/RetargetEditor/IKRetargetEditPoseMode.cpp b/Engine/Plugins/Animation/IKRig/Source/IKRigEditor/Private/RetargetEditor/IKRetargetEditPoseMode.cpp
--- a/Engine/Plugins/Animation/IKRig/Source/IKRigEditor/Private/RetargetEditor/IKRetargetEditPoseMode.cpp
+++ b/Engine/Plugins/Animation/IKRig/Source/IKRigEditor/Private/RetargetEditor/IKRetargetEditPoseMode.cpp
@@ -95,11 +95,11 @@ void FIKRetargetEditPoseMode::RenderSkeleton(
 	DrawConfig.SelectedBoneColor = GetMutableDefault<UPersonaOptions>()->SelectedBoneColor;
 	DrawConfig.ParentOfSelectedBoneColor = GetMutableDefault<UPersonaOptions>()->ParentOfSelectedBoneColor;
 
-	TArray<HHitProxy*> HitProxies;
-	HitProxies.AddUninitialized(NumBones);
+	TArray<TRefCountPtr<HHitProxy>> HitProxies;
+	HitProxies.Reserve(NumBones);
 	for (int32 Index = 0; Index < NumBones; ++Index)
 	{
-		HitProxies[Index] = new HIKRetargetEditorBoneProxy(RefSkeleton.GetBoneName(Index));
+		HitProxies.Add(new HIKRetargetEditorBoneProxy(RefSkeleton.GetBoneName(Index)));
 	}
 
 	// record selected bone indices
diff --git a/Engine/Plugins/Animation/IKRig/Source/IKRigEditor/Private/RigEditor/IKRigEditMode.cpp b/Engine/Plugins/Animation/IKRig/Source/IKRigEditor/Private/RigEditor/IKRigEditMode.cpp
--- a/Engine/Plugins/Animation/IKRig/Source/IKRigEditor/Private/RigEditor/IKRigEditMode.cpp
+++ b/Engine/Plugins/Animation/IKRig/Source/IKRigEditor/Private/RigEditor/IKRigEditMode.cpp
@@ -170,11 +170,11 @@ void FIKRigEditMode::RenderBones(FPrimitiveDrawInterface* PDI)
 	DrawConfig.SelectedBoneColor = GetMutableDefault<UPersonaOptions>()->SelectedBoneColor;
 	DrawConfig.ParentOfSelectedBoneColor = GetMutableDefault<UPersonaOptions>()->ParentOfSelectedBoneColor;
 
-	TArray<HHitProxy*> HitProxies;
-	HitProxies.AddUninitialized(NumBones);
+	TArray<TRefCountPtr<HHitProxy>> HitProxies;
+	HitProxies.Reserve(NumBones);
 	for (int32 Index = 0; Index < NumBones; ++Index)
 	{
-		HitProxies[Index] = new HIKRigEditorBoneProxy(RefSkeleton.GetBoneName(Index));
+		HitProxies.Add(new HIKRigEditorBoneProxy(RefSkeleton.GetBoneName(Index)));
 	}
 	
 	// get selected bones
diff --git a/Engine/Source/Editor/Persona/Private/AnimationEditorViewportClient.cpp b/Engine/Source/Editor/Persona/Private/AnimationEditorViewportClient.cpp
--- a/Engine/Source/Editor/Persona/Private/AnimationEditorViewportClient.cpp
+++ b/Engine/Source/Editor/Persona/Private/AnimationEditorViewportClient.cpp
@@ -1680,14 +1680,14 @@ void FAnimationViewportClient::DrawBones(
 	DrawConfig.SelectedBoneColor = GetMutableDefault<UPersonaOptions>()->SelectedBoneColor;
 	DrawConfig.ParentOfSelectedBoneColor = GetMutableDefault<UPersonaOptions>()->ParentOfSelectedBoneColor;
 
-	TArray<HHitProxy*> HitProxies;
-	HitProxies.AddUninitialized(RefSkeleton.GetNum());
+	TArray<TRefCountPtr<HHitProxy>> HitProxies;
 
 	if (bAddHitProxy)
 	{
+		HitProxies.Reserve(RefSkeleton.GetNum());
 		for (int32 Index = 0; Index < RefSkeleton.GetNum(); ++Index)
 		{
-			HitProxies[Index] = new HPersonaBoneHitProxy(Index, RefSkeleton.GetBoneName(Index));
+			HitProxies.Add(new HPersonaBoneHitProxy(Index, RefSkeleton.GetBoneName(Index)));
 		}
 	}
 
diff --git a/Engine/Source/Runtime/Engine/Private/SkeletalDebugRendering.cpp b/Engine/Source/Runtime/Engine/Private/SkeletalDebugRendering.cpp
--- a/Engine/Source/Runtime/Engine/Private/SkeletalDebugRendering.cpp
+++ b/Engine/Source/Runtime/Engine/Private/SkeletalDebugRendering.cpp
@@ -252,7 +252,7 @@ void DrawBonesFromPoseWatch(
 		UseWorldTransforms,
 		/*SelectedBones*/TArray<int32>(),
 		/*BoneColors*/TArray<FLinearColor>(),
-		/*HitProxies*/TArray<HHitProxy*>(),
+		/*HitProxies*/TArray<TRefCountPtr<HHitProxy>>(),
 		DrawConfig
 	);
 #endif
@@ -266,7 +266,7 @@ void DrawBones(
 	const TArray<FTransform>& WorldTransforms,
 	const TArray<int32>& InSelectedBones,
 	const TArray<FLinearColor>& BoneColors,
-	const TArray<HHitProxy*>& HitProxies,
+	const TArray<TRefCountPtr<HHitProxy>>& HitProxies,
 	const FSkelDebugDrawConfig& DrawConfig)
 {
 	// first determine which bones to draw, and which to filter out
diff --git a/Engine/Source/Runtime/Engine/Public/SkeletalDebugRendering.h b/Engine/Source/Runtime/Engine/Public/SkeletalDebugRendering.h
--- a/Engine/Source/Runtime/Engine/Public/SkeletalDebugRendering.h
+++ b/Engine/Source/Runtime/Engine/Public/SkeletalDebugRendering.h
@@ -129,6 +129,6 @@ ENGINE_API void DrawBones(
 	const TArray<FTransform>& WorldTransforms,
 	const TArray<int32>& InSelectedBones,
 	const TArray<FLinearColor>& BoneColors,
-	const TArray<HHitProxy*>& HitProxies,
+	const TArray<TRefCountPtr<HHitProxy>>& HitProxies,
 	const FSkelDebugDrawConfig& DrawConfig);
 }
\ No newline at end of file
