InitializeProperties on Initialize_AnyThread instead of just when the anim instance is initialized.

diff --git a/Engine/Plugins/Animation/IKRig/Source/IKRig/Private/AnimNodes/AnimNode_IKRig.cpp b/Engine/Plugins/Animation/IKRig/Source/IKRig/Private/AnimNodes/AnimNode_IKRig.cpp
--- a/Engine/Plugins/Animation/IKRig/Source/IKRig/Private/AnimNodes/AnimNode_IKRig.cpp
+++ b/Engine/Plugins/Animation/IKRig/Source/IKRig/Private/AnimNodes/AnimNode_IKRig.cpp
@@ -178,6 +178,8 @@ void FAnimNode_IKRig::Initialize_AnyThread(const FAnimationInitializeContext& Co
 	{
 		IKRigProcessor = NewObject<UIKRigProcessor>(Context.AnimInstanceProxy->GetSkelMeshComponent());	
 	}
+
+	InitializeProperties(Context.GetAnimInstanceObject(), GetTargetClass());
 }
 
 void FAnimNode_IKRig::Update_AnyThread(const FAnimationUpdateContext& Context)
@@ -331,7 +333,7 @@ void FAnimNode_IKRig::InitializeProperties(const UObject* InSourceInstance, UCla
 
 	check(SourcePropertyNames.Num() == DestPropertyNames.Num());
 
-	UClass* SourceClass = InSourceInstance->GetClass();
+	const UClass* SourceClass = InSourceInstance->GetClass();
 	for (int32 PropIndex = 0; PropIndex < SourcePropertyNames.Num(); ++PropIndex)
 	{
 		const FName& SourceName = SourcePropertyNames[PropIndex];
@@ -340,68 +342,76 @@ void FAnimNode_IKRig::InitializeProperties(const UObject* InSourceInstance, UCla
 		SourceProperties.Add(SourceProperty);
 		DestProperties.Add(nullptr);
 
-		if (SourceProperty)
+		// property not found
+		if (!SourceProperty)
+		{
+			continue;
+		}
+		
+		// store update functions for later use in PropagateInputProperties to avoid looking for properties
+		// while evaluating
+		const FName& GoalPropertyName = DestPropertyNames[PropIndex];
+
+		// find the right goal
+		const int32 GoalIndex = Goals.IndexOfByPredicate([&GoalPropertyName](const FIKRigGoal& InGoal)
+		{
+			return GoalPropertyName.ToString().EndsWith(InGoal.Name.ToString());
+		});
+
+		// goal not found?
+		if (!Goals.IsValidIndex(GoalIndex))
 		{
-			// store update functions for later use in PropagateInputProperties to avoid looking for properties
-			// while evaluating
-			const FName& GoalPropertyName = DestPropertyNames[PropIndex];
+			continue;
+		}
+		
+		FIKRigGoal& Goal = Goals[GoalIndex];
 
-			// find the right goal
-			const int32 GoalIndex = Goals.IndexOfByPredicate([&GoalPropertyName](const FIKRigGoal& InGoal)
+		const FString GoalPropStr = GoalPropertyName.ToString();
+
+		if (GoalPropStr.StartsWith(AlphaPosPropStr))
+		{
+			UpdateFunctions.Add([&Goal, SourceProperty](const UObject* InSourceInstance)
 			{
-			 	return GoalPropertyName.ToString().EndsWith(InGoal.Name.ToString());
+				const double* AlphaValue = SourceProperty->ContainerPtrToValuePtr<double>(InSourceInstance);
+				Goal.PositionAlpha = FMath::Clamp<float>(*AlphaValue, 0.f, 1.f);
 			});
-
-			if (Goals.IsValidIndex(GoalIndex))
+		}
+		else if (GoalPropStr.StartsWith(AlphaRotPropStr))
+		{
+			UpdateFunctions.Add([&Goal, SourceProperty](const UObject* InSourceInstance)
 			{
-				FIKRigGoal& Goal = Goals[GoalIndex];
-
-				const FString GoalPropStr = GoalPropertyName.ToString();
-
-				if (GoalPropStr.StartsWith(AlphaPosPropStr))
-				{
-					UpdateFunctions.Add([&Goal, SourceProperty](const UObject* InSourceInstance)
-					{
-						const double* AlphaValue = SourceProperty->ContainerPtrToValuePtr<double>(InSourceInstance);
-						Goal.PositionAlpha = FMath::Clamp<float>(*AlphaValue, 0.f, 1.f);
-					});
-				}
-				else if (GoalPropStr.StartsWith(AlphaRotPropStr))
-				{
-					UpdateFunctions.Add([&Goal, SourceProperty](const UObject* InSourceInstance)
-					{
-						const double* AlphaValue = SourceProperty->ContainerPtrToValuePtr<double>(InSourceInstance);
-						Goal.RotationAlpha = FMath::Clamp<float>(*AlphaValue, 0.f, 1.f);
-					});
-				}
-				else if (GoalPropStr.StartsWith(PositionPropStr))
-				{
-					UpdateFunctions.Add([&Goal, SourceProperty](const UObject* InSourceInstance)
-					{
-						Goal.Position = *SourceProperty->ContainerPtrToValuePtr<FVector>(InSourceInstance);
-					});
-				}
-				else if (GoalPropStr.StartsWith(RotationPropStr))
-				{
-					UpdateFunctions.Add([&Goal, SourceProperty](const UObject* InSourceInstance)
-					{
-						Goal.Rotation = *SourceProperty->ContainerPtrToValuePtr<FRotator>(InSourceInstance);
-					});
-				}
-			}
+				const double* AlphaValue = SourceProperty->ContainerPtrToValuePtr<double>(InSourceInstance);
+				Goal.RotationAlpha = FMath::Clamp<float>(*AlphaValue, 0.f, 1.f);
+			});
+		}
+		else if (GoalPropStr.StartsWith(PositionPropStr))
+		{
+			UpdateFunctions.Add([&Goal, SourceProperty](const UObject* InSourceInstance)
+			{
+				Goal.Position = *SourceProperty->ContainerPtrToValuePtr<FVector>(InSourceInstance);
+			});
+		}
+		else if (GoalPropStr.StartsWith(RotationPropStr))
+		{
+			UpdateFunctions.Add([&Goal, SourceProperty](const UObject* InSourceInstance)
+			{
+				Goal.Rotation = *SourceProperty->ContainerPtrToValuePtr<FRotator>(InSourceInstance);
+			});
 		}
 	}
 }
 
 void FAnimNode_IKRig::PropagateInputProperties(const UObject* InSourceInstance)
 {
-	if (InSourceInstance)
+	if (!InSourceInstance)
 	{
-		Algo::ForEach(UpdateFunctions, [InSourceInstance](const UpdateFunction& InFunc)
-		{
-		 	InFunc(InSourceInstance);
-		});
+		return;
 	}
+	
+	Algo::ForEach(UpdateFunctions, [InSourceInstance](const PropertyUpdateFunction& InFunc)
+	{
+		InFunc(InSourceInstance);
+	});
 }
 
 void FAnimNode_IKRig::ConditionalDebugDraw(
diff --git a/Engine/Plugins/Animation/IKRig/Source/IKRig/Public/AnimNodes/AnimNode_IKRig.h b/Engine/Plugins/Animation/IKRig/Source/IKRig/Public/AnimNodes/AnimNode_IKRig.h
--- a/Engine/Plugins/Animation/IKRig/Source/IKRig/Public/AnimNodes/AnimNode_IKRig.h
+++ b/Engine/Plugins/Animation/IKRig/Source/IKRig/Public/AnimNodes/AnimNode_IKRig.h
@@ -89,8 +89,8 @@ private:
 	TMap<FCompactPoseBoneIndex, int32, FDefaultSetAllocator, TCompactPoseBoneIndexMapKeyFuncs<int32>> CompactPoseToRigIndices;
 
 	/** Cached functions used to update goals using custom properties to avoid looking for then when evaluating */
-	using UpdateFunction = std::function<void(const UObject*)>;
-	TArray<UpdateFunction> UpdateFunctions;
+	using PropertyUpdateFunction = TFunction<void(const UObject*)>;
+	TArray<PropertyUpdateFunction> UpdateFunctions;
 
 	UPROPERTY(Transient)
 	float ActualAlpha;
