Add pole vector matching
Add chain column to IK Rig hierarchy view
Misc bug fixing

diff --git a/Engine/Plugins/Animation/IKRig/Source/IKRig/Private/IKRigProcessor.cpp b/Engine/Plugins/Animation/IKRig/Source/IKRig/Private/IKRigProcessor.cpp
--- a/Engine/Plugins/Animation/IKRig/Source/IKRig/Private/IKRigProcessor.cpp
+++ b/Engine/Plugins/Animation/IKRig/Source/IKRig/Private/IKRigProcessor.cpp
@@ -141,6 +141,19 @@ void UIKRigProcessor::Initialize(
 		Solvers.Add(Solver);
 	}
 
+	// validate retarget chains
+	const TArray<FBoneChain>& Chains = InRigAsset->GetRetargetChains();
+	TArray<int32> OutBoneIndices;
+	for (const FBoneChain& Chain : Chains)
+	{
+		if (!Skeleton.ValidateChainAndGetBones(Chain, OutBoneIndices))
+		{
+			Log.LogWarning(FText::Format(
+				LOCTEXT("InvalidRetargetChain", "Invalid Retarget Chain: '{0}'. End bone is not a child of the start bone in Skeletal Mesh, '{1}'."),
+				FText::FromString(Chain.ChainName.ToString()), FText::FromString(SkeletalMesh->GetName())));
+		}
+	}
+
 	Log.LogInfo(FText::Format(
 				LOCTEXT("SuccessfulInit", "IK Rig, '{0}' ready to run on {1}."),
 				FText::FromString(InRigAsset->GetName()), FText::FromString(SkeletalMesh->GetName())));
@@ -352,6 +365,11 @@ const FIKRigGoalContainer& UIKRigProcessor::GetGoalContainer() const
 	return GoalContainer;
 }
 
+const FGoalBone* UIKRigProcessor::GetGoalBone(const FName& GoalName) const
+{
+	return GoalBones.Find(GoalName);
+}
+
 FIKRigSkeleton& UIKRigProcessor::GetSkeletonWriteable()
 {
 	return Skeleton;
diff --git a/Engine/Plugins/Animation/IKRig/Source/IKRig/Private/IKRigSkeleton.cpp b/Engine/Plugins/Animation/IKRig/Source/IKRig/Private/IKRigSkeleton.cpp
--- a/Engine/Plugins/Animation/IKRig/Source/IKRig/Private/IKRigSkeleton.cpp
+++ b/Engine/Plugins/Animation/IKRig/Source/IKRig/Private/IKRigSkeleton.cpp
@@ -370,8 +370,10 @@ void FIKRigSkeleton::GetChainsInList(const TArray<int32>& SelectedBones, TArray<
     }
 }
 
-bool FIKRigSkeleton::GetBonesInChain(const FBoneChain& BoneChain, TArray<int32>& OutBoneIndices) const
+bool FIKRigSkeleton::ValidateChainAndGetBones(const FBoneChain& BoneChain, TArray<int32>& OutBoneIndices) const
 {
+	OutBoneIndices.Reset();
+	
 	// validate start and end bones exist and are not the root
 	const int32 StartIndex = GetBoneIndexFromName(BoneChain.StartBone.BoneName);
 	const int32 EndIndex = GetBoneIndexFromName(BoneChain.EndBone.BoneName);
diff --git a/Engine/Plugins/Animation/IKRig/Source/IKRig/Private/Retargeter/IKRetargetProcessor.cpp b/Engine/Plugins/Animation/IKRig/Source/IKRig/Private/Retargeter/IKRetargetProcessor.cpp
--- a/Engine/Plugins/Animation/IKRig/Source/IKRig/Private/Retargeter/IKRetargetProcessor.cpp
+++ b/Engine/Plugins/Animation/IKRig/Source/IKRig/Private/Retargeter/IKRetargetProcessor.cpp
@@ -22,7 +22,7 @@ void FRetargetSkeleton::Initialize(
 	const TArray<FBoneChain>& BoneChains,
 	const FName InRetargetPoseName,
 	const FIKRetargetPose* RetargetPose,
-	const FName& RetargetRootBone)
+	const FName RetargetRootBone)
 {
 	// reset all skeleton data
 	Reset();
@@ -73,7 +73,7 @@ void FRetargetSkeleton::Reset()
 void FRetargetSkeleton::GenerateRetargetPose(
 	const FName InRetargetPoseName,
 	const FIKRetargetPose* InRetargetPose,
-	const FName& RetargetRootBone)
+	const FName RetargetRootBone)
 {
 	// record the name of the retarget pose (prevents re-initialization if profile swaps it)
 	RetargetPoseName = InRetargetPoseName;
@@ -123,7 +123,7 @@ void FRetargetSkeleton::GenerateRetargetPose(
 
 int32 FRetargetSkeleton::FindBoneIndexByName(const FName InName) const
 {
-	return BoneNames.IndexOfByPredicate([&InName](const FName& BoneName)
+	return BoneNames.IndexOfByPredicate([&InName](const FName BoneName)
 	{
 		return BoneName == InName;
 	});
@@ -184,7 +184,9 @@ void FRetargetSkeleton::UpdateLocalTransformOfSingleBone(
 	const int32 ParentIndex = ParentIndices[BoneIndex];
 	if (ParentIndex == INDEX_NONE)
 	{
-		return; // root always in global space
+		// root bone, so just set the local pose to the global pose
+		OutLocalPose[BoneIndex] = InGlobalPose[BoneIndex];
+		return;
 	}
 	const FTransform& ChildGlobalTransform = InGlobalPose[BoneIndex];
 	const FTransform& ParentGlobalTransform = InGlobalPose[ParentIndex];
@@ -307,9 +309,9 @@ int32 FRetargetSkeleton::GetParentIndex(const int32 BoneIndex) const
 void FTargetSkeleton::Initialize(
 	USkeletalMesh* InSkeletalMesh,
 	const TArray<FBoneChain>& BoneChains,
-	const FName& InRetargetPoseName,
+	const FName InRetargetPoseName,
 	const FIKRetargetPose* RetargetPose,
-	const FName& RetargetRootBone)
+	const FName RetargetRootBone)
 {
 	Reset();
 	
@@ -1110,6 +1112,21 @@ bool FRetargetChainPairFK::Initialize(
 		return false;
 	}
 
+	// initialize the pole vector matcher for this chain
+	const bool bPoleVectorMatcherInitialized = PoleVectorMatcher.Initialize(
+		SourceBoneIndices,
+		TargetBoneIndices,
+		SourceSkeleton.RetargetGlobalPose,
+		TargetSkeleton.RetargetGlobalPose,
+		TargetSkeleton);
+	if (!bPoleVectorMatcherInitialized)
+	{
+		Log.LogWarning( FText::Format(
+			LOCTEXT("BadFKDecoder", "IK Retargeter failed to initialize pole matching for chain, '{0}', on Skeletal Mesh: '{1}'"),
+			FText::FromName(TargetBoneChainName), FText::FromString(TargetSkeleton.SkeletalMesh->GetName())));
+		return false;
+	}
+
 	return true;
 }
 
@@ -1166,6 +1183,7 @@ bool FRootRetargeter::InitializeSource(
 	FIKRigLogger& Log)
 {
 	// validate target root bone exists
+	Source.BoneName = SourceRootBoneName;
 	Source.BoneIndex = SourceSkeleton.FindBoneIndexByName(SourceRootBoneName);
 	if (Source.BoneIndex == INDEX_NONE)
 	{
@@ -1200,6 +1218,7 @@ bool FRootRetargeter::InitializeTarget(
 	FIKRigLogger& Log)
 {
 	// validate target root bone exists
+	Target.BoneName = TargetRootBoneName;
 	Target.BoneIndex = TargetSkeleton.FindBoneIndexByName(TargetRootBoneName);
 	if (Target.BoneIndex == INDEX_NONE)
 	{
@@ -1283,6 +1302,164 @@ void FRootRetargeter::DecodePose(TArray<FTransform>& OutTargetGlobalPose)
 	TargetRootTransform.SetRotation(Rotation);
 }
 
+bool FPoleVectorMatcher::Initialize(
+	const TArray<int32>& SourceIndices,
+	const TArray<int32>& TargetIndices,
+	const TArray<FTransform>& SourceGlobalPose,
+	const TArray<FTransform>& TargetGlobalPose,
+	const FRetargetSkeleton& TargetSkeleton)
+{
+	SourcePoleAxis = CalculateBestPoleAxisForChain(SourceIndices, SourceGlobalPose);
+	TargetPoleAxis = CalculateBestPoleAxisForChain(TargetIndices, TargetGlobalPose);
+
+	const FVector SourcePoleVector = CalculatePoleVector(SourcePoleAxis, SourceIndices, SourceGlobalPose);
+	const FVector TargetPoleVector = CalculatePoleVector(TargetPoleAxis, TargetIndices, TargetGlobalPose);
+	
+	TargetToSourceAngularOffsetAtRefPose = FMath::Acos(FVector::DotProduct(TargetPoleVector, SourcePoleVector));
+
+	// cache indices of bones in chain (and their children recursively) excluding children of the end bone
+	TArray<int32> AllChildrenIndices;
+	TargetSkeleton.GetChildrenIndicesRecursive(TargetIndices[0],AllChildrenIndices);
+	TArray<int32> ChildrenOfEndIndices = {TargetIndices.Last()};
+	TargetSkeleton.GetChildrenIndicesRecursive(TargetIndices.Last(),ChildrenOfEndIndices);
+	AllChildrenWithinChain.Reset();
+	for (int32 ChildIndex : AllChildrenIndices)
+	{
+		if (ChildrenOfEndIndices.Contains(ChildIndex))
+		{
+			continue;
+		}
+		AllChildrenWithinChain.Add(ChildIndex);
+	}
+
+	return true;
+}
+
+void FPoleVectorMatcher::MatchPoleVector(
+	const FTargetChainSettings& Settings,
+	const TArray<int32>& SourceIndices,
+	const TArray<int32>& TargetIndices,
+	const TArray<FTransform>& SourceGlobalPose,
+	TArray<FTransform>& OutTargetGlobalPose,
+	FRetargetSkeleton& TargetSkeleton)
+{
+	if (Settings.FK.PoleVectorMatching <= KINDA_SMALL_NUMBER &&
+		Settings.FK.PoleVectorOffset <= KINDA_SMALL_NUMBER)
+	{
+		return;
+	}
+	
+	// update local spaces of all bones within chain
+	TArray<FTransform> LocalSpaces;
+	for (const int32 ChildIndex : AllChildrenWithinChain)
+	{
+		const int32 ParentIndex = TargetSkeleton.ParentIndices[ChildIndex];
+		if (ParentIndex == INDEX_NONE)
+		{
+			LocalSpaces.Add(OutTargetGlobalPose[ChildIndex]);
+		}
+		const FTransform& ChildGlobalTransform = OutTargetGlobalPose[ChildIndex];
+		const FTransform& ParentGlobalTransform = OutTargetGlobalPose[ParentIndex];
+		const FTransform LocalTransform = ChildGlobalTransform.GetRelativeTransform(ParentGlobalTransform);
+		LocalSpaces.Add(LocalTransform);
+	}
+
+	// calculate rotation to match the target to the source pole vector
+	const FVector SourcePoleVector = CalculatePoleVector(SourcePoleAxis, SourceIndices, SourceGlobalPose);
+	const FVector TargetPoleVector = CalculatePoleVector(TargetPoleAxis, TargetIndices, OutTargetGlobalPose);
+	
+	const float RotateTargetToSource = FMath::Acos(FVector::DotProduct(SourcePoleVector, TargetPoleVector));
+	const float MatchPoleAngle = RotateTargetToSource - TargetToSourceAngularOffsetAtRefPose + FMath::DegreesToRadians(Settings.FK.PoleVectorOffset);
+
+	const FVector TargetChainNormal = GetChainNormal(TargetIndices, OutTargetGlobalPose);
+	const FQuat PoleVectorMatchingRotation = FQuat(TargetChainNormal, MatchPoleAngle);
+	const FQuat FinalRotation = FQuat::FastLerp(FQuat::Identity, PoleVectorMatchingRotation, Settings.FK.PoleVectorMatching).GetNormalized();
+
+	// rotate the base of the chain to match the pole vectors
+	FTransform& BaseOfChain = OutTargetGlobalPose[TargetIndices[0]];
+	BaseOfChain.SetRotation(FinalRotation * BaseOfChain.GetRotation());
+
+	// now update global pose of all bones within chain
+	for (int32 ChainIndex=0; ChainIndex < AllChildrenWithinChain.Num(); ++ChainIndex)
+	{
+		const int32 ChildIndex = AllChildrenWithinChain[ChainIndex];
+		const int32 ParentIndex = TargetSkeleton.ParentIndices[ChildIndex];
+		if (ParentIndex == INDEX_NONE)
+		{
+			// root always in global space already, no conversion required
+			OutTargetGlobalPose[ChildIndex] = LocalSpaces[ChainIndex];
+			return; 
+		}
+		const FTransform& ChildLocalTransform = LocalSpaces[ChainIndex];
+		const FTransform& ParentGlobalTransform = OutTargetGlobalPose[ParentIndex];
+		OutTargetGlobalPose[ChildIndex] = ChildLocalTransform * ParentGlobalTransform;
+	}
+}
+
+EAxis::Type FPoleVectorMatcher::CalculateBestPoleAxisForChain(
+	const TArray<int32>& BoneIndices,
+	const TArray<FTransform>& GlobalPose)
+{
+	check(!BoneIndices.IsEmpty());
+
+	TArray<EAxis::Type> PreferredAxes;
+
+	// handle 1-bone case, there is no "chain" to speak of just a single bone
+	// so we arbitrarily pick an axis and move on... realistically, this feature is fairly meaningless for a single bone
+	if (BoneIndices.Num() == 1)
+	{
+		return EAxis::Y;
+	}
+	
+	const FVector ChainOrigin = GlobalPose[BoneIndices[0]].GetLocation();
+	const FVector ChainAxisNormal = (GlobalPose[BoneIndices.Last()].GetLocation() - ChainOrigin).GetSafeNormal();
+	const EAxis::Type MostDifferentAxis = GetMostDifferentAxis(GlobalPose[BoneIndices[0]], ChainAxisNormal);
+	return MostDifferentAxis;
+}
+
+FVector FPoleVectorMatcher::CalculatePoleVector(
+	const EAxis::Type& PoleAxis,
+	const TArray<int32>& BoneIndices,
+	const TArray<FTransform>& GlobalPose)
+{
+	check(!BoneIndices.IsEmpty())
+
+	const FVector ChainNormal = GetChainNormal(BoneIndices, GlobalPose);
+	const FVector UnitPoleAxis = GlobalPose[BoneIndices[0]].GetUnitAxis(PoleAxis);
+	const FVector PoleVector = FVector::VectorPlaneProject(UnitPoleAxis, ChainNormal);
+	return PoleVector.GetSafeNormal();
+}
+
+EAxis::Type FPoleVectorMatcher::GetMostDifferentAxis(
+	const FTransform& Transform,
+	const FVector& InNormal)
+{
+	float MostDifferentDot = 2.0f;
+	EAxis::Type MostDifferentAxis = EAxis::Y;
+	const TArray<EAxis::Type> CardinalAxes = {EAxis::X, EAxis::Y, EAxis::Z};
+	for (const EAxis::Type Axis : CardinalAxes)
+	{
+		const FVector AxisVector = Transform.GetUnitAxis(Axis);
+		const float AbsAxisDotNormal = FMath::Abs(FVector::DotProduct(AxisVector, InNormal));
+		if (AbsAxisDotNormal < MostDifferentDot)
+		{
+			MostDifferentDot = AbsAxisDotNormal;
+			MostDifferentAxis = Axis;
+		}
+	}
+
+	return MostDifferentAxis;
+}
+
+FVector FPoleVectorMatcher::GetChainNormal(
+	const TArray<int32>& BoneIndices,
+	const TArray<FTransform>& GlobalPose)
+{
+	const FVector ChainOrigin = GlobalPose[BoneIndices[0]].GetLocation();
+	const FVector ChainAxis = GlobalPose[BoneIndices.Last()].GetLocation() - ChainOrigin;
+	return ChainAxis.GetSafeNormal();
+}
+
 UIKRetargetProcessor::UIKRetargetProcessor()
 {
 	const FName LogName = FName("IKRetarget_",GetUniqueID());
@@ -1530,7 +1707,7 @@ bool UIKRetargetProcessor::InitializeBoneChainPairs()
 bool UIKRetargetProcessor::InitializeIKRig(UObject* Outer, const USkeletalMesh* InSkeletalMesh)
 {
 	// get the target IK Rig asset
-	const UIKRigDefinition* IKRig = RetargeterAsset->GetTargetIKRig();
+	UIKRigDefinition* IKRig = RetargeterAsset->GetTargetIKRigWriteable();
 	
 	// gather list of excluded goals based on any chain mapping that has it's IK disabled
 	TArray<FName> GoalsToExclude;
@@ -1571,6 +1748,68 @@ bool UIKRetargetProcessor::InitializeIKRig(UObject* Outer, const USkeletalMesh*
 			return false;
 		}
 	}
+
+	// get a single retarget chain using the specified goal (if there is one)
+	auto GetRetargetChainByGoalName = [](const TArray<FBoneChain>& Chains, const FName GoalName) -> const FBoneChain*
+	{
+		for (const FBoneChain& Chain : Chains)
+		{
+			if (Chain.IKGoalName == GoalName)
+			{
+				return &Chain;
+			}
+		}
+
+		return nullptr;
+	};
+
+	// detect goals that are not being controlled by the retargeter and set them to "additive" mode
+	const TArray<FIKRigGoal>& GoalsInProcessor = IKRigProcessor->GetGoalContainer().GetGoalArray();
+	const TArray<FBoneChain>& ChainsInAsset = IKRig->GetRetargetChains();
+	const FIKRigSkeleton& IKRigSkeleton = IKRigProcessor->GetSkeleton();
+	for (const FIKRigGoal& GoalInProcessor : GoalsInProcessor)
+	{
+		if (GetRetargetChainByGoalName(ChainsInAsset, GoalInProcessor.Name))
+		{
+			// this goal belongs to a chain, so leave it alone
+			continue;
+		}
+
+		// this goal has no associated chain, so we set it to "additive" mode to prevent the goal
+		// from being left at the starting location (and thus completely breaking the pose)
+		// when in "additive" mode, any offset applied in the IK Rig will be maintained, but will be relative
+		// to the input pose, not in absolute/component space.
+		FIKRigGoal GoalToModify = GoalInProcessor;
+		GoalToModify.PositionSpace = EIKRigGoalSpace::Additive;
+		GoalToModify.RotationSpace = EIKRigGoalSpace::Additive;
+
+		// but for this to work, the goal's position/rotation must first be converted to "additive" space
+		const FGoalBone* GoalBone = IKRigProcessor->GetGoalBone(GoalToModify.Name);
+		const int32& BoneIndex = GoalBone->BoneIndex;
+		check(BoneIndex != INDEX_NONE)
+
+		const FTransform& GoalBoneGlobalTransform = IKRigSkeleton.RefPoseGlobal[BoneIndex];
+	
+		// convert position to additive space
+		if (GoalToModify.PositionSpace == EIKRigGoalSpace::World || GoalToModify.PositionSpace == EIKRigGoalSpace::Component)
+		{
+			GoalToModify.Position = GoalToModify.Position - GoalBoneGlobalTransform.GetLocation();
+		}
+
+		// convert rotation to additive space
+		if (GoalToModify.RotationSpace == EIKRigGoalSpace::World || GoalToModify.RotationSpace == EIKRigGoalSpace::Component)
+		{
+			GoalToModify.Rotation = (GoalToModify.Rotation.Quaternion() * GoalBoneGlobalTransform.GetRotation().Inverse()).Rotator();
+		}
+		
+		// assign the new goal forced to be additive
+		IKRigProcessor->SetIKGoal(GoalToModify);
+
+		// let the user know to avoid confusion
+		Log.LogWarning( FText::Format(
+		LOCTEXT("TargetIKGoalNotInChain", "The target IK Rig has a Goal, {0} that is not present in any of the target chains. It has been set to 'Additive' mode while retargeting."),
+		FText::FromName(GoalInProcessor.Name)));
+	}
 	
 	return true;
 }
@@ -1585,9 +1824,6 @@ TArray<FTransform>&  UIKRetargetProcessor::RunRetargeter(
 	// start from retarget pose
 	TargetSkeleton.OutputGlobalPose = TargetSkeleton.RetargetGlobalPose;
 
-	// get the global settings
-	const FRetargetGlobalSettings& GlobalSettings = RetargeterAsset->GetGlobalSettings();
-
 	// ROOT retargeting
 	if (GlobalSettings.bEnableRoot && bRootsInitialized)
 	{
@@ -1610,6 +1846,12 @@ TArray<FTransform>&  UIKRetargetProcessor::RunRetargeter(
 		RunIKRetarget(InSourceGlobalPose, TargetSkeleton.OutputGlobalPose, SpeedValuesFromCurves, DeltaTime);
 	}
 
+	// Pole Vector matching between source / target chains
+	if (GlobalSettings.bEnableFK && bAtLeastOneValidBoneChainPair)
+	{
+		RunPoleVectorMatching(InSourceGlobalPose, TargetSkeleton.OutputGlobalPose);
+	}
+
 	return TargetSkeleton.OutputGlobalPose;
 }
 
@@ -1700,11 +1942,26 @@ void UIKRetargetProcessor::RunIKRetarget(
 	IKRigProcessor->CopyOutputGlobalPoseToArray(OutTargetGlobalPose);
 }
 
+void UIKRetargetProcessor::RunPoleVectorMatching(
+	const TArray<FTransform>& InGlobalTransforms,
+	TArray<FTransform>& OutGlobalTransforms)
+{
+	for (FRetargetChainPairFK& ChainPair : ChainPairsFK)
+	{
+		ChainPair.PoleVectorMatcher.MatchPoleVector(
+			ChainPair.Settings,
+			ChainPair.SourceBoneIndices,
+			ChainPair.TargetBoneIndices,
+			InGlobalTransforms,
+			OutGlobalTransforms,
+			TargetSkeleton);
+	}
+}
+
 void UIKRetargetProcessor::RunStrideWarping(const TArray<FTransform>& InTargetGlobalPose)
 {
 	// check if we should apply stride warping
-	const FRetargetGlobalSettings& Settings = RetargeterAsset->GetGlobalSettings();
-	if (!(Settings.bIKWarping && bRootsInitialized))
+	if (!(GlobalSettings.bWarping && bRootsInitialized))
 	{
 		return;
 	}
@@ -1712,7 +1969,7 @@ void UIKRetargetProcessor::RunStrideWarping(const TArray<FTransform>& InTargetGl
 	// get sample points to use for "best fit" global body rotation
 	TArray<FVector> InitialPoints;
 	TArray<FVector> CurrentPoints;
-	if (Settings.DirectionSource == EWarpingDirectionSource::Goals)
+	if (GlobalSettings.DirectionSource == EWarpingDirectionSource::Goals)
 	{
 		// use goals to determine the body's rotation
 		for (const FRetargetChainPairIK& ChainPair : ChainPairsIK)
@@ -1736,9 +1993,9 @@ void UIKRetargetProcessor::RunStrideWarping(const TArray<FTransform>& InTargetGl
 	{
 		// use chain to determine the body's rotation
 		const FRetargetChainPairFK* SpineChain = ChainPairsFK.FindByPredicate(
-				[&Settings](const FRetargetChainPairFK& Other)
+				[this](const FRetargetChainPairFK& Other)
 				{
-					return Other.TargetBoneChainName == Settings.DirectionChain;
+					return Other.TargetBoneChainName == GlobalSettings.DirectionChain;
 				});
 		if (!SpineChain)
 		{
@@ -1775,9 +2032,9 @@ void UIKRetargetProcessor::RunStrideWarping(const TArray<FTransform>& InTargetGl
 	// generate axes based on body rotation
 	FTransform InitialGlobalGoalTransform = FTransform(FQuat::Identity, InitialGoalCentroid);
 	FTransform CurrentGlobalGoalTransform = FTransform(BestFitBodyRotation, CurrentGoalCentroid);
-	const FVector& Fwd = CurrentGlobalGoalTransform.TransformVector(Settings.GetAxisVector(Settings.ForwardDirection));
+	const FVector& Fwd = CurrentGlobalGoalTransform.TransformVector(GlobalSettings.GetAxisVector(GlobalSettings.ForwardDirection));
 	const FVector& Side = FVector::CrossProduct(Fwd, FVector::ZAxisVector);
-	const FVector SideOrig = FVector::CrossProduct(Settings.GetAxisVector(Settings.ForwardDirection), FVector::ZAxisVector);
+	const FVector SideOrig = FVector::CrossProduct(GlobalSettings.GetAxisVector(GlobalSettings.ForwardDirection), FVector::ZAxisVector);
 
 	// warp goal positions...
 	for (FRetargetChainPairIK& ChainPair : ChainPairsIK)
@@ -1797,16 +2054,16 @@ void UIKRetargetProcessor::RunStrideWarping(const TArray<FTransform>& InTargetGl
 		const FVector IntialGoalInCurrentSpace = CurrentGlobalGoalTransform.TransformPosition(InitialGoalInOrigSpace);
 		const FPlane FwdPlane(IntialGoalInCurrentSpace, Fwd);
 		const FVector GoalProjOnFwdPlane = FwdPlane.PointPlaneProject(WarpedGoalPosition, FwdPlane);
-		WarpedGoalPosition = GoalProjOnFwdPlane + ((WarpedGoalPosition - GoalProjOnFwdPlane) * Settings.WarpForwards);
+		WarpedGoalPosition = GoalProjOnFwdPlane + ((WarpedGoalPosition - GoalProjOnFwdPlane) * GlobalSettings.WarpForwards);
 
 		// sideways offset
 		// first determine which side the goal is on originally
 		float GoalSideMultiplier = FVector::DotProduct(ChainRetargeter.Target.InitialEndPosition.GetSafeNormal(), SideOrig);
 		// push goal by offset in the newly calculated sideways direction
-		WarpedGoalPosition += Side * Settings.SidewaysOffset * GoalSideMultiplier;
+		WarpedGoalPosition += Side * GlobalSettings.SidewaysOffset * GoalSideMultiplier;
 
 		// splay warping
-		WarpedGoalPosition = CurrentGoalCentroid + (WarpedGoalPosition - CurrentGoalCentroid) * Settings.WarpSplay; 
+		WarpedGoalPosition = CurrentGoalCentroid + (WarpedGoalPosition - CurrentGoalCentroid) * GlobalSettings.WarpSplay; 
 
 		// assign warped position
 		ChainRetargeter.Results.EndEffectorPosition = WarpedGoalPosition;
@@ -1910,37 +2167,21 @@ void UIKRetargetProcessor::ApplySettingsFromAsset()
 
 	// copy root settings
 	RootRetargeter.Settings.CopySettingsFromAsset(RetargeterAsset->GetRootSettingsUObject());
+
+	// copy global settings
+	GlobalSettings = RetargeterAsset->GetGlobalSettings();
 }
 
 void UIKRetargetProcessor::ApplySettingsFromProfile(const FRetargetProfile& Profile)
 {
-	auto ApplyRetargetPoseFromProfile = [this, Profile](ERetargetSourceOrTarget SourceOrTarget)
-	{
-		const bool bIsSource = SourceOrTarget == ERetargetSourceOrTarget::Source;
-		const FName PoseName = bIsSource ? Profile.SourceRetargetPoseName : Profile.TargetRetargetPoseName;
-		FRetargetSkeleton& RetargetSkeleton = bIsSource ? SourceSkeleton : TargetSkeleton;
-		if (PoseName == NAME_None || RetargetSkeleton.RetargetPoseName == PoseName)
-		{
-			return; // retarget pose not specified, or already in use
-		}
-		
-		const UIKRigDefinition* IKRig = bIsSource ? RetargeterAsset->GetSourceIKRig() : RetargeterAsset->GetTargetIKRig();
-		RetargetSkeleton.Initialize(
-			RetargetSkeleton.SkeletalMesh,
-			IKRig->GetRetargetChains(),
-			PoseName,
-			RetargeterAsset->GetRetargetPoseByName(SourceOrTarget, PoseName),
-			IKRig->GetRetargetRoot());		
-	};
-
 	// assign retarget poses specified in the profile
 	if (Profile.bApplySourceRetargetPose)
 	{
-		ApplyRetargetPoseFromProfile(ERetargetSourceOrTarget::Source);
+		ApplyNewRetargetPose(Profile.SourceRetargetPoseName, ERetargetSourceOrTarget::Source);
 	}
 	if (Profile.bApplyTargetRetargetPose)
 	{
-		ApplyRetargetPoseFromProfile(ERetargetSourceOrTarget::Target);
+		ApplyNewRetargetPose(Profile.TargetRetargetPoseName, ERetargetSourceOrTarget::Target);
 	}
 
 	// assign chain settings
@@ -1971,6 +2212,59 @@ void UIKRetargetProcessor::ApplySettingsFromProfile(const FRetargetProfile& Prof
 	{
 		RootRetargeter.Settings = Profile.RootSettings;
 	}
+
+	// assign global settings
+	if (Profile.bApplyGlobalSettings)
+	{
+		GlobalSettings = Profile.GlobalSettings;
+	}
+}
+
+void UIKRetargetProcessor::ApplyNewRetargetPose(
+	const FName NewRetargetPoseName,
+	ERetargetSourceOrTarget SourceOrTarget)
+{
+	const bool bIsSource = SourceOrTarget == ERetargetSourceOrTarget::Source;
+	FRetargetSkeleton& RetargetSkeleton = bIsSource ? SourceSkeleton : TargetSkeleton;
+	if (NewRetargetPoseName == NAME_None || RetargetSkeleton.RetargetPoseName == NewRetargetPoseName)
+	{
+		return; // retarget pose not specified, or already in use
+	}
+
+	const FIKRetargetPose* NewRetargetPose = RetargeterAsset->GetRetargetPoseByName(SourceOrTarget, NewRetargetPoseName);
+	if (!NewRetargetPose)
+	{
+		return;
+	}
+
+	const FName RootBoneName = bIsSource ? RootRetargeter.Source.BoneName : RootRetargeter.Target.BoneName;
+	if (RootBoneName == NAME_None)
+	{
+		return;
+	}
+
+	// re-generate the retarget pose
+	RetargetSkeleton.GenerateRetargetPose(NewRetargetPoseName, NewRetargetPose, RootBoneName);
+
+	// re-initialize the bone chains using the newly generated retarget pose
+	const UIKRigDefinition* SourceIKRig = RetargeterAsset->GetSourceIKRig();
+	const UIKRigDefinition* TargetIKRig = RetargeterAsset->GetTargetIKRig();
+	for (FRetargetChainPairFK& FKChainPair : ChainPairsFK)
+	{
+		const FBoneChain* SourceBoneChain = SourceIKRig->GetRetargetChainByName(FKChainPair.SourceBoneChainName);
+		const FBoneChain* TargetBoneChain = TargetIKRig->GetRetargetChainByName(FKChainPair.TargetBoneChainName);
+		FKChainPair.Initialize(*SourceBoneChain, *TargetBoneChain, SourceSkeleton, TargetSkeleton, Log);
+	}
+	for (FRetargetChainPairIK& IKChainPair : ChainPairsIK)
+	{
+		const FBoneChain* SourceBoneChain = SourceIKRig->GetRetargetChainByName(IKChainPair.SourceBoneChainName);
+		const FBoneChain* TargetBoneChain = TargetIKRig->GetRetargetChainByName(IKChainPair.TargetBoneChainName);
+		IKChainPair.Initialize(*SourceBoneChain, *TargetBoneChain, SourceSkeleton, TargetSkeleton, Log);
+	}
+
+	// re-initialize the root
+	RootRetargeter.InitializeSource(RootRetargeter.Source.BoneName, SourceSkeleton, Log);
+	RootRetargeter.InitializeTarget(RootRetargeter.Target.BoneName, TargetSkeleton, Log);
 }
 
 #undef LOCTEXT_NAMESPACE
diff --git a/Engine/Plugins/Animation/IKRig/Source/IKRig/Private/Retargeter/IKRetargeter.cpp b/Engine/Plugins/Animation/IKRig/Source/IKRig/Private/Retargeter/IKRetargeter.cpp
--- a/Engine/Plugins/Animation/IKRig/Source/IKRig/Private/Retargeter/IKRetargeter.cpp
+++ b/Engine/Plugins/Animation/IKRig/Source/IKRig/Private/Retargeter/IKRetargeter.cpp
@@ -305,6 +305,44 @@ const FRetargetProfile* UIKRetargeter::GetProfileByName(const FName& ProfileName
 	return Profiles.Find(ProfileName);
 }
 
+FTargetChainSettings UIKRetargeter::GetChainUsingGoalFromRetargetAsset(
+	const UIKRetargeter* RetargetAsset,
+	const FName IKGoalName)
+{
+	FTargetChainSettings EmptySettings;
+
+	if (!RetargetAsset)
+	{
+		return EmptySettings;
+	}
+
+	const UIKRigDefinition* IKRig = RetargetAsset->GetTargetIKRig();
+	if (!IKRig)
+	{
+		return EmptySettings;
+	}
+
+	const TArray<FBoneChain>& RetargetChains = IKRig->GetRetargetChains();
+	const FBoneChain* ChainWithGoal = nullptr;
+	for (const FBoneChain& RetargetChain : RetargetChains)
+	{
+		if (RetargetChain.IKGoalName == IKGoalName)
+		{
+			ChainWithGoal = &RetargetChain;
+			break;
+		}
+	}
+
+	if (!ChainWithGoal)
+	{
+		return EmptySettings;
+	}
+
+	// found a chain using the specified goal, return a copy of it's settings
+	const FTargetChainSettings* ChainSettings = RetargetAsset->GetChainSettingsByName(ChainWithGoal->ChainName);
+	return ChainSettings ? *ChainSettings : EmptySettings;
+}
+
 FTargetChainSettings UIKRetargeter::GetChainSettingsFromRetargetAsset(
 	const UIKRetargeter* RetargetAsset,
 	const FName TargetChainName,
@@ -381,6 +419,56 @@ void UIKRetargeter::GetRootSettingsFromRetargetAsset(
 	OutSettings =  RetargetAsset->GetRootSettingsUObject()->Settings;
 }
 
+FTargetRootSettings UIKRetargeter::GetRootSettingsFromRetargetProfile(FRetargetProfile& RetargetProfile)
+{
+	return RetargetProfile.RootSettings;
+}
+
+void UIKRetargeter::GetGlobalSettingsFromRetargetAsset(
+	const UIKRetargeter* RetargetAsset,
+	const FName OptionalProfileName,
+	FRetargetGlobalSettings& OutSettings)
+{
+	if (!RetargetAsset)
+	{
+		OutSettings = FRetargetGlobalSettings();
+		return;
+	}
+	
+	// optionally get the root settings from a profile
+	if (OptionalProfileName != NAME_None)
+	{
+		if (const FRetargetProfile* RetargetProfile = RetargetAsset->GetProfileByName(OptionalProfileName))
+		{
+			if (RetargetProfile->bApplyGlobalSettings)
+			{
+				OutSettings = RetargetProfile->GlobalSettings;
+				return;
+			}
+		}
+		
+		// could not find profile, so return default settings
+		OutSettings = FRetargetGlobalSettings();
+		return;
+	}
+
+	// return the base root settings
+	OutSettings = RetargetAsset->GetGlobalSettings();
+}
+
+FRetargetGlobalSettings UIKRetargeter::GetGlobalSettingsFromRetargetProfile(FRetargetProfile& RetargetProfile)
+{
+	return RetargetProfile.GlobalSettings;
+}
+
+void UIKRetargeter::SetGlobalSettingsInRetargetProfile(
+	FRetargetProfile& RetargetProfile,
+	const FRetargetGlobalSettings& GlobalSettings)
+{
+	RetargetProfile.GlobalSettings = GlobalSettings;
+	RetargetProfile.bApplyGlobalSettings = true;
+}
+
 void UIKRetargeter::SetRootSettingsInRetargetProfile(
 	FRetargetProfile& RetargetProfile,
 	const FTargetRootSettings& RootSettings)
diff --git a/Engine/Plugins/Animation/IKRig/Source/IKRig/Public/IKRigDefinition.h b/Engine/Plugins/Animation/IKRig/Source/IKRig/Public/IKRigDefinition.h
--- a/Engine/Plugins/Animation/IKRig/Source/IKRig/Public/IKRigDefinition.h
+++ b/Engine/Plugins/Animation/IKRig/Source/IKRig/Public/IKRigDefinition.h
@@ -16,7 +16,7 @@ class IPropertyHandle;
 class UIKRigSolver;
 
 #if WITH_EDITORONLY_DATA
-UENUM(BlueprintType)
+UENUM()
 enum class EIKRigGoalPreviewMode : uint8
 {
 	Additive		UMETA(DisplayName = "Additive"),
@@ -128,7 +128,7 @@ public:
 #endif
 };
 
-USTRUCT(Blueprintable)
+USTRUCT()
 struct IKRIG_API FBoneChain
 {
 	GENERATED_BODY()
@@ -155,7 +155,7 @@ struct IKRIG_API FBoneChain
 	FName IKGoalName;
 };
 
-USTRUCT(Blueprintable)
+USTRUCT()
 struct IKRIG_API FRetargetDefinition
 {
 	GENERATED_BODY()
diff --git a/Engine/Plugins/Animation/IKRig/Source/IKRig/Public/IKRigProcessor.h b/Engine/Plugins/Animation/IKRig/Source/IKRig/Public/IKRigProcessor.h
--- a/Engine/Plugins/Animation/IKRig/Source/IKRig/Public/IKRigProcessor.h
+++ b/Engine/Plugins/Animation/IKRig/Source/IKRig/Public/IKRigProcessor.h
@@ -87,6 +87,8 @@ public:
 
 	/** Get access to the internal goal data (read only) */
 	const FIKRigGoalContainer& GetGoalContainer() const;
+	/** Get the bone associated with a goal */
+	const FGoalBone* GetGoalBone(const FName& GoalName) const;
 	
 	/** Get read/write access to the internal skeleton data */
 	FIKRigSkeleton& GetSkeletonWriteable();
diff --git a/Engine/Plugins/Animation/IKRig/Source/IKRig/Public/IKRigSkeleton.h b/Engine/Plugins/Animation/IKRig/Source/IKRig/Public/IKRigSkeleton.h
--- a/Engine/Plugins/Animation/IKRig/Source/IKRig/Public/IKRigSkeleton.h
+++ b/Engine/Plugins/Animation/IKRig/Source/IKRig/Public/IKRigSkeleton.h
@@ -112,7 +112,7 @@ struct IKRIG_API FIKRigSkeleton
 	void GetChainsInList(const TArray<int32>& SelectedBones, TArray<FBoneChain>& OutChains) const;
 
 	// get indices of bones in this chain, ordered from tip to root. returns false if chain is invalid.
-	bool GetBonesInChain(const FBoneChain& Chain, TArray<int32>& OutBoneIndices) const;
+	bool ValidateChainAndGetBones(const FBoneChain& Chain, TArray<int32>& OutBoneIndices) const;
 
 	// given a set of bones, returns a set of indices of bones that are mirrored across X axis (YZ plane)
 	// return false if any of the bones do not have a mirrored pair within the MaxDistanceThreshold
diff --git a/Engine/Plugins/Animation/IKRig/Source/IKRig/Public/Retargeter/IKRetargetProcessor.h b/Engine/Plugins/Animation/IKRig/Source/IKRig/Public/Retargeter/IKRetargetProcessor.h
--- a/Engine/Plugins/Animation/IKRig/Source/IKRig/Public/Retargeter/IKRetargetProcessor.h
+++ b/Engine/Plugins/Animation/IKRig/Source/IKRig/Public/Retargeter/IKRetargetProcessor.h
@@ -37,14 +37,14 @@ struct IKRIG_API FRetargetSkeleton
 		const TArray<FBoneChain>& BoneChains,
 		const FName InRetargetPoseName,
 		const FIKRetargetPose* RetargetPose,
-		const FName& RetargetRootBone);
+		const FName RetargetRootBone);
 
 	void Reset();
 
 	void GenerateRetargetPose(
 		const FName InRetargetPoseName,
 		const FIKRetargetPose* InRetargetPose,
-		const FName& RetargetRootBone);
+		const FName RetargetRootBone);
 
 	int32 FindBoneIndexByName(const FName InName) const;
 
@@ -99,9 +99,9 @@ struct FTargetSkeleton : public FRetargetSkeleton
 	void Initialize(
 		USkeletalMesh* InSkeletalMesh,
 		const TArray<FBoneChain>& BoneChains,
-		const FName& InRetargetPoseName,
+		const FName InRetargetPoseName,
 		const FIKRetargetPose* RetargetPose,
-		const FName& RetargetRootBone);
+		const FName RetargetRootBone);
 
 	void Reset();
 
@@ -130,6 +130,7 @@ struct IKRIG_API FResolvedBoneChain
 
 struct FRootSource
 {
+	FName BoneName;
 	int32 BoneIndex;
 	FQuat InitialRotation;
 	float InitialHeightInverse;
@@ -141,6 +142,7 @@ struct FRootSource
 
 struct FRootTarget
 {
+	FName BoneName;
 	int32 BoneIndex;
 	FVector InitialPosition;
 	FQuat InitialRotation;
@@ -178,6 +180,48 @@ struct FRootRetargeter
 	}
 };
 
+struct FPoleVectorMatcher
+{
+	EAxis::Type SourcePoleAxis;
+	EAxis::Type TargetPoleAxis;
+	float TargetToSourceAngularOffsetAtRefPose;
+	TArray<int32> AllChildrenWithinChain;
+
+	bool Initialize(
+		const TArray<int32>& SourceIndices,
+		const TArray<int32>& TargetIndices,
+		const TArray<FTransform> &SourceGlobalPose,
+		const TArray<FTransform> &TargetGlobalPose,
+		const FRetargetSkeleton& TargetSkeleton);
+
+	void MatchPoleVector(
+		const FTargetChainSettings& Settings,
+		const TArray<int32>& SourceIndices,
+		const TArray<int32>& TargetIndices,
+		const TArray<FTransform> &SourceGlobalPose,
+		TArray<FTransform> &OutTargetGlobalPose,
+		FRetargetSkeleton& TargetSkeleton);
+	
+private:
+
+	EAxis::Type CalculateBestPoleAxisForChain(
+		const TArray<int32>& BoneIndices,
+		const TArray<FTransform>& GlobalPose);
+	
+	static FVector CalculatePoleVector(
+		const EAxis::Type& PoleAxis,
+		const TArray<int32>& BoneIndices,
+		const TArray<FTransform>& GlobalPose);
+
+	static EAxis::Type GetMostDifferentAxis(
+		const FTransform& Transform,
+		const FVector& InNormal);
+
+	static FVector GetChainNormal(
+		const TArray<int32>& BoneIndices,
+		const TArray<FTransform>& GlobalPose);
+};
+
 struct FChainFK
 {
 	TArray<FTransform> InitialGlobalTransforms;
@@ -245,8 +289,15 @@ struct FChainDecoderFK : public FChainFK
 		const FTargetSkeleton& TargetSkeleton,
 		TArray<FTransform> &InOutGlobalPose);
 
-private:
+	void MatchPoleVector(
+		const FTargetChainSettings& Settings,
+		const TArray<int32>& TargetBoneIndices,
+		FChainEncoderFK& SourceChain,
+		const FTargetSkeleton& TargetSkeleton,
+		TArray<FTransform> &InOutGlobalPose);
 
+private:
+	
 	FTransform GetTransformAtParam(
 		const TArray<FTransform>& Transforms,
 		const TArray<float>& InParams,
@@ -344,11 +395,9 @@ struct FRetargetChainPair
 	FTargetChainSettings Settings;
 	
 	TArray<int32> SourceBoneIndices;
-	
 	TArray<int32> TargetBoneIndices;
 	
 	FName SourceBoneChainName;
-	
 	FName TargetBoneChainName;
 
 	virtual ~FRetargetChainPair() = default;
@@ -372,8 +421,8 @@ private:
 struct FRetargetChainPairFK : FRetargetChainPair
 {
 	FChainEncoderFK FKEncoder;
-
 	FChainDecoderFK FKDecoder;
+	FPoleVectorMatcher PoleVectorMatcher;
 	
 	virtual bool Initialize(
         const FBoneChain& SourceBoneChain,
@@ -386,9 +435,7 @@ struct FRetargetChainPairFK : FRetargetChainPair
 struct FRetargetChainPairIK : FRetargetChainPair
 {
 	FChainRetargeterIK IKChainRetargeter;
-	
 	FName IKGoalName;
-	
 	FName PoleVectorGoalName;
 
 	virtual bool Initialize(
@@ -529,6 +576,9 @@ private:
 
 	/** The Source/Target pair of Root Bones retargeted with scaled translation */
 	FRootRetargeter RootRetargeter;
+
+	/** The currently used global settings (driven either by source asset or a profile) */
+	FRetargetGlobalSettings GlobalSettings;
 	
 	/** Initializes the FRootRetargeter */
 	bool InitializeRoots();
@@ -552,6 +602,12 @@ private:
 		const TMap<FName, float>& SpeedValuesFromCurves,
 		const float DeltaTime);
 
+	/** Internal retarget phase for the pole matching feature of FK chains. */
+	void RunPoleVectorMatching(const TArray<FTransform>& InGlobalTransforms, TArray<FTransform>& OutGlobalTransforms);
+
 	/** Runs in the after the base IK retarget to apply stride warping to IK goals. */
 	void RunStrideWarping(const TArray<FTransform>& InTargeGlobalPose);
+
+	/** Does a partial reinitialization (at runtime) whenever the retarget pose is swapped to a different one. */
+	void ApplyNewRetargetPose(const FName NewRetargetPoseName, ERetargetSourceOrTarget SourceOrTarget);
 };
diff --git a/Engine/Plugins/Animation/IKRig/Source/IKRig/Public/Retargeter/IKRetargetProfile.h b/Engine/Plugins/Animation/IKRig/Source/IKRig/Public/Retargeter/IKRetargetProfile.h
--- a/Engine/Plugins/Animation/IKRig/Source/IKRig/Public/Retargeter/IKRetargetProfile.h
+++ b/Engine/Plugins/Animation/IKRig/Source/IKRig/Public/Retargeter/IKRetargetProfile.h
@@ -16,37 +16,46 @@ struct FRetargetProfile
 	GENERATED_BODY()
 	
 public:
+
+	// If true, the TARGET Retarget Pose specified in this profile will be applied to the Retargeter (when plugged into the Retargeter).
+	UPROPERTY(EditAnywhere, BlueprintReadWrite, Category=RetargetPoses)
+	bool bApplyTargetRetargetPose = false;
+	
 	// Override the TARGET Retarget Pose to use when this profile is active.
 	// The pose must be present in the Retarget Asset and is not applied unless bApplyTargetRetargetPose is true.
-	UPROPERTY(BlueprintReadWrite, Category=RetargetPoses)
+	UPROPERTY(EditAnywhere, BlueprintReadWrite, Category=RetargetPoses, meta=(EditCondition="bApplyTargetRetargetPose"))
 	FName TargetRetargetPoseName;
 
-	// If true, the TARGET Retarget Pose specified in this profile will be applied to the Retargeter (when plugged into the Retargeter).
-	UPROPERTY(BlueprintReadWrite, Category=RetargetPoses)
-	bool bApplyTargetRetargetPose = false;
+	// If true, the Source Retarget Pose specified in this profile will be applied to the Retargeter (when plugged into the Retargeter).
+	UPROPERTY(EditAnywhere, BlueprintReadWrite, Category=RetargetPoses)
+	bool bApplySourceRetargetPose = false;
 
 	// Override the SOURCE Retarget Pose to use when this profile is active.
 	// The pose must be present in the Retarget Asset and is not applied unless bApplySourceRetargetPose is true.
-	UPROPERTY(BlueprintReadWrite, Category=RetargetPoses)
+	UPROPERTY(EditAnywhere, BlueprintReadWrite, Category=RetargetPoses, meta=(EditCondition="bApplySourceRetargetPose"))
 	FName SourceRetargetPoseName;
-	
-	// If true, the Source Retarget Pose specified in this profile will be applied to the Retargeter (when plugged into the Retargeter).
-	UPROPERTY(BlueprintReadWrite, Category=RetargetPoses)
-	bool bApplySourceRetargetPose = false;
 
+	// If true, the Chain Settings stored in this profile will be applied to the Retargeter (when plugged into the Retargeter).
+	UPROPERTY(EditAnywhere, BlueprintReadWrite, Category=ChainSettings)
+	bool bApplyChainSettings = true;
+	
 	// A (potentially sparse) set of setting overrides for the target chains (only applied when bApplyChainSettings is true).
-	UPROPERTY(BlueprintReadWrite, Category=ChainSettings)
+	UPROPERTY(EditAnywhere, BlueprintReadWrite, Category=ChainSettings, meta=(EditCondition="bApplyChainSettings"))
 	TMap<FName, FTargetChainSettings> ChainSettings;
 
-	// If true, the Chain Settings stored in this profile will be applied to the Retargeter (when plugged into the Retargeter).
-	UPROPERTY(BlueprintReadWrite, Category=ChainSettings)
-	bool bApplyChainSettings = true;
+	// If true, the root settings stored in this profile will be applied to the Retargeter (when plugged into the Retargeter).
+	UPROPERTY(EditAnywhere, BlueprintReadWrite, Category=RootSettings)
+	bool bApplyRootSettings = false;
 
 	// Retarget settings to control behavior of the retarget root motion (not applied unless bApplyRootSettings is true)
-	UPROPERTY(BlueprintReadWrite, Category=RootSettings)
+	UPROPERTY(EditAnywhere, BlueprintReadWrite, Category=RootSettings, meta=(EditCondition="bApplyRootSettings"))
 	FTargetRootSettings RootSettings;
 
-	// If true, the root settings stored in this profile will be applied to the Retargeter (when plugged into the Retargeter).
-	UPROPERTY(BlueprintReadWrite, Category=RootSettings)
-	bool bApplyRootSettings = false;
+	// If true, the global settings stored in this profile will be applied to the Retargeter (when plugged into the Retargeter).
+	UPROPERTY(EditAnywhere, BlueprintReadWrite, Category=GlobalSettings)
+	bool bApplyGlobalSettings = false;
+
+	// Retarget settings to control global behavior, like Stride Warping (not applied unless bApplyGlobalSettings is true)
+	UPROPERTY(EditAnywhere, BlueprintReadWrite, Category=GlobalSettings, meta=(EditCondition="bApplyGlobalSettings"))
+	FRetargetGlobalSettings GlobalSettings;
 };
diff --git a/Engine/Plugins/Animation/IKRig/Source/IKRig/Public/Retargeter/IKRetargetSettings.h b/Engine/Plugins/Animation/IKRig/Source/IKRig/Public/Retargeter/IKRetargetSettings.h
--- a/Engine/Plugins/Animation/IKRig/Source/IKRig/Public/Retargeter/IKRetargetSettings.h
+++ b/Engine/Plugins/Animation/IKRig/Source/IKRig/Public/Retargeter/IKRetargetSettings.h
@@ -144,7 +144,7 @@ struct FTargetChainIKSettings
 	float Extension = 1.0f;
 
 	/** Is this IK goal affected by the stride warping (in Global Settings). Typically this is true for all feet, but not for hands.  */
-	UPROPERTY(EditAnywhere, Category = Warping, meta = (DisplayName = "Affected by IK Warping"))
+	UPROPERTY(EditAnywhere, BlueprintReadWrite, Category = Warping, meta = (DisplayName = "Affected by IK Warping"))
 	bool bAffectedByIKWarping = true;
 };
 
@@ -256,7 +256,7 @@ enum class EWarpingDirectionSource
 	Chain,
 };
 
-USTRUCT()
+USTRUCT(BlueprintType)
 struct IKRIG_API FRetargetGlobalSettings
 {
 	GENERATED_USTRUCT_BODY()
@@ -282,8 +282,8 @@ struct IKRIG_API FRetargetGlobalSettings
 	/** Enable IK Warping.
 	 * These options allow for global modifications to all IK Goals that have "Affected by IK Warping" turned on (the default).
 	 * "Affected by IK Warping" can be found in the IK settings of a retarget chain. */
-	UPROPERTY(EditAnywhere, Category = Warping)
-	bool bIKWarping = false;
+	UPROPERTY(EditAnywhere, BlueprintReadWrite, Category = Warping)
+	bool bWarping = false;
 
 	/** Defines the source used to determine the forward direction as the character animates in world space. Default is "Goals".
 	 * This method uses various positions on the character to create a "best fit" global rotation that approximates the facing direction of the character over time.
@@ -292,32 +292,32 @@ struct IKRIG_API FRetargetGlobalSettings
 	 * Goals: uses the positions of the IK goals to approximate the facing direction. This is best used on characters with a vertical spine, like bipeds.
 	 * Chain: uses the positions of the bones in a retarget chain to approximate the facing direction. This is best for characters with a horizontal spine, like quadrupeds.
 	 */
-	UPROPERTY(EditAnywhere, Category = Warping)
+	UPROPERTY(EditAnywhere, BlueprintReadWrite, Category = Warping)
 	EWarpingDirectionSource DirectionSource = EWarpingDirectionSource::Goals;
 	
 	/** The world space axis that represents the forward facing direction for your character. By default, in Unreal, this is +Y.  */
-	UPROPERTY(EditAnywhere, Category = Warping)
+	UPROPERTY(EditAnywhere, BlueprintReadWrite, Category = Warping)
 	EBasicAxis ForwardDirection = EBasicAxis::Y;
 
 	/** When using the "Chain" option as a Direction Source, this defines the chain to use to determine the facing direction of the character.
 	 * Typically this would be the chain that contains the Spine bones. */
-	UPROPERTY(EditAnywhere, Category = Warping, meta=(EditCondition="DirectionSource == EWarpingDirectionSource::Chain"))
+	UPROPERTY(EditAnywhere, BlueprintReadWrite, Category = Warping, meta=(EditCondition="DirectionSource == EWarpingDirectionSource::Chain"))
 	FName DirectionChain;
 
 	/** Range 0 to Inf. Default 1. Warps IK goal positions in the forward direction. Useful for stride warping.
 	 * Values below 1 will create smaller, squashed strides. Values greater than 1 will create stretched, longer strides.*/
-	UPROPERTY(EditAnywhere, Category = Warping, meta = (UIMin = "0.0", UIMax = "5.0"))
+	UPROPERTY(EditAnywhere, BlueprintReadWrite, Category = Warping, meta = (UIMin = "0.0", UIMax = "5.0"))
 	float WarpForwards = 1.0f;
 
 	/** Range -+Inf. Default is 0. A static offset in world units to move the IK goals perpendicular to the forward direction.
 	 * Values less than zero will move the goals towards the center-line of the character. Values greater than zero push the goals outwards.*/
-	UPROPERTY(EditAnywhere, Category = Warping, meta = (UIMin = "-10.0", UIMax = "100.0"))
+	UPROPERTY(EditAnywhere, BlueprintReadWrite, Category = Warping, meta = (UIMin = "-10.0", UIMax = "100.0"))
 	float SidewaysOffset = 0.0f;
 
 	/** Range 0 to +Inf. Default is 1.0f.
 	 * Values below 1 pull all the goals towards the average of all the goals (towards each other).
 	 * Values greater than 1 push the goals apart.*/
-	UPROPERTY(EditAnywhere, Category = Warping, meta = (UIMin = "0.0", UIMax = "2.0"))
+	UPROPERTY(EditAnywhere, BlueprintReadWrite, Category = Warping, meta = (UIMin = "0.0", UIMax = "2.0"))
 	float WarpSplay = 1.0f;
 
 	static FVector GetAxisVector(const EBasicAxis& Axis)
diff --git a/Engine/Plugins/Animation/IKRig/Source/IKRig/Public/Retargeter/IKRetargeter.h b/Engine/Plugins/Animation/IKRig/Source/IKRig/Public/Retargeter/IKRetargeter.h
--- a/Engine/Plugins/Animation/IKRig/Source/IKRig/Public/Retargeter/IKRetargeter.h
+++ b/Engine/Plugins/Animation/IKRig/Source/IKRig/Public/Retargeter/IKRetargeter.h
@@ -240,6 +240,14 @@ public:
 	/* Get the retarget profile by name (may be null) */
 	const FRetargetProfile* GetProfileByName(const FName& ProfileName) const;
 
+	/** BLUEPRINT GETTERS */
+
+	/** Returns the chain settings associated with a given Goal in an IK Retargeter Asset using the given profile name (optional) */
+	UFUNCTION(BlueprintCallable, BlueprintPure, Category=RetargetAsset)
+	static FTargetChainSettings GetChainUsingGoalFromRetargetAsset(
+		const UIKRetargeter* RetargetAsset,
+		const FName IKGoalName);
+	
 	/** Returns the chain settings associated with a given target chain in an IK Retargeter Asset using the given profile name (optional) */
 	UFUNCTION(BlueprintCallable, BlueprintPure, Category=RetargetAsset)
 	static FTargetChainSettings GetChainSettingsFromRetargetAsset(
@@ -247,6 +255,12 @@ public:
 		const FName TargetChainName,
 		const FName OptionalProfileName);
 
+	/** Returns the chain settings associated with a given target chain in the supplied Retarget Profile. */
+	UFUNCTION(BlueprintCallable, BlueprintPure, Category=RetargetProfile)
+	static FTargetChainSettings GetChainSettingsFromRetargetProfile(
+		UPARAM(ref) FRetargetProfile& RetargetProfile,
+		const FName TargetChainName);
+
 	/** Returns the root settings in an IK Retargeter Asset using the given profile name (optional) */
 	UFUNCTION(BlueprintCallable, BlueprintPure, Category=RetargetAsset)
 	static void GetRootSettingsFromRetargetAsset(
@@ -254,11 +268,28 @@ public:
 		const FName OptionalProfileName,
 		FTargetRootSettings& OutSettings);
 
-	/** Returns the chain settings associated with a given target chain in the supplied Retarget Profile. */
+	/** Returns the root settings in the supplied Retarget Profile. */
 	UFUNCTION(BlueprintCallable, BlueprintPure, Category=RetargetProfile)
-	static FTargetChainSettings GetChainSettingsFromRetargetProfile(
+	static FTargetRootSettings GetRootSettingsFromRetargetProfile(UPARAM(ref) FRetargetProfile& RetargetProfile);
+
+	/** Returns the global settings in an IK Retargeter Asset using the given profile name (optional) */
+	UFUNCTION(BlueprintCallable, BlueprintPure, Category=RetargetAsset)
+	static void GetGlobalSettingsFromRetargetAsset(
+		const UIKRetargeter* RetargetAsset,
+		const FName OptionalProfileName,
+		FRetargetGlobalSettings& OutSettings);
+
+	/** Returns the global settings in the supplied Retarget Profile. */
+	UFUNCTION(BlueprintCallable, BlueprintPure, Category=RetargetProfile)
+	static FRetargetGlobalSettings GetGlobalSettingsFromRetargetProfile(UPARAM(ref) FRetargetProfile& RetargetProfile);
+
+	/** BLUEPRINT SETTERS */
+
+	/** Set the global settings in a retarget profile (will set bApplyGlobalSettings to true). */
+	UFUNCTION(BlueprintCallable, Category=RetargetProfile)
+	static void SetGlobalSettingsInRetargetProfile(
 		UPARAM(ref) FRetargetProfile& RetargetProfile,
-		const FName TargetChainName);
+		const FRetargetGlobalSettings& GlobalSettings);
 
 	/** Set the root settings in a retarget profile (will set bApplyRootSettings to true). */
 	UFUNCTION(BlueprintCallable, Category=RetargetProfile)
diff --git a/Engine/Plugins/Animation/IKRig/Source/IKRigEditor/Private/RetargetEditor/IKRetargetAnimInstance.cpp b/Engine/Plugins/Animation/IKRig/Source/IKRigEditor/Private/RetargetEditor/IKRetargetAnimInstance.cpp
--- a/Engine/Plugins/Animation/IKRig/Source/IKRigEditor/Private/RetargetEditor/IKRetargetAnimInstance.cpp
+++ b/Engine/Plugins/Animation/IKRig/Source/IKRigEditor/Private/RetargetEditor/IKRetargetAnimInstance.cpp
@@ -89,15 +89,23 @@ void FAnimNode_PreviewRetargetPose::Evaluate_AnyThread(FPoseContext& Output)
 	{
 		FTransform& RootTransform = RetargetGlobalPose[RootBoneIndex];
 		RootTransform.AddToTranslation(RetargetPose->GetRootTranslationDelta());
+		
 		// update local transform of retarget root
 		const int32 ParentIndex = RefSkeleton.GetParentIndex(RootBoneIndex);
-		if (ParentIndex != INDEX_NONE)
+		if (ParentIndex == INDEX_NONE)
+		{
+			RetargetLocalPose[RootBoneIndex] = RootTransform;
+		}
+		else
 		{
 			const FTransform& ChildGlobalTransform = RetargetGlobalPose[RootBoneIndex];
 			const FTransform& ParentGlobalTransform = RetargetGlobalPose[ParentIndex];
 			RetargetLocalPose[RootBoneIndex] = ChildGlobalTransform.GetRelativeTransform(ParentGlobalTransform);
 		}
 	}
+
+	// update GLOBAL space pose after root translation (for editor to query)
+	FAnimationRuntime::FillUpComponentSpaceTransforms(RefSkeleton, RetargetLocalPose, RetargetGlobalPose);
 	
 	// copy to the compact output pose
 	const int32 NumBones = Output.Pose.GetNumBones();
diff --git a/Engine/Plugins/Animation/IKRig/Source/IKRigEditor/Private/RetargetEditor/IKRetargetDefaultMode.cpp b/Engine/Plugins/Animation/IKRig/Source/IKRigEditor/Private/RetargetEditor/IKRetargetDefaultMode.cpp
--- a/Engine/Plugins/Animation/IKRig/Source/IKRigEditor/Private/RetargetEditor/IKRetargetDefaultMode.cpp
+++ b/Engine/Plugins/Animation/IKRig/Source/IKRigEditor/Private/RetargetEditor/IKRetargetDefaultMode.cpp
@@ -10,6 +10,7 @@
 #include "Animation/DebugSkelMeshComponent.h"
 #include "RetargetEditor/IKRetargetEditor.h"
 #include "RetargetEditor/IKRetargetHitProxies.h"
+#include "ReferenceSkeleton.h"
 
 
 #define LOCTEXT_NAMESPACE "IKRetargetDefaultMode"
@@ -18,14 +19,47 @@ FName FIKRetargetDefaultMode::ModeName("IKRetargetAssetDefaultMode");
 
 bool FIKRetargetDefaultMode::GetCameraTarget(FSphere& OutTarget) const
 {
-	// target skeletal mesh
-	if (const TSharedPtr<FIKRetargetEditorController> Controller = EditorController.Pin())
+	const TSharedPtr<FIKRetargetEditorController> Controller = EditorController.Pin();
+	if (!Controller.IsValid())
+	{
+		return false;
+	}	
+
+	// target the selected bones
+	const TArray<FName>& SelectedBoneNames = Controller->GetSelectedBones();
+	if (!SelectedBoneNames.IsEmpty())
 	{
-		if (Controller->SourceSkelMeshComponent)
+		TArray<FVector> TargetPoints;
+		const UDebugSkelMeshComponent* SkeletalMeshComponent = Controller->GetSkeletalMeshComponent(Controller->GetSourceOrTarget());
+		const FReferenceSkeleton& RefSkeleton = SkeletalMeshComponent->GetReferenceSkeleton();
+		TArray<int32> ChildrenIndices;
+		for (const FName& SelectedBoneName : SelectedBoneNames)
 		{
-			OutTarget = Controller->SourceSkelMeshComponent->Bounds.GetSphere();
-			return true;
+			const int32 BoneIndex = RefSkeleton.FindBoneIndex(SelectedBoneName);
+			if (BoneIndex == INDEX_NONE)
+			{
+				continue;
+			}
+
+			TargetPoints.Add(SkeletalMeshComponent->GetBoneTransform(BoneIndex).GetLocation());
+			ChildrenIndices.Reset();
+			//RefSkeleton.GetDirectChildBones(BoneIndex, ChildrenIndices);
+			for (const int32 ChildIndex : ChildrenIndices)
+			{
+				TargetPoints.Add(SkeletalMeshComponent->GetBoneTransform(ChildIndex).GetLocation());
+			}
 		}
+		
+		// create a sphere that contains all the target points
+		OutTarget = FSphere(&TargetPoints[0], TargetPoints.Num());
+		return true;
+	}
+
+	// target the selected mesh
+	if (const UPrimitiveComponent* SelectedMesh = Controller->GetSelectedMesh())
+	{
+		OutTarget = SelectedMesh->Bounds.GetSphere();
+		return true;
 	}
 	
 	return false;
@@ -47,15 +81,6 @@ void FIKRetargetDefaultMode::Initialize()
 	{
 		return; 
 	}
-	
-	// register selection callback overrides
-	TInlineComponentArray<UPrimitiveComponent*> PrimitiveComponents;
-	PrimitiveComponents.Add(Controller->SourceSkelMeshComponent);
-	PrimitiveComponents.Add(Controller->TargetSkelMeshComponent);
-	for (UPrimitiveComponent* PrimitiveComponent : PrimitiveComponents)
-	{
-		PrimitiveComponent->SelectionOverrideDelegate = UPrimitiveComponent::FSelectionOverride::CreateRaw(this, &FIKRetargetDefaultMode::ComponentSelectionOverride);
-	}
 
 	// update offsets on preview meshes
 	Controller->AddOffsetToMeshComponent(FVector::ZeroVector, Controller->SourceSkelMeshComponent);
@@ -73,12 +98,18 @@ void FIKRetargetDefaultMode::Render(const FSceneView* View, FViewport* Viewport,
 		return;
 	}
 	
-	FIKRetargetEditorController* Controller = EditorController.Pin().Get();
-	if (Controller->GetRetargeterMode() == ERetargeterOutputMode::EditRetargetPose)
-	{
-		return;
-	}
-	
+	const FIKRetargetEditorController* Controller = EditorController.Pin().Get();
+
+	// render source and target skeletons
+	Controller->RenderSkeleton(PDI, ERetargetSourceOrTarget::Source);
+	Controller->RenderSkeleton(PDI, ERetargetSourceOrTarget::Target);
+
+	// render all the chain and root debug proxies
+	RenderDebugProxies(PDI, Controller);
+}
+
+void FIKRetargetDefaultMode::RenderDebugProxies(FPrimitiveDrawInterface* PDI, const FIKRetargetEditorController* Controller) const
+{
 	const UIKRetargeter* Asset = Controller->AssetController->GetAsset();
 	if (!Asset->bDebugDraw)
 	{
@@ -304,10 +335,10 @@ bool FIKRetargetDefaultMode::HandleClick(FEditorViewportClient* InViewportClient
 	}
 
 	// did we click on a bone in the viewport?
-	const bool bHitBone = HitProxy && HitProxy->IsA(HPersonaBoneHitProxy::StaticGetType());
+	const bool bHitBone = HitProxy && HitProxy->IsA(HIKRetargetEditorBoneProxy::StaticGetType());
 	if (bLeftButtonClicked && bHitBone)
 	{
-		const HPersonaBoneHitProxy* BoneProxy = static_cast<HPersonaBoneHitProxy*>(HitProxy);
+		const HIKRetargetEditorBoneProxy* BoneProxy = static_cast<HIKRetargetEditorBoneProxy*>(HitProxy);
 		const TArray BoneNames{BoneProxy->BoneName};
 		constexpr bool bFromHierarchy = false;
 		Controller->EditBoneSelection(BoneNames, EditMode, bFromHierarchy);
@@ -436,11 +467,6 @@ void FIKRetargetDefaultMode::Enter()
 
 	// record which skeleton is being viewed/edited
 	SkeletonMode = Controller->GetSourceOrTarget();
-
-	// allow selection of meshes in this mode
-	Controller->Editor.Pin()->GetPersonaToolkit()->GetPreviewScene()->SetAllowMeshHitProxies(true);
-	Controller->SourceSkelMeshComponent->bSelectable = true;
-	Controller->TargetSkelMeshComponent->bSelectable = true;
 }
 
 void FIKRetargetDefaultMode::Exit()
@@ -450,28 +476,8 @@ void FIKRetargetDefaultMode::Exit()
 	{
 		return; 
 	}
-
-	// editor can be closed while in edit mode
-	if (Controller->Editor.IsValid())
-	{
-		// disable selection in other modes
-		Controller->Editor.Pin()->GetPersonaToolkit()->GetPreviewScene()->SetAllowMeshHitProxies(false);
-		Controller->SourceSkelMeshComponent->bSelectable = false;
-		Controller->TargetSkelMeshComponent->bSelectable = false;
-	}
-
-	IPersonaEditMode::Exit();
-}
-
-bool FIKRetargetDefaultMode::ComponentSelectionOverride(const UPrimitiveComponent* InComponent) const
-{
-	const TSharedPtr<FIKRetargetEditorController> Controller = EditorController.Pin();
-	if (!Controller.IsValid())
-	{
-		return false; 
-	}
 	
-	return InComponent == Controller->GetSelectedMesh();
+	IPersonaEditMode::Exit();
 }
 
 UDebugSkelMeshComponent* FIKRetargetDefaultMode::GetCurrentlyEditedMesh() const
@@ -504,15 +510,6 @@ void FIKRetargetDefaultMode::Tick(FEditorViewportClient* ViewportClient, float D
 	{
 		Initialize();
 	}
-
-	// update skeleton drawing mode
-	const TSharedPtr<FIKRetargetEditorController> Controller = EditorController.Pin();
-	if (Controller.IsValid())
-	{
-		const bool bEditingSource = Controller->GetSourceOrTarget() == ERetargetSourceOrTarget::Source;
-		Controller->SourceSkelMeshComponent->SkeletonDrawMode = bEditingSource ? ESkeletonDrawMode::Default : ESkeletonDrawMode::GreyedOut;
-		Controller->TargetSkelMeshComponent->SkeletonDrawMode = !bEditingSource ? ESkeletonDrawMode::Default : ESkeletonDrawMode::GreyedOut;
-	}
 }
 
 void FIKRetargetDefaultMode::DrawHUD(FEditorViewportClient* ViewportClient, FViewport* Viewport, const FSceneView* View, FCanvas* Canvas)
diff --git a/Engine/Plugins/Animation/IKRig/Source/IKRigEditor/Private/RetargetEditor/IKRetargetDetails.cpp b/Engine/Plugins/Animation/IKRig/Source/IKRigEditor/Private/RetargetEditor/IKRetargetDetails.cpp
--- a/Engine/Plugins/Animation/IKRig/Source/IKRigEditor/Private/RetargetEditor/IKRetargetDetails.cpp
+++ b/Engine/Plugins/Animation/IKRig/Source/IKRigEditor/Private/RetargetEditor/IKRetargetDetails.cpp
@@ -787,38 +787,55 @@ void FRetargetChainSettingsCustomization::CustomizeDetails(IDetailLayoutBuilder&
 	// determine the current setup state of the chain (is it ready for FK and/or IK?)
 	// FK: it can only run FK retarget if it's mapped to a source chain
 	// IK: can only run IK if FK is setup AND it's mapped to a goal that is connected to a solver (in the target IK rig)
-	auto CheckChainSetupForIKAndFK([this](const TObjectPtr<URetargetChainSettings> ChainSettings, bool& OutFKSetup, bool&OutIKSetup)
+	auto CheckChainSetupForIKAndFK([this](
+		const TObjectPtr<URetargetChainSettings> ChainSettings,
+		bool& OutFKSetup,
+		bool& OutIKSetup,
+		bool& OutIKConnected)
 	{
 		OutFKSetup = true;
 		OutIKSetup = true;
+		OutIKConnected = true;
 		
 		if (ChainSettings->SourceChain == NAME_None || ChainSettings->TargetChain == NAME_None)
 		{
 			OutFKSetup = false;
 			OutIKSetup = false;
+			OutIKConnected = false;
+			return;
 		}
 
-		if (!Controller->AssetController->IsChainRunningIK(ChainSettings.Get()))
+		const FName ChainGoal = Controller->AssetController->GetChainGoal(ChainSettings.Get());
+		if (ChainGoal == NAME_None)
 		{
 			OutIKSetup = false;
 		}
+
+		if (!Controller->AssetController->IsChainGoalConnectedToASolver(ChainGoal))
+		{
+			OutIKConnected = false;
+		}
 	});
 
 	// get the FK and IK status of the first selected chain
 	bool bFirstChainFKSetup;
 	bool bFirstChainIKSetup;
-	CheckChainSetupForIKAndFK(ChainSettingsObjects[0].Get(), bFirstChainFKSetup, bFirstChainIKSetup);
+	bool bFirstChainIKConnected;
+	CheckChainSetupForIKAndFK(ChainSettingsObjects[0].Get(), bFirstChainFKSetup, bFirstChainIKSetup, bFirstChainIKConnected);
 
 	// check to see if there is a mix of chains with valid FK or IK setups
 	bool bIsFKSetupMultipleValues = false;
 	bool bIsIKSetupMultipleValues = false;
+	bool bIsIKConnectedMultipleValues = false;
 	for (const TWeakObjectPtr<URetargetChainSettings> ChainSettings : ChainSettingsObjects)
 	{
 		bool bIsFKSetup;
 		bool bIsIKSetup;
-		CheckChainSetupForIKAndFK(ChainSettings.Get(), bIsFKSetup, bIsIKSetup);
+		bool bIsIKConnected;
+		CheckChainSetupForIKAndFK(ChainSettings.Get(), bIsFKSetup, bIsIKSetup, bIsIKConnected);
 		bIsFKSetupMultipleValues = bIsFKSetup != bFirstChainFKSetup ? true : bIsFKSetupMultipleValues;
 		bIsIKSetupMultipleValues = bIsIKSetup != bFirstChainIKSetup ? true : bIsIKSetupMultipleValues;
+		bIsIKConnectedMultipleValues = bIsIKConnected != bFirstChainIKConnected ? true : bIsIKConnectedMultipleValues;
 	}
 
 	// are we editing multiple chains?
@@ -898,9 +915,24 @@ void FRetargetChainSettingsCustomization::CustomizeDetails(IDetailLayoutBuilder&
 
 	// IK settings
 	const FText NoIKGoalWarning = LOCTEXT( "IKDisabled_Label", "Disabled: No IK Goal specified." );
+	const FText IKNotConnectedWarning = LOCTEXT( "IKDisconnected_Label", "Disabled: IK Goal is not connected to any solvers." );
 	const FText IKEnabledMultipleValuesWarning = LOCTEXT( "IKDisabledMultipleValues_Label", "Some selected chain(s) do not have an IK Goal." );
+	const bool bIKIsConnected = bFirstChainIKConnected && !bIsIKConnectedMultipleValues;
 	const bool bIKEditingEnabled = bFirstChainIKSetup && !bIsIKSetupMultipleValues;
-	const FText IKWarningToUse = !bFKEditingEnabled ? FKNotMappedWarning : !bIKEditingEnabled ? NoIKGoalWarning : FText();
+	FText IKWarningToUse = FText();
+	if (!bFKEditingEnabled)
+	{
+		IKWarningToUse = FKNotMappedWarning;
+	}
+	else if (!bIKEditingEnabled)
+	{
+		IKWarningToUse = NoIKGoalWarning;
+	}
+	else if (!bIKIsConnected)
+	{
+		IKWarningToUse = IKNotConnectedWarning;
+	}
+	
 	AddSettingsSection(
 		DetailBuilder,
 		SettingsCategory,
@@ -909,7 +941,7 @@ void FRetargetChainSettingsCustomization::CustomizeDetails(IDetailLayoutBuilder&
 		LOCTEXT("IKGroup_Label", "IK"),
 		FTargetChainIKSettings::StaticStruct(),
 		GET_MEMBER_NAME_STRING_CHECKED(FTargetChainIKSettings, EnableIK),
-		bIKEditingEnabled,
+		bIKEditingEnabled && bIKIsConnected,
 		IKWarningToUse);
 
 	// Plant settings
@@ -970,7 +1002,7 @@ void FRetargetChainSettingsCustomization::AddSettingsSection(
 				.OnCheckStateChanged_Lambda([this, IsEnabledProperty](ECheckBoxState State)
 				{
 					IsEnabledProperty->SetValue(State == ECheckBoxState::Checked);
-					Controller->OnRetargeterNeedsInitialized(Controller->AssetController->GetAsset());
+					Controller->HandleRetargeterNeedsInitialized(Controller->AssetController->GetAsset());
 				})
 			];
 	}
@@ -1092,7 +1124,7 @@ void FRetargetGlobalSettingsCustomization::CustomizeDetails(IDetailLayoutBuilder
 
 	auto GetPropertyHandle = [&DetailBuilder, &StructPropertyPath](const FString& PropertyName)
 	{
-		const FString EnableWarpingPropertyName = GET_MEMBER_NAME_STRING_CHECKED(FRetargetGlobalSettings, bIKWarping);
+		const FString EnableWarpingPropertyName = GET_MEMBER_NAME_STRING_CHECKED(FRetargetGlobalSettings, bWarping);
 		const FString PropertyPath = FString::Printf(TEXT("%s.%s"), *StructPropertyPath, *PropertyName);
 		return DetailBuilder.GetProperty(*PropertyPath);
 	};
@@ -1109,11 +1141,12 @@ void FRetargetGlobalSettingsCustomization::CustomizeDetails(IDetailLayoutBuilder
 	AlphaGroup.AddPropertyRow(GetPropertyHandle(GET_MEMBER_NAME_STRING_CHECKED(FRetargetGlobalSettings, bEnableIK)));
 
 	// stride warping group
-	IDetailGroup& WarpingGroup = SettingsCategory.AddGroup("Warping", LOCTEXT("Warping_Label", "Stride Warping"), false, true);
+	const FText WarpingTitleLabel = LOCTEXT("Warping_Label", "Stride Warping");
+	IDetailGroup& WarpingGroup = SettingsCategory.AddGroup("Warping", WarpingTitleLabel, false, true);
 
 	// header with "enable" checkbox that disables whole group
-	const TSharedRef<IPropertyHandle> EnableWarpingPropertyHandle = GetPropertyHandle(GET_MEMBER_NAME_STRING_CHECKED(FRetargetGlobalSettings, bIKWarping));
-	WarpingGroup.HeaderProperty(EnableWarpingPropertyHandle);
+	const TSharedRef<IPropertyHandle> EnableWarpingPropertyHandle = GetPropertyHandle(GET_MEMBER_NAME_STRING_CHECKED(FRetargetGlobalSettings, bWarping));
+	WarpingGroup.HeaderProperty(EnableWarpingPropertyHandle).DisplayName(WarpingTitleLabel);
 	
 	const TAttribute<bool> EditCondition = TAttribute<bool>::Create([EnableWarpingPropertyHandle]()
 	{
diff --git a/Engine/Plugins/Animation/IKRig/Source/IKRigEditor/Private/RetargetEditor/IKRetargetEditPoseMode.cpp b/Engine/Plugins/Animation/IKRig/Source/IKRigEditor/Private/RetargetEditor/IKRetargetEditPoseMode.cpp
--- a/Engine/Plugins/Animation/IKRig/Source/IKRigEditor/Private/RetargetEditor/IKRetargetEditPoseMode.cpp
+++ b/Engine/Plugins/Animation/IKRig/Source/IKRigEditor/Private/RetargetEditor/IKRetargetEditPoseMode.cpp
@@ -55,73 +55,8 @@ void FIKRetargetEditPoseMode::Render(const FSceneView* View, FViewport* Viewport
 		return;
 	}
 
-	RenderSkeleton(PDI, Controller);
-}
-
-void FIKRetargetEditPoseMode::RenderSkeleton(
-	FPrimitiveDrawInterface* PDI,
-	const FIKRetargetEditorController* Controller)
-{
-	const UDebugSkelMeshComponent* MeshComponent = Controller->GetSkeletalMeshComponent(SourceOrTarget);
-	FTransform ComponentTransform = MeshComponent->GetComponentTransform();
-	const FReferenceSkeleton& RefSkeleton = MeshComponent->GetReferenceSkeleton();
-	const int32 NumBones = RefSkeleton.GetNum();
-
-	// get world transforms of bones
-	TArray<FBoneIndexType> RequiredBones;
-	RequiredBones.AddUninitialized(NumBones);
-	TArray<FTransform> WorldTransforms;
-	WorldTransforms.AddUninitialized(NumBones);
-	for (int32 Index=0; Index<NumBones; ++Index)
-	{
-		RequiredBones[Index] = Index;
-		WorldTransforms[Index] = MeshComponent->GetBoneTransform(Index, ComponentTransform);
-	}
-	
-	const UIKRetargeter* Asset = Controller->AssetController->GetAsset();
-	const float BoneDrawSize = Asset->BoneDrawSize;
-	const float MaxDrawRadius = Controller->TargetSkelMeshComponent->Bounds.SphereRadius * 0.01f;
-	const float BoneRadius = FMath::Min(1.0f, MaxDrawRadius) * BoneDrawSize;
-	
-	UPersonaOptions* ConfigOption = UPersonaOptions::StaticClass()->GetDefaultObject<UPersonaOptions>();
-	
-	FSkelDebugDrawConfig DrawConfig;
-	DrawConfig.BoneDrawMode = (EBoneDrawMode::Type)ConfigOption->DefaultBoneDrawSelection;
-	DrawConfig.BoneDrawSize = BoneRadius;
-	DrawConfig.bAddHitProxy = true;
-	DrawConfig.bForceDraw = false;
-	DrawConfig.DefaultBoneColor = GetMutableDefault<UPersonaOptions>()->DefaultBoneColor;
-	DrawConfig.AffectedBoneColor = GetMutableDefault<UPersonaOptions>()->AffectedBoneColor;
-	DrawConfig.SelectedBoneColor = GetMutableDefault<UPersonaOptions>()->SelectedBoneColor;
-	DrawConfig.ParentOfSelectedBoneColor = GetMutableDefault<UPersonaOptions>()->ParentOfSelectedBoneColor;
-
-	TArray<TRefCountPtr<HHitProxy>> HitProxies;
-	HitProxies.Reserve(NumBones);
-	for (int32 Index = 0; Index < NumBones; ++Index)
-	{
-		HitProxies.Add(new HIKRetargetEditorBoneProxy(RefSkeleton.GetBoneName(Index)));
-	}
-
-	// record selected bone indices
-	const TArray<FName>& SelectedBoneNames = Controller->GetSelectedBones();
-	TArray<int32> SelectedBones;
-	for (const FName& SelectedBoneName : SelectedBoneNames)
-	{
-		int32 SelectedBoneIndex = RefSkeleton.FindBoneIndex(SelectedBoneName);
-		SelectedBones.Add(SelectedBoneIndex);
-	}
-
-	SkeletalDebugRendering::DrawBones(
-		PDI,
-		ComponentTransform.GetLocation(),
-		RequiredBones,
-		RefSkeleton,
-		WorldTransforms,
-		SelectedBones,
-		TArray<FLinearColor>(),
-		HitProxies,
-		DrawConfig
-	);
+	Controller->RenderSkeleton(PDI, ERetargetSourceOrTarget::Source);
+	Controller->RenderSkeleton(PDI, ERetargetSourceOrTarget::Target);
 }
 
 bool FIKRetargetEditPoseMode::AllowWidgetMove()
@@ -424,10 +359,6 @@ void FIKRetargetEditPoseMode::Enter()
 	
 	// record skeleton to edit (must be constant between enter/exit)
 	SourceOrTarget = Controller->GetSourceOrTarget();
-	// hide skeleton we're editing so that we can draw our own editable version of it
-	const bool bEditingSource = SourceOrTarget == ERetargetSourceOrTarget::Source;
-	Controller->SourceSkelMeshComponent->SkeletonDrawMode = bEditingSource ? ESkeletonDrawMode::Hidden : ESkeletonDrawMode::GreyedOut;
-	Controller->TargetSkelMeshComponent->SkeletonDrawMode = !bEditingSource ? ESkeletonDrawMode::Hidden : ESkeletonDrawMode::GreyedOut;
 }
 
 void FIKRetargetEditPoseMode::Exit()
diff --git a/Engine/Plugins/Animation/IKRig/Source/IKRigEditor/Private/RetargetEditor/IKRetargetEditor.cpp b/Engine/Plugins/Animation/IKRig/Source/IKRigEditor/Private/RetargetEditor/IKRetargetEditor.cpp
--- a/Engine/Plugins/Animation/IKRig/Source/IKRigEditor/Private/RetargetEditor/IKRetargetEditor.cpp
+++ b/Engine/Plugins/Animation/IKRig/Source/IKRigEditor/Private/RetargetEditor/IKRetargetEditor.cpp
@@ -81,6 +81,7 @@ void FIKRetargetEditor::InitAssetEditor(
 	GetEditorModeManager().ActivateMode(FIKRetargetDefaultMode::ModeName);
 	FIKRetargetDefaultMode* DefaultMode = GetEditorModeManager().GetActiveModeTyped<FIKRetargetDefaultMode>(FIKRetargetDefaultMode::ModeName);
 	DefaultMode->SetEditorController(EditorController);
+	GetEditorModeManager().DeactivateMode(FIKRetargetDefaultMode::ModeName);
 
 	// give edit pose mode a pointer to the editor controller
 	GetEditorModeManager().ActivateMode(FIKRetargetEditPoseMode::ModeName);
@@ -271,12 +272,12 @@ TStatId FIKRetargetEditor::GetStatId() const
 
 void FIKRetargetEditor::PostUndo(bool bSuccess)
 {
-	EditorController->OnRetargeterNeedsInitialized(EditorController->AssetController->GetAsset());
+	EditorController->HandleRetargeterNeedsInitialized(EditorController->AssetController->GetAsset());
 }
 
 void FIKRetargetEditor::PostRedo(bool bSuccess)
 {
-	EditorController->OnRetargeterNeedsInitialized(EditorController->AssetController->GetAsset());
+	EditorController->HandleRetargeterNeedsInitialized(EditorController->AssetController->GetAsset());
 }
 
 void FIKRetargetEditor::HandleViewportCreated(const TSharedRef<class IPersonaViewport>& InViewport)
@@ -343,6 +344,14 @@ void FIKRetargetEditor::HandlePreviewSceneCreated(const TSharedRef<IPersonaPrevi
 	EditorController->SourceSkelMeshComponent = NewObject<UDebugSkelMeshComponent>(Actor);
 	EditorController->TargetSkelMeshComponent = NewObject<UDebugSkelMeshComponent>(Actor);
 
+	// hide skeletons, we want to do custom rendering
+	EditorController->SourceSkelMeshComponent->SkeletonDrawMode = ESkeletonDrawMode::Hidden;
+	EditorController->TargetSkelMeshComponent->SkeletonDrawMode = ESkeletonDrawMode::Hidden;
+
+	// don't want selectable meshes as it gets in the way of bone selection
+	EditorController->SourceSkelMeshComponent->bSelectable = false;
+	EditorController->TargetSkelMeshComponent->bSelectable = false;
+	
 	// setup an apply an anim instance to the skeletal mesh component
 	EditorController->SourceAnimInstance = NewObject<UIKRetargetAnimInstance>(EditorController->SourceSkelMeshComponent, TEXT("IKRetargetSourceAnimScriptInstance"));
 	EditorController->TargetAnimInstance = NewObject<UIKRetargetAnimInstance>(EditorController->TargetSkelMeshComponent, TEXT("IKRetargetTargetAnimScriptInstance"));
@@ -359,10 +368,6 @@ void FIKRetargetEditor::HandlePreviewSceneCreated(const TSharedRef<IPersonaPrevi
 	InPersonaPreviewScene->SetPreviewMeshComponent(EditorController->SourceSkelMeshComponent);
 	InPersonaPreviewScene->SetPreviewMesh(SourceMesh);
 	InPersonaPreviewScene->SetAdditionalMeshesSelectable(false);
-
-	// SetPreviewMesh() sets this flag true, which the render uses to filter out objects for selection highlighting...
-	// but since we want to be able to select the mesh in this viewport, we have to set it back to false
-	EditorController->SourceSkelMeshComponent->bCanHighlightSelectedSections = false;
 	
 	InPersonaPreviewScene->AddComponent(EditorController->SourceSkelMeshComponent, FTransform::Identity);
 	InPersonaPreviewScene->AddComponent(EditorController->TargetSkelMeshComponent, FTransform::Identity);
@@ -442,6 +447,9 @@ void FIKRetargetEditor::OnFinishedChangingDetails(const FPropertyChangedEvent& P
 		USkeletalMesh* TargetMesh = EditorController->GetSkeletalMesh(ERetargetSourceOrTarget::Target);
 		EditorController->SourceSkelMeshComponent->SetSkeletalMesh(SourceMesh);
 		EditorController->TargetSkelMeshComponent->SetSkeletalMesh(TargetMesh);
+
+		// reset bone selections in case of incompatible indices
+		EditorController->ClearSelection();
 	
 		// apply mesh to the preview scene
 		TSharedRef<IPersonaPreviewScene> PreviewScene = GetPersonaToolkit()->GetPreviewScene();
diff --git a/Engine/Plugins/Animation/IKRig/Source/IKRigEditor/Private/RetargetEditor/IKRetargetEditorController.cpp b/Engine/Plugins/Animation/IKRig/Source/IKRigEditor/Private/RetargetEditor/IKRetargetEditorController.cpp
--- a/Engine/Plugins/Animation/IKRig/Source/IKRigEditor/Private/RetargetEditor/IKRetargetEditorController.cpp
+++ b/Engine/Plugins/Animation/IKRig/Source/IKRigEditor/Private/RetargetEditor/IKRetargetEditorController.cpp
@@ -6,13 +6,16 @@
 #include "AssetToolsModule.h"
 #include "EditorModeManager.h"
 #include "IContentBrowserSingleton.h"
+#include "SkeletalDebugRendering.h"
 #include "Widgets/Input/SButton.h"
 #include "Animation/DebugSkelMeshComponent.h"
 #include "Animation/PoseAsset.h"
+#include "Preferences/PersonaOptions.h"
 #include "RetargetEditor/IKRetargetAnimInstance.h"
 #include "RetargetEditor/IKRetargetDefaultMode.h"
 #include "RetargetEditor/IKRetargetEditPoseMode.h"
 #include "RetargetEditor/IKRetargetEditor.h"
+#include "RetargetEditor/IKRetargetHitProxies.h"
 #include "RetargetEditor/SIKRetargetChainMapList.h"
 #include "RetargetEditor/SIKRetargetHierarchy.h"
 #include "Retargeter/IKRetargeter.h"
@@ -33,6 +36,10 @@ void FIKRetargetEditorController::Initialize(TSharedPtr<FIKRetargetEditor> InEdi
 	PoseExporter = MakeShared<FIKRetargetPoseExporter>();
 	PoseExporter->Initialize(SharedThis(this));
 
+	SelectedBoneNames.Add(ERetargetSourceOrTarget::Source);
+	SelectedBoneNames.Add(ERetargetSourceOrTarget::Target);
+	LastSelectedItem = ERetargetSelectionType::NONE;
+
 	// clean the asset before editing
 	const bool bForceReinitialization = false;
 	AssetController->CleanChainMapping(bForceReinitialization);
@@ -43,13 +50,14 @@ void FIKRetargetEditorController::Initialize(TSharedPtr<FIKRetargetEditor> InEdi
 	BindToIKRigAsset(AssetController->GetAsset()->GetTargetIKRigWriteable());
 
 	// bind callback when retargeter needs reinitialized
-	ReInitDelegateHandle = AssetController->OnRetargeterNeedsInitialized().AddSP(this, &FIKRetargetEditorController::OnRetargeterNeedsInitialized);
+	ReInitDelegateHandle = AssetController->OnRetargeterNeedsInitialized().AddSP(this, &FIKRetargetEditorController::HandleRetargeterNeedsInitialized);
 }
 
 void FIKRetargetEditorController::Close() const
 {
 	AssetController->OnRetargeterNeedsInitialized().Remove(ReInitDelegateHandle);
 
+	
 	if (UIKRigController* IKRigController = UIKRigController::GetIKRigController(BoundToIKRig))
 	{
 		IKRigController->OnIKRigNeedsInitialized().Remove(ReInitDelegateHandle);
@@ -70,13 +78,13 @@ void FIKRetargetEditorController::BindToIKRigAsset(UIKRigDefinition* InIKRig)
 	UIKRigController* Controller = UIKRigController::GetIKRigController(InIKRig);
 	if (!Controller->OnIKRigNeedsInitialized().IsBoundToObject(this))
 	{
-		ReInitIKDelegateHandle = Controller->OnIKRigNeedsInitialized().AddSP(this, &FIKRetargetEditorController::OnIKRigNeedsInitialized);
-		RenameChainDelegateHandle = Controller->OnRetargetChainRenamed().AddSP(this, &FIKRetargetEditorController::OnRetargetChainRenamed);
-		RemoveChainDelegateHandle = Controller->OnRetargetChainRemoved().AddSP(this, &FIKRetargetEditorController::OnRetargetChainRemoved);
+		ReInitIKDelegateHandle = Controller->OnIKRigNeedsInitialized().AddSP(this, &FIKRetargetEditorController::HandleIKRigNeedsInitialized);
+		RenameChainDelegateHandle = Controller->OnRetargetChainRenamed().AddSP(this, &FIKRetargetEditorController::HandleRetargetChainRenamed);
+		RemoveChainDelegateHandle = Controller->OnRetargetChainRemoved().AddSP(this, &FIKRetargetEditorController::HandleRetargetChainRemoved);
 	}
 }
 
-void FIKRetargetEditorController::OnIKRigNeedsInitialized(UIKRigDefinition* ModifiedIKRig) const
+void FIKRetargetEditorController::HandleIKRigNeedsInitialized(UIKRigDefinition* ModifiedIKRig) const
 {
 	UIKRetargeter* Retargeter = AssetController->GetAsset();
 	
@@ -90,24 +98,24 @@ void FIKRetargetEditorController::OnIKRigNeedsInitialized(UIKRigDefinition* Modi
 	}
 
 	// the target anim instance has the RetargetPoseFromMesh node which needs reinitialized with new asset version
-	OnRetargeterNeedsInitialized(Retargeter);
+	HandleRetargeterNeedsInitialized(Retargeter);
 }
 
-void FIKRetargetEditorController::OnRetargetChainRenamed(UIKRigDefinition* ModifiedIKRig, FName OldName, FName NewName) const
+void FIKRetargetEditorController::HandleRetargetChainRenamed(UIKRigDefinition* ModifiedIKRig, FName OldName, FName NewName) const
 {
 	check(ModifiedIKRig)
 	
 	AssetController->OnRetargetChainRenamed(ModifiedIKRig, OldName, NewName);
 }
 
-void FIKRetargetEditorController::OnRetargetChainRemoved(UIKRigDefinition* ModifiedIKRig, const FName& InChainRemoved) const
+void FIKRetargetEditorController::HandleRetargetChainRemoved(UIKRigDefinition* ModifiedIKRig, const FName& InChainRemoved) const
 {
 	check(ModifiedIKRig)
 	AssetController->OnRetargetChainRemoved(ModifiedIKRig, InChainRemoved);
 	RefreshAllViews();
 }
 
-void FIKRetargetEditorController::OnRetargeterNeedsInitialized(UIKRetargeter* Retargeter) const
+void FIKRetargetEditorController::HandleRetargeterNeedsInitialized(UIKRetargeter* Retargeter) const
 {
 	// clear the output log
 	ClearOutputLog();
@@ -335,6 +343,11 @@ void FIKRetargetEditorController::ClearOutputLog() const
 
 void FIKRetargetEditorController::RefreshAllViews() const
 {
+	//if (!Editor.IsValid())
+	//{
+	//	return;
+	//}
+	
 	Editor.Pin()->RegenerateMenusAndToolbars();
 	RefreshDetailsView();
 	RefreshChainsView();
@@ -463,38 +476,28 @@ void FIKRetargetEditorController::SetSourceOrTargetMode(ERetargetSourceOrTarget
 	{
 		return;
 	}
-
-	// clear the selection on old skeleton
-	ClearSelection();
 	
 	// store the new skeleton mode
 	CurrentlyEditingSourceOrTarget = NewMode;
 
-	switch (GetRetargeterMode())
+	// if we switch source/target while in edit mode we need to re-enter that mode
+	if (GetRetargeterMode() == ERetargeterOutputMode::EditRetargetPose)
 	{
-		case ERetargeterOutputMode::EditRetargetPose:
-		{
-			FEditorModeTools& EditorModeManager = Editor.Pin()->GetEditorModeManager();
-			FIKRetargetEditPoseMode* EditMode = EditorModeManager.GetActiveModeTyped<FIKRetargetEditPoseMode>(FIKRetargetEditPoseMode::ModeName);
-			if (EditMode)
-			{
-				// FIKRetargetEditPoseMode::Enter() is reentrant and written so we can switch between editing
-				// source / target skeleton without having to enter/exit the mode; just call Enter() again
-				EditMode->Enter();
-			}
-			break;
-		}
-		case ERetargeterOutputMode::ShowRetargetPose:
-		case ERetargeterOutputMode::RunRetarget:
+		FEditorModeTools& EditorModeManager = Editor.Pin()->GetEditorModeManager();
+		FIKRetargetEditPoseMode* EditMode = EditorModeManager.GetActiveModeTyped<FIKRetargetEditPoseMode>(FIKRetargetEditPoseMode::ModeName);
+		if (EditMode)
 		{
-			// toggle visibility of currently active skeleton
-			const bool bEditingSource = NewMode == ERetargetSourceOrTarget::Source;
-			SourceSkelMeshComponent->SkeletonDrawMode = bEditingSource ? ESkeletonDrawMode::Default : ESkeletonDrawMode::GreyedOut;
-			TargetSkelMeshComponent->SkeletonDrawMode = bEditingSource ? ESkeletonDrawMode::GreyedOut : ESkeletonDrawMode::Default;
-			break;
+			// FIKRetargetEditPoseMode::Enter() is reentrant and written so we can switch between editing
+			// source / target skeleton without having to enter/exit the mode; just call Enter() again
+			EditMode->Enter();
 		}
-		default:
-			checkNoEntry();
+	}
+
+	// make sure details panel updates with selected bone on OTHER skeleton
+	if (LastSelectedItem == ERetargetSelectionType::BONE)
+	{
+		const TArray<FName>& SelectedBones = GetSelectedBones();
+		EditBoneSelection(SelectedBones, ESelectionEdit::Replace);
 	}
 	
 	RefreshAllViews();
@@ -503,13 +506,13 @@ void FIKRetargetEditorController::SetSourceOrTargetMode(ERetargetSourceOrTarget
 void FIKRetargetEditorController::SetSelectedMesh(UPrimitiveComponent* InMeshComponent)
 {
 	SelectedMesh = InMeshComponent;
-	SourceSkelMeshComponent->PushSelectionToProxy();
-	TargetSkelMeshComponent->PushSelectionToProxy();
-	SourceSkelMeshComponent->MarkRenderStateDirty();
-	TargetSkelMeshComponent->MarkRenderStateDirty();
+	if (SelectedMesh)
+	{
+		LastSelectedItem = ERetargetSelectionType::MESH;
+	}
 }
 
-UPrimitiveComponent* FIKRetargetEditorController::GetSelectedMesh()
+UPrimitiveComponent* FIKRetargetEditorController::GetSelectedMesh() const
 {
 	return SelectedMesh;
 }
@@ -526,6 +529,8 @@ void FIKRetargetEditorController::EditBoneSelection(
 		return;
 	}
 
+	LastSelectedItem = ERetargetSelectionType::BONE;
+
 	// deselect mesh
 	SetSelectedMesh(nullptr);
 	SetRootSelected(false);
@@ -536,7 +541,7 @@ void FIKRetargetEditorController::EditBoneSelection(
 		{
 			for (const FName& BoneName : InBoneNames)
 			{
-				SelectedBones.AddUnique(BoneName);
+				SelectedBoneNames[CurrentlyEditingSourceOrTarget].AddUnique(BoneName);
 			}
 			
 			break;
@@ -545,38 +550,19 @@ void FIKRetargetEditorController::EditBoneSelection(
 		{
 			for (const FName& BoneName : InBoneNames)
 			{
-				SelectedBones.Remove(BoneName);
+				SelectedBoneNames[CurrentlyEditingSourceOrTarget].Remove(BoneName);
 			}
 			break;
 		}
 		case ESelectionEdit::Replace:
 		{
-			SelectedBones = InBoneNames;
+			SelectedBoneNames[CurrentlyEditingSourceOrTarget] = InBoneNames;
 			break;
 		}
 		default:
 			checkNoEntry();
 	}
 
-	// convert to bone indices
-	const FReferenceSkeleton& RefSkeleton = DebugComponent->GetSkeletalMeshAsset()->GetRefSkeleton();
-	TArray<int32> SelectedBoneIndices;
-	for (const FName& Bone : SelectedBones)
-	{
-		int32 BoneIndex = RefSkeleton.FindBoneIndex(Bone);
-		SelectedBoneIndices.Add(BoneIndex);
-		
-		if (BoneIndex == INDEX_NONE)
-		{
-			ensureMsgf(false, TEXT("Incoming selection list is not compatible with the currently edited skeleton."));
-			SelectedBoneIndices.Reset();
-			break;
-		}
-	}
-	
-	// apply selection to debug mesh component so rendering knows
-	DebugComponent->BonesOfInterest = SelectedBoneIndices;
-
 	// update hierarchy view
 	if (!bFromHierarchyView)
 	{
@@ -584,14 +570,14 @@ void FIKRetargetEditorController::EditBoneSelection(
 	}
 
 	// update details
-	if (SelectedBones.IsEmpty())
+	if (SelectedBoneNames[CurrentlyEditingSourceOrTarget].IsEmpty())
 	{
 		SetDetailsObject(AssetController->GetAsset());
 	}
 	else
 	{
 		TArray<UObject*> SelectedBoneDetails;
-		for (const FName& SelectedBone : SelectedBones)
+		for (const FName& SelectedBone : SelectedBoneNames[CurrentlyEditingSourceOrTarget])
 		{
 			TObjectPtr<UIKRetargetBoneDetails> BoneDetails = GetOrCreateBoneDetailsObject(SelectedBone);
 			SelectedBoneDetails.Add(BoneDetails);
@@ -608,6 +594,8 @@ void FIKRetargetEditorController::EditChainSelection(
 	// deselect others
 	SetSelectedMesh(nullptr);
 	SetRootSelected(false);
+
+	LastSelectedItem = ERetargetSelectionType::CHAIN;
 	
 	// update selection set based on edit mode
 	switch (EditMode)
@@ -670,17 +658,20 @@ void FIKRetargetEditorController::EditChainSelection(
 void FIKRetargetEditorController::SetRootSelected(const bool bIsSelected)
 {
 	bIsRootSelected = bIsSelected;
-
-	if (bIsSelected)
+	if (!bIsSelected)
 	{
-		URetargetRootSettings* RootSettings = AssetController->GetAsset()->GetRootSettingsUObject();
-		if (RootSettings->EditorController.Get() != this)
-		{
-			RootSettings->EditorController = SharedThis(this);	
-		}
-		
-		SetDetailsObject(RootSettings);
+		return;
 	}
+	
+	URetargetRootSettings* RootSettings = AssetController->GetAsset()->GetRootSettingsUObject();
+	if (RootSettings->EditorController.Get() != this)
+	{
+		RootSettings->EditorController = SharedThis(this);	
+	}
+	
+	SetDetailsObject(RootSettings);
+
+	LastSelectedItem = ERetargetSelectionType::ROOT;
 }
 
 void FIKRetargetEditorController::ClearSelection(const bool bKeepBoneSelection)
@@ -699,14 +690,17 @@ void FIKRetargetEditorController::ClearSelection(const bool bKeepBoneSelection)
 	// clear bone selection
 	if (!bKeepBoneSelection)
 	{
-		const TArray<FName> Empty;
-		constexpr bool bFromHierarchy = false;
-		EditBoneSelection(Empty, ESelectionEdit::Replace, bFromHierarchy);
+		SetSelectedMesh(nullptr);
+		SetRootSelected(false);
+		SelectedBoneNames[ERetargetSourceOrTarget::Source].Reset();
+		SelectedBoneNames[ERetargetSourceOrTarget::Target].Reset();
 
 		// show global details
 		SetDetailsObject(AssetController->GetAsset());
 	}
 
+	LastSelectedItem = ERetargetSelectionType::NONE;
+
 	RefreshDetailsView();
 }
 
@@ -1002,7 +996,7 @@ void FIKRetargetEditorController::HandleResetAllBones() const
 void FIKRetargetEditorController::HandleResetSelectedBones() const
 {
 	const FName CurrentPose = AssetController->GetCurrentRetargetPoseName(CurrentlyEditingSourceOrTarget);
-	AssetController->ResetRetargetPose(CurrentPose, GetSelectedBones(), GetSourceOrTarget());
+	AssetController->ResetRetargetPose(CurrentPose, GetSelectedBones(), CurrentlyEditingSourceOrTarget);
 }
 
 void FIKRetargetEditorController::HandleResetSelectedAndChildrenBones() const
@@ -1017,7 +1011,7 @@ void FIKRetargetEditorController::HandleResetSelectedAndChildrenBones() const
 	
 	// get list of all children of selected bones
 	TArray<int32> AllChildrenIndices;
-	for (const FName& SelectedBone : SelectedBones)
+	for (const FName& SelectedBone : SelectedBoneNames[CurrentlyEditingSourceOrTarget])
 	{
 		const int32 SelectedBoneIndex = RefSkeleton.FindBoneIndex(SelectedBone);
 		AllChildrenIndices.Add(SelectedBoneIndex);
@@ -1033,7 +1027,7 @@ void FIKRetargetEditorController::HandleResetSelectedAndChildrenBones() const
 	}
 
 	// merge total list of all selected bones and their children
-	TArray<FName> BonesToReset = SelectedBones;
+	TArray<FName> BonesToReset = SelectedBoneNames[CurrentlyEditingSourceOrTarget];
 	for (const int32 ChildIndex : AllChildrenIndices)
 	{
 		BonesToReset.AddUnique(RefSkeleton.GetBoneName(ChildIndex));
@@ -1138,7 +1132,104 @@ void FIKRetargetEditorController::AddReferencedObjects(FReferenceCollector& Coll
 	{
 		Collector.AddReferencedObject(Pair.Value);
 	}
-};
+}
+
+void FIKRetargetEditorController::RenderSkeleton(FPrimitiveDrawInterface* PDI, ERetargetSourceOrTarget InSourceOrTarget) const
+{
+	const UDebugSkelMeshComponent* MeshComponent = GetSkeletalMeshComponent(InSourceOrTarget);
+	const FTransform ComponentTransform = MeshComponent->GetComponentTransform();
+	const FReferenceSkeleton& RefSkeleton = MeshComponent->GetReferenceSkeleton();
+	const int32 NumBones = RefSkeleton.GetNum();
+
+	// get world transforms of bones
+	TArray<FBoneIndexType> RequiredBones;
+	RequiredBones.AddUninitialized(NumBones);
+	TArray<FTransform> WorldTransforms;
+	WorldTransforms.AddUninitialized(NumBones);
+	for (int32 Index=0; Index<NumBones; ++Index)
+	{
+		RequiredBones[Index] = Index;
+		WorldTransforms[Index] = MeshComponent->GetBoneTransform(Index, ComponentTransform);
+	}
+	
+	const UIKRetargeter* Asset = AssetController->GetAsset();
+	const float BoneDrawSize = Asset->BoneDrawSize;
+	const float MaxDrawRadius = MeshComponent->Bounds.SphereRadius * 0.01f;
+	const float BoneRadius = FMath::Min(1.0f, MaxDrawRadius) * BoneDrawSize;
+	const bool bIsSelectable = InSourceOrTarget == GetSourceOrTarget();
+	const FLinearColor DefaultColor = bIsSelectable ? GetMutableDefault<UPersonaOptions>()->DefaultBoneColor : GetMutableDefault<UPersonaOptions>()->DisabledBoneColor;
+	
+	UPersonaOptions* ConfigOption = UPersonaOptions::StaticClass()->GetDefaultObject<UPersonaOptions>();
+	
+	FSkelDebugDrawConfig DrawConfig;
+	DrawConfig.BoneDrawMode = (EBoneDrawMode::Type)ConfigOption->DefaultBoneDrawSelection;
+	DrawConfig.BoneDrawSize = BoneRadius;
+	DrawConfig.bAddHitProxy = bIsSelectable;
+	DrawConfig.bForceDraw = false;
+	DrawConfig.DefaultBoneColor = DefaultColor;
+	DrawConfig.AffectedBoneColor = GetMutableDefault<UPersonaOptions>()->AffectedBoneColor;
+	DrawConfig.SelectedBoneColor = GetMutableDefault<UPersonaOptions>()->SelectedBoneColor;
+	DrawConfig.ParentOfSelectedBoneColor = GetMutableDefault<UPersonaOptions>()->ParentOfSelectedBoneColor;
+
+	TArray<TRefCountPtr<HHitProxy>> HitProxies;
+	TArray<int32> SelectedBones;
+	
+	// create hit proxies and selection set only for the currently active skeleton
+	if (bIsSelectable)
+	{
+		HitProxies.Reserve(NumBones);
+		for (int32 Index = 0; Index < NumBones; ++Index)
+		{
+			HitProxies.Add(new HIKRetargetEditorBoneProxy(RefSkeleton.GetBoneName(Index), Index, InSourceOrTarget));
+		}
+
+		// record selected bone indices
+		for (const FName& SelectedBoneName : SelectedBoneNames[InSourceOrTarget])
+		{
+			int32 SelectedBoneIndex = RefSkeleton.FindBoneIndex(SelectedBoneName);
+			SelectedBones.Add(SelectedBoneIndex);
+		}
+	}
+
+	// generate bone colors, blue on selected chains
+	TArray<FLinearColor> BoneColors;
+	{
+		// set all to default color
+		BoneColors.Init(DefaultColor, RefSkeleton.GetNum());
+
+		// highlight selected chains in blue
+		const TArray<FName>& SelectedChainNames = GetSelectedChains();
+		const FLinearColor HighlightedColor = FLinearColor::Blue;
+		const bool bIsSource = InSourceOrTarget == ERetargetSourceOrTarget::Source;
+		if (const UIKRetargetProcessor* Processor = GetRetargetProcessor())
+		{
+			const TArray<FRetargetChainPairFK>& FKChains = Processor->GetFKChainPairs();
+			for (const FRetargetChainPairFK& FKChain : FKChains)
+			{
+				if (SelectedChainNames.Contains(FKChain.TargetBoneChainName))
+				{
+					const TArray<int32>& BoneIndices = bIsSource ? FKChain.SourceBoneIndices : FKChain.TargetBoneIndices;
+					for (const int32& BoneIndex : BoneIndices)
+					{
+						BoneColors[BoneIndex] = HighlightedColor;
+					}	
+				}
+			}
+		}
+	}
+
+	SkeletalDebugRendering::DrawBones(
+		PDI,
+		ComponentTransform.GetLocation(),
+		RequiredBones,
+		RefSkeleton,
+		WorldTransforms,
+		SelectedBones,
+		BoneColors,
+		HitProxies,
+		DrawConfig
+	);
+}
 
 FText FIKRetargetEditorController::GetCurrentPoseName() const
 {
diff --git a/Engine/Plugins/Animation/IKRig/Source/IKRigEditor/Private/RetargetEditor/IKRetargeterController.cpp b/Engine/Plugins/Animation/IKRig/Source/IKRigEditor/Private/RetargetEditor/IKRetargeterController.cpp
--- a/Engine/Plugins/Animation/IKRig/Source/IKRigEditor/Private/RetargetEditor/IKRetargeterController.cpp
+++ b/Engine/Plugins/Animation/IKRig/Source/IKRigEditor/Private/RetargetEditor/IKRetargeterController.cpp
@@ -324,13 +324,25 @@ const TArray<TObjectPtr<URetargetChainSettings>>& UIKRetargeterController::GetCh
 	return Asset->ChainSettings;
 }
 
-bool UIKRetargeterController::IsChainRunningIK(const TObjectPtr<URetargetChainSettings> ChainSettings) const
+FName UIKRetargeterController::GetChainGoal(const TObjectPtr<URetargetChainSettings> ChainSettings) const
 {
 	if (ChainSettings.IsNull())
 	{
-		return false;
+		return NAME_None;
+	}
+	
+	UIKRigDefinition* TargetIKRig = GetIKRigWriteable(ERetargetSourceOrTarget::Target);
+	if (!TargetIKRig)
+	{
+		return NAME_None;
 	}
 	
+	const UIKRigController* IKRigController = UIKRigController::GetIKRigController(TargetIKRig);
+	return IKRigController->GetRetargetChainGoal(ChainSettings->TargetChain);
+}
+
+bool UIKRetargeterController::IsChainGoalConnectedToASolver(const FName& GoalName) const
+{
 	UIKRigDefinition* TargetIKRig = GetIKRigWriteable(ERetargetSourceOrTarget::Target);
 	if (!TargetIKRig)
 	{
@@ -338,8 +350,7 @@ bool UIKRetargeterController::IsChainRunningIK(const TObjectPtr<URetargetChainSe
 	}
 	
 	const UIKRigController* IKRigController = UIKRigController::GetIKRigController(TargetIKRig);
-	const FName ChainGoalName = IKRigController->GetRetargetChainGoal(ChainSettings->TargetChain);
-	return IKRigController->IsGoalConnectedToAnySolver(ChainGoalName);
+	return IKRigController->IsGoalConnectedToAnySolver(GoalName);
 }
 
 void UIKRetargeterController::AddRetargetPose(
diff --git a/Engine/Plugins/Animation/IKRig/Source/IKRigEditor/Private/RetargetEditor/SIKRetargetAssetBrowser.cpp b/Engine/Plugins/Animation/IKRig/Source/IKRigEditor/Private/RetargetEditor/SIKRetargetAssetBrowser.cpp
--- a/Engine/Plugins/Animation/IKRig/Source/IKRigEditor/Private/RetargetEditor/SIKRetargetAssetBrowser.cpp
+++ b/Engine/Plugins/Animation/IKRig/Source/IKRigEditor/Private/RetargetEditor/SIKRetargetAssetBrowser.cpp
@@ -66,8 +66,8 @@ void SIKRetargetAssetBrowser::RefreshView()
 	AssetPickerConfig.bShowTypeInColumnView = true;
 	AssetPickerConfig.OnShouldFilterAsset = FOnShouldFilterAsset::CreateSP(this, &SIKRetargetAssetBrowser::OnShouldFilterAsset);
 	AssetPickerConfig.DefaultFilterMenuExpansion = EAssetTypeCategories::Animation;
-	
 	AssetPickerConfig.OnAssetDoubleClicked = FOnAssetSelected::CreateSP(this, &SIKRetargetAssetBrowser::OnAssetDoubleClicked);
+	AssetPickerConfig.OnGetAssetContextMenu = FOnGetAssetContextMenu::CreateSP(this, &SIKRetargetAssetBrowser::OnGetAssetContextMenu);
 	AssetPickerConfig.GetCurrentSelectionDelegates.Add(&GetCurrentSelectionDelegate);
 	AssetPickerConfig.bAllowNullSelection = false;
 
@@ -81,12 +81,52 @@ void SIKRetargetAssetBrowser::RefreshView()
 
 	// Also hide the type column by default (but allow users to enable it, so don't use bShowTypeInColumnView)
 	AssetPickerConfig.HiddenColumnNames.Add(TEXT("Class"));
+	AssetPickerConfig.HiddenColumnNames.Add(TEXT("Has Virtualized Data"));
 
-	FContentBrowserModule& ContentBrowserModule = FModuleManager::Get().LoadModuleChecked<FContentBrowserModule>(TEXT("ContentBrowser"));
-
+	const FContentBrowserModule& ContentBrowserModule = FModuleManager::Get().LoadModuleChecked<FContentBrowserModule>(TEXT("ContentBrowser"));
 	AssetBrowserBox->SetContent(ContentBrowserModule.Get().CreateAssetPicker(AssetPickerConfig));
 }
 
+TSharedPtr<SWidget> SIKRetargetAssetBrowser::OnGetAssetContextMenu(const TArray<FAssetData>& SelectedAssets) const
+{
+	if (SelectedAssets.Num() <= 0)
+	{
+		return nullptr;
+	}
+
+	UObject* SelectedAsset = SelectedAssets[0].GetAsset();
+	if (SelectedAsset == nullptr)
+	{
+		return nullptr;
+	}
+	
+	FMenuBuilder MenuBuilder(true, MakeShared<FUICommandList>());
+
+	MenuBuilder.BeginSection(TEXT("Asset"), LOCTEXT("AssetSectionLabel", "Asset"));
+	{
+		MenuBuilder.AddMenuEntry(
+			LOCTEXT("Browse", "Browse to Asset"),
+			LOCTEXT("BrowseTooltip", "Browses to the associated asset and selects it in the most recently used Content Browser (summoning one if necessary)"),
+			FSlateIcon(FAppStyle::GetAppStyleSetName(), "SystemWideCommands.FindInContentBrowser.Small"),
+			FUIAction(
+				FExecuteAction::CreateLambda([SelectedAsset] ()
+				{
+					if (SelectedAsset)
+					{
+						const TArray<FAssetData>& Assets = { SelectedAsset };
+						const FContentBrowserModule& ContentBrowserModule = FModuleManager::Get().LoadModuleChecked<FContentBrowserModule>("ContentBrowser");
+						ContentBrowserModule.Get().SyncBrowserToAssets(Assets);
+					}
+				}),
+				FCanExecuteAction::CreateLambda([] () { return true; })
+			)
+		);
+	}
+	MenuBuilder.EndSection();
+
+	return MenuBuilder.MakeWidget();
+}
+
 FReply SIKRetargetAssetBrowser::OnExportButtonClicked()
 {
 	FIKRetargetEditorController* Controller = EditorController.Pin().Get();
@@ -282,7 +322,7 @@ void SBatchExportDialog::Construct(const FArguments& InArgs)
 						.VAlign(VAlign_Center)
 						.Padding(2, 1)
 						[
-							SNew(STextBlock).Text(LOCTEXT("DuplicateAndRetarget_Prefix", "Prefix:"))
+							SNew(STextBlock).Text(LOCTEXT("DuplicateAndRetarget_Prefix", "Add Prefix:"))
 						]
 
 						+SHorizontalBox::Slot()
@@ -315,7 +355,7 @@ void SBatchExportDialog::Construct(const FArguments& InArgs)
 						.VAlign(VAlign_Center)
 						.Padding(2, 1)
 						[
-							SNew(STextBlock).Text(LOCTEXT("DuplicateAndRetarget_Suffix", "Suffix:"))
+							SNew(STextBlock).Text(LOCTEXT("DuplicateAndRetarget_Suffix", "Add Suffix:"))
 						]
 
 						+SHorizontalBox::Slot()
@@ -344,62 +384,61 @@ void SBatchExportDialog::Construct(const FArguments& InArgs)
 					[
 						SNew(SHorizontalBox)
 						+SHorizontalBox::Slot()
-						.AutoWidth()
+						.HAlign(HAlign_Right)
 						.VAlign(VAlign_Center)
-						.Padding(2, 1)
 						[
-							SNew(STextBlock).Text(LOCTEXT("DuplicateAndRetarget_Search", "Search:"))
+							SNew(SHorizontalBox)
+							+SHorizontalBox::Slot()
+							.Padding(2, 1)
+							[
+								SNew(STextBlock).Text(LOCTEXT("DuplicateAndRetarget_Search", "Search for:"))
+							]
+
+							+SHorizontalBox::Slot()
+							[
+								SNew(SEditableTextBox)
+								.Text_Lambda([this]
+								{
+									return FText::FromString(BatchContext.NameRule.ReplaceFrom);
+								})
+								.MinDesiredWidth(TextInputWidths)
+								.OnTextChanged_Lambda([this](const FText& InText)
+								{
+									BatchContext.NameRule.ReplaceFrom = InText.ToString();
+									UpdateExampleText();
+								})
+								.IsReadOnly(false)
+								.RevertTextOnEscape(true)
+							]
 						]
-
+						
 						+SHorizontalBox::Slot()
-						.AutoWidth()
-						[
-							SNew(SEditableTextBox)
-							.Text_Lambda([this]
-							{
-								return FText::FromString(BatchContext.NameRule.ReplaceFrom);
-							})
-							.MinDesiredWidth(TextInputWidths)
-							.OnTextChanged_Lambda([this](const FText& InText)
-							{
-								BatchContext.NameRule.ReplaceFrom = InText.ToString();
-								UpdateExampleText();
-							})
-							.IsReadOnly(false)
-							.RevertTextOnEscape(true)
-						]
-					]
-
-					+SVerticalBox::Slot()
-					.AutoHeight()
-					.Padding(2, 1)
-					.HAlign(HAlign_Right)
-					[
-						SNew(SHorizontalBox)
-						+SHorizontalBox::Slot()
-						.AutoWidth()
+						.HAlign(HAlign_Right)
 						.VAlign(VAlign_Center)
-						.Padding(2, 1)
 						[
-							SNew(STextBlock).Text(LOCTEXT("DuplicateAndRetarget_Replace", "Replace:"))
-						]
-
-						+SHorizontalBox::Slot()
-						.AutoWidth()
-						[
-							SNew(SEditableTextBox)
-							.Text_Lambda([this]
-							{
-								return FText::FromString(BatchContext.NameRule.ReplaceTo);
-							})
-							.MinDesiredWidth(TextInputWidths)
-							.OnTextChanged_Lambda([this](const FText& InText)
-							{
-								BatchContext.NameRule.ReplaceTo = InText.ToString();
-								UpdateExampleText();
-							})
-							.IsReadOnly(false)
-							.RevertTextOnEscape(true)
+							SNew(SHorizontalBox)
+							+SHorizontalBox::Slot()
+							.Padding(2, 1)
+							[
+								SNew(STextBlock).Text(LOCTEXT("DuplicateAndRetarget_Replace", "Replace with:"))
+							]
+
+							+SHorizontalBox::Slot()
+							[
+								SNew(SEditableTextBox)
+								.Text_Lambda([this]
+								{
+									return FText::FromString(BatchContext.NameRule.ReplaceTo);
+								})
+								.MinDesiredWidth(TextInputWidths)
+								.OnTextChanged_Lambda([this](const FText& InText)
+								{
+									BatchContext.NameRule.ReplaceTo = InText.ToString();
+									UpdateExampleText();
+								})
+								.IsReadOnly(false)
+								.RevertTextOnEscape(true)
+							]
 						]
 					]
 
diff --git a/Engine/Plugins/Animation/IKRig/Source/IKRigEditor/Private/RigEditor/IKRigController.cpp b/Engine/Plugins/Animation/IKRig/Source/IKRigEditor/Private/RigEditor/IKRigController.cpp
--- a/Engine/Plugins/Animation/IKRig/Source/IKRigEditor/Private/RigEditor/IKRigController.cpp
+++ b/Engine/Plugins/Animation/IKRig/Source/IKRigEditor/Private/RigEditor/IKRigController.cpp
@@ -367,7 +367,10 @@ FName UIKRigController::GetUniqueRetargetChainName(const FName& NameToMakeUnique
 	return FName(NameToMakeUnique, Number);
 }
 
-bool UIKRigController::ValidateChain(const FName& ChainName, TSet<int32>& OutChainIndices) const
+bool UIKRigController::ValidateChain(
+	const FName& ChainName,
+	const FIKRigSkeleton* OptionalSkeleton,
+	TSet<int32>& OutChainIndices) const
 {
 	const FBoneChain* Chain = Asset->GetRetargetChainByName(ChainName);
 	if (!Chain)
@@ -375,21 +378,33 @@ bool UIKRigController::ValidateChain(const FName& ChainName, TSet<int32>& OutCha
 		return false; // chain doesn't exist
 	}
 
-	const FIKRigSkeleton &Skeleton = GetIKRigSkeleton();
+	const FIKRigSkeleton &Skeleton = OptionalSkeleton ? *OptionalSkeleton : GetIKRigSkeleton();
 	const int32 StartBoneIndex = Skeleton.GetBoneIndexFromName(Chain->StartBone.BoneName);
 	const int32 EndBoneIndex = Skeleton.GetBoneIndexFromName(Chain->EndBone.BoneName);
-	
-	if (StartBoneIndex == INDEX_NONE)
+
+	const bool bHasStartBone = StartBoneIndex != INDEX_NONE;
+	const bool bHasEndBone = StartBoneIndex != INDEX_NONE;
+
+	// chain has neither start nor end bone
+	if (!bHasStartBone && !bHasEndBone)
 	{
-		return false; // no start bone specified, not allowed
+		return false;
 	}
 
-	if (EndBoneIndex == INDEX_NONE)
+	// has only a start bone, this is a single bone "chain" which is fine
+	if (bHasStartBone && !bHasEndBone)
 	{
 		OutChainIndices.Add(StartBoneIndex);
-		return true; // no end bone specified, this is a single bone "chain" which is fine
+		return true; 
 	}
 
+	// has only a end bone, not valid
+	if (!bHasStartBone && bHasEndBone)
+	{
+		OutChainIndices.Add(EndBoneIndex);
+		return false; 
+	}
+	
 	// this chain has a start AND an end bone so we must verify that end bone is child of start bone
 	int32 NextBoneIndex = EndBoneIndex;
 	while (true)
@@ -404,11 +419,60 @@ bool UIKRigController::ValidateChain(const FName& ChainName, TSet<int32>& OutCha
 		if (NextBoneIndex == INDEX_NONE)
 		{
 			// oops, we walked all the way past the root without finding the start bone
+			OutChainIndices.Reset();
+			OutChainIndices.Add(EndBoneIndex);
+			OutChainIndices.Add(StartBoneIndex);
 			return false;
 		}
 	}
 }
 
+FName UIKRigController::GetRetargetChainFromBone(const FName& BoneName, const FIKRigSkeleton* OptionalSkeleton) const
+{
+	const FIKRigSkeleton& Skeleton = OptionalSkeleton ? *OptionalSkeleton : GetIKRigSkeleton();
+	const int32 BoneIndex = Skeleton.GetBoneIndexFromName(BoneName);
+
+	if (BoneName == GetRetargetRoot())
+	{
+		return FName("Retarget Root");
+	}
+	
+	const TArray<FBoneChain>& Chains = GetRetargetChains();
+	TSet<int32> OutBoneIndices;
+	for (const FBoneChain& Chain : Chains)
+	{
+		OutBoneIndices.Reset();
+		if (ValidateChain(Chain.ChainName, OptionalSkeleton, OutBoneIndices))
+		{
+			if (OutBoneIndices.Contains(BoneIndex))
+			{
+				return Chain.ChainName;
+			}
+		}
+	}
+
+	return NAME_None;
+}
+
+FName UIKRigController::GetRetargetChainFromGoal(const FName& GoalName) const
+{
+	if (GoalName == NAME_None)
+	{
+		return NAME_None;
+	}
+	
+	const TArray<FBoneChain>& Chains = GetRetargetChains();
+	for (const FBoneChain& Chain : Chains)
+	{
+		if (Chain.IKGoalName == GoalName)
+		{
+			return Chain.ChainName;
+		}
+	}
+
+	return NAME_None;
+}
+
 // -------------------------------------------------------
 // SKELETON
 //
@@ -677,7 +741,7 @@ UIKRigEffectorGoal* UIKRigController::AddNewGoal(const FName& GoalName, const FN
 		return nullptr; // goal already exists!
 	}
 	
-	FScopedTransaction Transaction(LOCTEXT("AddSetting_Label", "Add Bone Setting"));
+	FScopedTransaction Transaction(LOCTEXT("AddNewGoal_Label", "Add New Goal"));
 	Asset->Modify();
 
 	UIKRigEffectorGoal* NewGoal = NewObject<UIKRigEffectorGoal>(Asset, UIKRigEffectorGoal::StaticClass(), NAME_None, RF_Transactional);
@@ -688,6 +752,12 @@ UIKRigEffectorGoal* UIKRigController::AddNewGoal(const FName& GoalName, const FN
 	// set initial transform
 	NewGoal->InitialTransform = GetRefPoseTransformOfBone(NewGoal->BoneName);
 	NewGoal->CurrentTransform = NewGoal->InitialTransform;
+
+	// connect the new goal to all the solvers
+	for (int32 SolverIndex=0; SolverIndex<GetNumSolvers(); ++SolverIndex)
+	{
+		ConnectGoalToSolver(*NewGoal, SolverIndex);
+	}
  
 	BroadcastNeedsReinitialized();
 	BroadcastGoalsChange();
@@ -869,10 +939,10 @@ FName UIKRigController::GetBoneForGoal(const FName& GoalName) const
 	return NAME_None;
 }
 
-const UIKRigEffectorGoal* UIKRigController::GetGoalForBone(const FName& BoneName) const
+UIKRigEffectorGoal* UIKRigController::GetGoalForBone(const FName& BoneName) const
 {
-	const TArray<UIKRigEffectorGoal*>& AllGoals = GetAllGoals();
-	for (const UIKRigEffectorGoal* Goal : AllGoals)
+	TArray<UIKRigEffectorGoal*>& AllGoals = GetAllGoals();
+	for (UIKRigEffectorGoal* Goal : AllGoals)
 	{
 		if (Goal->BoneName == BoneName)
 		{
@@ -947,7 +1017,7 @@ bool UIKRigController::IsGoalConnectedToAnySolver(const FName& GoalName) const
 	return false;
 }
 
-const TArray<UIKRigEffectorGoal*>& UIKRigController::GetAllGoals() const
+TArray<UIKRigEffectorGoal*>& UIKRigController::GetAllGoals() const
 {
 	return Asset->Goals;
 }
diff --git a/Engine/Plugins/Animation/IKRig/Source/IKRigEditor/Private/RigEditor/IKRigEditMode.cpp b/Engine/Plugins/Animation/IKRig/Source/IKRigEditor/Private/RigEditor/IKRigEditMode.cpp
--- a/Engine/Plugins/Animation/IKRig/Source/IKRigEditor/Private/RigEditor/IKRigEditMode.cpp
+++ b/Engine/Plugins/Animation/IKRig/Source/IKRigEditor/Private/RigEditor/IKRigEditMode.cpp
@@ -35,35 +35,36 @@ bool FIKRigEditMode::GetCameraTarget(FSphere& OutTarget) const
 	Controller->GetSelectedGoalNames(OutGoalNames);
 	Controller->GetSelectedBoneNames(OutBoneNames);
 	
-	if (!OutGoalNames.IsEmpty() || !OutBoneNames.IsEmpty())
+	if (!(OutGoalNames.IsEmpty() && OutBoneNames.IsEmpty()))
 	{
-		TArray<FVector> GoalPoints;
+		TArray<FVector> TargetPoints;
 
 		// get goal locations
 		for (const FName& GoalName : OutGoalNames)
 		{			
-			GoalPoints.Add(Controller->AssetController->GetGoal(GoalName)->CurrentTransform.GetLocation());
+			TargetPoints.Add(Controller->AssetController->GetGoal(GoalName)->CurrentTransform.GetLocation());
 		}
 
 		// get bone locations
-		const FIKRigSkeleton& FIKRigSkeleton = Controller->AssetController->GetIKRigSkeleton();
+		const FIKRigSkeleton* CurrentIKRigSkeleton = Controller->GetCurrentIKRigSkeleton();
+		const FIKRigSkeleton& Skeleton = CurrentIKRigSkeleton ? *CurrentIKRigSkeleton : Controller->AssetController->GetIKRigSkeleton();
 		for (const FName& BoneName : OutBoneNames)
 		{
-			const int32 BoneIndex = FIKRigSkeleton.GetBoneIndexFromName(BoneName);
+			const int32 BoneIndex = Skeleton.GetBoneIndexFromName(BoneName);
 			if (BoneIndex != INDEX_NONE)
 			{
 				TArray<int32> Children;
-				FIKRigSkeleton.GetChildIndices(BoneIndex, Children);
-				for (int32 ChildIndex: Children)
+				Skeleton.GetChildIndices(BoneIndex, Children);
+				for (const int32 ChildIndex: Children)
 				{
-					GoalPoints.Add(FIKRigSkeleton.CurrentPoseGlobal[ChildIndex].GetLocation());
+					TargetPoints.Add(Skeleton.CurrentPoseGlobal[ChildIndex].GetLocation());
 				}
-				GoalPoints.Add(FIKRigSkeleton.CurrentPoseGlobal[BoneIndex].GetLocation());
+				TargetPoints.Add(Skeleton.CurrentPoseGlobal[BoneIndex].GetLocation());
 			}
 		}
 
 		// create a sphere that contains all the goal points
-		OutTarget = FSphere(&GoalPoints[0], GoalPoints.Num());
+		OutTarget = FSphere(&TargetPoints[0], TargetPoints.Num());
 		return true;
 	}
 
@@ -212,6 +213,7 @@ void FIKRigEditMode::GetBoneColors(
 {
 	const FLinearColor DefaultColor = GetMutableDefault<UPersonaOptions>()->DefaultBoneColor;
 	const FLinearColor HighlightedColor = FLinearColor::Blue;
+	const FLinearColor ErrorColor = FLinearColor::Red;
 
 	// set all to default color
 	OutBoneColors.SetNum(RefSkeleton.GetNum());
@@ -262,15 +264,14 @@ void FIKRigEditMode::GetBoneColors(
 				return;
 			}
 
+			const FIKRigSkeleton& Skeleton = Processor->GetSkeleton();
 			for (const FName& SelectedChainName : SelectedChainNames)
 			{
 				TSet<int32> OutChainBoneIndices;
-				if (Controller->AssetController->ValidateChain(SelectedChainName, OutChainBoneIndices))
+				const bool bIsValid = Controller->AssetController->ValidateChain(SelectedChainName, &Skeleton, OutChainBoneIndices);
+				for (const int32 IndexOfBoneInChain : OutChainBoneIndices)
 				{
-					for (const int32 IndexOfBoneInChain : OutChainBoneIndices)
-					{
-						OutBoneColors[IndexOfBoneInChain] = HighlightedColor;
-					}
+					OutBoneColors[IndexOfBoneInChain] = bIsValid ? HighlightedColor : ErrorColor;
 				}
 			}
 		}
diff --git a/Engine/Plugins/Animation/IKRig/Source/IKRigEditor/Private/RigEditor/IKRigEditorController.cpp b/Engine/Plugins/Animation/IKRig/Source/IKRigEditor/Private/RigEditor/IKRigEditorController.cpp
--- a/Engine/Plugins/Animation/IKRig/Source/IKRigEditor/Private/RigEditor/IKRigEditorController.cpp
+++ b/Engine/Plugins/Animation/IKRig/Source/IKRigEditor/Private/RigEditor/IKRigEditorController.cpp
@@ -9,12 +9,13 @@
 #include "Dialogs/Dialogs.h"
 
 #include "RigEditor/IKRigController.h"
-#include "RigEditor/SIKRigSkeleton.h"
+#include "RigEditor/SIKRigHierarchy.h"
 #include "RigEditor/SIKRigSolverStack.h"
 #include "RigEditor/IKRigAnimInstance.h"
 #include "RigEditor/IKRigToolkit.h"
 #include "RigEditor/SIKRigAssetBrowser.h"
 #include "RigEditor/SIKRigOutputLog.h"
+#include "Solvers/IKRig_PBIKSolver.h"
 
 #if WITH_EDITOR
 
@@ -198,8 +199,10 @@ void FRetargetChainAnalyzer::AssignBestGuessName(FBoneChain& Chain, const FIKRig
 	ChainToBonesMap.Add(FString("Ring"), {"ring"});
 	ChainToBonesMap.Add(FString("Pinky"), {"pinky"});
 
+	ChainToBonesMap.Add(FString("Root"), {"root"});
+
 	TArray<int32> BonesInChainIndices;
-	const bool bIsChainValid = IKRigSkeleton.GetBonesInChain(Chain, BonesInChainIndices);
+	const bool bIsChainValid = IKRigSkeleton.ValidateChainAndGetBones(Chain, BonesInChainIndices);
 	if (!bIsChainValid)
 	{
 		Chain.ChainName = GetDefaultChainName();
@@ -319,6 +322,19 @@ UIKRigProcessor* FIKRigEditorController::GetIKRigProcessor() const
 	return nullptr;
 }
 
+const FIKRigSkeleton* FIKRigEditorController::GetCurrentIKRigSkeleton() const
+{
+	if (const UIKRigProcessor* Processor = GetIKRigProcessor())
+	{
+		if (Processor->IsInitialized())
+		{
+			return &Processor->GetSkeleton();	
+		}
+	}
+
+	return nullptr;
+}
+
 void FIKRigEditorController::OnIKRigNeedsInitialized(UIKRigDefinition* ModifiedIKRig)
 {
 	if (ModifiedIKRig != AssetController->GetAsset())
@@ -326,6 +342,8 @@ void FIKRigEditorController::OnIKRigNeedsInitialized(UIKRigDefinition* ModifiedI
 		return;
 	}
 
+	ClearOutputLog();
+	
 	AnimInstance->SetProcessorNeedsInitialized();
 
 	// Initialize editor's instances on request
@@ -383,8 +401,11 @@ void FIKRigEditorController::ClearOutputLog() const
 {
 	if (OutputLogView.IsValid())
 	{
-		OutputLogView->ClearLog();
-		GetIKRigProcessor()->Log.Clear();
+		OutputLogView.Get()->ClearLog();
+		if (const UIKRigProcessor* Processor = GetIKRigProcessor())
+		{
+			Processor->Log.Clear();
+		}
 	}
 }
 
@@ -393,17 +414,7 @@ void FIKRigEditorController::AddNewGoals(const TArray<FName>& GoalNames, const T
 	check(GoalNames.Num() == BoneNames.Num());
 
 	// add a default solver if there isn't one already
-	if (AssetController->GetNumSolvers() == 0)
-	{
-		if (!PromptToAddSolver())
-		{
-			return; // user cancelled
-		}
-	}
-
-	// get selected solvers
-	TArray<TSharedPtr<FSolverStackElement>> SelectedSolvers;
-	GetSelectedSolvers(SelectedSolvers);
+	PromptToAddDefaultSolver();
 
 	// create goals
 	FName LastCreatedGoalName = NAME_None;
@@ -418,12 +429,6 @@ void FIKRigEditorController::AddNewGoals(const TArray<FName>& GoalNames, const T
 		{
 			continue; // already exists
 		}
-		
-		// connect the new goal to all the selected solvers
-		for (const TSharedPtr<FSolverStackElement>& SolverElement : SelectedSolvers)
-		{
-			AssetController->ConnectGoalToSolver(*NewGoal, SolverElement->IndexInStack);	
-		}
 
 		// ask user if they want to assign this goal to a chain (if there is one on this bone)
 		PromptToAddGoalToChain(NewGoal);
@@ -609,7 +614,7 @@ void FIKRigEditorController::CreateNewRetargetChains()
 	RefreshAllViews();
 }
 
-bool FIKRigEditorController::PromptToAddSolver() const
+bool FIKRigEditorController::PromptToAddDefaultSolver() const
 {
 	if (AssetController->GetNumSolvers() > 0)
 	{
@@ -642,7 +647,17 @@ bool FIKRigEditorController::PromptToAddSolver() const
 		SolverTypes.Add(SolverType);
 	}
 
+	// select "full body IK" by default
 	TSharedPtr<FIKRigSolverTypeAndName> SelectedSolver = SolverTypes[0];
+	for (TSharedPtr<FIKRigSolverTypeAndName>& SolverType :SolverTypes)
+	{
+		if (SolverType->SolverType == UIKRigPBIKSolver::StaticClass())
+		{
+			SelectedSolver = SolverType;
+			break;
+		}
+	}
+	
 	TSharedRef<SComboBox<TSharedPtr<FIKRigSolverTypeAndName>>> SolverOptionBox = SNew(SComboBox<TSharedPtr<FIKRigSolverTypeAndName>>)
 	.OptionsSource(&SolverTypes)
 	.OnGenerateWidget_Lambda([](TSharedPtr<FIKRigSolverTypeAndName> Item)
@@ -658,7 +673,7 @@ bool FIKRigEditorController::PromptToAddSolver() const
 		SNew(STextBlock)
 		.MinDesiredWidth(200)
 		.Text_Lambda([&SelectedSolver]()
-			{
+		{
 			return SelectedSolver->NiceName;
 		})
 	];
@@ -670,8 +685,8 @@ bool FIKRigEditorController::PromptToAddSolver() const
 			SolverOptionBox
 		]
 		.Buttons({
-			SCustomDialog::FButton(LOCTEXT("OK", "OK")),
-			SCustomDialog::FButton(LOCTEXT("Cancel", "Cancel"))
+			SCustomDialog::FButton(LOCTEXT("AddSolver", "Add Solver")),
+			SCustomDialog::FButton(LOCTEXT("Skip", "Skip"))
 	});
 
 	if (AddSolverDialog->ShowModal() != 0)
@@ -687,104 +702,6 @@ bool FIKRigEditorController::PromptToAddSolver() const
 	return true;
 }
 
-void FIKRigEditorController::PromptToAddGoalToChain(UIKRigEffectorGoal* NewGoal) const
-{
-	const TArray<FBoneChain>& AllRetargetChains = AssetController->GetRetargetChains();
-	FName ChainToAddGoalTo = NAME_None;
-	for (const FBoneChain& Chain : AllRetargetChains)
-	{
-		if (Chain.EndBone == NewGoal->BoneName)
-		{
-			ChainToAddGoalTo = Chain.ChainName;
-		}
-	}
-
-	if (ChainToAddGoalTo == NAME_None)
-	{
-		return;
-	}
-
-	TSharedRef<SCustomDialog> AddGoalToChainDialog = SNew(SCustomDialog)
-		.Title(FText(LOCTEXT("AddGoalToNewChainTitle", "Add Goal to Retarget Chain")))
-		.Content()
-		[
-			SNew(STextBlock)
-			.Text(FText::Format(LOCTEXT("AddNewGoalToChainLabel", "Add new goal, {0} to retarget chain, {1}?"), FText::FromName(NewGoal->GoalName), FText::FromName(ChainToAddGoalTo)))
-		]
-		.Buttons({
-			SCustomDialog::FButton(LOCTEXT("OK", "OK")),
-			SCustomDialog::FButton(LOCTEXT("Cancel", "Cancel"))
-	});
-
-	if (AddGoalToChainDialog->ShowModal() != 0)
-	{
-		return; // cancel button pressed, or window closed
-	}
-
-	AssetController->SetRetargetChainGoal(ChainToAddGoalTo, NewGoal->GoalName);
-}
-
-bool FIKRigEditorController::IsElementConnectedToSolver(TSharedRef<FIKRigTreeElement> TreeElement, int32 SolverIndex)
-{
-	if (!AssetController->GetSolverArray().IsValidIndex(SolverIndex))
-	{
-		return false; // not a valid solver index
-	}
-
-	const UIKRigSolver* Solver = AssetController->GetSolver(SolverIndex);
-	if (TreeElement->ElementType == IKRigTreeElementType::BONE)
-	{
-		// is this bone affected by this solver?
-		return Solver->IsBoneAffectedBySolver(TreeElement->BoneName, AssetController->GetIKRigSkeleton());
-	}
-
-	if (TreeElement->ElementType == IKRigTreeElementType::BONE_SETTINGS)
-	{
-		// is this bone setting belonging to the solver?
-		return (Solver->GetBoneSetting(TreeElement->BoneSettingBoneName) != nullptr);
-	}
-
-	if (TreeElement->ElementType == IKRigTreeElementType::GOAL)
-	{
-		// is goal connected to the solver?
-		return AssetController->IsGoalConnectedToSolver(TreeElement->GoalName, SolverIndex);
-	}
-
-	if (TreeElement->ElementType == IKRigTreeElementType::SOLVERGOAL)
-	{
-		// is this an effector for this solver?
-		return TreeElement->EffectorIndex == SolverIndex;
-	}
-
-	checkNoEntry();
-	return false;
-}
-
-bool FIKRigEditorController::IsElementConnectedToAnySolver(TSharedRef<FIKRigTreeElement> TreeElement)
-{
-	const int32 NumSolvers = AssetController->GetNumSolvers();
-	for (int32 SolverIndex=0; SolverIndex<NumSolvers; ++SolverIndex)
-	{
-		if (IsElementConnectedToSolver(TreeElement, SolverIndex))
-		{
-			return true;
-		}
-	}
-
-	return false;
-}
-
-bool FIKRigEditorController::IsElementExcludedBone(TSharedRef<FIKRigTreeElement> TreeElement)
-{
-	if (TreeElement->ElementType != IKRigTreeElementType::BONE)
-	{
-		return false;
-	}
-	
-	// is this bone excluded?
-	return AssetController->GetBoneExcluded(TreeElement->BoneName);
-}
-
 void FIKRigEditorController::ShowDetailsForBone(const FName BoneName) const
 {
 	BoneDetails->SetBone(BoneName);
@@ -920,6 +837,43 @@ void FIKRigEditorController::SetDetailsView(const TSharedPtr<IDetailsView>& InDe
 	ShowEmptyDetails();
 }
 
+void FIKRigEditorController::PromptToAddGoalToChain(UIKRigEffectorGoal* NewGoal) const
+{
+	const TArray<FBoneChain>& AllRetargetChains = AssetController->GetRetargetChains();
+	FName ChainToAddGoalTo = NAME_None;
+	for (const FBoneChain& Chain : AllRetargetChains)
+	{
+		if (Chain.EndBone == NewGoal->BoneName)
+		{
+			ChainToAddGoalTo = Chain.ChainName;
+		}
+	}
+
+	if (ChainToAddGoalTo == NAME_None)
+	{
+		return;
+	}
+
+	TSharedRef<SCustomDialog> AddGoalToChainDialog = SNew(SCustomDialog)
+		.Title(FText(LOCTEXT("AddGoalToNewChainTitle", "Add Goal to Retarget Chain")))
+		.Content()
+		[
+			SNew(STextBlock)
+			.Text(FText::Format(LOCTEXT("AddNewGoalToChainLabel", "Add new goal, {0} to retarget chain, {1}?"), FText::FromName(NewGoal->GoalName), FText::FromName(ChainToAddGoalTo)))
+		]
+		.Buttons({
+			SCustomDialog::FButton(LOCTEXT("AddGoal", "Add Goal")),
+			SCustomDialog::FButton(LOCTEXT("Skip", "Skip"))
+	});
+
+	if (AddGoalToChainDialog->ShowModal() != 0)
+	{
+		return; // cancel button pressed, or window closed
+	}
+
+	AssetController->SetRetargetChainGoal(ChainToAddGoalTo, NewGoal->GoalName);
+}
+
 void FIKRigEditorController::PromptToAddNewRetargetChain(const FBoneChain& BoneChain) const
 {
 	const TSharedPtr<FStructOnScope> StructToDisplay = MakeShareable(new FStructOnScope(FIKRigRetargetChainSettings::StaticStruct(), (uint8*)&BoneChain));
@@ -929,8 +883,18 @@ void FIKRigEditorController::PromptToAddNewRetargetChain(const FBoneChain& BoneC
 	SGenericDialogWidget::FArguments DialogArguments;
 	DialogArguments.OnOkPressed_Lambda([&BoneChain, this] ()
 	{
+		// prompt to add a goal to the chain (skipped if already has a goal)
+		UIKRigEffectorGoal* NewGoal = PromptToAddGoalToNewChain(BoneChain);
+
 		// add the retarget chain
 		AssetController->AddRetargetChain(BoneChain);
+
+		// ask user if they want to assign this goal to a chain (if there is one on this bone)
+		if (NewGoal)
+		{
+			PromptToAddGoalToChain(NewGoal);
+		}
+		
 		RefreshAllViews();
 	});
 
@@ -941,6 +905,44 @@ void FIKRigEditorController::PromptToAddNewRetargetChain(const FBoneChain& BoneC
 		true);
 }
 
+UIKRigEffectorGoal* FIKRigEditorController::PromptToAddGoalToNewChain(const FBoneChain& BoneChain) const
+{
+	// check if there is already a goal on the end bone of this chain, if there is, then we're done
+	UIKRigEffectorGoal* ExistingGoal = AssetController->GetGoalForBone(BoneChain.EndBone.BoneName);
+	if (ExistingGoal)
+	{
+		return ExistingGoal;
+	}
+
+	// ask user if they want to add a goal to the end of this chain
+	const TSharedRef<SCustomDialog> AddGoalToChainDialog = SNew(SCustomDialog)
+		.Title(FText(LOCTEXT("AddGoalToNewChainTitleLabel", "Add Goal to New Chain")))
+		.Content()
+		[
+			SNew(STextBlock)
+			.Text(FText::Format(
+				LOCTEXT("AddGoalToNewChainLabel", "Add goal to end bone, {0} of new chain, {1}?"),
+				FText::FromName(BoneChain.IKGoalName), 
+				FText::FromName(BoneChain.EndBone.BoneName)))
+		]
+		.Buttons({
+			SCustomDialog::FButton(LOCTEXT("AddGoal", "Add Goal")),
+			SCustomDialog::FButton(LOCTEXT("Skip", "Skip"))
+	});
+
+	if (AddGoalToChainDialog->ShowModal() != 0)
+	{
+		return nullptr; // cancel button pressed, or window closed
+	}
+
+	// add a default solver if there isn't one already
+	PromptToAddDefaultSolver();
+
+	// add the goal
+	const FName NewGoalName = FName(FText::Format(LOCTEXT("GoalOnNewChainName", "{0}_Goal"), FText::FromName(BoneChain.ChainName)).ToString());
+	return AssetController->AddNewGoal(NewGoalName, BoneChain.EndBone.BoneName);
+}
+
 void FIKRigEditorController::PromptToAddGoalToNewlyMirroredChain(const FBoneChain& OldChain, const FBoneChain& MirroredChain) const
 {
 	const TSharedRef<SCustomDialog> AddGoalToChainDialog = SNew(SCustomDialog)
@@ -951,8 +953,8 @@ void FIKRigEditorController::PromptToAddGoalToNewlyMirroredChain(const FBoneChai
 			.Text(FText::Format(LOCTEXT("AddGoalToMirroredChainLabel", "Mirror goal, {0} on chain, {1}?"), FText::FromName(OldChain.IKGoalName), FText::FromName(OldChain.IKGoalName)))
 		]
 		.Buttons({
-			SCustomDialog::FButton(LOCTEXT("Ok", "OK")),
-			SCustomDialog::FButton(LOCTEXT("Cancel", "Cancel"))
+			SCustomDialog::FButton(LOCTEXT("AddGoal", "Add Goal")),
+			SCustomDialog::FButton(LOCTEXT("Skip", "Skip"))
 	});
 
 	if (AddGoalToChainDialog->ShowModal() != 0)
@@ -960,8 +962,10 @@ void FIKRigEditorController::PromptToAddGoalToNewlyMirroredChain(const FBoneChai
 		return; // cancel button pressed, or window closed
 	}
 
-	UIKRigEffectorGoal* OldGoal = AssetController->GetGoal(OldChain.IKGoalName);
+	const UIKRigEffectorGoal* OldGoal = AssetController->GetGoal(OldChain.IKGoalName);
 	const FName NewGoalName = FName(OldChain.IKGoalName.ToString() + "_Mirrored");
+	// add a default solver if there isn't one already
+	PromptToAddDefaultSolver();
 	UIKRigEffectorGoal* NewGoal = AssetController->AddNewGoal(NewGoalName, MirroredChain.EndBone.BoneName);
 
 	// connect to the same solvers
diff --git a/Engine/Plugins/Animation/IKRig/Source/IKRigEditor/Private/RigEditor/IKRigOutputLogTabSummoner.cpp b/Engine/Plugins/Animation/IKRig/Source/IKRigEditor/Private/RigEditor/IKRigOutputLogTabSummoner.cpp
--- a/Engine/Plugins/Animation/IKRig/Source/IKRigEditor/Private/RigEditor/IKRigOutputLogTabSummoner.cpp
+++ b/Engine/Plugins/Animation/IKRig/Source/IKRigEditor/Private/RigEditor/IKRigOutputLogTabSummoner.cpp
@@ -34,8 +34,9 @@ TSharedRef<SWidget> FIKRigOutputLogTabSummoner::CreateTabBody(const FWorkflowTab
 {
 	const TSharedRef<FIKRigEditorController>& Controller = IKRigEditor.Pin()->GetController();
 	const FName LogName = Controller->GetIKRigProcessor()->Log.GetLogTarget();
-	
-	return SNew(SIKRigOutputLog, LogName);
+	TSharedRef<SIKRigOutputLog> LogView = SNew(SIKRigOutputLog, LogName);
+	Controller->SetOutputLogView(LogView);
+	return LogView;
 }
 
 #undef LOCTEXT_NAMESPACE 
diff --git a/Engine/Plugins/Animation/IKRig/Source/IKRigEditor/Private/RigEditor/IKRigSkeletonTabSummoner.cpp b/Engine/Plugins/Animation/IKRig/Source/IKRigEditor/Private/RigEditor/IKRigSkeletonTabSummoner.cpp
--- a/Engine/Plugins/Animation/IKRig/Source/IKRigEditor/Private/RigEditor/IKRigSkeletonTabSummoner.cpp
+++ b/Engine/Plugins/Animation/IKRig/Source/IKRigEditor/Private/RigEditor/IKRigSkeletonTabSummoner.cpp
@@ -5,7 +5,7 @@
 #include "IDocumentation.h"
 #include "RigEditor/IKRigEditorStyle.h"
 #include "RigEditor/IKRigToolkit.h"
-#include "RigEditor/SIKRigSkeleton.h"
+#include "RigEditor/SIKRigHierarchy.h"
 
 #define LOCTEXT_NAMESPACE "IKRigSkeletonTabSummoner"
 
@@ -31,7 +31,7 @@ TSharedPtr<SToolTip> FIKRigSkeletonTabSummoner::CreateTabToolTipWidget(const FWo
 
 TSharedRef<SWidget> FIKRigSkeletonTabSummoner::CreateTabBody(const FWorkflowTabSpawnInfo& Info) const
 {
-	return SNew(SIKRigSkeleton, IKRigEditor.Pin()->GetController());
+	return SNew(SIKRigHierarchy, IKRigEditor.Pin()->GetController());
 }
 
 #undef LOCTEXT_NAMESPACE 
diff --git a/Engine/Plugins/Animation/IKRig/Source/IKRigEditor/Private/RigEditor/IKRigToolkit.cpp b/Engine/Plugins/Animation/IKRig/Source/IKRigEditor/Private/RigEditor/IKRigToolkit.cpp
--- a/Engine/Plugins/Animation/IKRig/Source/IKRigEditor/Private/RigEditor/IKRigToolkit.cpp
+++ b/Engine/Plugins/Animation/IKRig/Source/IKRigEditor/Private/RigEditor/IKRigToolkit.cpp
@@ -192,12 +192,14 @@ void FIKRigEditorToolkit::HandleOnPreviewSceneSettingsCustomized(IDetailLayoutBu
 
 void FIKRigEditorToolkit::PostUndo(bool bSuccess)
 {
+	EditorController->ClearOutputLog();
 	EditorController->AssetController->BroadcastNeedsReinitialized();
 	EditorController->RefreshAllViews();
 }
 
 void FIKRigEditorToolkit::PostRedo(bool bSuccess)
 {
+	EditorController->ClearOutputLog();
 	EditorController->AssetController->BroadcastNeedsReinitialized();
 	EditorController->RefreshAllViews();
 }
diff --git a/Engine/Plugins/Animation/IKRig/Source/IKRigEditor/Private/RigEditor/SIKRigAssetBrowser.cpp b/Engine/Plugins/Animation/IKRig/Source/IKRigEditor/Private/RigEditor/SIKRigAssetBrowser.cpp
--- a/Engine/Plugins/Animation/IKRig/Source/IKRigEditor/Private/RigEditor/SIKRigAssetBrowser.cpp
+++ b/Engine/Plugins/Animation/IKRig/Source/IKRigEditor/Private/RigEditor/SIKRigAssetBrowser.cpp
@@ -46,8 +46,8 @@ void SIKRigAssetBrowser::RefreshView()
 	AssetPickerConfig.bShowTypeInColumnView = true;
 	AssetPickerConfig.OnShouldFilterAsset = FOnShouldFilterAsset::CreateSP(this, &SIKRigAssetBrowser::OnShouldFilterAsset);
 	AssetPickerConfig.DefaultFilterMenuExpansion = EAssetTypeCategories::Animation;
-	
 	AssetPickerConfig.OnAssetDoubleClicked = FOnAssetSelected::CreateSP(this, &SIKRigAssetBrowser::OnAssetDoubleClicked);
+	AssetPickerConfig.OnGetAssetContextMenu = FOnGetAssetContextMenu::CreateSP(this, &SIKRigAssetBrowser::OnGetAssetContextMenu);
 	AssetPickerConfig.GetCurrentSelectionDelegates.Add(&GetCurrentSelectionDelegate);
 	AssetPickerConfig.bAllowNullSelection = false;
 
@@ -62,11 +62,50 @@ void SIKRigAssetBrowser::RefreshView()
 	// Also hide the type column by default (but allow users to enable it, so don't use bShowTypeInColumnView)
 	AssetPickerConfig.HiddenColumnNames.Add(TEXT("Class"));
 
-	FContentBrowserModule& ContentBrowserModule = FModuleManager::Get().LoadModuleChecked<FContentBrowserModule>(TEXT("ContentBrowser"));
-
+	const FContentBrowserModule& ContentBrowserModule = FModuleManager::Get().LoadModuleChecked<FContentBrowserModule>(TEXT("ContentBrowser"));
 	AssetBrowserBox->SetContent(ContentBrowserModule.Get().CreateAssetPicker(AssetPickerConfig));
 }
 
+TSharedPtr<SWidget> SIKRigAssetBrowser::OnGetAssetContextMenu(const TArray<FAssetData>& SelectedAssets) const
+{
+	if (SelectedAssets.Num() <= 0)
+	{
+		return nullptr;
+	}
+
+	UObject* SelectedAsset = SelectedAssets[0].GetAsset();
+	if (SelectedAsset == nullptr)
+	{
+		return nullptr;
+	}
+	
+	FMenuBuilder MenuBuilder(true, MakeShared<FUICommandList>());
+
+	MenuBuilder.BeginSection(TEXT("Asset"), LOCTEXT("AssetSectionLabel", "Asset"));
+	{
+		MenuBuilder.AddMenuEntry(
+			LOCTEXT("Browse", "Browse to Asset"),
+			LOCTEXT("BrowseTooltip", "Browses to the associated asset and selects it in the most recently used Content Browser (summoning one if necessary)"),
+			FSlateIcon(FAppStyle::GetAppStyleSetName(), "SystemWideCommands.FindInContentBrowser.Small"),
+			FUIAction(
+				FExecuteAction::CreateLambda([SelectedAsset] ()
+				{
+					if (SelectedAsset)
+					{
+						const TArray<FAssetData>& Assets = { SelectedAsset };
+						const FContentBrowserModule& ContentBrowserModule = FModuleManager::Get().LoadModuleChecked<FContentBrowserModule>("ContentBrowser");
+						ContentBrowserModule.Get().SyncBrowserToAssets(Assets);
+					}
+				}),
+				FCanExecuteAction::CreateLambda([] () { return true; })
+			)
+		);
+	}
+	MenuBuilder.EndSection();
+
+	return MenuBuilder.MakeWidget();
+}
+
 void SIKRigAssetBrowser::OnAssetDoubleClicked(const FAssetData& AssetData)
 {
 	if (!AssetData.GetAsset())
diff --git a/Engine/Plugins/Animation/IKRig/Source/IKRigEditor/Private/RigEditor/SIKRigHierarchy.cpp b/Engine/Plugins/Animation/IKRig/Source/IKRigEditor/Private/RigEditor/SIKRigHierarchy.cpp
new file mode 100644
--- /dev/null
+++ b/Engine/Plugins/Animation/IKRig/Source/IKRigEditor/Private/RigEditor/SIKRigHierarchy.cpp
@@ -0,0 +1,2013 @@
+// Copyright Epic Games, Inc. All Rights Reserved.
+
+#include "RigEditor/SIKRigHierarchy.h"
+
+#include "IKRigProcessor.h"
+#include "IKRigSolver.h"
+#include "Engine/SkeletalMesh.h"
+#include "IPersonaToolkit.h"
+#include "SKismetInspector.h"
+#include "Dialogs/Dialogs.h"
+#include "SPositiveActionButton.h"
+#include "RigEditor/IKRigEditorController.h"
+#include "RigEditor/IKRigEditorStyle.h"
+#include "RigEditor/IKRigSkeletonCommands.h"
+#include "RigEditor/IKRigToolkit.h"
+#include "RigEditor/SIKRigSolverStack.h"
+#include "Widgets/Text/SInlineEditableTextBlock.h"
+#include "Framework/MultiBox/MultiBoxBuilder.h"
+#include "Framework/Commands/UICommandList.h"
+#include "Preferences/PersonaOptions.h"
+#include "Widgets/Input/SSearchBox.h"
+
+#define LOCTEXT_NAMESPACE "SIKRigHierarchy"
+
+static const FName IKRigElementColumnName(TEXT("ElementName"));
+static const FName IKRigChainColumnName(TEXT("RetargetChainName"));
+
+FIKRigTreeElement::FIKRigTreeElement(const FText& InKey, IKRigTreeElementType InType,
+									 const TSharedRef<FIKRigEditorController>& InEditorController)
+	: Key(InKey)
+	, ElementType(InType)
+	, EditorController(InEditorController)
+	, OptionalBoneDetails(nullptr)
+{}
+
+void FIKRigTreeElement::RequestRename()
+{
+	OnRenameRequested.ExecuteIfBound();
+}
+
+TWeakObjectPtr< UObject > FIKRigTreeElement::GetObject() const
+{
+	const TSharedPtr<FIKRigEditorController>& Controller = EditorController.Pin();
+	if (!Controller.IsValid())
+	{
+		return nullptr; 
+	}
+
+	if (Controller->AssetController == nullptr)
+	{
+		return nullptr;
+	}
+
+	switch(ElementType)
+	{
+	case IKRigTreeElementType::BONE:
+		if (OptionalBoneDetails == nullptr)
+		{
+			OptionalBoneDetails = Controller->CreateBoneDetails(AsShared());
+		}
+		return OptionalBoneDetails;
+	case IKRigTreeElementType::BONE_SETTINGS:
+		return Controller->AssetController->GetSettingsForBone(BoneSettingBoneName, BoneSettingsSolverIndex);
+	case IKRigTreeElementType::GOAL:
+		return Controller->AssetController->GetGoal(GoalName);
+	case IKRigTreeElementType::SOLVERGOAL:
+		if (const UIKRigSolver* SolverWithEffector = Controller->AssetController->GetSolver(EffectorIndex))
+		{
+			return SolverWithEffector->GetGoalSettings(EffectorGoalName);
+		}
+	default:
+		return nullptr;
+	}
+}
+
+FName FIKRigTreeElement::GetChainName() const
+{
+	const UIKRigController* AssetController = EditorController.Pin()->AssetController;
+	const FIKRigSkeleton* Skeleton  = EditorController.Pin()->GetCurrentIKRigSkeleton();
+	
+	switch(ElementType)
+	{
+	case IKRigTreeElementType::BONE:
+		return AssetController->GetRetargetChainFromBone(BoneName, Skeleton);
+	case IKRigTreeElementType::BONE_SETTINGS:
+		return AssetController->GetRetargetChainFromBone(BoneSettingBoneName, Skeleton);
+	case IKRigTreeElementType::GOAL:
+		return AssetController->GetRetargetChainFromGoal(GoalName);
+	case IKRigTreeElementType::SOLVERGOAL:
+		return AssetController->GetRetargetChainFromGoal(EffectorGoalName);
+	default:
+		return NAME_None;
+	}
+}
+
+void SIKRigHierarchyItem::Construct(const FArguments& InArgs)
+{
+	WeakTreeElement = InArgs._TreeElement;
+	EditorController = InArgs._EditorController;
+	HierarchyView = InArgs._HierarchyView;
+
+	TSharedPtr<FIKRigTreeElement> TreeElement = WeakTreeElement.Pin();
+	
+	// is this element affected by the selected solver?
+	bool bIsConnectedToSelectedSolver;
+	const int32 SelectedSolver = EditorController.Pin()->GetSelectedSolverIndex();
+	if (SelectedSolver == INDEX_NONE)
+	{
+		bIsConnectedToSelectedSolver = HierarchyView.Pin()->IsElementConnectedToAnySolver(WeakTreeElement);
+	}
+	else
+	{
+		bIsConnectedToSelectedSolver = HierarchyView.Pin()->IsElementConnectedToSolver(WeakTreeElement, SelectedSolver);
+	}
+
+	// determine text style
+	FTextBlockStyle NormalText = FIKRigEditorStyle::Get().GetWidgetStyle<FTextBlockStyle>("IKRig.Tree.NormalText");
+	FTextBlockStyle ItalicText = FIKRigEditorStyle::Get().GetWidgetStyle<FTextBlockStyle>("IKRig.Tree.ItalicText");
+	FSlateFontInfo TextFont;
+	FSlateColor TextColor;
+	if (bIsConnectedToSelectedSolver)
+	{
+		// elements connected to the selected solver are green
+		TextFont = ItalicText.Font;
+		TextColor = NormalText.ColorAndOpacity;
+	}
+	else
+	{
+		TextFont = NormalText.Font;
+		TextColor = FLinearColor(0.2f, 0.2f, 0.2f, 0.5f);
+	}
+
+	// determine which icon to use for tree element
+	const FSlateBrush* Brush = FAppStyle::Get().GetBrush("SkeletonTree.Bone");
+	switch(TreeElement->ElementType)
+	{
+		case IKRigTreeElementType::BONE:
+			if (!HierarchyView.Pin()->IsElementExcludedBone(WeakTreeElement))
+			{
+				Brush = FAppStyle::Get().GetBrush("SkeletonTree.Bone");
+			}
+			else
+			{
+				Brush = FAppStyle::Get().GetBrush("SkeletonTree.BoneNonWeighted");
+			}
+			break;
+		case IKRigTreeElementType::BONE_SETTINGS:
+			Brush = FIKRigEditorStyle::Get().GetBrush("IKRig.Tree.BoneWithSettings");
+			break;
+		case IKRigTreeElementType::GOAL:
+			Brush = FIKRigEditorStyle::Get().GetBrush("IKRig.Tree.Goal");
+			break;
+		case IKRigTreeElementType::SOLVERGOAL:
+			Brush = FIKRigEditorStyle::Get().GetBrush("IKRig.Tree.Effector");
+			break;
+		default:
+			checkNoEntry();
+	}
+
+	TSharedPtr<SHorizontalBox> RowBox;
+
+	ChildSlot
+	[
+		SAssignNew(RowBox, SHorizontalBox)
+		+ SHorizontalBox::Slot()
+		.MaxWidth(18)
+		.FillWidth(1.0)
+		.HAlign(HAlign_Left)
+		.VAlign(VAlign_Center)
+		[
+			SNew(SImage)
+			.Image(Brush)
+		]
+	];
+	
+	if (TreeElement->ElementType == IKRigTreeElementType::BONE)
+	{
+		RowBox->AddSlot()
+        .AutoWidth()
+        .VAlign(VAlign_Center)
+        [
+			SNew(STextBlock)
+			.Text(this, &SIKRigHierarchyItem::GetName)
+			.Font(TextFont)
+			.ColorAndOpacity(TextColor)
+        ];
+	}
+	else
+	{
+		TSharedPtr< SInlineEditableTextBlock > InlineWidget;
+		RowBox->AddSlot()
+        .AutoWidth()
+        .VAlign(VAlign_Center)
+        [
+	        SAssignNew(InlineWidget, SInlineEditableTextBlock)
+		    .Text(this, &SIKRigHierarchyItem::GetName)
+		    .Font(TextFont)
+			.ColorAndOpacity(TextColor)
+			.OnVerifyTextChanged(this, &SIKRigHierarchyItem::OnVerifyNameChanged)
+		    .OnTextCommitted(this, &SIKRigHierarchyItem::OnNameCommitted)
+		    .MultiLine(false)
+        ];
+		TreeElement->OnRenameRequested.BindSP(InlineWidget.Get(), &SInlineEditableTextBlock::EnterEditingMode);
+	}
+}
+
+void SIKRigHierarchyItem::OnNameCommitted(const FText& InText, ETextCommit::Type InCommitType) const
+{
+	check(WeakTreeElement.IsValid());
+
+	if (!(InCommitType == ETextCommit::OnEnter || InCommitType == ETextCommit::OnUserMovedFocus))
+	{
+		return; // make sure user actually intends to commit a name change
+	}
+
+	const TSharedPtr<FIKRigEditorController> Controller = EditorController.Pin();
+	if (!Controller.IsValid())
+	{
+		return; 
+	}
+	
+	const FText OldText = WeakTreeElement.Pin()->Key;
+	const FName OldName = WeakTreeElement.Pin()->GoalName;
+	const FName PotentialNewName = FName(InText.ToString());
+	const FName NewName = Controller->AssetController->RenameGoal(OldName, PotentialNewName);
+	if (NewName != NAME_None)
+	{
+		WeakTreeElement.Pin()->Key = FText::FromName(NewName);
+		WeakTreeElement.Pin()->GoalName = NewName;
+	}
+	
+	Controller->RefreshAllViews();
+	HierarchyView.Pin()->ReplaceItemInSelection(OldText, WeakTreeElement.Pin()->Key);
+}
+
+bool SIKRigHierarchyItem::OnVerifyNameChanged(const FText& InText, FText& OutErrorMessage) const
+{
+	// TODO let the user know when/why the goal can't be renamed 
+	return true;
+}
+
+FText SIKRigHierarchyItem::GetName() const
+{
+	return WeakTreeElement.Pin()->Key;
+}
+
+TSharedRef<SWidget> SIKRigSkeletonRow::GenerateWidgetForColumn(const FName& ColumnName)
+{
+	// the actual bone/goal/settings element in the hierarchy
+	if (ColumnName == IKRigElementColumnName)
+	{
+		return SNew(SHorizontalBox)
+		+ SHorizontalBox::Slot()
+		.AutoWidth()
+		[
+			SNew( SExpanderArrow, SharedThis(this) )
+			.ShouldDrawWires(true)
+		]
+
+		+ SHorizontalBox::Slot()
+		.AutoWidth()
+		[
+			SNew(SIKRigHierarchyItem)
+			.EditorController(EditorController)
+			.HierarchyView(HierarchyView)
+			.TreeElement(WeakTreeElement)
+		];
+	}
+
+	// the column of chain names for each hierarchy element
+	if (ColumnName == IKRigChainColumnName)
+	{
+		const FName ChainName = WeakTreeElement.Pin()->GetChainName();
+		const FText ChainText = ChainName == NAME_None ? FText::FromString("") : FText::FromName(ChainName);
+		const FSlateFontInfo Font = ChainName != "(Retarget Root)" ? FAppStyle::Get().GetFontStyle("NormalFont") : FAppStyle::Get().GetFontStyle("BoldFont");
+		
+		TSharedPtr< SHorizontalBox > RowBox;
+		SAssignNew(RowBox, SHorizontalBox)
+		+ SHorizontalBox::Slot()
+		.AutoWidth()
+		[
+			SNew(STextBlock).Text(ChainText).Font(Font)
+		];
+
+		return RowBox.ToSharedRef();
+	}
+	
+	return SNullWidget::NullWidget;
+	
+}
+
+TSharedRef<FIKRigSkeletonDragDropOp> FIKRigSkeletonDragDropOp::New(TWeakPtr<FIKRigTreeElement> InElement)
+{
+	TSharedRef<FIKRigSkeletonDragDropOp> Operation = MakeShared<FIKRigSkeletonDragDropOp>();
+	Operation->Element = InElement;
+	Operation->Construct();
+	return Operation;
+}
+
+TSharedPtr<SWidget> FIKRigSkeletonDragDropOp::GetDefaultDecorator() const
+{
+	return SNew(SBorder)
+        .Visibility(EVisibility::Visible)
+        .BorderImage(FAppStyle::GetBrush("Menu.Background"))
+        [
+            SNew(STextBlock)
+            .Text(FText::FromString(Element.Pin()->Key.ToString()))
+        ];
+}
+
+FReply SIKRigSkeletonRow::HandleDragDetected(const FGeometry& MyGeometry, const FPointerEvent& MouseEvent)
+{
+	const TArray<TSharedPtr<FIKRigTreeElement>> SelectedItems = HierarchyView.Pin()->GetSelectedItems();
+	if (SelectedItems.Num() != 1)
+	{
+		return FReply::Unhandled();
+	}
+
+	const TSharedPtr<FIKRigTreeElement> DraggedElement = SelectedItems[0];
+	if (DraggedElement->ElementType != IKRigTreeElementType::GOAL)
+	{
+		return FReply::Unhandled();
+	}
+	
+	if (MouseEvent.IsMouseButtonDown(EKeys::LeftMouseButton))
+	{
+		const TSharedRef<FIKRigSkeletonDragDropOp> DragDropOp = FIKRigSkeletonDragDropOp::New(DraggedElement);
+		return FReply::Handled().BeginDragDrop(DragDropOp);
+	}
+
+	return FReply::Unhandled();
+}
+
+TOptional<EItemDropZone> SIKRigSkeletonRow::HandleCanAcceptDrop(
+	const FDragDropEvent& DragDropEvent,
+	EItemDropZone DropZone,
+    TSharedPtr<FIKRigTreeElement> TargetItem)
+{
+	TOptional<EItemDropZone> ReturnedDropZone;
+	
+	const TSharedPtr<FIKRigSkeletonDragDropOp> DragDropOp = DragDropEvent.GetOperationAs<FIKRigSkeletonDragDropOp>();
+	if (DragDropOp.IsValid())
+	{
+		if (TargetItem.Get()->ElementType == IKRigTreeElementType::BONE)
+        {
+        	ReturnedDropZone = EItemDropZone::BelowItem;	
+        }
+	}
+	
+	return ReturnedDropZone;
+}
+
+FReply SIKRigSkeletonRow::HandleAcceptDrop(
+	const FDragDropEvent& DragDropEvent,
+	EItemDropZone DropZone,
+    TSharedPtr<FIKRigTreeElement> TargetItem)
+{
+	const TSharedPtr<FIKRigSkeletonDragDropOp> DragDropOp = DragDropEvent.GetOperationAs<FIKRigSkeletonDragDropOp>();
+	if (!DragDropOp.IsValid())
+	{
+		return FReply::Unhandled();
+	}
+	
+	const TSharedPtr<FIKRigEditorController> Controller = EditorController.Pin();
+	if (!Controller.IsValid())
+	{
+		return FReply::Handled();
+	}
+
+	const FIKRigTreeElement& DraggedElement = *DragDropOp.Get()->Element.Pin().Get();
+	ensure(DraggedElement.ElementType == IKRigTreeElementType::GOAL);		// drag only supported for goals
+	ensure(TargetItem.Get()->ElementType == IKRigTreeElementType::BONE);	// drop only supported for bones
+
+	// re-parent the goal to a different bone
+	const UIKRigController* AssetController = Controller->AssetController;
+	const bool bWasReparented = AssetController->SetGoalBone(DraggedElement.GoalName, TargetItem.Get()->BoneName);
+	if (bWasReparented)
+	{
+		Controller->RefreshAllViews();
+	}
+	
+	return FReply::Handled();
+}
+
+void SIKRigHierarchy::Construct(const FArguments& InArgs, TSharedRef<FIKRigEditorController> InEditorController)
+{
+	EditorController = InEditorController;
+	EditorController.Pin()->SetHierarchyView(SharedThis(this));
+	TextFilter = MakeShareable(new FTextFilterExpressionEvaluator(ETextFilterExpressionEvaluatorMode::BasicString));
+	CommandList = MakeShared<FUICommandList>();
+	BindCommands();
+	
+	ChildSlot
+    [
+        SNew(SVerticalBox)
+
+        +SVerticalBox::Slot()
+		.Padding(2.0f)
+		.AutoHeight()
+		[
+			SNew(SHorizontalBox)
+			+SHorizontalBox::Slot()
+			.AutoWidth()
+			.VAlign(VAlign_Center)
+			.Padding(FMargin(6.f, 0.0))
+			[
+				SNew(SPositiveActionButton)
+				.OnGetMenuContent(this, &SIKRigHierarchy::CreateAddNewMenu)
+				.Icon(FAppStyle::Get().GetBrush("Icons.Plus"))
+			]
+
+			+SHorizontalBox::Slot()
+			.FillWidth(1.0f)
+			[
+				SNew(SSearchBox)
+				.SelectAllTextWhenFocused(true)
+				.OnTextChanged( this, &SIKRigHierarchy::OnFilterTextChanged )
+				.HintText(LOCTEXT( "SearchBoxHint", "Filter Hierarchy Tree..."))
+			]
+
+			+SHorizontalBox::Slot()
+			.AutoWidth()
+			.Padding(FMargin(6.f, 0.0))
+			.VAlign(VAlign_Center)
+			[
+				SNew(SComboButton)
+				.ComboButtonStyle(&FAppStyle::Get().GetWidgetStyle<FComboButtonStyle>("SimpleComboButton"))
+				.ForegroundColor(FSlateColor::UseStyle())
+				.ContentPadding(2.0f)
+				.OnGetMenuContent( this, &SIKRigHierarchy::CreateFilterMenuWidget )
+				.ToolTipText(LOCTEXT( "FilterMenuLabel", "Filter hierarchy tree options."))
+				.HasDownArrow(true)
+				.ButtonContent()
+				[
+					SNew(SImage)
+					.Image(FAppStyle::Get().GetBrush("Icons.Settings"))
+					.ColorAndOpacity(FSlateColor::UseForeground())
+				]
+			]
+		]
+		
+        +SVerticalBox::Slot()
+        .Padding(0.0f, 0.0f)
+        [
+            SNew(SBorder)
+            .Padding(2.0f)
+            .BorderImage(FAppStyle::GetBrush("SCSEditor.TreePanel"))
+            [
+                SAssignNew(TreeView, SIKRigSkeletonTreeView)
+                .TreeItemsSource(&RootElements)
+                .SelectionMode(ESelectionMode::Multi)
+                .OnGenerateRow_Lambda( [this](
+					TSharedPtr<FIKRigTreeElement> InItem,
+					const TSharedRef<STableViewBase>& OwnerTable) -> TSharedRef<ITableRow>
+                {
+                	return SNew(SIKRigSkeletonRow, OwnerTable)
+                	.EditorController(EditorController.Pin().ToSharedRef())
+                	.TreeElement(InItem)
+                	.HierarchyView(SharedThis(this));
+                })
+                .OnGetChildren(this, &SIKRigHierarchy::HandleGetChildrenForTree)
+                .OnSelectionChanged(this, &SIKRigHierarchy::OnSelectionChanged)
+                .OnContextMenuOpening(this, &SIKRigHierarchy::CreateContextMenu)
+                .OnMouseButtonClick(this, &SIKRigHierarchy::OnItemClicked)
+                .OnMouseButtonDoubleClick(this, &SIKRigHierarchy::OnItemDoubleClicked)
+                .OnSetExpansionRecursive(this, &SIKRigHierarchy::OnSetExpansionRecursive)
+                .HighlightParentNodesForSelection(false)
+                .ItemHeight(24)
+                .HeaderRow
+				(
+					SNew(SHeaderRow)
+					
+					+ SHeaderRow::Column(IKRigElementColumnName)
+					.DefaultLabel(LOCTEXT("IKRigElementColumnLabel", "Rig Element"))
+					.FillWidth(0.7f)
+						
+					+ SHeaderRow::Column(IKRigChainColumnName)
+					.DefaultLabel(LOCTEXT("IKRigChainColumnLabel", "Retarget Chain"))
+					.FillWidth(0.3f)
+				)
+            ]
+        ]
+    ];
+
+	constexpr bool IsInitialSetup = true;
+	RefreshTreeView(IsInitialSetup);
+}
+
+void SIKRigHierarchy::AddSelectedItemFromViewport(
+	const FName& ItemName,
+	IKRigTreeElementType ItemType,
+	const bool bReplace)
+{	
+	// nothing to add
+	if (ItemName == NAME_None)
+	{
+		return;
+	}
+
+	// record what was already selected
+	TArray<TSharedPtr<FIKRigTreeElement>> PreviouslySelectedItems = TreeView->GetSelectedItems();
+	// add/remove items as needed
+	for (const TSharedPtr<FIKRigTreeElement>& Item : AllElements)
+	{
+		bool bIsBeingAdded = false;
+		switch (ItemType)
+		{
+		case IKRigTreeElementType::GOAL:
+			if (ItemName == Item->GoalName)
+			{
+				bIsBeingAdded = true;
+			}
+			break;
+		case IKRigTreeElementType::BONE:
+			if (ItemName == Item->BoneName)
+			{
+				bIsBeingAdded = true;
+			}
+			break;
+		default:
+			ensureMsgf(false, TEXT("IKRig cannot select anything but bones and goals in viewport."));
+			return;
+		}
+
+		if (bReplace)
+		{
+			if (bIsBeingAdded)
+			{
+				TreeView->ClearSelection();
+				AddItemToSelection(Item);
+				return;
+			}
+			
+			continue;
+		}
+
+		// remove if already selected (invert)
+		if (bIsBeingAdded && PreviouslySelectedItems.Contains(Item))
+		{
+			RemoveItemFromSelection(Item);
+			continue;
+		}
+
+		// add if being added
+		if (bIsBeingAdded)
+		{
+			AddItemToSelection(Item);
+			continue;
+		}
+	}
+}
+
+void SIKRigHierarchy::AddItemToSelection(const TSharedPtr<FIKRigTreeElement>& InItem)
+{
+	TreeView->SetItemSelection(InItem, true, ESelectInfo::Direct);
+    
+	if(GetDefault<UPersonaOptions>()->bExpandTreeOnSelection)
+	{
+		TSharedPtr<FIKRigTreeElement> ItemToExpand = InItem->Parent;
+		while(ItemToExpand.IsValid())
+		{
+			TreeView->SetItemExpansion(ItemToExpand, true);
+			ItemToExpand = ItemToExpand->Parent;
+		}
+	}
+    
+	TreeView->RequestScrollIntoView(InItem);
+}
+
+void SIKRigHierarchy::RemoveItemFromSelection(const TSharedPtr<FIKRigTreeElement>& InItem)
+{
+	TreeView->SetItemSelection(InItem, false, ESelectInfo::Direct);
+}
+
+void SIKRigHierarchy::ReplaceItemInSelection(const FText& OldName, const FText& NewName)
+{
+	for (const TSharedPtr<FIKRigTreeElement>& Item : AllElements)
+	{
+		// remove old selection
+		if (Item->Key.EqualTo(OldName))
+		{
+			TreeView->SetItemSelection(Item, false, ESelectInfo::Direct);
+		}
+		// add new selection
+		if (Item->Key.EqualTo(NewName))
+		{
+			TreeView->SetItemSelection(Item, true, ESelectInfo::Direct);
+		}
+	}
+}
+
+void SIKRigHierarchy::GetSelectedBoneChains(TArray<FBoneChain>& OutChains)
+{
+	const TSharedPtr<FIKRigEditorController> Controller = EditorController.Pin();
+	if (!Controller.IsValid())
+	{
+		return; 
+	}
+
+	const FIKRigSkeleton& Skeleton = Controller->AssetController->GetIKRigSkeleton();
+
+	// get selected bones
+	TArray<TSharedPtr<FIKRigTreeElement>> SelectedBoneItems;
+	GetSelectedBones(SelectedBoneItems);
+
+	// get selected bone indices
+	TArray<int32> SelectedBones;
+	for (const TSharedPtr<FIKRigTreeElement>& BoneItem : SelectedBoneItems)
+	{
+		const FName BoneName = BoneItem.Get()->BoneName;
+		const int32 BoneIndex = Skeleton.GetBoneIndexFromName(BoneName);
+		SelectedBones.Add(BoneIndex);
+	}
+
+	// get bones in chains
+	Skeleton.GetChainsInList(SelectedBones, OutChains);
+	
+	// add goals (if there are any on the end bone of the chain)
+	for (FBoneChain& Chain : OutChains)
+	{
+		if (const UIKRigEffectorGoal* Goal = Controller->AssetController->GetGoalForBone(Chain.EndBone.BoneName))
+		{
+			Chain.IKGoalName = Goal->GoalName;
+		}
+	}
+}
+
+TArray<TSharedPtr<FIKRigTreeElement>> SIKRigHierarchy::GetSelectedItems() const
+{
+	return TreeView->GetSelectedItems();
+}
+
+bool SIKRigHierarchy::HasSelectedItems() const
+{
+	return TreeView->GetNumItemsSelected() > 0;
+}
+
+bool SIKRigHierarchy::IsElementConnectedToSolver(TWeakPtr<FIKRigTreeElement> InTreeElement, int32 SolverIndex)
+{
+	const UIKRigController* AssetController = EditorController.Pin()->AssetController;
+	const TSharedPtr<FIKRigTreeElement> TreeElement = InTreeElement.Pin();
+	
+	if (!AssetController->GetSolverArray().IsValidIndex(SolverIndex))
+	{
+		return false; // not a valid solver index
+	}
+
+	const UIKRigSolver* Solver = AssetController->GetSolver(SolverIndex);
+	if (TreeElement->ElementType == IKRigTreeElementType::BONE)
+	{
+		// is this bone affected by this solver?
+		return Solver->IsBoneAffectedBySolver(TreeElement->BoneName, AssetController->GetIKRigSkeleton());
+	}
+
+	if (TreeElement->ElementType == IKRigTreeElementType::BONE_SETTINGS)
+	{
+		// is this bone setting belonging to the solver?
+		return (Solver->GetBoneSetting(TreeElement->BoneSettingBoneName) != nullptr);
+	}
+
+	if (TreeElement->ElementType == IKRigTreeElementType::GOAL)
+	{
+		// is goal connected to the solver?
+		return AssetController->IsGoalConnectedToSolver(TreeElement->GoalName, SolverIndex);
+	}
+
+	if (TreeElement->ElementType == IKRigTreeElementType::SOLVERGOAL)
+	{
+		// is this an effector for this solver?
+		return TreeElement->EffectorIndex == SolverIndex;
+	}
+
+	checkNoEntry();
+	return false;
+}
+
+bool SIKRigHierarchy::IsElementConnectedToAnySolver(TWeakPtr<FIKRigTreeElement> InTreeElement)
+{
+	const UIKRigController* AssetController = EditorController.Pin()->AssetController;
+	const TSharedPtr<FIKRigTreeElement> TreeElement = InTreeElement.Pin();
+	
+	const int32 NumSolvers = AssetController->GetNumSolvers();
+	for (int32 SolverIndex=0; SolverIndex<NumSolvers; ++SolverIndex)
+	{
+		if (IsElementConnectedToSolver(TreeElement, SolverIndex))
+		{
+			return true;
+		}
+	}
+
+	return false;
+}
+
+bool SIKRigHierarchy::IsElementExcludedBone(TWeakPtr<FIKRigTreeElement> TreeElement)
+{
+	if (TreeElement.Pin()->ElementType != IKRigTreeElementType::BONE)
+	{
+		return false;
+	}
+	
+	// is this bone excluded?
+	const UIKRigController* AssetController = EditorController.Pin()->AssetController;
+	return AssetController->GetBoneExcluded(TreeElement.Pin()->BoneName);
+}
+
+void SIKRigHierarchy::BindCommands()
+{
+	const FIKRigSkeletonCommands& Commands = FIKRigSkeletonCommands::Get();
+
+	CommandList->MapAction(Commands.NewGoal,
+        FExecuteAction::CreateSP(this, &SIKRigHierarchy::HandleNewGoal),
+        FCanExecuteAction::CreateSP(this, &SIKRigHierarchy::CanAddNewGoal));
+	
+	CommandList->MapAction(Commands.DeleteElement,
+        FExecuteAction::CreateSP(this, &SIKRigHierarchy::HandleDeleteElement),
+        FCanExecuteAction::CreateSP(this, &SIKRigHierarchy::CanDeleteElement));
+
+	CommandList->MapAction(Commands.ConnectGoalToSolver,
+        FExecuteAction::CreateSP(this, &SIKRigHierarchy::HandleConnectGoalToSolver),
+        FCanExecuteAction::CreateSP(this, &SIKRigHierarchy::CanConnectGoalToSolvers));
+
+	CommandList->MapAction(Commands.DisconnectGoalFromSolver,
+        FExecuteAction::CreateSP(this, &SIKRigHierarchy::HandleDisconnectGoalFromSolver),
+        FCanExecuteAction::CreateSP(this, &SIKRigHierarchy::CanDisconnectGoalFromSolvers));
+
+	CommandList->MapAction(Commands.SetRootBoneOnSolver,
+        FExecuteAction::CreateSP(this, &SIKRigHierarchy::HandleSetRootBoneOnSolvers),
+        FCanExecuteAction::CreateSP(this, &SIKRigHierarchy::CanSetRootBoneOnSolvers));
+
+	CommandList->MapAction(Commands.SetEndBoneOnSolver,
+	FExecuteAction::CreateSP(this, &SIKRigHierarchy::HandleSetEndBoneOnSolvers),
+	FCanExecuteAction::CreateSP(this, &SIKRigHierarchy::CanSetEndBoneOnSolvers),
+	FIsActionChecked(),
+	FIsActionButtonVisible::CreateSP(this, &SIKRigHierarchy::HasEndBoneCompatibleSolverSelected));
+
+	CommandList->MapAction(Commands.AddBoneSettings,
+        FExecuteAction::CreateSP(this, &SIKRigHierarchy::HandleAddBoneSettings),
+        FCanExecuteAction::CreateSP(this, &SIKRigHierarchy::CanAddBoneSettings));
+
+	CommandList->MapAction(Commands.RemoveBoneSettings,
+        FExecuteAction::CreateSP(this, &SIKRigHierarchy::HandleRemoveBoneSettings),
+        FCanExecuteAction::CreateSP(this, &SIKRigHierarchy::CanRemoveBoneSettings));
+
+	CommandList->MapAction(Commands.ExcludeBone,
+		FExecuteAction::CreateSP(this, &SIKRigHierarchy::HandleExcludeBone),
+		FCanExecuteAction::CreateSP(this, &SIKRigHierarchy::CanExcludeBone));
+
+	CommandList->MapAction(Commands.IncludeBone,
+		FExecuteAction::CreateSP(this, &SIKRigHierarchy::HandleIncludeBone),
+		FCanExecuteAction::CreateSP(this, &SIKRigHierarchy::CanIncludeBone));
+
+	CommandList->MapAction(Commands.NewRetargetChain,
+		FExecuteAction::CreateSP(this, &SIKRigHierarchy::HandleNewRetargetChain),
+		FCanExecuteAction::CreateSP(this, &SIKRigHierarchy::CanAddNewRetargetChain));
+
+	CommandList->MapAction(Commands.SetRetargetRoot,
+		FExecuteAction::CreateSP(this, &SIKRigHierarchy::HandleSetRetargetRoot),
+		FCanExecuteAction::CreateSP(this, &SIKRigHierarchy::CanSetRetargetRoot));
+
+	CommandList->MapAction(Commands.ClearRetargetRoot,
+		FExecuteAction::CreateSP(this, &SIKRigHierarchy::HandleClearRetargetRoot),
+		FCanExecuteAction::CreateSP(this, &SIKRigHierarchy::CanClearRetargetRoot));
+
+	CommandList->MapAction(Commands.RenameGoal,
+	FExecuteAction::CreateSP(this, &SIKRigHierarchy::HandleRenameGoal),
+		FCanExecuteAction::CreateSP(this, &SIKRigHierarchy::CanRenameGoal));
+}
+
+void SIKRigHierarchy::FillContextMenu(FMenuBuilder& MenuBuilder)
+{
+	const FIKRigSkeletonCommands& Actions = FIKRigSkeletonCommands::Get();
+
+	const TArray<TSharedPtr<FIKRigTreeElement>> SelectedItems = TreeView->GetSelectedItems();
+	if (SelectedItems.IsEmpty())
+	{
+		MenuBuilder.AddWidget(
+			SNew(SBox)
+			.HAlign(HAlign_Center)
+			[
+				SNew(STextBlock).Text(LOCTEXT("InvalidSelection", "Nothing selected."))
+			], FText(), false);
+	}
+
+	MenuBuilder.BeginSection("AddRemoveGoals", LOCTEXT("AddRemoveGoalOperations", "Goals"));
+	MenuBuilder.AddMenuEntry(Actions.NewGoal);
+	MenuBuilder.EndSection();
+	
+	MenuBuilder.BeginSection("ConnectGoals", LOCTEXT("ConnectGoalOperations", "Connect Goals To Solvers"));
+	MenuBuilder.AddMenuEntry(Actions.ConnectGoalToSolver);
+	MenuBuilder.AddMenuEntry(Actions.DisconnectGoalFromSolver);
+	MenuBuilder.EndSection();
+
+	MenuBuilder.BeginSection("BoneSettings", LOCTEXT("BoneSettingsOperations", "Bone Settings"));
+	MenuBuilder.AddMenuEntry(Actions.SetRootBoneOnSolver);
+	MenuBuilder.AddMenuEntry(Actions.SetEndBoneOnSolver);
+	MenuBuilder.AddMenuEntry(Actions.AddBoneSettings);
+	MenuBuilder.AddMenuEntry(Actions.RemoveBoneSettings);
+	MenuBuilder.EndSection();
+
+	MenuBuilder.BeginSection("IncludeExclude", LOCTEXT("IncludeExcludeOperations", "Exclude Bones"));
+	MenuBuilder.AddMenuEntry(Actions.ExcludeBone);
+	MenuBuilder.AddMenuEntry(Actions.IncludeBone);
+	MenuBuilder.EndSection();
+
+	MenuBuilder.BeginSection("Retargeting", LOCTEXT("RetargetingOperations", "Retargeting"));
+	MenuBuilder.AddMenuEntry(Actions.SetRetargetRoot);
+	MenuBuilder.AddMenuEntry(Actions.ClearRetargetRoot);
+	MenuBuilder.AddMenuEntry(Actions.NewRetargetChain);
+	MenuBuilder.EndSection();
+}
+
+void SIKRigHierarchy::HandleNewGoal() const
+{
+	const TSharedPtr<FIKRigEditorController> Controller = EditorController.Pin();
+	if (!Controller.IsValid())
+	{
+		return; 
+	}
+
+	// get names of selected bones and default goal names for them
+	TArray<FName> GoalNames;
+	TArray<FName> BoneNames;
+	const TArray<TSharedPtr<FIKRigTreeElement>> SelectedItems = TreeView.Get()->GetSelectedItems();
+	for (const TSharedPtr<FIKRigTreeElement>& Item : SelectedItems)
+	{
+		if (Item->ElementType != IKRigTreeElementType::BONE)
+		{
+			continue; // can only add goals to bones
+		}
+
+		// build default name for the new goal
+		const FName BoneName = Item->BoneName;
+		const FName NewGoalName = FName(BoneName.ToString() + "_Goal");
+		
+		GoalNames.Add(NewGoalName);
+		BoneNames.Add(BoneName);
+	}
+	
+	// add new goals
+	Controller->AddNewGoals(GoalNames, BoneNames);
+}
+
+bool SIKRigHierarchy::CanAddNewGoal() const
+{
+	// is anything selected?
+	TArray<TSharedPtr<FIKRigTreeElement>> SelectedItems = TreeView->GetSelectedItems();
+	if (SelectedItems.IsEmpty())
+	{
+		return false;
+	}
+
+	// can only add goals to selected bones
+	for (const TSharedPtr<FIKRigTreeElement>& Item : SelectedItems)
+	{
+		if (Item->ElementType != IKRigTreeElementType::BONE)
+		{
+			return false;
+		}
+	}
+
+	return true;
+}
+
+void SIKRigHierarchy::HandleDeleteElement()
+{
+	const TSharedPtr<FIKRigEditorController> Controller = EditorController.Pin();
+	if (!Controller.IsValid())
+	{
+		return; 
+	}
+
+	TArray<TSharedPtr<FIKRigTreeElement>> SelectedItems = TreeView->GetSelectedItems();
+	for (const TSharedPtr<FIKRigTreeElement>& SelectedItem : SelectedItems)
+	{
+		switch(SelectedItem->ElementType)
+		{
+			case IKRigTreeElementType::GOAL:
+				Controller->AssetController->RemoveGoal(SelectedItem->GoalName);
+				break;
+			case IKRigTreeElementType::SOLVERGOAL:
+				Controller->AssetController->DisconnectGoalFromSolver(SelectedItem->EffectorGoalName, SelectedItem->EffectorIndex);
+				break;
+			case IKRigTreeElementType::BONE_SETTINGS:
+				Controller->AssetController->RemoveBoneSetting(SelectedItem->BoneSettingBoneName, SelectedItem->BoneSettingsSolverIndex);
+				break;
+			default:
+				break; // can't delete anything else
+		}
+	}
+	
+	RefreshTreeView();
+
+	Controller->ShowEmptyDetails();
+	// update all views
+	Controller->RefreshAllViews();
+}
+bool SIKRigHierarchy::CanDeleteElement() const
+{
+	// is anything selected?
+	TArray<TSharedPtr<FIKRigTreeElement>> SelectedItems = TreeView->GetSelectedItems();
+	if (SelectedItems.IsEmpty())
+	{
+		return false;
+	}
+
+	// are all selected items goals, effectors or bone settings?
+	for (const TSharedPtr<FIKRigTreeElement>& Item : SelectedItems)
+	{
+		const bool bIsBone = Item->ElementType == IKRigTreeElementType::BONE;
+		if (bIsBone)
+		{
+			return false;
+		}
+	}
+
+	return true;
+}
+
+void SIKRigHierarchy::HandleConnectGoalToSolver()
+{
+	const bool bConnect = true; //connect
+	ConnectSelectedGoalsToSelectedSolvers(bConnect);
+}
+
+void SIKRigHierarchy::HandleDisconnectGoalFromSolver()
+{
+	const bool bConnect = false; //disconnect
+	ConnectSelectedGoalsToSelectedSolvers(bConnect);
+}
+
+bool SIKRigHierarchy::CanConnectGoalToSolvers() const
+{
+	const bool bCountOnlyConnected = false;
+	const int32 NumDisconnected = GetNumSelectedGoalToSolverConnections(bCountOnlyConnected);
+	return NumDisconnected > 0;
+}
+
+bool SIKRigHierarchy::CanDisconnectGoalFromSolvers() const
+{
+	const bool bCountOnlyConnected = true;
+	const int32 NumConnected = GetNumSelectedGoalToSolverConnections(bCountOnlyConnected);
+	return NumConnected > 0;
+}
+
+void SIKRigHierarchy::ConnectSelectedGoalsToSelectedSolvers(bool bConnect)
+{
+	const TSharedPtr<FIKRigEditorController> Controller = EditorController.Pin();
+	if (!Controller.IsValid())
+	{
+		return; 
+	}
+	
+	TArray<TSharedPtr<FIKRigTreeElement>> SelectedGoals;
+	GetSelectedGoals(SelectedGoals);
+	TArray<TSharedPtr<FSolverStackElement>> SelectedSolvers;
+	Controller->GetSelectedSolvers(SelectedSolvers);
+
+	UIKRigController* AssetController = Controller->AssetController;
+	for (const TSharedPtr<FIKRigTreeElement>& GoalElement : SelectedGoals)
+	{
+		const FName GoalName = GoalElement->GoalName;
+		const int32 GoalIndex = AssetController->GetGoalIndex(GoalName);
+		check(GoalIndex != INDEX_NONE);
+		const UIKRigEffectorGoal& EffectorGoal = *AssetController->GetGoal(GoalIndex);
+		for (const TSharedPtr<FSolverStackElement>& SolverElement : SelectedSolvers)
+		{
+			if (bConnect)
+			{
+				AssetController->ConnectGoalToSolver(EffectorGoal, SolverElement->IndexInStack);	
+			}
+			else
+			{
+				AssetController->DisconnectGoalFromSolver(EffectorGoal.GoalName, SolverElement->IndexInStack);	
+			}
+		}
+	}
+
+	// add/remove new effector under goal in skeleton view
+	RefreshTreeView();
+}
+
+int32 SIKRigHierarchy::GetNumSelectedGoalToSolverConnections(bool bCountOnlyConnected) const
+{
+	const TSharedPtr<FIKRigEditorController> Controller = EditorController.Pin();
+	if (!Controller.IsValid())
+	{
+		return 0;
+	}
+	
+	TArray<TSharedPtr<FIKRigTreeElement>> SelectedGoals;
+	GetSelectedGoals(SelectedGoals);
+	TArray<TSharedPtr<FSolverStackElement>> SelectedSolvers;
+	Controller->GetSelectedSolvers(SelectedSolvers);
+
+	int32 NumMatched = 0;
+	for (const TSharedPtr<FIKRigTreeElement>& Goal : SelectedGoals)
+	{
+		for (const TSharedPtr<FSolverStackElement>& Solver : SelectedSolvers)
+		{
+			const bool bIsConnected = Controller->AssetController->IsGoalConnectedToSolver(Goal->GoalName, Solver->IndexInStack);
+			if (bIsConnected == bCountOnlyConnected)
+			{
+				++NumMatched;
+			}
+		}
+	}
+
+	return NumMatched;
+}
+
+void SIKRigHierarchy::HandleSetRootBoneOnSolvers()
+{
+	const TSharedPtr<FIKRigEditorController> Controller = EditorController.Pin();
+	if (!Controller.IsValid())
+	{
+		return;
+	}
+
+    // get name of selected root bone
+    TArray<TSharedPtr<FIKRigTreeElement>> SelectedBones;
+	GetSelectedBones(SelectedBones);
+	const FName RootBoneName = SelectedBones[0]->BoneName;
+
+	// apply to all selected solvers (ignored on solvers that don't accept a root bone)
+	UIKRigController* AssetController = Controller->AssetController;
+	TArray<TSharedPtr<FSolverStackElement>> SelectedSolvers;
+	Controller->GetSelectedSolvers(SelectedSolvers);
+	int32 SolverToShow = 0;
+	for (const TSharedPtr<FSolverStackElement>& Solver : SelectedSolvers)
+	{
+		AssetController->SetRootBone(RootBoneName, Solver->IndexInStack);
+		SolverToShow = Solver->IndexInStack;
+	}
+
+	// show solver that had it's root bone updated
+	Controller->ShowDetailsForSolver(SolverToShow);
+	
+	// show new icon when bone has settings applied
+	RefreshTreeView();
+}
+
+bool SIKRigHierarchy::CanSetRootBoneOnSolvers()
+{
+	const TSharedPtr<FIKRigEditorController> Controller = EditorController.Pin();
+	if (!Controller.IsValid())
+	{
+		return false;
+	}
+	
+	// must have at least 1 bone selected
+	TArray<TSharedPtr<FIKRigTreeElement>> SelectedBones;
+	GetSelectedBones(SelectedBones);
+	if (SelectedBones.Num() != 1)
+	{
+		return false;
+	}
+
+	// must have at least 1 solver selected that accepts root bones
+	UIKRigController* AssetController = Controller->AssetController;
+	TArray<TSharedPtr<FSolverStackElement>> SelectedSolvers;
+	Controller->GetSelectedSolvers(SelectedSolvers);
+	for (const TSharedPtr<FSolverStackElement>& Solver : SelectedSolvers)
+	{
+		if (AssetController->GetSolver(Solver->IndexInStack)->RequiresRootBone())
+		{
+			return true;
+		}
+	}
+	
+	return false;
+}
+
+void SIKRigHierarchy::HandleSetEndBoneOnSolvers()
+{
+	const TSharedPtr<FIKRigEditorController> Controller = EditorController.Pin();
+	if (!Controller.IsValid())
+	{
+		return;
+	}
+
+    // get name of selected root bone
+    TArray<TSharedPtr<FIKRigTreeElement>> SelectedBones;
+	GetSelectedBones(SelectedBones);
+	const FName RootBoneName = SelectedBones[0]->BoneName;
+
+	// apply to all selected solvers (ignored on solvers that don't accept a root bone)
+	UIKRigController* AssetController = Controller->AssetController;
+	TArray<TSharedPtr<FSolverStackElement>> SelectedSolvers;
+	Controller->GetSelectedSolvers(SelectedSolvers);
+	int32 SolverToShow = 0;
+	for (const TSharedPtr<FSolverStackElement>& Solver : SelectedSolvers)
+	{
+		AssetController->SetEndBone(RootBoneName, Solver->IndexInStack);
+		SolverToShow = Solver->IndexInStack;
+	}
+
+	// show solver that had it's root bone updated
+	Controller->ShowDetailsForSolver(SolverToShow);
+	
+	// show new icon when bone has settings applied
+	RefreshTreeView();
+}
+
+bool SIKRigHierarchy::CanSetEndBoneOnSolvers() const
+{
+	const TSharedPtr<FIKRigEditorController> Controller = EditorController.Pin();
+	if (!Controller.IsValid())
+	{
+		return false;
+	}
+	
+	// must have at least 1 bone selected
+	TArray<TSharedPtr<FIKRigTreeElement>> SelectedBones;
+	GetSelectedBones(SelectedBones);
+	if (SelectedBones.Num() != 1)
+	{
+		return false;
+	}
+
+	// must have at least 1 solver selected that accepts end bones
+	UIKRigController* AssetController = Controller->AssetController;
+	TArray<TSharedPtr<FSolverStackElement>> SelectedSolvers;
+	Controller->GetSelectedSolvers(SelectedSolvers);
+	for (const TSharedPtr<FSolverStackElement>& Solver : SelectedSolvers)
+	{
+		if (AssetController->GetSolver(Solver->IndexInStack)->RequiresEndBone())
+		{
+			return true;
+		}
+	}
+	
+	return false;
+}
+
+bool SIKRigHierarchy::HasEndBoneCompatibleSolverSelected() const
+{
+	return CanSetEndBoneOnSolvers();
+}
+
+void SIKRigHierarchy::HandleAddBoneSettings()
+{
+	const TSharedPtr<FIKRigEditorController> Controller = EditorController.Pin();
+	if (!Controller.IsValid())
+	{
+		return;
+	}
+	
+	// get selected bones
+	TArray<TSharedPtr<FIKRigTreeElement>> SelectedBones;
+	GetSelectedBones(SelectedBones);
+	
+	// add settings for bone on all selected solvers (ignored if already present)
+	UIKRigController* AssetController = Controller->AssetController;
+	TArray<TSharedPtr<FSolverStackElement>> SelectedSolvers;
+	Controller->GetSelectedSolvers(SelectedSolvers);
+	FName BoneNameForSettings;
+	int32 SolverIndex = INDEX_NONE;
+	for (const TSharedPtr<FIKRigTreeElement>& BoneItem : SelectedBones)
+	{
+		for (const TSharedPtr<FSolverStackElement>& Solver : SelectedSolvers)
+		{
+			AssetController->AddBoneSetting(BoneItem->BoneName, Solver->IndexInStack);
+			BoneNameForSettings = BoneItem->BoneName;
+			SolverIndex = Solver->IndexInStack;
+        }
+	}
+
+	Controller->ShowDetailsForBoneSettings(BoneNameForSettings, SolverIndex);
+
+	// show new icon when bone has settings applied
+	RefreshTreeView();
+}
+
+bool SIKRigHierarchy::CanAddBoneSettings()
+{
+	const TSharedPtr<FIKRigEditorController> Controller = EditorController.Pin();
+	if (!Controller.IsValid())
+	{
+		return false;
+	}
+	
+    // must have at least 1 bone selected
+    TArray<TSharedPtr<FIKRigTreeElement>> SelectedBones;
+	GetSelectedBones(SelectedBones);
+	if (SelectedBones.IsEmpty())
+	{
+		return false;
+	}
+
+	// must have at least 1 solver selected that does not already have a bone setting for the selected bones
+	UIKRigController* AssetController = Controller->AssetController;
+	TArray<TSharedPtr<FSolverStackElement>> SelectedSolvers;
+	Controller->GetSelectedSolvers(SelectedSolvers);
+	for (const TSharedPtr<FIKRigTreeElement>& BoneItem : SelectedBones)
+	{
+		for (const TSharedPtr<FSolverStackElement>& Solver : SelectedSolvers)
+		{
+			if (AssetController->CanAddBoneSetting(BoneItem->BoneName, Solver->IndexInStack))
+			{
+				return true;
+			}
+        }
+	}
+	
+	return false;
+}
+
+void SIKRigHierarchy::HandleRemoveBoneSettings()
+{
+	const TSharedPtr<FIKRigEditorController> Controller = EditorController.Pin();
+	if (!Controller.IsValid())
+	{
+		return;
+	}
+	
+	// get selected bones
+	TArray<TSharedPtr<FIKRigTreeElement>> SelectedBones;
+	GetSelectedBones(SelectedBones);
+	
+	// add settings for bone on all selected solvers (ignored if already present)
+	UIKRigController* AssetController = Controller->AssetController;
+	TArray<TSharedPtr<FSolverStackElement>> SelectedSolvers;
+	Controller->GetSelectedSolvers(SelectedSolvers);
+	FName BoneToShowInDetailsView;
+	for (const TSharedPtr<FIKRigTreeElement>& BoneItem : SelectedBones)
+	{
+		for (const TSharedPtr<FSolverStackElement>& Solver : SelectedSolvers)
+		{
+			AssetController->RemoveBoneSetting(BoneItem->BoneName, Solver->IndexInStack);
+			BoneToShowInDetailsView = BoneItem->BoneName;
+		}
+	}
+
+	Controller->ShowDetailsForBone(BoneToShowInDetailsView);
+	
+	// show new icon when bone has settings applied
+	RefreshTreeView();
+}
+
+bool SIKRigHierarchy::CanRemoveBoneSettings()
+{
+	const TSharedPtr<FIKRigEditorController> Controller = EditorController.Pin();
+	if (!Controller.IsValid())
+	{
+		return false;
+	}
+	
+    // must have at least 1 bone selected
+    TArray<TSharedPtr<FIKRigTreeElement>> SelectedBones;
+	GetSelectedBones(SelectedBones);
+	if (SelectedBones.IsEmpty())
+	{
+		return false;
+	}
+
+	// must have at least 1 solver selected that has a bone setting for 1 of the selected bones
+	UIKRigController* AssetController = Controller->AssetController;
+	TArray<TSharedPtr<FSolverStackElement>> SelectedSolvers;
+	Controller->GetSelectedSolvers(SelectedSolvers);
+	for (const TSharedPtr<FIKRigTreeElement>& BoneItem : SelectedBones)
+	{
+		for (const TSharedPtr<FSolverStackElement>& Solver : SelectedSolvers)
+		{
+			if (AssetController->CanRemoveBoneSetting(BoneItem->BoneName, Solver->IndexInStack))
+			{
+				return true;
+			}
+		}
+	}
+	
+	return false;
+}
+
+void SIKRigHierarchy::HandleExcludeBone()
+{
+	const TSharedPtr<FIKRigEditorController> Controller = EditorController.Pin();
+	if (!Controller.IsValid())
+	{
+		return;
+	}
+	
+	// exclude selected bones
+	TArray<TSharedPtr<FIKRigTreeElement>> SelectedBones;
+	GetSelectedBones(SelectedBones);
+	for (const TSharedPtr<FIKRigTreeElement>& BoneItem : SelectedBones)
+	{
+		Controller->AssetController->SetBoneExcluded(BoneItem->BoneName, true);
+	}
+
+	// show greyed out bone name after being excluded
+	RefreshTreeView();
+}
+
+bool SIKRigHierarchy::CanExcludeBone()
+{
+	const TSharedPtr<FIKRigEditorController> Controller = EditorController.Pin();
+	if (!Controller.IsValid())
+	{
+		return false;
+	}
+	
+	// must have at least 1 bone selected that is INCLUDED
+	TArray<TSharedPtr<FIKRigTreeElement>> SelectedBones;
+	GetSelectedBones(SelectedBones);
+	for (const TSharedPtr<FIKRigTreeElement>& BoneItem : SelectedBones)
+	{
+		if (!Controller->AssetController->GetBoneExcluded(BoneItem->BoneName))
+		{
+			return true;
+		}
+	}
+
+	return false;
+}
+
+void SIKRigHierarchy::HandleIncludeBone()
+{
+	const TSharedPtr<FIKRigEditorController> Controller = EditorController.Pin();
+	if (!Controller.IsValid())
+	{
+		return;
+	}
+	
+	// exclude selected bones
+	TArray<TSharedPtr<FIKRigTreeElement>> SelectedBones;
+	GetSelectedBones(SelectedBones);
+	for (const TSharedPtr<FIKRigTreeElement>& BoneItem : SelectedBones)
+	{
+		Controller->AssetController->SetBoneExcluded(BoneItem->BoneName, false);
+	}
+
+	// show normal bone name after being included
+	RefreshTreeView();
+}
+
+bool SIKRigHierarchy::CanIncludeBone()
+{
+	const TSharedPtr<FIKRigEditorController> Controller = EditorController.Pin();
+	if (!Controller.IsValid())
+	{
+		return false;
+	}
+	
+	// must have at least 1 bone selected that is EXCLUDED
+	TArray<TSharedPtr<FIKRigTreeElement>> SelectedBones;
+	GetSelectedBones(SelectedBones);
+	for (const TSharedPtr<FIKRigTreeElement>& BoneItem : SelectedBones)
+	{
+		if (Controller->AssetController->GetBoneExcluded(BoneItem->BoneName))
+		{
+			return true;
+		}
+	}
+
+	return false;
+}
+
+void SIKRigHierarchy::HandleNewRetargetChain()
+{
+	const TSharedPtr<FIKRigEditorController> Controller = EditorController.Pin();
+	if (!Controller.IsValid())
+	{
+		return;
+	}
+
+	Controller->CreateNewRetargetChains();
+}
+
+bool SIKRigHierarchy::CanAddNewRetargetChain()
+{
+	TArray<TSharedPtr<FIKRigTreeElement>> SelectedBones;
+	GetSelectedBones(SelectedBones);
+	return !SelectedBones.IsEmpty();
+}
+
+void SIKRigHierarchy::HandleSetRetargetRoot()
+{
+	const TSharedPtr<FIKRigEditorController> Controller = EditorController.Pin();
+	if (!Controller.IsValid())
+	{
+		return;
+	}
+	
+	// get selected bones
+	TArray<TSharedPtr<FIKRigTreeElement>> SelectedBones;
+	GetSelectedBones(SelectedBones);
+
+	// must have at least 1 bone selected
+	if (SelectedBones.IsEmpty())
+	{
+		return;
+	}
+
+	// set the first selected bone as the retarget root
+	Controller->AssetController->SetRetargetRoot(SelectedBones[0]->BoneName);
+
+	// show root bone after being set
+	Controller->RefreshAllViews();
+}
+
+bool SIKRigHierarchy::CanSetRetargetRoot()
+{
+	TArray<TSharedPtr<FIKRigTreeElement>> SelectedBones;
+	GetSelectedBones(SelectedBones);
+	return !SelectedBones.IsEmpty();
+}
+
+void SIKRigHierarchy::HandleClearRetargetRoot()
+{
+	const TSharedPtr<FIKRigEditorController> Controller = EditorController.Pin();
+	if (!Controller.IsValid())
+	{
+		return;
+	}
+	
+	Controller->AssetController->SetRetargetRoot(NAME_None);
+	Controller->RefreshAllViews();
+}
+
+bool SIKRigHierarchy::CanClearRetargetRoot()
+{
+	const TSharedPtr<FIKRigEditorController> Controller = EditorController.Pin();
+	if (!Controller.IsValid())
+	{
+		return false;
+	}
+
+	return Controller->AssetController->GetRetargetRoot() != NAME_None;
+}
+
+bool SIKRigHierarchy::IsBoneInSelection(TArray<TSharedPtr<FIKRigTreeElement>>& SelectedBoneItems, const FName& BoneName)
+{
+	for (const TSharedPtr<FIKRigTreeElement>& Item : SelectedBoneItems)
+	{
+		if (Item->BoneName == BoneName)
+		{
+			return true;
+		}
+	}
+	return false;
+}
+
+void SIKRigHierarchy::GetSelectedBones(TArray<TSharedPtr<FIKRigTreeElement>>& OutBoneItems) const
+{
+	const TArray<TSharedPtr<FIKRigTreeElement>> SelectedItems = TreeView->GetSelectedItems();
+	for (const TSharedPtr<FIKRigTreeElement>& Item : SelectedItems)
+	{
+		if (Item->ElementType == IKRigTreeElementType::BONE)
+		{
+			OutBoneItems.Add(Item);
+        }
+	}
+}
+
+void SIKRigHierarchy::GetSelectedBoneNames(TArray<FName>& OutSelectedBoneNames) const
+{
+	TArray<TSharedPtr<FIKRigTreeElement>> OutSelectedBones;
+	GetSelectedBones(OutSelectedBones);
+	OutSelectedBoneNames.Reset();
+	for (TSharedPtr<FIKRigTreeElement> SelectedBoneItem : OutSelectedBones)
+	{
+		OutSelectedBoneNames.Add(SelectedBoneItem->BoneName);
+	}
+}
+
+void SIKRigHierarchy::GetSelectedGoals(TArray<TSharedPtr<FIKRigTreeElement>>& OutSelectedGoals) const
+{
+	OutSelectedGoals.Reset();
+	TArray<TSharedPtr<FIKRigTreeElement>> SelectedItems = TreeView->GetSelectedItems();
+	for (const TSharedPtr<FIKRigTreeElement>& Item : SelectedItems)
+	{
+		if (Item->ElementType == IKRigTreeElementType::GOAL)
+		{
+			OutSelectedGoals.Add(Item);
+		}
+	}
+}
+
+int32 SIKRigHierarchy::GetNumSelectedGoals()
+{
+	TArray<TSharedPtr<FIKRigTreeElement>> OutSelectedGoals;
+	GetSelectedGoals(OutSelectedGoals);
+	return OutSelectedGoals.Num();
+}
+
+void SIKRigHierarchy::GetSelectedGoalNames(TArray<FName>& OutSelectedGoalNames) const
+{
+	TArray<TSharedPtr<FIKRigTreeElement>> OutSelectedGoals;
+	GetSelectedGoals(OutSelectedGoals);
+	OutSelectedGoalNames.Reset();
+	for (TSharedPtr<FIKRigTreeElement> SelectedGoalItem : OutSelectedGoals)
+	{
+		OutSelectedGoalNames.Add(SelectedGoalItem->GoalName);
+	}
+}
+
+bool SIKRigHierarchy::IsGoalSelected(const FName& GoalName)
+{
+	TArray<TSharedPtr<FIKRigTreeElement>> OutSelectedGoals;
+	GetSelectedGoals(OutSelectedGoals);
+	for (TSharedPtr<FIKRigTreeElement> SelectedGoalItem : OutSelectedGoals)
+	{
+		if (SelectedGoalItem->GoalName == GoalName)
+		{
+			return true;
+		}
+	}
+	return false;
+}
+
+void SIKRigHierarchy::HandleRenameGoal() const
+{
+	TArray<TSharedPtr<FIKRigTreeElement>> SelectedGoals;
+	GetSelectedGoals(SelectedGoals);
+	if (SelectedGoals.Num() != 1)
+	{
+		return;
+	}
+	
+	SelectedGoals[0]->RequestRename();
+}
+
+bool SIKRigHierarchy::CanRenameGoal() const
+{
+	TArray<TSharedPtr<FIKRigTreeElement>> SelectedGoals;
+	GetSelectedGoals(SelectedGoals);
+	return SelectedGoals.Num() == 1;
+}
+
+TSharedRef<SWidget> SIKRigHierarchy::CreateFilterMenuWidget()
+{
+	const FUIAction FilterShowAllAction = FUIAction(
+		FExecuteAction::CreateLambda([this]
+		{
+			FilterOptions.ResetShowOptions();
+			FilterOptions.bShowAll = true; 
+			RefreshTreeView();
+		}),
+		FCanExecuteAction(),
+		FIsActionChecked::CreateLambda([this]()
+		{
+			return FilterOptions.bShowAll;
+		}));
+	
+	const FUIAction FilterOnlyBonesAction = FUIAction(
+		FExecuteAction::CreateLambda([this]
+		{
+			FilterOptions.ResetShowOptions();
+			FilterOptions.bShowOnlyBones = true;
+			RefreshTreeView();
+		}),
+		FCanExecuteAction(),
+		FIsActionChecked::CreateLambda([this]()
+		{
+			return FilterOptions.bShowOnlyBones;
+		}));
+
+	const FUIAction FilterOnlyGoalsAction = FUIAction(
+		FExecuteAction::CreateLambda([this]
+		{
+			FilterOptions.ResetShowOptions();
+			FilterOptions.bShowOnlyGoals = true;
+			RefreshTreeView();
+		}),
+		FCanExecuteAction(),
+		FIsActionChecked::CreateLambda([this]()
+		{
+			return FilterOptions.bShowOnlyGoals;
+		}));
+
+	const FUIAction FilterOnlyBoneSettingsAction = FUIAction(
+		FExecuteAction::CreateLambda([this]
+		{
+			FilterOptions.ResetShowOptions();
+			FilterOptions.bShowOnlyBoneSettings = true;
+			RefreshTreeView();
+		}),
+		FCanExecuteAction(),
+		FIsActionChecked::CreateLambda([this]()
+		{
+			return FilterOptions.bShowOnlyBoneSettings;
+		}));
+
+	const FUIAction FilterUnaffectedBonesAction = FUIAction(
+		FExecuteAction::CreateLambda([this]
+		{
+			FilterOptions.bHideUnaffectedBones = !FilterOptions.bHideUnaffectedBones;
+			RefreshTreeView();
+		}),
+		FCanExecuteAction(),
+		FIsActionChecked::CreateLambda([this]()
+		{
+			return FilterOptions.bHideUnaffectedBones;
+		}));
+	
+	static constexpr bool CloseAfterSelection = true;
+	FMenuBuilder MenuBuilder(CloseAfterSelection, CommandList);
+
+	MenuBuilder.BeginSection("Hierarchy Filters", LOCTEXT("HierarchyFiltersSection", "Hierarchy Filters"));
+
+	MenuBuilder.AddMenuEntry(
+		LOCTEXT("ShowAllTypesLabel", "Show All Types"),
+		LOCTEXT("ShowAllTypesTooltip", "Show all object types in the hierarchy."),
+		FSlateIcon(),
+		FilterShowAllAction,
+		NAME_None,
+		EUserInterfaceActionType::RadioButton);
+
+	MenuBuilder.AddMenuEntry(
+		LOCTEXT("OnlyBonesLabel", "Show Only Bones"),
+		LOCTEXT("OnlyBonesTooltip", "Hide everything except bones."),
+		FSlateIcon(),
+		FilterOnlyBonesAction,
+		NAME_None,
+		EUserInterfaceActionType::RadioButton);
+
+	MenuBuilder.AddMenuEntry(
+		LOCTEXT("OnlyGoalsLabel", "Show Only Goals"),
+		LOCTEXT("OnlyGoalsTooltip", "Hide everything except IK goals."),
+		FSlateIcon(),
+		FilterOnlyGoalsAction,
+		NAME_None,
+		EUserInterfaceActionType::RadioButton);
+
+	MenuBuilder.AddMenuEntry(
+		LOCTEXT("OnlyBoneSettingsLabel", "Show Only Bone Settings"),
+		LOCTEXT("OnlyBoneSettingsTooltip", "Hide everything except bone settings."),
+		FSlateIcon(),
+		FilterOnlyBoneSettingsAction,
+		NAME_None,
+		EUserInterfaceActionType::RadioButton);
+
+	MenuBuilder.AddMenuEntry(
+		LOCTEXT("UnaffectedBonesLabel", "Hide Unaffected Bones"),
+		LOCTEXT("UnaffectedBonesTooltip", "Hide all bones that are unaffected by any solvers."),
+		FSlateIcon(),
+		FilterUnaffectedBonesAction,
+		NAME_None,
+		EUserInterfaceActionType::Check);
+	
+	MenuBuilder.EndSection();
+
+	MenuBuilder.BeginSection("Clear", LOCTEXT("ClearMapFiltersSection", "Clear"));
+	
+	MenuBuilder.AddMenuEntry(
+		LOCTEXT("ClearSkeletonFilterLabel", "Clear Filters"),
+		LOCTEXT("ClearSkeletonFilterTooltip", "Clear all filters to show all hierarchy elements."),
+		FSlateIcon(),
+		FUIAction(FExecuteAction::CreateLambda([this]
+		{
+			FilterOptions = FIKRigHierarchyFilterOptions();
+			RefreshTreeView();
+		})));
+
+	MenuBuilder.EndSection();
+	
+	return MenuBuilder.MakeWidget();
+}
+
+void SIKRigHierarchy::OnFilterTextChanged(const FText& SearchText)
+{
+	TextFilter->SetFilterText(SearchText);
+	RefreshTreeView();
+}
+
+void SIKRigHierarchy::RefreshTreeView(bool IsInitialSetup /*=false*/)
+{
+	const TSharedPtr<FIKRigEditorController> Controller = EditorController.Pin();
+	if (!Controller.IsValid())
+	{
+		return;
+	}
+
+	const TSharedRef<FIKRigEditorController> ControllerRef = Controller.ToSharedRef();
+	
+	// save expansion and selection state
+	TreeView->SaveAndClearState();
+
+	// reset all tree items
+	RootElements.Reset();
+	AllElements.Reset();
+
+	// validate we have a skeleton to load
+	UIKRigController* AssetController = Controller->AssetController;
+	const FIKRigSkeleton& Skeleton = AssetController->GetIKRigSkeleton();
+	if (Skeleton.BoneNames.IsEmpty())
+	{
+		TreeView->RequestTreeRefresh();
+		return;
+	}
+
+	// get all goals
+	const TArray<UIKRigEffectorGoal*>& Goals = AssetController->GetAllGoals();
+	
+	// get all solvers
+	const TArray<UIKRigSolver*>& Solvers = AssetController->GetSolverArray();
+	// record bone element indices
+	TMap<FName, int32> BoneTreeElementIndices;
+
+	auto FilterString = [this](const FString& StringToTest) ->bool
+	{
+		return TextFilter->TestTextFilter(FBasicStringFilterExpressionContext(StringToTest));
+	};
+
+	// create all bone elements
+	for (int32 BoneIndex=0; BoneIndex<Skeleton.BoneNames.Num(); ++BoneIndex)
+	{
+		// create "Bone" tree element for this bone
+		const FName BoneName = Skeleton.BoneNames[BoneIndex];
+		const FText BoneDisplayName = FText::FromName(BoneName);
+		TSharedPtr<FIKRigTreeElement> BoneElement = MakeShared<FIKRigTreeElement>(BoneDisplayName, IKRigTreeElementType::BONE, ControllerRef);
+		BoneElement.Get()->BoneName = BoneName;
+		const int32 BoneElementIndex = AllElements.Add(BoneElement);
+		BoneTreeElementIndices.Add(BoneName, BoneElementIndex);
+
+		// store pointer to parent (if there is one)
+		const int32 ParentIndex = Skeleton.ParentIndices[BoneIndex];
+		if (ParentIndex != INDEX_NONE)
+		{
+			// get parent tree element
+			const FName ParentBoneName = Skeleton.BoneNames[ParentIndex];
+			const TSharedPtr<FIKRigTreeElement> ParentBoneTreeElement = AllElements[BoneTreeElementIndices[ParentBoneName]];
+			BoneElement->UnFilteredParent = ParentBoneTreeElement;
+		}
+
+		// apply type filter to bones
+		if (FilterOptions.bShowOnlyBoneSettings || FilterOptions.bShowOnlyGoals)
+		{
+			BoneElement->bIsHidden = true;
+		}
+
+		if (FilterOptions.bHideUnaffectedBones)
+		{
+			
+		}
+		
+		// apply text filter to bones
+		if (!(TextFilter->GetFilterText().IsEmpty() || FilterString(BoneName.ToString())))
+		{
+			BoneElement->bIsHidden = true;
+		}
+
+		// create all "Bone Setting" tree elements for this bone
+		for (int32 SolverIndex=0; SolverIndex<Solvers.Num(); ++SolverIndex)
+		{
+			if (!Solvers[SolverIndex]->GetBoneSetting(BoneName))
+			{
+				continue;
+			}
+			
+			const FText SolverDisplayName = FText::FromString(AssetController->GetSolverUniqueName(SolverIndex));
+			const FText BoneSettingDisplayName = FText::Format(LOCTEXT("BoneSettings", "{0} bone settings."), SolverDisplayName);
+			TSharedPtr<FIKRigTreeElement> SettingsItem = MakeShared<FIKRigTreeElement>(BoneSettingDisplayName, IKRigTreeElementType::BONE_SETTINGS, ControllerRef);
+			SettingsItem->BoneSettingBoneName = BoneName;
+			SettingsItem->BoneSettingsSolverIndex = SolverIndex;
+			AllElements.Add(SettingsItem);
+			
+			// store hierarchy pointer for item
+			SettingsItem->UnFilteredParent = BoneElement;
+
+			// apply type filter to bone settings
+			if (FilterOptions.bShowOnlyBones || FilterOptions.bShowOnlyGoals)
+			{
+				SettingsItem->bIsHidden = true;
+			}
+
+			// apply text filter to bone settings
+			if (!(TextFilter->GetFilterText().IsEmpty() || FilterString(BoneName.ToString())))
+			{
+				SettingsItem->bIsHidden = true;
+			}
+		}
+
+		// create all "Goal" and "Effector" tree elements for this bone
+		for (const UIKRigEffectorGoal* Goal : Goals)
+		{
+			if (Goal->BoneName != BoneName)
+			{
+				continue;
+			}
+			
+			// make new element for goal
+			const FText GoalDisplayName = FText::FromName(Goal->GoalName);
+			TSharedPtr<FIKRigTreeElement> GoalItem = MakeShared<FIKRigTreeElement>(GoalDisplayName, IKRigTreeElementType::GOAL, ControllerRef);
+			GoalItem->GoalName = Goal->GoalName;
+			AllElements.Add(GoalItem);
+
+			// apply type filter to goals
+			if (FilterOptions.bShowOnlyBones || FilterOptions.bShowOnlyBoneSettings)
+			{
+				GoalItem->bIsHidden = true;
+			}
+
+			// apply text filter to goals
+			if (!(TextFilter->GetFilterText().IsEmpty() ||
+				FilterString(Goal->GoalName.ToString()) ||
+				FilterString(Goal->BoneName.ToString())))
+			{
+				GoalItem->bIsHidden = true;
+			}
+
+			// store hierarchy pointer for goal
+			GoalItem->UnFilteredParent = BoneElement;
+
+			// add all solver settings connected to this goal
+			for (int32 SolverIndex=0; SolverIndex<Solvers.Num(); ++SolverIndex)
+			{
+				UObject* GoalSettings = AssetController->GetGoalSettingsForSolver(Goal->GoalName, SolverIndex);
+				if (!GoalSettings)
+				{
+					continue;
+				}
+				
+				// make new element for effector
+				const FText SolverDisplayName = FText::FromString(AssetController->GetSolverUniqueName(SolverIndex));
+				const FText EffectorDisplayName = FText::Format(LOCTEXT("GoalSettingsForSolver", "{0} goal settings."), SolverDisplayName);
+				TSharedPtr<FIKRigTreeElement> EffectorItem = MakeShared<FIKRigTreeElement>(EffectorDisplayName, IKRigTreeElementType::SOLVERGOAL, ControllerRef);
+				EffectorItem->EffectorIndex = SolverIndex;
+				EffectorItem->EffectorGoalName = Goal->GoalName;
+				AllElements.Add(EffectorItem);
+
+				// apply filter to effectors (treated as goals for filtering purposes)
+				EffectorItem->bIsHidden = GoalItem->bIsHidden;
+				
+				// store hierarchy pointer for effectors
+				EffectorItem->UnFilteredParent = GoalItem;
+			}
+		}
+	}
+
+	// resolve parent/children pointers on all tree elements, taking into consideration the filter options
+	// (elements are parented to their nearest non-hidden/filtered parent element)
+	for (TSharedPtr<FIKRigTreeElement>& Element : AllElements)
+	{
+		if (Element->bIsHidden)
+		{
+			continue;
+		}
+		
+		// find first parent that is not filtered
+		TSharedPtr<FIKRigTreeElement> ParentElement = Element->UnFilteredParent;
+		while (true)
+		{
+			if (!ParentElement.IsValid())
+			{
+				break;
+			}
+
+			if (!ParentElement->bIsHidden)
+			{
+				break;
+			}
+
+			ParentElement = ParentElement->UnFilteredParent;
+		}
+
+		if (ParentElement.IsValid())
+		{
+			// store pointer to child on parent
+			ParentElement->Children.Add(Element);
+			// store pointer to parent on child
+			Element->Parent = ParentElement;
+		}
+		else
+		{
+			// has no parent, store a root element
+			RootElements.Add(Element);
+		}
+	}
+
+	// expand all elements upon the initial construction of the tree
+	if (IsInitialSetup)
+	{
+		for (TSharedPtr<FIKRigTreeElement> RootElement : RootElements)
+		{
+			SetExpansionRecursive(RootElement, false, true);
+		}
+	}
+	else
+	{
+		// restore expansion and selection state
+		for (const TSharedPtr<FIKRigTreeElement>& Element : AllElements)
+		{
+			TreeView->RestoreState(Element);
+		}
+	}
+	
+	TreeView->RequestTreeRefresh();
+}
+
+void SIKRigHierarchy::HandleGetChildrenForTree(
+	TSharedPtr<FIKRigTreeElement> InItem,
+	TArray<TSharedPtr<FIKRigTreeElement>>& OutChildren)
+{
+	OutChildren = InItem.Get()->Children;
+}
+
+void SIKRigHierarchy::OnSelectionChanged(TSharedPtr<FIKRigTreeElement> Selection, ESelectInfo::Type SelectInfo)
+{
+	const TSharedPtr<FIKRigEditorController> Controller = EditorController.Pin();
+	if (!Controller.IsValid())
+	{
+		return;
+	}
+
+	// update details view
+	const TArray<TSharedPtr<FIKRigTreeElement>> SelectedItems = TreeView->GetSelectedItems();
+	Controller->ShowDetailsForElements(SelectedItems);
+
+	// NOTE: we may want to set the last selected item here
+}
+
+TSharedRef<SWidget> SIKRigHierarchy::CreateAddNewMenu()
+{
+	constexpr bool CloseAfterSelection = true;
+	FMenuBuilder MenuBuilder(CloseAfterSelection, CommandList);
+	FillContextMenu(MenuBuilder);
+	return MenuBuilder.MakeWidget();
+}
+
+TSharedPtr<SWidget> SIKRigHierarchy::CreateContextMenu()
+{
+	constexpr bool CloseAfterSelection = true;
+	FMenuBuilder MenuBuilder(CloseAfterSelection, CommandList);
+	FillContextMenu(MenuBuilder);
+	return MenuBuilder.MakeWidget();
+}
+
+void SIKRigHierarchy::OnItemClicked(TSharedPtr<FIKRigTreeElement> InItem)
+{
+	const TSharedPtr<FIKRigEditorController> Controller = EditorController.Pin();
+	if (!Controller.IsValid())
+	{
+		return;
+	}
+	
+	// to rename an item, you have to select it first, then click on it again within a time limit (slow double click)
+	const bool ClickedOnSameItem = TreeView->LastSelected.Pin().Get() == InItem.Get();
+	const uint32 CurrentCycles = FPlatformTime::Cycles();
+	const double SecondsPassed = static_cast<double>(CurrentCycles - TreeView->LastClickCycles) * FPlatformTime::GetSecondsPerCycle();
+	if (ClickedOnSameItem && SecondsPassed > 0.25f)
+	{
+		RegisterActiveTimer(0.f,
+            FWidgetActiveTimerDelegate::CreateLambda([this](double, float) {
+            HandleRenameGoal();
+            return EActiveTimerReturnType::Stop;
+        }));
+	}
+
+	TreeView->LastClickCycles = CurrentCycles;
+	TreeView->LastSelected = InItem;
+	Controller->SetLastSelectedType(EIKRigSelectionType::Hierarchy);
+}
+
+void SIKRigHierarchy::OnItemDoubleClicked(TSharedPtr<FIKRigTreeElement> InItem)
+{
+	if (TreeView->IsItemExpanded(InItem))
+	{
+		SetExpansionRecursive(InItem, false, false);
+	}
+	else
+	{
+		SetExpansionRecursive(InItem, false, true);
+	}
+}
+
+void SIKRigHierarchy::OnSetExpansionRecursive(TSharedPtr<FIKRigTreeElement> InItem, bool bShouldBeExpanded)
+{
+	SetExpansionRecursive(InItem, false, bShouldBeExpanded);
+}
+
+void SIKRigHierarchy::SetExpansionRecursive(
+	TSharedPtr<FIKRigTreeElement> InElement,
+	bool bTowardsParent,
+    bool bShouldBeExpanded)
+{
+	TreeView->SetItemExpansion(InElement, bShouldBeExpanded);
+    
+    if (bTowardsParent)
+    {
+    	if (InElement->Parent.Get())
+    	{
+    		SetExpansionRecursive(InElement->Parent, bTowardsParent, bShouldBeExpanded);
+    	}
+    }
+    else
+    {
+    	for (int32 ChildIndex = 0; ChildIndex < InElement->Children.Num(); ++ChildIndex)
+    	{
+    		SetExpansionRecursive(InElement->Children[ChildIndex], bTowardsParent, bShouldBeExpanded);
+    	}
+    }
+}
+
+FReply SIKRigHierarchy::OnKeyDown(const FGeometry& MyGeometry, const FKeyEvent& InKeyEvent)
+{
+	const FKey Key = InKeyEvent.GetKey();
+
+	const TSharedPtr<FIKRigEditorController> Controller = EditorController.Pin();
+	if (!Controller.IsValid())
+	{
+		return FReply::Handled();
+	}
+
+	if (CommandList.IsValid() && CommandList->ProcessCommandBindings(InKeyEvent))
+	{
+		return FReply::Handled();
+	}
+	
+	return FReply::Unhandled();
+}
+
+#undef LOCTEXT_NAMESPACE
\ No newline at end of file
diff --git a/Engine/Plugins/Animation/IKRig/Source/IKRigEditor/Private/RigEditor/SIKRigRetargetChainList.cpp b/Engine/Plugins/Animation/IKRig/Source/IKRigEditor/Private/RigEditor/SIKRigRetargetChainList.cpp
--- a/Engine/Plugins/Animation/IKRig/Source/IKRigEditor/Private/RigEditor/SIKRigRetargetChainList.cpp
+++ b/Engine/Plugins/Animation/IKRig/Source/IKRigEditor/Private/RigEditor/SIKRigRetargetChainList.cpp
@@ -604,6 +604,9 @@ void SIKRigRetargetChainList::RefreshView()
 	}
 
 	ListView->RequestListRefresh();
+
+	// refresh the tree to show updated chain column info
+	Controller->RefreshTreeView();
 }
 
 TSharedRef<ITableRow> SIKRigRetargetChainList::MakeListRowWidget(
@@ -640,7 +643,7 @@ FReply SIKRigRetargetChainList::OnKeyDown(const FGeometry& MyGeometry, const FKe
 			return FReply::Unhandled();
 		}
 
-		UIKRigController* AssetController = Controller->AssetController;
+		const UIKRigController* AssetController = Controller->AssetController;
 		AssetController->RemoveRetargetChain(SelectedItems[0]->ChainName);
 
 		RefreshView();
@@ -692,7 +695,7 @@ void SIKRigRetargetChainList::SortChainList()
 	}
 }
 
-void SIKRigRetargetChainList::MirrorSelectedChains()
+void SIKRigRetargetChainList::MirrorSelectedChains() const
 {
 	const TArray<FName> SelectedChainNames = GetSelectedChains();
 	if (SelectedChainNames.IsEmpty())
@@ -713,7 +716,7 @@ void SIKRigRetargetChainList::MirrorSelectedChains()
 		}
 		
 		TArray<int32> BonesInChainIndices;
-		const bool bIsChainValid = IKRigSkeleton.GetBonesInChain(*Chain, BonesInChainIndices);
+		const bool bIsChainValid = IKRigSkeleton.ValidateChainAndGetBones(*Chain, BonesInChainIndices);
 		if (!bIsChainValid)
 		{
 			continue;
diff --git a/Engine/Plugins/Animation/IKRig/Source/IKRigEditor/Private/RigEditor/SIKRigSkeleton.cpp b/Engine/Plugins/Animation/IKRig/Source/IKRigEditor/Private/RigEditor/SIKRigSkeleton.cpp
deleted file mode 100644
diff --git a/Engine/Plugins/Animation/IKRig/Source/IKRigEditor/Public/RetargetEditor/IKRetargetDefaultMode.h b/Engine/Plugins/Animation/IKRig/Source/IKRigEditor/Public/RetargetEditor/IKRetargetDefaultMode.h
--- a/Engine/Plugins/Animation/IKRig/Source/IKRigEditor/Public/RetargetEditor/IKRetargetDefaultMode.h
+++ b/Engine/Plugins/Animation/IKRig/Source/IKRigEditor/Public/RetargetEditor/IKRetargetDefaultMode.h
@@ -55,8 +55,7 @@ public:
 	/** END FEdMode interface */
 
 private:
-	
-	bool ComponentSelectionOverride(const UPrimitiveComponent* InComponent) const;
+	void RenderDebugProxies(FPrimitiveDrawInterface* PDI, const FIKRetargetEditorController* Controller) const;
 	static void ApplyOffsetToMeshTransform(const FVector& Offset, USceneComponent* Component);
 
 	// the skeleton currently being edited
diff --git a/Engine/Plugins/Animation/IKRig/Source/IKRigEditor/Public/RetargetEditor/IKRetargetEditPoseMode.h b/Engine/Plugins/Animation/IKRig/Source/IKRigEditor/Public/RetargetEditor/IKRetargetEditPoseMode.h
--- a/Engine/Plugins/Animation/IKRig/Source/IKRigEditor/Public/RetargetEditor/IKRetargetEditPoseMode.h
+++ b/Engine/Plugins/Animation/IKRig/Source/IKRigEditor/Public/RetargetEditor/IKRetargetEditPoseMode.h
@@ -81,11 +81,6 @@ private:
 
 	// get the scale and offset associated with the currently edited skeletal mesh component
 	void GetEditedComponentScaleAndOffset(float& OutScale, FVector& OutOffset) const;
-	
-	// render skeleton in viewport
-	void RenderSkeleton(
-		FPrimitiveDrawInterface* PDI,
-		const FIKRetargetEditorController* Controller);
 
 	UE::Widget::EWidgetMode CurrentWidgetMode;
 
diff --git a/Engine/Plugins/Animation/IKRig/Source/IKRigEditor/Public/RetargetEditor/IKRetargetEditorController.h b/Engine/Plugins/Animation/IKRig/Source/IKRigEditor/Public/RetargetEditor/IKRetargetEditorController.h
--- a/Engine/Plugins/Animation/IKRig/Source/IKRigEditor/Public/RetargetEditor/IKRetargetEditorController.h
+++ b/Engine/Plugins/Animation/IKRig/Source/IKRigEditor/Public/RetargetEditor/IKRetargetEditorController.h
@@ -39,6 +39,15 @@ enum class ESelectionEdit : uint8
 	Replace	// replace selection entirely
 };
 
+enum class ERetargetSelectionType : uint8
+{
+	BONE,
+	CHAIN,
+	MESH,
+	ROOT,
+	NONE
+};
+
 /** a home for cross-widget communication to synchronize state across all tabs and viewport */
 class FIKRetargetEditorController : public TSharedFromThis<FIKRetargetEditorController>, FGCObject
 {
@@ -53,13 +62,13 @@ public:
 	/** Bind callbacks to this IK Rig */
 	void BindToIKRigAsset(UIKRigDefinition* InIKRig);
 	/** callback when IK Rig asset requires reinitialization */
-	void OnIKRigNeedsInitialized(UIKRigDefinition* ModifiedIKRig) const;
+	void HandleIKRigNeedsInitialized(UIKRigDefinition* ModifiedIKRig) const;
 	/** callback when IK Rig asset's retarget chain has been renamed */
-	void OnRetargetChainRenamed(UIKRigDefinition* ModifiedIKRig, FName OldName, FName NewName) const;
+	void HandleRetargetChainRenamed(UIKRigDefinition* ModifiedIKRig, FName OldName, FName NewName) const;
 	/** callback when IK Rig asset's retarget chain has been removed */
-	void OnRetargetChainRemoved(UIKRigDefinition* ModifiedIKRig, const FName& InChainRemoved) const;
+	void HandleRetargetChainRemoved(UIKRigDefinition* ModifiedIKRig, const FName& InChainRemoved) const;
 	/** callback when IK Retargeter asset requires reinitialization */
-	void OnRetargeterNeedsInitialized(UIKRetargeter* Retargeter) const;
+	void HandleRetargeterNeedsInitialized(UIKRetargeter* Retargeter) const;
 	
 	/** all modifications to the data model should go through this controller */
 	TObjectPtr<UIKRetargeterController> AssetController;
@@ -164,13 +173,13 @@ public:
 	void EditBoneSelection(
 		const TArray<FName>& InBoneNames,
 		ESelectionEdit EditMode,
-		const bool bFromHierarchyView);
-	const TArray<FName>& GetSelectedBones() const {return SelectedBones; };
+		const bool bFromHierarchyView = false);
+	const TArray<FName>& GetSelectedBones() const {return SelectedBoneNames[CurrentlyEditingSourceOrTarget]; };
 	/* END bone selection */
 	
 	/** SELECTION - MESHES (viewport view) */
 	void SetSelectedMesh(UPrimitiveComponent* InComponent);
-	UPrimitiveComponent* GetSelectedMesh();
+	UPrimitiveComponent* GetSelectedMesh() const;
 	void AddOffsetToMeshComponent(const FVector& Offset, USceneComponent* MeshComponent) const;
 	/* END mesh selection */
 
@@ -249,6 +258,9 @@ public:
 	virtual FString GetReferencerName() const override { return TEXT("Retarget Editor"); };
 	/** END FGCObject interface */
 
+	/** render the skeleton's in the viewport (either source or target) */
+	void RenderSkeleton(FPrimitiveDrawInterface* PDI, ERetargetSourceOrTarget SourceOrTarget) const;
+
 private:
 
 	/** asset properties tab */
@@ -274,11 +286,11 @@ private:
 	bool bIsRootSelected;
 	UPrimitiveComponent* SelectedMesh;
 	TArray<FName> SelectedChains;
-	TArray<FName> SelectedBones;
+	TMap<ERetargetSourceOrTarget, TArray<FName>> SelectedBoneNames;
+	ERetargetSelectionType LastSelectedItem;
 	UPROPERTY()
 	TMap<FName,TObjectPtr<UIKRetargetBoneDetails>> AllBoneDetails;
 
-
 	/** to remove delegates when editor is closed */
 	FDelegateHandle ReInitDelegateHandle;
 	UIKRigDefinition* BoundToIKRig;
diff --git a/Engine/Plugins/Animation/IKRig/Source/IKRigEditor/Public/RetargetEditor/IKRetargetHitProxies.h b/Engine/Plugins/Animation/IKRig/Source/IKRigEditor/Public/RetargetEditor/IKRetargetHitProxies.h
--- a/Engine/Plugins/Animation/IKRig/Source/IKRigEditor/Public/RetargetEditor/IKRetargetHitProxies.h
+++ b/Engine/Plugins/Animation/IKRig/Source/IKRigEditor/Public/RetargetEditor/IKRetargetHitProxies.h
@@ -4,6 +4,7 @@
 
 #include "CoreMinimal.h"
 #include "HitProxies.h"
+#include "IKRetargetAnimInstanceProxy.h"
 
 // allow bone selection to edit retarget pose
 struct HIKRetargetEditorBoneProxy : public HHitProxy
@@ -11,10 +12,17 @@ struct HIKRetargetEditorBoneProxy : public HHitProxy
 	DECLARE_HIT_PROXY();
 
 	FName BoneName;
+	int32 BoneIndex;
+	ERetargetSourceOrTarget SourceOrTarget;
 	
-	HIKRetargetEditorBoneProxy(const FName& InBoneName)
-		: HHitProxy(HPP_World)
-		, BoneName(InBoneName) {}
+	HIKRetargetEditorBoneProxy(
+		const FName& InBoneName,
+		const int32 InBoneIndex,
+		ERetargetSourceOrTarget InSourceOrTarget)
+		: HHitProxy(HPP_World),
+		BoneName(InBoneName),
+		BoneIndex(InBoneIndex),
+		SourceOrTarget(InSourceOrTarget){}
 
 	virtual EMouseCursor::Type GetMouseCursor()
 	{
@@ -27,7 +35,6 @@ struct HIKRetargetEditorBoneProxy : public HHitProxy
 	}
 };
 
-
 // select chains/goals to edit chain settings
 struct HIKRetargetEditorChainProxy : public HHitProxy
 {
diff --git a/Engine/Plugins/Animation/IKRig/Source/IKRigEditor/Public/RetargetEditor/IKRetargeterController.h b/Engine/Plugins/Animation/IKRig/Source/IKRigEditor/Public/RetargetEditor/IKRetargeterController.h
--- a/Engine/Plugins/Animation/IKRig/Source/IKRigEditor/Public/RetargetEditor/IKRetargeterController.h
+++ b/Engine/Plugins/Animation/IKRig/Source/IKRigEditor/Public/RetargetEditor/IKRetargeterController.h
@@ -63,8 +63,10 @@ public:
 	void SetSourceChainForTargetChain(URetargetChainSettings* ChainMap, FName SourceChainToMapTo) const;
 	/** Get read-only access to the list of chain mappings */
 	const TArray<TObjectPtr<URetargetChainSettings>>& GetChainMappings() const;
-	/** Get whether the given chain has a valid IK goal assigned to it */
-	bool IsChainRunningIK(const TObjectPtr<URetargetChainSettings> ChainSettings) const;
+	/** Get the name of the goal assigned to a chain */
+	FName GetChainGoal(const TObjectPtr<URetargetChainSettings> ChainSettings) const;
+	/** Get whether the given chain's IK goal is connected to a solver */
+	bool IsChainGoalConnectedToASolver(const FName& GoalName) const;
 	/** END RETARGET CHAIN MAPPING */
 
 	/** RETARGET POSE EDITING
diff --git a/Engine/Plugins/Animation/IKRig/Source/IKRigEditor/Public/RetargetEditor/SIKRetargetAssetBrowser.h b/Engine/Plugins/Animation/IKRig/Source/IKRigEditor/Public/RetargetEditor/SIKRetargetAssetBrowser.h
--- a/Engine/Plugins/Animation/IKRig/Source/IKRigEditor/Public/RetargetEditor/SIKRetargetAssetBrowser.h
+++ b/Engine/Plugins/Animation/IKRig/Source/IKRigEditor/Public/RetargetEditor/SIKRetargetAssetBrowser.h
@@ -29,6 +29,7 @@ private:
 	
 	void OnAssetDoubleClicked(const FAssetData& AssetData);
 	bool OnShouldFilterAsset(const struct FAssetData& AssetData);
+	TSharedPtr<SWidget> OnGetAssetContextMenu(const TArray<FAssetData>& SelectedAssets) const;
 
 	/** Used to get the currently selected assets */
 	FGetCurrentSelectionDelegate GetCurrentSelectionDelegate;
diff --git a/Engine/Plugins/Animation/IKRig/Source/IKRigEditor/Public/RigEditor/IKRigController.h b/Engine/Plugins/Animation/IKRig/Source/IKRigEditor/Public/RigEditor/IKRigController.h
--- a/Engine/Plugins/Animation/IKRig/Source/IKRigEditor/Public/RigEditor/IKRigController.h
+++ b/Engine/Plugins/Animation/IKRig/Source/IKRigEditor/Public/RigEditor/IKRigController.h
@@ -93,7 +93,7 @@ public:
 	/** The the Bone associated with the given Goal. */
 	FName GetBoneForGoal(const FName& GoalName) const;
 	/** The the Goal associated with the given Bone (may be null. */
-	const UIKRigEffectorGoal* GetGoalForBone(const FName& BoneName) const;
+	UIKRigEffectorGoal* GetGoalForBone(const FName& BoneName) const;
 	/** Connect the given Goal to the given Solver. This creates an "Effector" with settings specific to this Solver.*/
 	bool ConnectGoalToSolver(const UIKRigEffectorGoal& Goal, int32 SolverIndex) const;
 	/** Disconnect the given Goal from the given Solver. This removes the Effector that associates the Goal with the Solver.*/
@@ -106,8 +106,8 @@ public:
 	int32 GetGoalIndex(const FName& InGoalName, const ENameCase CompareMethod = ENameCase::IgnoreCase) const;
 	/** Get the name of Goal at the given index. */
 	FName GetGoalName(const int32& GoalIndex) const;
-	/** Get read-only access to the list of Goals. */
-	const TArray<UIKRigEffectorGoal*>& GetAllGoals() const;
+	/** Get access to the list of Goals. */
+	TArray<UIKRigEffectorGoal*>& GetAllGoals() const;
 	/** Get read-only access to the Goal at the given index. */
 	const UIKRigEffectorGoal* GetGoal(int32 GoalIndex) const;
 	/** Get read-write access to the Goal with the given name. */
@@ -178,8 +178,16 @@ public:
 	void SortRetargetChains() const;
 	/** Make unique name for a retargeting bone chain. Adds a numbered suffix to make it unique.*/
 	FName GetUniqueRetargetChainName(const FName& NameToMakeUnique) const;
-	/** Returns true if this is a valid chain. Produces array of bone indices between start and end (inclusive). */
-	bool ValidateChain(const FName& ChainName, TSet<int32>& OutChainIndices) const;
+	/** Returns true if this is a valid chain. Produces array of bone indices between start and end (inclusive).
+	 * Optionally provide a runtime skeleton from an IKRigProcessor to get indices for a running instance (otherwise uses stored hierarchy in asset)*/
+	bool ValidateChain(
+		const FName& ChainName,
+		const FIKRigSkeleton* OptionalSkeleton,
+		TSet<int32>& OutChainIndices) const;
+	/** Get the name of the retarget chain that contains the given Bone. Returns NAME_None if Bone not in a Chain. */
+	FName GetRetargetChainFromBone(const FName& BoneName, const FIKRigSkeleton* OptionalSkeleton) const;
+	/** Get the name of the retarget chain that contains the given Goal. Returns NAME_None if Goal not in a Chain. */
+	FName GetRetargetChainFromGoal(const FName& GoalName) const;
 	/** END retarget chains */
 
 	// force all currently connected processors to reinitialize using latest asset state
diff --git a/Engine/Plugins/Animation/IKRig/Source/IKRigEditor/Public/RigEditor/IKRigEditorController.h b/Engine/Plugins/Animation/IKRig/Source/IKRigEditor/Public/RigEditor/IKRigEditorController.h
--- a/Engine/Plugins/Animation/IKRig/Source/IKRigEditor/Public/RigEditor/IKRigEditorController.h
+++ b/Engine/Plugins/Animation/IKRig/Source/IKRigEditor/Public/RigEditor/IKRigEditorController.h
@@ -15,7 +15,7 @@ class SIKRigOutputLog;
 class UIKRigAnimInstance;
 class FIKRigEditorToolkit;
 class SIKRigSolverStack;
-class SIKRigSkeleton;
+class SIKRigHierarchy;
 class UIKRigController;
 class FSolverStackElement;
 class FIKRigTreeElement;
@@ -115,6 +115,8 @@ public:
 
 	/** get the currently active processor running the IK Rig in the editor */
 	UIKRigProcessor* GetIKRigProcessor() const;
+	/** get the currently running IKRig skeleton (if there is a running processor) */
+	const FIKRigSkeleton* GetCurrentIKRigSkeleton() const;
 
 	/** callback when IK Rig requires re-initialization */
 	void OnIKRigNeedsInitialized(UIKRigDefinition* ModifiedIKRig);
@@ -159,17 +161,6 @@ public:
 	
 	/** create new retarget chains from selected bones (or single empty chain if no selection) */
 	void CreateNewRetargetChains();
-	/** show user the new retarget chain they are about to create (provides option to edit name) */
-	void PromptToAddNewRetargetChain(const FBoneChain& BoneChain) const;
-	/** prompt user if they want to add a goal to a newly mirrored chain */
-	void PromptToAddGoalToNewlyMirroredChain(const FBoneChain& OldChain, const FBoneChain& MirroredChain) const;
-	
-	/** determine if the element is connected to the selected solver */
-	bool IsElementConnectedToSolver(TSharedRef<FIKRigTreeElement> TreeElement, int32 SolverIndex);
-	/** determine if the element is connected to ANY solver */
-	bool IsElementConnectedToAnySolver(TSharedRef<FIKRigTreeElement> TreeElement);
-	/** determine if the element is an excluded bone*/
-	bool IsElementExcludedBone(TSharedRef<FIKRigTreeElement> TreeElement);
 	
 	/** show single transform of bone in details view */
 	void ShowDetailsForBone(const FName BoneName) const;
@@ -191,11 +182,24 @@ public:
 	/** set details tab view */
 	void SetDetailsView(const TSharedPtr<class IDetailsView>& InDetailsView);
 	/** set skeleton tab view */
-	void SetSkeletonsView(const TSharedPtr<SIKRigSkeleton>& InSkeletonView){ SkeletonView = InSkeletonView; };
+	void SetHierarchyView(const TSharedPtr<SIKRigHierarchy>& InSkeletonView){ SkeletonView = InSkeletonView; };
 	/** set solver stack tab view */
 	void SetSolverStackView(const TSharedPtr<SIKRigSolverStack>& InSolverStackView){ SolverStackView = InSolverStackView; };
 	/** set retargeting tab view */
 	void SetRetargetingView(const TSharedPtr<SIKRigRetargetChainList>& InRetargetingView){ RetargetingView = InRetargetingView; };
+	/** set output log view */
+	void SetOutputLogView(const TSharedPtr<SIKRigOutputLog>& InOutputLogView){ OutputLogView = InOutputLogView; };
+
+	/** right after importing a skeleton, we ask user what solver they want to use */
+	bool PromptToAddDefaultSolver() const;
+	/** show user the new retarget chain they are about to create (provides option to edit name) */
+	void PromptToAddNewRetargetChain(const FBoneChain& BoneChain) const;
+	/** prompt user if they want to add a goal to a newly created chain */
+	UIKRigEffectorGoal* PromptToAddGoalToNewChain(const FBoneChain& BoneChain) const;
+	/** prompt user if they want to add a goal to a newly mirrored chain */
+	void PromptToAddGoalToNewlyMirroredChain(const FBoneChain& OldChain, const FBoneChain& MirroredChain) const;
+	/** right after creating a goal, we ask user if they want it assigned to a retarget chain */
+	void PromptToAddGoalToChain(UIKRigEffectorGoal* NewGoal) const;
 
 	/** play preview animation on running anim instance in editor (before IK) */
 	void PlayAnimationAsset(UAnimationAsset* AssetToPlay);
@@ -237,12 +241,6 @@ public:
 	
 private:
 	
-	/** right after importing a skeleton, we ask user what solver they want to use */
-	bool PromptToAddSolver() const;
-
-	/** right after creating a goal, we ask user if they want it assigned to a retarget chain */
-	void PromptToAddGoalToChain(UIKRigEffectorGoal* NewGoal) const;
-
 	/** Initializes editor's solvers instances */
 	void InitializeSolvers() const;
 
@@ -250,7 +248,7 @@ private:
 	TSharedPtr<IDetailsView> DetailsView;
 
 	/** the skeleton tree view */
-	TSharedPtr<SIKRigSkeleton> SkeletonView;
+	TSharedPtr<SIKRigHierarchy> SkeletonView;
 	
 	/** the solver stack view */
 	TSharedPtr<SIKRigSolverStack> SolverStackView;
diff --git a/Engine/Plugins/Animation/IKRig/Source/IKRigEditor/Public/RigEditor/IKRigToolkit.h b/Engine/Plugins/Animation/IKRig/Source/IKRigEditor/Public/RigEditor/IKRigToolkit.h
--- a/Engine/Plugins/Animation/IKRig/Source/IKRigEditor/Public/RigEditor/IKRigToolkit.h
+++ b/Engine/Plugins/Animation/IKRig/Source/IKRigEditor/Public/RigEditor/IKRigToolkit.h
@@ -10,7 +10,7 @@
 #include "IKRigMode.h"
 #include "IPersonaPreviewScene.h"
 #include "PersonaAssetEditorToolkit.h"
-#include "SIKRigSkeleton.h"
+#include "SIKRigHierarchy.h"
 
 namespace IKRigEditorModes
 {
diff --git a/Engine/Plugins/Animation/IKRig/Source/IKRigEditor/Public/RigEditor/SIKRigAssetBrowser.h b/Engine/Plugins/Animation/IKRig/Source/IKRigEditor/Public/RigEditor/SIKRigAssetBrowser.h
--- a/Engine/Plugins/Animation/IKRig/Source/IKRigEditor/Public/RigEditor/SIKRigAssetBrowser.h
+++ b/Engine/Plugins/Animation/IKRig/Source/IKRigEditor/Public/RigEditor/SIKRigAssetBrowser.h
@@ -26,7 +26,8 @@ private:
 	
 	void OnAssetDoubleClicked(const FAssetData& AssetData);
 	bool OnShouldFilterAsset(const struct FAssetData& AssetData);
-
+	TSharedPtr<SWidget> OnGetAssetContextMenu(const TArray<FAssetData>& SelectedAssets) const;
+	
 	/** Used to get the currently selected assets */
 	FGetCurrentSelectionDelegate GetCurrentSelectionDelegate;
 		
diff --git a/Engine/Plugins/Animation/IKRig/Source/IKRigEditor/Public/RigEditor/SIKRigHierarchy.h b/Engine/Plugins/Animation/IKRig/Source/IKRigEditor/Public/RigEditor/SIKRigHierarchy.h
new file mode 100644
--- /dev/null
+++ b/Engine/Plugins/Animation/IKRig/Source/IKRigEditor/Public/RigEditor/SIKRigHierarchy.h
@@ -0,0 +1,304 @@
+// Copyright Epic Games, Inc. All Rights Reserved.
+
+#pragma once
+
+#include "CoreMinimal.h"
+#include "EditorUndoClient.h"
+#include "SBaseHierarchyTreeView.h"
+#include "SIKRigRetargetChainList.h"
+#include "DragAndDrop/DecoratedDragDropOp.h"
+#include "Widgets/SCompoundWidget.h"
+#include "Widgets/Views/STreeView.h"
+
+#include "SIKRigSolverStack.h"
+
+class FIKRigEditorController;
+class SIKRigHierarchy;
+class FIKRigEditorToolkit;
+class USkeletalMesh;
+class UIKRigBoneDetails;
+
+enum class IKRigTreeElementType { BONE, GOAL, SOLVERGOAL, BONE_SETTINGS };
+
+class FIKRigTreeElement : public TSharedFromThis<FIKRigTreeElement>
+{
+public:
+	
+	FIKRigTreeElement(
+		const FText& InKey,
+		IKRigTreeElementType InType,
+		const TSharedRef<FIKRigEditorController>& InEditorController);
+
+	FText Key;
+	IKRigTreeElementType ElementType;
+
+	/** determines if this item is filtered out of the view */
+	bool bIsHidden = false;
+	TSharedPtr<FIKRigTreeElement> UnFilteredParent;
+	
+	TSharedPtr<FIKRigTreeElement> Parent;
+	TArray<TSharedPtr<FIKRigTreeElement>> Children;
+
+	/** effector meta-data (if it is an effector) */
+	FName EffectorGoalName = NAME_None;
+	int32 EffectorIndex = INDEX_NONE;
+
+	/** bone setting meta-data (if it is a bone setting) */
+	FName BoneSettingBoneName = NAME_None;
+	int32 BoneSettingsSolverIndex = INDEX_NONE;
+
+	/** name of bone if it is one */
+	FName BoneName = NAME_None;
+	
+	/** name of goal if it is one */
+	FName GoalName = NAME_None;
+
+	/** delegate for when the context menu requests a rename */
+	void RequestRename();
+	DECLARE_DELEGATE(FOnRenameRequested);
+	FOnRenameRequested OnRenameRequested;
+
+	/** get the underlying object */
+	TWeakObjectPtr< UObject > GetObject() const;
+
+	/** get the name of the retarget chain this element belongs to (or NAME_None if not in a chain) */
+	FName GetChainName() const;
+
+private:
+	/** centralized editor controls */
+	TWeakPtr<FIKRigEditorController> EditorController;
+
+	/** on demand bone details object */	
+	mutable TObjectPtr<UIKRigBoneDetails> OptionalBoneDetails = nullptr;
+};
+
+class SIKRigHierarchyItem : public SCompoundWidget
+{
+	SLATE_BEGIN_ARGS(SIKRigHierarchyItem) {}
+	SLATE_ARGUMENT(TWeakPtr<FIKRigEditorController>, EditorController)
+	SLATE_ARGUMENT(TWeakPtr<FIKRigTreeElement>, TreeElement)
+	SLATE_ARGUMENT(TWeakPtr<SIKRigHierarchy>, HierarchyView)
+	SLATE_END_ARGS()
+
+	void Construct(const FArguments& InArgs);
+	
+private:
+
+	bool OnVerifyNameChanged(const FText& InText, FText& OutErrorMessage) const;
+	void OnNameCommitted(const FText& InText, ETextCommit::Type InCommitType) const;
+	FText GetName() const;
+	
+	TWeakPtr<FIKRigTreeElement> WeakTreeElement;
+	TWeakPtr<FIKRigEditorController> EditorController;
+	TWeakPtr<SIKRigHierarchy> HierarchyView;
+};
+
+class FIKRigSkeletonDragDropOp : public FDecoratedDragDropOp
+{
+public:
+	DRAG_DROP_OPERATOR_TYPE(FIKRigSkeletonDragDropOp, FDecoratedDragDropOp)
+	static TSharedRef<FIKRigSkeletonDragDropOp> New(TWeakPtr<FIKRigTreeElement> InElement);
+	virtual TSharedPtr<SWidget> GetDefaultDecorator() const override;
+	TWeakPtr<FIKRigTreeElement> Element;
+};
+
+class SIKRigSkeletonRow : public SMultiColumnTableRow<TSharedPtr<FIKRigTreeElement>>
+{
+public:
+
+	SLATE_BEGIN_ARGS(SIKRigSkeletonRow) {}
+	SLATE_ARGUMENT(TSharedPtr<FIKRigEditorController>, EditorController)
+	SLATE_ARGUMENT(TSharedPtr<FIKRigTreeElement>, TreeElement)
+	SLATE_ARGUMENT(TSharedPtr<SIKRigHierarchy>, HierarchyView)
+	SLATE_END_ARGS()
+
+	void Construct(
+		const FArguments& InArgs,
+		const TSharedRef<STableViewBase>& InOwnerTable)
+	{
+		OwnerTable = InOwnerTable;
+		WeakTreeElement = InArgs._TreeElement;
+		EditorController = InArgs._EditorController;
+		HierarchyView = InArgs._HierarchyView;
+		
+		SMultiColumnTableRow<TSharedPtr<FIKRigTreeElement>>::Construct(
+			FSuperRowType::FArguments()
+			.OnDragDetected(this, &SIKRigSkeletonRow::HandleDragDetected)
+			.OnCanAcceptDrop(this, &SIKRigSkeletonRow::HandleCanAcceptDrop)
+			.OnAcceptDrop(this, &SIKRigSkeletonRow::HandleAcceptDrop)
+			,InOwnerTable);
+	}
+
+	/** Overridden from SMultiColumnTableRow.  Generates a widget for this column of the list view. */
+	virtual TSharedRef<SWidget> GenerateWidgetForColumn(const FName& ColumnName) override;
+	/** End SMultiColumnTableRow */
+
+	/** drag and drop */
+	FReply HandleDragDetected(const FGeometry& MyGeometry, const FPointerEvent& MouseEvent);
+	TOptional<EItemDropZone> HandleCanAcceptDrop(const FDragDropEvent& DragDropEvent, EItemDropZone DropZone, TSharedPtr<FIKRigTreeElement> TargetItem);
+	FReply HandleAcceptDrop(const FDragDropEvent& DragDropEvent, EItemDropZone DropZone, TSharedPtr<FIKRigTreeElement> TargetItem);
+	/** END drag and drop */
+
+private:
+
+	TSharedPtr<STableViewBase> OwnerTable;
+	TWeakPtr<FIKRigTreeElement> WeakTreeElement;
+	TWeakPtr<FIKRigEditorController> EditorController;
+	TWeakPtr<SIKRigHierarchy> HierarchyView;
+};
+
+struct FIKRigHierarchyFilterOptions
+{
+	bool bShowAll = true;
+	bool bShowOnlyBones = false;
+	bool bShowOnlyGoals = false;
+	bool bShowOnlyBoneSettings = false;
+	bool bHideUnaffectedBones = false;
+
+	void ResetShowOptions()
+	{
+		bShowAll = false;
+		bShowOnlyBones = false;
+		bShowOnlyGoals = false;
+		bShowOnlyBoneSettings = false;
+	}
+};
+
+typedef SBaseHierarchyTreeView<FIKRigTreeElement> SIKRigSkeletonTreeView;
+
+class SIKRigHierarchy : public SCompoundWidget, public FEditorUndoClient
+{
+public:
+	SLATE_BEGIN_ARGS(SIKRigHierarchy) {}
+	SLATE_END_ARGS()
+
+    ~SIKRigHierarchy() {};
+
+	void Construct(const FArguments& InArgs, TSharedRef<FIKRigEditorController> InEditorController);
+
+	/** selection state queries */
+	static bool IsBoneInSelection(TArray<TSharedPtr<FIKRigTreeElement>>& SelectedBoneItems, const FName& BoneName);
+	void GetSelectedBones(TArray<TSharedPtr<FIKRigTreeElement>>& OutBoneItems) const;
+	void GetSelectedBoneNames(TArray<FName>& OutSelectedBoneNames) const;
+	void GetSelectedGoals(TArray<TSharedPtr<FIKRigTreeElement>>& OutSelectedGoals) const;
+	int32 GetNumSelectedGoals();
+	void GetSelectedGoalNames(TArray<FName>& OutSelectedGoalNames) const;
+	bool IsGoalSelected(const FName& GoalName);
+	void AddSelectedItemFromViewport(
+		const FName& ItemName,
+		IKRigTreeElementType ItemType,
+		const bool bReplace);
+	void GetSelectedBoneChains(TArray<FBoneChain>& OutChains);
+	TArray<TSharedPtr<FIKRigTreeElement>> GetSelectedItems() const;
+	bool HasSelectedItems() const;
+	/** END selection state queries */
+
+	/** determine if the element is connected to the selected solver */
+	bool IsElementConnectedToSolver(TWeakPtr<FIKRigTreeElement> TreeElement, int32 SolverIndex);
+	/** determine if the element is connected to ANY solver */
+	bool IsElementConnectedToAnySolver(TWeakPtr<FIKRigTreeElement> TreeElement);
+	/** determine if the element is an excluded bone*/
+	bool IsElementExcludedBone(TWeakPtr<FIKRigTreeElement> TreeElement);
+
+private:
+	/** SWidget interface */
+	virtual FReply OnKeyDown(const FGeometry& MyGeometry, const FKeyEvent& InKeyEvent);
+	/** END SWidget interface */
+	
+	/** Bind commands that this widget handles */
+	void BindCommands();
+
+	/** directly, programmatically add item to selection */
+	void AddItemToSelection(const TSharedPtr<FIKRigTreeElement>& InItem);
+	/** directly, programmatically add item to selection */
+	void RemoveItemFromSelection(const TSharedPtr<FIKRigTreeElement>& InItem);
+	/** directly, programmatically replace item in selection */
+	void ReplaceItemInSelection(const FText& OldName, const FText& NewName);
+
+	/** creating / renaming / deleting goals */
+	void HandleNewGoal() const;
+	bool CanAddNewGoal() const;
+	void HandleRenameGoal() const;
+	bool CanRenameGoal() const;
+	void HandleDeleteElement();
+	bool CanDeleteElement() const;
+	/** END creating / renaming / deleting goals */
+
+	/** connecting/disconnecting goals to solvers */
+	void HandleConnectGoalToSolver();
+	void HandleDisconnectGoalFromSolver();
+	bool CanConnectGoalToSolvers() const;
+	bool CanDisconnectGoalFromSolvers() const;
+	void ConnectSelectedGoalsToSelectedSolvers(bool bConnect);
+	int32 GetNumSelectedGoalToSolverConnections(bool bCountOnlyConnected) const;
+	/** END connecting/disconnecting goals to solvers */
+
+	/** setting root bone */
+	void HandleSetRootBoneOnSolvers();
+	bool CanSetRootBoneOnSolvers();
+	/** END setting root bone */
+
+	/** setting end bone */
+	void HandleSetEndBoneOnSolvers();
+	bool CanSetEndBoneOnSolvers() const;
+	bool HasEndBoneCompatibleSolverSelected() const;
+	/** END setting end bone */
+
+	/** per-bone settings */
+	void HandleAddBoneSettings();
+	bool CanAddBoneSettings();
+	void HandleRemoveBoneSettings();
+	bool CanRemoveBoneSettings();
+	/** END per-bone settings */
+	
+	/** exclude/include bones */
+	void HandleExcludeBone();
+	bool CanExcludeBone();
+	void HandleIncludeBone();
+	bool CanIncludeBone();
+	/** END exclude/include bones */
+
+	/** retarget chains */
+	void HandleNewRetargetChain();
+	bool CanAddNewRetargetChain();
+	void HandleSetRetargetRoot();
+	bool CanSetRetargetRoot();
+	void HandleClearRetargetRoot();
+	bool CanClearRetargetRoot();
+	/** END retarget chains */
+
+	/** centralized editor controls (facilitate cross-communication between multiple UI elements)*/
+	TWeakPtr<FIKRigEditorController> EditorController;
+	
+	/** command list we bind to */
+	TSharedPtr<FUICommandList> CommandList;
+	
+	/** tree view widget */
+	TSharedPtr<SIKRigSkeletonTreeView> TreeView;
+	TArray<TSharedPtr<FIKRigTreeElement>> RootElements;
+	TArray<TSharedPtr<FIKRigTreeElement>> AllElements;
+
+	/** filtering the tree with search box */
+	TSharedRef<SWidget> CreateFilterMenuWidget();
+	void OnFilterTextChanged(const FText& SearchText);
+	TSharedPtr<FTextFilterExpressionEvaluator> TextFilter;
+	FIKRigHierarchyFilterOptions FilterOptions;
+	
+	/** tree view callbacks */
+	void RefreshTreeView(bool IsInitialSetup=false);
+	void HandleGetChildrenForTree(TSharedPtr<FIKRigTreeElement> InItem, TArray<TSharedPtr<FIKRigTreeElement>>& OutChildren);
+	void OnSelectionChanged(TSharedPtr<FIKRigTreeElement> Selection, ESelectInfo::Type SelectInfo);
+	TSharedRef< SWidget > CreateAddNewMenu();
+	TSharedPtr< SWidget > CreateContextMenu();
+	void OnItemClicked(TSharedPtr<FIKRigTreeElement> InItem);
+	void OnItemDoubleClicked(TSharedPtr<FIKRigTreeElement> InItem);
+	void OnSetExpansionRecursive(TSharedPtr<FIKRigTreeElement> InItem, bool bShouldBeExpanded);
+	void SetExpansionRecursive(TSharedPtr<FIKRigTreeElement> InElement, bool bTowardsParent, bool bShouldBeExpanded);
+	void FillContextMenu(FMenuBuilder& MenuBuilder);
+	/** END tree view callbacks */
+
+	friend SIKRigHierarchyItem;
+	friend FIKRigEditorController;
+	friend SIKRigSolverStack;
+	friend SIKRigRetargetChainList;
+};
diff --git a/Engine/Plugins/Animation/IKRig/Source/IKRigEditor/Public/RigEditor/SIKRigRetargetChainList.h b/Engine/Plugins/Animation/IKRig/Source/IKRigEditor/Public/RigEditor/SIKRigRetargetChainList.h
--- a/Engine/Plugins/Animation/IKRig/Source/IKRigEditor/Public/RigEditor/SIKRigRetargetChainList.h
+++ b/Engine/Plugins/Animation/IKRig/Source/IKRigEditor/Public/RigEditor/SIKRigRetargetChainList.h
@@ -172,7 +172,7 @@ private:
 	void SortChainList();
 
 	/** mirror the selected chains (right-click menu callback) */
-	void MirrorSelectedChains();
+	void MirrorSelectedChains() const;
 
 	friend SIKRigRetargetChainRow;
 	friend FIKRigEditorController;
diff --git a/Engine/Plugins/Animation/IKRig/Source/IKRigEditor/Public/RigEditor/SIKRigSkeleton.h b/Engine/Plugins/Animation/IKRig/Source/IKRigEditor/Public/RigEditor/SIKRigSkeleton.h
deleted file mode 100644
