Fixing static analysis warnings in SIKRigHierarchy.cpp

diff --git a/Engine/Plugins/Animation/IKRig/Source/IKRigEditor/Private/RigEditor/SIKRigHierarchy.cpp b/Engine/Plugins/Animation/IKRig/Source/IKRigEditor/Private/RigEditor/SIKRigHierarchy.cpp
--- a/Engine/Plugins/Animation/IKRig/Source/IKRigEditor/Private/RigEditor/SIKRigHierarchy.cpp
+++ b/Engine/Plugins/Animation/IKRig/Source/IKRigEditor/Private/RigEditor/SIKRigHierarchy.cpp
@@ -271,7 +271,7 @@ TSharedRef<SWidget> SIKRigSkeletonRow::GenerateWidgetForColumn(const FName& Colu
 	if (ColumnName == IKRigChainColumnName)
 	{
 		const FName ChainName = WeakTreeElement.Pin()->GetChainName();
-		const FText ChainText = ChainName == NAME_None ? FText::FromString("") : FText::FromName(ChainName);
+		const FText ChainText = (ChainName == NAME_None) ? FText::FromString("") : FText::FromName(ChainName);
 		const FSlateFontInfo Font = ChainName != "(Retarget Root)" ? FAppStyle::Get().GetFontStyle("NormalFont") : FAppStyle::Get().GetFontStyle("BoldFont");
 		
 		TSharedPtr< SHorizontalBox > RowBox;
