Disabled editing of retarget root translation in XY plane.

diff --git a/Engine/Plugins/Animation/IKRig/Source/IKRig/Private/Retargeter/IKRetargeter.cpp b/Engine/Plugins/Animation/IKRig/Source/IKRig/Private/Retargeter/IKRetargeter.cpp
--- a/Engine/Plugins/Animation/IKRig/Source/IKRig/Private/Retargeter/IKRetargeter.cpp
+++ b/Engine/Plugins/Animation/IKRig/Source/IKRig/Private/Retargeter/IKRetargeter.cpp
@@ -240,11 +240,17 @@ FVector FIKRetargetPose::GetRootTranslationDelta() const
 void FIKRetargetPose::SetRootTranslationDelta(const FVector& TranslationDelta)
 {
 	RootTranslationOffset = TranslationDelta;
+	// only allow vertical offset of root in retarget pose
+	RootTranslationOffset.X = 0.0f;
+	RootTranslationOffset.Y = 0.0f;
 }
 
 void FIKRetargetPose::AddToRootTranslationDelta(const FVector& TranslateDelta)
 {
 	RootTranslationOffset += TranslateDelta;
+	// only allow vertical offset of root in retarget pose
+	RootTranslationOffset.X = 0.0f;
+	RootTranslationOffset.Y = 0.0f;
 }
 
 void FIKRetargetPose::SortHierarchically(const FIKRigSkeleton& Skeleton)
