Auto populate source mesh preview when making a new IK retargeter asset.

diff --git a/Engine/Plugins/Animation/IKRig/Source/IKRigEditor/Private/RetargetEditor/IKRetargeterController.cpp b/Engine/Plugins/Animation/IKRig/Source/IKRigEditor/Private/RetargetEditor/IKRetargeterController.cpp
--- a/Engine/Plugins/Animation/IKRig/Source/IKRigEditor/Private/RetargetEditor/IKRetargeterController.cpp
+++ b/Engine/Plugins/Animation/IKRig/Source/IKRigEditor/Private/RetargetEditor/IKRetargeterController.cpp
@@ -39,6 +39,7 @@ UIKRetargeter* UIKRetargeterController::GetAsset() const
 void UIKRetargeterController::SetSourceIKRig(UIKRigDefinition* SourceIKRig)
 {
 	Asset->SourceIKRigAsset = SourceIKRig;
+	Asset->SourcePreviewMesh = Asset->SourceIKRigAsset->PreviewSkeletalMesh;
 }
 
 USkeletalMesh* UIKRetargeterController::GetPreviewMesh(const ERetargetSourceOrTarget& SourceOrTarget) const
