Preventing chain renaming from adding incorrect suffix when undone.

diff --git a/Engine/Plugins/Animation/IKRig/Source/IKRigEditor/Private/RigEditor/SIKRigRetargetChainList.cpp b/Engine/Plugins/Animation/IKRig/Source/IKRigEditor/Private/RigEditor/SIKRigRetargetChainList.cpp
--- a/Engine/Plugins/Animation/IKRig/Source/IKRigEditor/Private/RigEditor/SIKRigRetargetChainList.cpp
+++ b/Engine/Plugins/Animation/IKRig/Source/IKRigEditor/Private/RigEditor/SIKRigRetargetChainList.cpp
@@ -257,11 +257,17 @@ void SIKRigRetargetChainRow::OnRenameChain(const FText& InText, ETextCommit::Typ
 		return; 
 	}
 
+	// prevent setting name to the same name
 	const FName OldName = ChainElement.Pin()->ChainName;
 	const FName NewName = FName(*InText.ToString());
 	if (OldName == NewName)
 	{
-		// most reliable way to catch multiple commits
+		return;
+	}
+
+	// prevent multiple commits when pressing enter
+	if (CommitType == ETextCommit::OnEnter)
+	{
 		return;
 	}
 	
