Fixing global scale factor incorrectly calculated in retargeter.

diff --git a/Engine/Plugins/Animation/IKRig/Source/IKRig/Private/Retargeter/IKRetargetProcessor.cpp b/Engine/Plugins/Animation/IKRig/Source/IKRig/Private/Retargeter/IKRetargetProcessor.cpp
--- a/Engine/Plugins/Animation/IKRig/Source/IKRig/Private/Retargeter/IKRetargetProcessor.cpp
+++ b/Engine/Plugins/Animation/IKRig/Source/IKRig/Private/Retargeter/IKRetargetProcessor.cpp
@@ -1232,6 +1232,10 @@ bool FRootRetargeter::InitializeTarget(
 	Target.InitialHeight = TargetInitialTransform.GetTranslation().Z;
 	Target.InitialRotation = TargetInitialTransform.GetRotation();
 	Target.InitialPosition = TargetInitialTransform.GetTranslation();
+
+	// initialize the global scale factor
+	const float ScaleFactor = Source.InitialHeightInverse * Target.InitialHeight;
+	GlobalScaleFactor.Set(ScaleFactor, ScaleFactor, ScaleFactor);
 	
 	return true;
 }
diff --git a/Engine/Plugins/Animation/IKRig/Source/IKRig/Public/Retargeter/IKRetargetProcessor.h b/Engine/Plugins/Animation/IKRig/Source/IKRig/Public/Retargeter/IKRetargetProcessor.h
--- a/Engine/Plugins/Animation/IKRig/Source/IKRig/Public/Retargeter/IKRetargetProcessor.h
+++ b/Engine/Plugins/Animation/IKRig/Source/IKRig/Public/Retargeter/IKRetargetProcessor.h
@@ -176,8 +176,11 @@ struct FRootRetargeter
 
 	FVector GetGlobalScaleVector() const
 	{
-		return FVector(Settings.ScaleHorizontal, Settings.ScaleHorizontal, Settings.ScaleVertical);
+		return GlobalScaleFactor * FVector(Settings.ScaleHorizontal, Settings.ScaleHorizontal, Settings.ScaleVertical);
 	}
+
+private:
+	FVector GlobalScaleFactor;
 };
 
 struct FPoleVectorMatcher
