Fixing "reset to default" for chain settings in retargeter.

diff --git a/Engine/Plugins/Animation/IKRig/Source/IKRig/Private/Retargeter/IKRetargetSettings.cpp b/Engine/Plugins/Animation/IKRig/Source/IKRig/Private/Retargeter/IKRetargetSettings.cpp
--- a/Engine/Plugins/Animation/IKRig/Source/IKRig/Private/Retargeter/IKRetargetSettings.cpp
+++ b/Engine/Plugins/Animation/IKRig/Source/IKRig/Private/Retargeter/IKRetargetSettings.cpp
@@ -3,11 +3,48 @@
 #include "Retargeter/IKRetargetSettings.h"
 #include "Retargeter/IKRetargeter.h"
 
+bool FTargetChainSpeedPlantSettings::operator==(const FTargetChainSpeedPlantSettings& Other) const
+{
+	return EnableSpeedPlanting == Other.EnableSpeedPlanting
+	&& SpeedCurveName == Other.SpeedCurveName
+	&& FMath::IsNearlyEqualByULP(SpeedThreshold, Other.SpeedThreshold)
+	&& FMath::IsNearlyEqualByULP(UnplantStiffness, Other.UnplantStiffness)
+	&& FMath::IsNearlyEqualByULP(UnplantCriticalDamping, Other.UnplantCriticalDamping);
+}
+
+bool FTargetChainFKSettings::operator==(const FTargetChainFKSettings& Other) const
+{
+	return EnableFK == Other.EnableFK
+	   && RotationMode == Other.RotationMode
+	   && FMath::IsNearlyEqualByULP(RotationAlpha,Other.RotationAlpha)
+	   && TranslationMode == Other.TranslationMode
+	   && FMath::IsNearlyEqualByULP(TranslationAlpha, Other.TranslationAlpha)
+	   && PoleVectorMatching == Other.PoleVectorMatching
+	   && FMath::IsNearlyEqualByULP(PoleVectorOffset, Other.PoleVectorOffset);
+}
+
+bool FTargetChainIKSettings::operator==(const FTargetChainIKSettings& Other) const
+{
+	return EnableIK == Other.EnableIK
+	&& FMath::IsNearlyEqualByULP(BlendToSource, Other.BlendToSource)
+	&& BlendToSourceWeights.Equals(Other.BlendToSourceWeights)
+	&& StaticOffset.Equals(Other.StaticOffset)
+	&& StaticLocalOffset.Equals(Other.StaticLocalOffset)
+	&& StaticRotationOffset.Equals(Other.StaticRotationOffset)
+	&& FMath::IsNearlyEqualByULP(Extension, Other.Extension)
+	&& bAffectedByIKWarping == Other.bAffectedByIKWarping;
+}
+
 void FTargetChainSettings::CopySettingsFromAsset(const URetargetChainSettings* AssetChainSettings)
 {
 	*this = AssetChainSettings->Settings;
 }
 
+bool FTargetChainSettings::operator==(const FTargetChainSettings& Other) const
+{
+	return FK == Other.FK && IK == Other.IK && SpeedPlanting == Other.SpeedPlanting;
+}
+
 void FTargetRootSettings::CopySettingsFromAsset(const URetargetRootSettings* AssetRootSettings)
 {
 	*this = AssetRootSettings->Settings;
diff --git a/Engine/Plugins/Animation/IKRig/Source/IKRig/Public/Retargeter/IKRetargetSettings.h b/Engine/Plugins/Animation/IKRig/Source/IKRig/Public/Retargeter/IKRetargetSettings.h
--- a/Engine/Plugins/Animation/IKRig/Source/IKRig/Public/Retargeter/IKRetargetSettings.h
+++ b/Engine/Plugins/Animation/IKRig/Source/IKRig/Public/Retargeter/IKRetargetSettings.h
@@ -26,9 +26,9 @@ enum class ERetargetRotationMode : uint8
 };
 
 USTRUCT(BlueprintType)
-struct FTargetChainSpeedPlantSettings
+struct IKRIG_API FTargetChainSpeedPlantSettings
 {
-	GENERATED_USTRUCT_BODY()
+	GENERATED_BODY()
 
 	/** The name of the curve on the source animation that contains the speed of the end effector bone.*/
 	UPROPERTY(EditAnywhere, BlueprintReadWrite, Category = "Plant IK by Speed", meta = (ClampMin = "0.0", ClampMax = "100.0", UIMin = "0.0", UIMax = "100.0"))
@@ -50,10 +50,12 @@ struct FTargetChainSpeedPlantSettings
 	// How much damping to apply to the spring (0 means no damping, 1 means critically damped which means no oscillation)
 	UPROPERTY(EditAnywhere, BlueprintReadWrite, Category = "Plant IK by Speed", meta = (ClampMin = "0.0", ClampMax = "10.0", UIMin = "0.0", UIMax = "1.0"))
 	float UnplantCriticalDamping = 1.0f;
+
+	bool operator==(const FTargetChainSpeedPlantSettings& Other) const;
 };
 
 USTRUCT(BlueprintType)
-struct FTargetChainFKSettings
+struct IKRIG_API FTargetChainFKSettings
 {
 	GENERATED_BODY()
 
@@ -98,12 +100,14 @@ struct FTargetChainFKSettings
 	/** Range +/- 180. Default 0. An angular offset, in degrees, for the pole direction of the chain. */
 	UPROPERTY(EditAnywhere, BlueprintReadWrite, Category = "FK Adjustments", meta = (UIMin = "-180.0", UIMax = "180.0"))
 	float PoleVectorOffset = 0.0f;
+
+	bool operator==(const FTargetChainFKSettings& Other) const;
 };
 
 USTRUCT(BlueprintType)
-struct FTargetChainIKSettings
+struct IKRIG_API FTargetChainIKSettings
 {
-	GENERATED_USTRUCT_BODY()
+	GENERATED_BODY()
 
 	/** Whether to modify the location of the IK goal on this chain. Default is true.
 	 * NOTE: This only has an effect if the chain has an IK Goal assigned to it in the Target IK Rig asset.
@@ -146,12 +150,14 @@ struct FTargetChainIKSettings
 	/** Is this IK goal affected by the stride warping (in Global Settings). Typically this is true for all feet, but not for hands.  */
 	UPROPERTY(EditAnywhere, BlueprintReadWrite, Category = Warping, meta = (DisplayName = "Affected by IK Warping"))
 	bool bAffectedByIKWarping = true;
+
+	bool operator==(const FTargetChainIKSettings& Other) const;
 };
 
 USTRUCT(BlueprintType)
-struct FTargetChainSettings
+struct IKRIG_API FTargetChainSettings
 {
-	GENERATED_USTRUCT_BODY()
+	GENERATED_BODY()
 
 public:
 
@@ -171,12 +177,14 @@ public:
 	FTargetChainSpeedPlantSettings SpeedPlanting;
 	
 	void CopySettingsFromAsset(const URetargetChainSettings* AssetChainSettings);
+
+	bool operator==(const FTargetChainSettings& Other) const;
 };
 
 USTRUCT(BlueprintType)
 struct IKRIG_API FTargetRootSettings
 {
-	GENERATED_USTRUCT_BODY()
+	GENERATED_BODY()
 
 public:
 
@@ -259,7 +267,7 @@ enum class EWarpingDirectionSource
 USTRUCT(BlueprintType)
 struct IKRIG_API FRetargetGlobalSettings
 {
-	GENERATED_USTRUCT_BODY()
+	GENERATED_BODY()
 	
 	/** When false, the motion of the Retarget Root bone is not copied from the source. Useful for debugging issues with the root settings.
 	 * Note: the retargeting order is: Root > FK > IK
diff --git a/Engine/Plugins/Animation/IKRig/Source/IKRigEditor/Private/RetargetEditor/SIKRetargetChainMapList.cpp b/Engine/Plugins/Animation/IKRig/Source/IKRigEditor/Private/RetargetEditor/SIKRetargetChainMapList.cpp
--- a/Engine/Plugins/Animation/IKRig/Source/IKRigEditor/Private/RetargetEditor/SIKRetargetChainMapList.cpp
+++ b/Engine/Plugins/Animation/IKRig/Source/IKRigEditor/Private/RetargetEditor/SIKRetargetChainMapList.cpp
@@ -16,6 +16,7 @@
 #include "SSearchableComboBox.h"
 #include "Misc/TextFilterExpressionEvaluator.h"
 #include "Widgets/Input/SSearchBox.h"
+#include "Retargeter/IKRetargetSettings.h"
 
 #define LOCTEXT_NAMESPACE "SIKRigRetargetChains"
 
@@ -130,9 +131,9 @@ TSharedRef< SWidget > SIKRetargetChainMapRow::GenerateWidgetForColumn(const FNam
 		.Padding(3.0f, 1.0f)
 		[
 			SNew(SButton)
-			//.OnClicked(this, &SIKRetargetChainMapRow::OnResetToDefaultClicked) TODO 
-			//.Visibility(this, &SIKRetargetChainMapRow::GetResetToDefaultVisibility) TODO 
-			.ToolTipText(LOCTEXT("ResetToDefaultToolTip", "Reset to Default"))
+			.OnClicked(this, &SIKRetargetChainMapRow::OnResetToDefaultClicked) 
+			.Visibility(this, &SIKRetargetChainMapRow::GetResetToDefaultVisibility) 
+			.ToolTipText(LOCTEXT("ResetChainToDefaultToolTip", "Reset Chain Settings to Default"))
 			.ButtonStyle(FAppStyle::Get(), "NoBorder")
 			.Content()
 			[
@@ -147,6 +148,19 @@ TSharedRef< SWidget > SIKRetargetChainMapRow::GenerateWidgetForColumn(const FNam
 	return SNullWidget::NullWidget;
 }
 
+FReply SIKRetargetChainMapRow::OnResetToDefaultClicked()
+{
+	ChainMapElement.Pin()->ChainMap->Settings = FTargetChainSettings();
+	return FReply::Handled();
+}
+
+EVisibility SIKRetargetChainMapRow::GetResetToDefaultVisibility() const
+{
+	const FTargetChainSettings& DefaultSettings = FTargetChainSettings();
+	const FTargetChainSettings& Settings = ChainMapElement.Pin()->ChainMap->Settings;
+	return (Settings == DefaultSettings) ? EVisibility::Hidden : EVisibility::Visible;
+}
+
 void SIKRetargetChainMapRow::OnSourceChainComboSelectionChanged(TSharedPtr<FString> InName, ESelectInfo::Type SelectInfo)
 {
 	UIKRetargeterController* RetargeterController = ChainMapList.Pin()->GetRetargetController();
diff --git a/Engine/Plugins/Animation/IKRig/Source/IKRigEditor/Public/RetargetEditor/SIKRetargetChainMapList.h b/Engine/Plugins/Animation/IKRig/Source/IKRigEditor/Public/RetargetEditor/SIKRetargetChainMapList.h
--- a/Engine/Plugins/Animation/IKRig/Source/IKRigEditor/Public/RetargetEditor/SIKRetargetChainMapList.h
+++ b/Engine/Plugins/Animation/IKRig/Source/IKRigEditor/Public/RetargetEditor/SIKRetargetChainMapList.h
@@ -60,6 +60,10 @@ public:
 	void OnSourceChainComboSelectionChanged(TSharedPtr<FString> InName, ESelectInfo::Type SelectInfo);
 	
 private:
+
+	FReply OnResetToDefaultClicked();
+
+	EVisibility GetResetToDefaultVisibility() const;
 	
 	FText GetSourceChainName() const;
 
