Fixing issue, adding retarget chain in IK Rig asset does not update the retarget asset editor that is using it.

diff --git a/Engine/Plugins/Animation/IKRig/Source/IKRigEditor/Private/RetargetEditor/IKRetargetEditorController.cpp b/Engine/Plugins/Animation/IKRig/Source/IKRigEditor/Private/RetargetEditor/IKRetargetEditorController.cpp
--- a/Engine/Plugins/Animation/IKRig/Source/IKRigEditor/Private/RetargetEditor/IKRetargetEditorController.cpp
+++ b/Engine/Plugins/Animation/IKRig/Source/IKRigEditor/Private/RetargetEditor/IKRetargetEditorController.cpp
@@ -56,13 +56,13 @@ void FIKRetargetEditorController::Initialize(TSharedPtr<FIKRetargetEditor> InEdi
 void FIKRetargetEditorController::Close() const
 {
 	AssetController->OnRetargeterNeedsInitialized().Remove(ReInitDelegateHandle);
-
 	
 	if (UIKRigController* IKRigController = UIKRigController::GetIKRigController(BoundToIKRig))
 	{
 		IKRigController->OnIKRigNeedsInitialized().Remove(ReInitDelegateHandle);
-		IKRigController->OnRetargetChainRenamed().Remove(RenameChainDelegateHandle);
+		IKRigController->OnRetargetChainAdded().Remove(AddedChainDelegateHandle);
 		IKRigController->OnRetargetChainRemoved().Remove(RemoveChainDelegateHandle);
+		IKRigController->OnRetargetChainRenamed().Remove(RenameChainDelegateHandle);
 	}
 }
 
@@ -79,8 +79,9 @@ void FIKRetargetEditorController::BindToIKRigAsset(UIKRigDefinition* InIKRig)
 	if (!Controller->OnIKRigNeedsInitialized().IsBoundToObject(this))
 	{
 		ReInitIKDelegateHandle = Controller->OnIKRigNeedsInitialized().AddSP(this, &FIKRetargetEditorController::HandleIKRigNeedsInitialized);
-		RenameChainDelegateHandle = Controller->OnRetargetChainRenamed().AddSP(this, &FIKRetargetEditorController::HandleRetargetChainRenamed);
+		AddedChainDelegateHandle = Controller->OnRetargetChainAdded().AddSP(this, &FIKRetargetEditorController::HandleRetargetChainAdded);
 		RemoveChainDelegateHandle = Controller->OnRetargetChainRemoved().AddSP(this, &FIKRetargetEditorController::HandleRetargetChainRemoved);
+		RenameChainDelegateHandle = Controller->OnRetargetChainRenamed().AddSP(this, &FIKRetargetEditorController::HandleRetargetChainRenamed);
 	}
 }
 
@@ -101,10 +102,16 @@ void FIKRetargetEditorController::HandleIKRigNeedsInitialized(UIKRigDefinition*
 	HandleRetargeterNeedsInitialized(Retargeter);
 }
 
+void FIKRetargetEditorController::HandleRetargetChainAdded(UIKRigDefinition* ModifiedIKRig) const
+{
+	check(ModifiedIKRig)
+	AssetController->OnRetargetChainAdded(ModifiedIKRig);
+	RefreshAllViews();
+}
+
 void FIKRetargetEditorController::HandleRetargetChainRenamed(UIKRigDefinition* ModifiedIKRig, FName OldName, FName NewName) const
 {
 	check(ModifiedIKRig)
-	
 	AssetController->OnRetargetChainRenamed(ModifiedIKRig, OldName, NewName);
 }
 
diff --git a/Engine/Plugins/Animation/IKRig/Source/IKRigEditor/Private/RetargetEditor/IKRetargeterController.cpp b/Engine/Plugins/Animation/IKRig/Source/IKRigEditor/Private/RetargetEditor/IKRetargeterController.cpp
--- a/Engine/Plugins/Animation/IKRig/Source/IKRigEditor/Private/RetargetEditor/IKRetargeterController.cpp
+++ b/Engine/Plugins/Animation/IKRig/Source/IKRigEditor/Private/RetargetEditor/IKRetargeterController.cpp
@@ -260,6 +260,20 @@ void UIKRetargeterController::AutoMapChains() const
 	BroadcastNeedsReinitialized();
 }
 
+void UIKRetargeterController::OnRetargetChainAdded(UIKRigDefinition* IKRig) const
+{
+	const bool bIsTargetRig = IKRig == Asset->GetTargetIKRig();
+	if (!bIsTargetRig)
+	{
+		// if a source chain is added, it will simply be available as a new option, no need to reinitialize until it's used
+		return;
+	}
+
+	// add the new chain to the mapping data
+	constexpr bool bForceReinitialization = true;
+	CleanChainMapping(bForceReinitialization);
+}
+
 void UIKRetargeterController::OnRetargetChainRenamed(UIKRigDefinition* IKRig, FName OldChainName, FName NewChainName) const
 {
 	const bool bIsSourceRig = IKRig == Asset->GetSourceIKRig();
diff --git a/Engine/Plugins/Animation/IKRig/Source/IKRigEditor/Private/RigEditor/IKRigController.cpp b/Engine/Plugins/Animation/IKRig/Source/IKRigEditor/Private/RigEditor/IKRigController.cpp
--- a/Engine/Plugins/Animation/IKRig/Source/IKRigEditor/Private/RigEditor/IKRigController.cpp
+++ b/Engine/Plugins/Animation/IKRig/Source/IKRigEditor/Private/RigEditor/IKRigController.cpp
@@ -161,6 +161,8 @@ void UIKRigController::AddRetargetChain(const FBoneChain& BoneChain) const
 	Asset->Modify();
 	
 	Asset->RetargetDefinition.BoneChains.Emplace(ChainToAdd);
+	
+	RetargetChainAdded.Broadcast(Asset);
 	BroadcastNeedsReinitialized();
 }
 
diff --git a/Engine/Plugins/Animation/IKRig/Source/IKRigEditor/Public/RetargetEditor/IKRetargetEditorController.h b/Engine/Plugins/Animation/IKRig/Source/IKRigEditor/Public/RetargetEditor/IKRetargetEditorController.h
--- a/Engine/Plugins/Animation/IKRig/Source/IKRigEditor/Public/RetargetEditor/IKRetargetEditorController.h
+++ b/Engine/Plugins/Animation/IKRig/Source/IKRigEditor/Public/RetargetEditor/IKRetargetEditorController.h
@@ -63,6 +63,8 @@ public:
 	void BindToIKRigAsset(UIKRigDefinition* InIKRig);
 	/** callback when IK Rig asset requires reinitialization */
 	void HandleIKRigNeedsInitialized(UIKRigDefinition* ModifiedIKRig) const;
+	/** callback when IK Rig asset's retarget chain's have been added or removed*/
+	void HandleRetargetChainAdded(UIKRigDefinition* ModifiedIKRig) const;
 	/** callback when IK Rig asset's retarget chain has been renamed */
 	void HandleRetargetChainRenamed(UIKRigDefinition* ModifiedIKRig, FName OldName, FName NewName) const;
 	/** callback when IK Rig asset's retarget chain has been removed */
@@ -295,6 +297,7 @@ private:
 	FDelegateHandle ReInitDelegateHandle;
 	UIKRigDefinition* BoundToIKRig;
 	FDelegateHandle ReInitIKDelegateHandle;
+	FDelegateHandle AddedChainDelegateHandle;
 	FDelegateHandle RenameChainDelegateHandle;
 	FDelegateHandle RemoveChainDelegateHandle;
 };
diff --git a/Engine/Plugins/Animation/IKRig/Source/IKRigEditor/Public/RetargetEditor/IKRetargeterController.h b/Engine/Plugins/Animation/IKRig/Source/IKRigEditor/Public/RetargetEditor/IKRetargeterController.h
--- a/Engine/Plugins/Animation/IKRig/Source/IKRigEditor/Public/RetargetEditor/IKRetargeterController.h
+++ b/Engine/Plugins/Animation/IKRig/Source/IKRigEditor/Public/RetargetEditor/IKRetargeterController.h
@@ -55,6 +55,8 @@ public:
 	void CleanChainMapping(const bool bForceReinitialization=true) const;
 	/** Use fuzzy string search to find "best" Source chain to map to each Target chain */
 	void AutoMapChains() const;
+	/** Callback when IK Rig chain is added or removed. */
+	void OnRetargetChainAdded(UIKRigDefinition* IKRig) const;
 	/** Callback when IK Rig chain is renamed. Retains existing mappings using the new name */
 	void OnRetargetChainRenamed(UIKRigDefinition* IKRig, FName OldChainName, FName NewChainName) const;
 	/** Callback when IK Rig chain is removed. */
diff --git a/Engine/Plugins/Animation/IKRig/Source/IKRigEditor/Public/RigEditor/IKRigController.h b/Engine/Plugins/Animation/IKRig/Source/IKRigEditor/Public/RigEditor/IKRigController.h
--- a/Engine/Plugins/Animation/IKRig/Source/IKRigEditor/Public/RigEditor/IKRigController.h
+++ b/Engine/Plugins/Animation/IKRig/Source/IKRigEditor/Public/RigEditor/IKRigController.h
@@ -202,20 +202,24 @@ private:
 	DECLARE_MULTICAST_DELEGATE_OneParam(FOnIKRigNeedsInitialized, UIKRigDefinition*);
 	FOnIKRigNeedsInitialized IKRigNeedsInitialized;
 
-	/** Called whenever a retarget chain is renamed.*/
-	DECLARE_MULTICAST_DELEGATE_ThreeParams(FOnRetargetChainRenamed, UIKRigDefinition*, FName /*old name*/, FName /*new name*/);
-	FOnRetargetChainRenamed RetargetChainRenamed;
+	/** Called whenever a retarget chain is added.*/
+	DECLARE_MULTICAST_DELEGATE_OneParam(FOnRetargetChainAdded, UIKRigDefinition*);
+	FOnRetargetChainAdded RetargetChainAdded;
 
 	/** Called whenever a retarget chain is removed.*/
 	DECLARE_MULTICAST_DELEGATE_TwoParams(FOnRetargetChainRemoved, UIKRigDefinition*, const FName& /*chain name*/);
 	FOnRetargetChainRemoved RetargetChainRemoved;
 
+	/** Called whenever a retarget chain is renamed.*/
+	DECLARE_MULTICAST_DELEGATE_ThreeParams(FOnRetargetChainRenamed, UIKRigDefinition*, FName /*old name*/, FName /*new name*/);
+	FOnRetargetChainRenamed RetargetChainRenamed;
+
 public:
 	
 	FOnIKRigNeedsInitialized& OnIKRigNeedsInitialized(){ return IKRigNeedsInitialized; };
-
-	FOnRetargetChainRenamed& OnRetargetChainRenamed(){ return RetargetChainRenamed; };
+	FOnRetargetChainAdded& OnRetargetChainAdded(){ return RetargetChainAdded; };
 	FOnRetargetChainRemoved& OnRetargetChainRemoved(){ return RetargetChainRemoved; };
+	FOnRetargetChainRenamed& OnRetargetChainRenamed(){ return RetargetChainRenamed; };
 
 private:
 
