Fixes crash when editing IK Rig and RTG assets together

diff --git a/Engine/Plugins/Animation/IKRig/Source/IKRigEditor/Private/RetargetEditor/IKRetargetEditor.cpp b/Engine/Plugins/Animation/IKRig/Source/IKRigEditor/Private/RetargetEditor/IKRetargetEditor.cpp
--- a/Engine/Plugins/Animation/IKRig/Source/IKRigEditor/Private/RetargetEditor/IKRetargetEditor.cpp
+++ b/Engine/Plugins/Animation/IKRig/Source/IKRigEditor/Private/RetargetEditor/IKRetargetEditor.cpp
@@ -430,8 +430,7 @@ void FIKRetargetEditor::OnFinishedChangingDetails(const FPropertyChangedEvent& P
 	if (bTargetIKRigChanged || bSourceIKRigChanged)
 	{
 		EditorController->ClearOutputLog();
-		EditorController->BindToIKRigAsset(AssetController->GetAsset()->GetTargetIKRigWriteable());
-		EditorController->BindToIKRigAsset(AssetController->GetAsset()->GetSourceIKRigWriteable());
+		EditorController->BindToIKRigAssets(AssetController->GetAsset());
 		EditorController->AssetController->CleanChainMapping();
 		EditorController->AssetController->AutoMapChains();
 	}
diff --git a/Engine/Plugins/Animation/IKRig/Source/IKRigEditor/Private/RetargetEditor/IKRetargetEditorController.cpp b/Engine/Plugins/Animation/IKRig/Source/IKRigEditor/Private/RetargetEditor/IKRetargetEditorController.cpp
--- a/Engine/Plugins/Animation/IKRig/Source/IKRigEditor/Private/RetargetEditor/IKRetargetEditorController.cpp
+++ b/Engine/Plugins/Animation/IKRig/Source/IKRigEditor/Private/RetargetEditor/IKRetargetEditorController.cpp
@@ -26,6 +26,27 @@
 #define LOCTEXT_NAMESPACE "IKRetargetEditorController"
 
 
+FBoundIKRig::FBoundIKRig(UIKRigDefinition* InIKRig, const FIKRetargetEditorController& InController)
+{
+	check(IKRig);
+	IKRig = InIKRig;
+	UIKRigController* IKRigController = UIKRigController::GetIKRigController(InIKRig);
+	ReInitIKDelegateHandle = IKRigController->OnIKRigNeedsInitialized().AddSP(&InController, &FIKRetargetEditorController::HandleIKRigNeedsInitialized);
+	AddedChainDelegateHandle = IKRigController->OnRetargetChainAdded().AddSP(&InController, &FIKRetargetEditorController::HandleRetargetChainAdded);
+	RemoveChainDelegateHandle = IKRigController->OnRetargetChainRemoved().AddSP(&InController, &FIKRetargetEditorController::HandleRetargetChainRemoved);
+	RenameChainDelegateHandle = IKRigController->OnRetargetChainRenamed().AddSP(&InController, &FIKRetargetEditorController::HandleRetargetChainRenamed);
+}
+
+void FBoundIKRig::UnBind() const
+{
+	check(IKRig);
+	UIKRigController* IKRigController = UIKRigController::GetIKRigController(IKRig);
+	IKRigController->OnIKRigNeedsInitialized().Remove(ReInitIKDelegateHandle);
+	IKRigController->OnRetargetChainAdded().Remove(AddedChainDelegateHandle);
+	IKRigController->OnRetargetChainRemoved().Remove(RemoveChainDelegateHandle);
+	IKRigController->OnRetargetChainRenamed().Remove(RenameChainDelegateHandle);
+}
+
 void FIKRetargetEditorController::Initialize(TSharedPtr<FIKRetargetEditor> InEditor, UIKRetargeter* InAsset)
 {
 	Editor = InEditor;
@@ -46,57 +67,51 @@ void FIKRetargetEditorController::Initialize(TSharedPtr<FIKRetargetEditor> InEdi
 	AssetController->CleanPoseLists(bForceReinitialization);
 
 	// bind callbacks when SOURCE or TARGET IK Rigs are modified
-	BindToIKRigAsset(AssetController->GetAsset()->GetSourceIKRigWriteable());
-	BindToIKRigAsset(AssetController->GetAsset()->GetTargetIKRigWriteable());
+	BindToIKRigAssets(AssetController->GetAsset());
 
 	// bind callback when retargeter needs reinitialized
-	ReInitDelegateHandle = AssetController->OnRetargeterNeedsInitialized().AddSP(this, &FIKRetargetEditorController::HandleRetargeterNeedsInitialized);
+	RetargeterReInitDelegateHandle = AssetController->OnRetargeterNeedsInitialized().AddSP(this, &FIKRetargetEditorController::HandleRetargeterNeedsInitialized);
 }
 
-void FIKRetargetEditorController::Close() const
+void FIKRetargetEditorController::Close()
 {
-	AssetController->OnRetargeterNeedsInitialized().Remove(ReInitDelegateHandle);
-	
-	if (UIKRigController* IKRigController = UIKRigController::GetIKRigController(BoundToIKRig))
+	AssetController->OnRetargeterNeedsInitialized().Remove(RetargeterReInitDelegateHandle);
+
+	for (const FBoundIKRig& BoundIKRig : BoundIKRigs)
 	{
-		IKRigController->OnIKRigNeedsInitialized().Remove(ReInitDelegateHandle);
-		IKRigController->OnRetargetChainAdded().Remove(AddedChainDelegateHandle);
-		IKRigController->OnRetargetChainRemoved().Remove(RemoveChainDelegateHandle);
-		IKRigController->OnRetargetChainRenamed().Remove(RenameChainDelegateHandle);
+		BoundIKRig.UnBind();
 	}
 }
 
-void FIKRetargetEditorController::BindToIKRigAsset(UIKRigDefinition* InIKRig)
+void FIKRetargetEditorController::BindToIKRigAssets(UIKRetargeter* InAsset)
 {
-	BoundToIKRig = InIKRig;
-	
-	if (!InIKRig)
+	// unbind previously bound IK Rigs
+	for (const FBoundIKRig& BoundIKRig : BoundIKRigs)
 	{
-		return;
+		BoundIKRig.UnBind();
 	}
-	
-	UIKRigController* Controller = UIKRigController::GetIKRigController(InIKRig);
-	if (!Controller->OnIKRigNeedsInitialized().IsBoundToObject(this))
+
+	BoundIKRigs.Empty();
+
+	if (InAsset->GetSourceIKRigWriteable())
+	{
+		BoundIKRigs.Emplace(FBoundIKRig(InAsset->GetSourceIKRigWriteable(), *this));
+	}
+
+	if (InAsset->GetTargetIKRigWriteable())
 	{
-		ReInitIKDelegateHandle = Controller->OnIKRigNeedsInitialized().AddSP(this, &FIKRetargetEditorController::HandleIKRigNeedsInitialized);
-		AddedChainDelegateHandle = Controller->OnRetargetChainAdded().AddSP(this, &FIKRetargetEditorController::HandleRetargetChainAdded);
-		RemoveChainDelegateHandle = Controller->OnRetargetChainRemoved().AddSP(this, &FIKRetargetEditorController::HandleRetargetChainRemoved);
-		RenameChainDelegateHandle = Controller->OnRetargetChainRenamed().AddSP(this, &FIKRetargetEditorController::HandleRetargetChainRenamed);
+		BoundIKRigs.Emplace(FBoundIKRig(InAsset->GetTargetIKRigWriteable(), *this));
 	}
 }
 
 void FIKRetargetEditorController::HandleIKRigNeedsInitialized(UIKRigDefinition* ModifiedIKRig) const
 {
 	UIKRetargeter* Retargeter = AssetController->GetAsset();
-	
 	check(ModifiedIKRig && Retargeter)
 	 
 	const bool bIsSource = ModifiedIKRig == Retargeter->GetSourceIKRig();
 	const bool bIsTarget = ModifiedIKRig == Retargeter->GetTargetIKRig();
-	if (!(bIsSource || bIsTarget))
-	{
-		return;
-	}
+	check(bIsSource || bIsTarget);
 
 	// the target anim instance has the RetargetPoseFromMesh node which needs reinitialized with new asset version
 	HandleRetargeterNeedsInitialized(Retargeter);
diff --git a/Engine/Plugins/Animation/IKRig/Source/IKRigEditor/Public/RetargetEditor/IKRetargetEditorController.h b/Engine/Plugins/Animation/IKRig/Source/IKRigEditor/Public/RetargetEditor/IKRetargetEditorController.h
--- a/Engine/Plugins/Animation/IKRig/Source/IKRigEditor/Public/RetargetEditor/IKRetargetEditorController.h
+++ b/Engine/Plugins/Animation/IKRig/Source/IKRigEditor/Public/RetargetEditor/IKRetargetEditorController.h
@@ -48,6 +48,18 @@ enum class ERetargetSelectionType : uint8
 	NONE
 };
 
+struct FBoundIKRig
+{
+	FBoundIKRig(UIKRigDefinition* InIKRig, const FIKRetargetEditorController& InController);
+	void UnBind() const;
+	
+	UIKRigDefinition* IKRig;
+	FDelegateHandle ReInitIKDelegateHandle;
+	FDelegateHandle AddedChainDelegateHandle;
+	FDelegateHandle RenameChainDelegateHandle;
+	FDelegateHandle RemoveChainDelegateHandle;
+};
+
 /** a home for cross-widget communication to synchronize state across all tabs and viewport */
 class FIKRetargetEditorController : public TSharedFromThis<FIKRetargetEditorController>, FGCObject
 {
@@ -58,9 +70,9 @@ public:
 	/** Initialize the editor */
 	void Initialize(TSharedPtr<FIKRetargetEditor> InEditor, UIKRetargeter* InAsset);
 	/** Close the editor */
-	void Close() const;
+	void Close();
 	/** Bind callbacks to this IK Rig */
-	void BindToIKRigAsset(UIKRigDefinition* InIKRig);
+	void BindToIKRigAssets(UIKRetargeter* InAsset);
 	/** callback when IK Rig asset requires reinitialization */
 	void HandleIKRigNeedsInitialized(UIKRigDefinition* ModifiedIKRig) const;
 	/** callback when IK Rig asset's retarget chain's have been added or removed*/
@@ -71,6 +83,7 @@ public:
 	void HandleRetargetChainRemoved(UIKRigDefinition* ModifiedIKRig, const FName& InChainRemoved) const;
 	/** callback when IK Retargeter asset requires reinitialization */
 	void HandleRetargeterNeedsInitialized(UIKRetargeter* Retargeter) const;
+	FDelegateHandle RetargeterReInitDelegateHandle;
 	
 	/** all modifications to the data model should go through this controller */
 	TObjectPtr<UIKRetargeterController> AssetController;
@@ -293,12 +306,7 @@ private:
 	UPROPERTY()
 	TMap<FName,TObjectPtr<UIKRetargetBoneDetails>> AllBoneDetails;
 
-	/** to remove delegates when editor is closed */
-	FDelegateHandle ReInitDelegateHandle;
-	UIKRigDefinition* BoundToIKRig;
-	FDelegateHandle ReInitIKDelegateHandle;
-	FDelegateHandle AddedChainDelegateHandle;
-	FDelegateHandle RenameChainDelegateHandle;
-	FDelegateHandle RemoveChainDelegateHandle;
+	/** ik rigs bound to this editor (will receive callbacks when requiring reinitialization*/
+	TArray<FBoundIKRig> BoundIKRigs;
 };
 
