Fixed bug, reset chain settings not updating details panel.
Added undo transaction for reseting settings.

diff --git a/Engine/Plugins/Animation/IKRig/Source/IKRigEditor/Private/RetargetEditor/SIKRetargetChainMapList.cpp b/Engine/Plugins/Animation/IKRig/Source/IKRigEditor/Private/RetargetEditor/SIKRetargetChainMapList.cpp
--- a/Engine/Plugins/Animation/IKRig/Source/IKRigEditor/Private/RetargetEditor/SIKRetargetChainMapList.cpp
+++ b/Engine/Plugins/Animation/IKRig/Source/IKRigEditor/Private/RetargetEditor/SIKRetargetChainMapList.cpp
@@ -2,7 +2,7 @@
 
 #include "RetargetEditor/SIKRetargetChainMapList.h"
 
-
+#include "ScopedTransaction.h"
 #include "Framework/Commands/UICommandList.h"
 #include "Framework/MultiBox/MultiBoxBuilder.h"
 #include "IPersonaToolkit.h"
@@ -150,7 +150,7 @@ TSharedRef< SWidget > SIKRetargetChainMapRow::GenerateWidgetForColumn(const FNam
 
 FReply SIKRetargetChainMapRow::OnResetToDefaultClicked()
 {
-	ChainMapElement.Pin()->ChainMap->Settings = FTargetChainSettings();
+	ChainMapList.Pin()->ResetChainSettings(ChainMapElement.Pin()->ChainMap.Get());
 	return FReply::Handled();
 }
 
@@ -390,6 +390,14 @@ void SIKRetargetChainMapList::ClearSelection() const
 	ListView->ClearSelection();
 }
 
+void SIKRetargetChainMapList::ResetChainSettings(URetargetChainSettings* ChainMap) const
+{
+	FScopedTransaction Transaction(LOCTEXT("ResetChainSettings", "Reset Retarget Chain Settings"));
+	ChainMap->Modify();
+	ChainMap->Settings = FTargetChainSettings();
+	EditorController.Pin()->RefreshDetailsView();
+}
+
 UIKRetargeterController* SIKRetargetChainMapList::GetRetargetController() const
 {
 	const TSharedPtr<FIKRetargetEditorController> Controller = EditorController.Pin();
diff --git a/Engine/Plugins/Animation/IKRig/Source/IKRigEditor/Public/RetargetEditor/SIKRetargetChainMapList.h b/Engine/Plugins/Animation/IKRig/Source/IKRigEditor/Public/RetargetEditor/SIKRetargetChainMapList.h
--- a/Engine/Plugins/Animation/IKRig/Source/IKRigEditor/Public/RetargetEditor/SIKRetargetChainMapList.h
+++ b/Engine/Plugins/Animation/IKRig/Source/IKRigEditor/Public/RetargetEditor/SIKRetargetChainMapList.h
@@ -98,6 +98,8 @@ public:
 	void Construct(const FArguments& InArgs, TSharedRef<FIKRetargetEditorController> InEditorController);
 
 	void ClearSelection() const;
+
+	void ResetChainSettings(URetargetChainSettings* Settings) const;
 	
 private:
 	
