Fixing static analysis warnings in IKRetargetEditorController.cpp

diff --git a/Engine/Plugins/Animation/IKRig/Source/IKRigEditor/Private/RetargetEditor/IKRetargetEditorController.cpp b/Engine/Plugins/Animation/IKRig/Source/IKRigEditor/Private/RetargetEditor/IKRetargetEditorController.cpp
--- a/Engine/Plugins/Animation/IKRig/Source/IKRigEditor/Private/RetargetEditor/IKRetargetEditorController.cpp
+++ b/Engine/Plugins/Animation/IKRig/Source/IKRigEditor/Private/RetargetEditor/IKRetargetEditorController.cpp
@@ -28,7 +28,7 @@
 
 FBoundIKRig::FBoundIKRig(UIKRigDefinition* InIKRig, const FIKRetargetEditorController& InController)
 {
-	check(IKRig);
+	check(InIKRig);
 	IKRig = InIKRig;
 	UIKRigController* IKRigController = UIKRigController::GetIKRigController(InIKRig);
 	ReInitIKDelegateHandle = IKRigController->OnIKRigNeedsInitialized().AddSP(&InController, &FIKRetargetEditorController::HandleIKRigNeedsInitialized);
