Various fixes to chain mirroring functionality in IK Rig editor.

diff --git a/Engine/Plugins/Animation/IKRig/Source/IKRigEditor/Private/RigEditor/IKRigController.cpp b/Engine/Plugins/Animation/IKRig/Source/IKRigEditor/Private/RigEditor/IKRigController.cpp
--- a/Engine/Plugins/Animation/IKRig/Source/IKRigEditor/Private/RigEditor/IKRigController.cpp
+++ b/Engine/Plugins/Animation/IKRig/Source/IKRigEditor/Private/RigEditor/IKRigController.cpp
@@ -154,7 +154,7 @@ bool UIKRigController::DoesBoneHaveSettings(const FName& BoneName) const
 	return false;
 }
 
-void UIKRigController::AddRetargetChain(const FBoneChain& BoneChain) const
+FName UIKRigController::AddRetargetChain(const FBoneChain& BoneChain) const
 {
 	FBoneChain ChainToAdd = BoneChain;
 	ChainToAdd.ChainName = GetUniqueRetargetChainName(BoneChain.ChainName);
@@ -162,10 +162,12 @@ void UIKRigController::AddRetargetChain(const FBoneChain& BoneChain) const
 	FScopedTransaction Transaction(LOCTEXT("AddRetargetChain_Label", "Add Retarget Chain"));
 	Asset->Modify();
 	
-	Asset->RetargetDefinition.BoneChains.Emplace(ChainToAdd);
+	const int32 NewChainIndex = Asset->RetargetDefinition.BoneChains.Emplace(ChainToAdd);
 	
 	RetargetChainAdded.Broadcast(Asset);
 	BroadcastNeedsReinitialized();
+
+	return Asset->RetargetDefinition.BoneChains[NewChainIndex].ChainName;
 }
 
 bool UIKRigController::RemoveRetargetChain(const FName& ChainName) const
diff --git a/Engine/Plugins/Animation/IKRig/Source/IKRigEditor/Private/RigEditor/IKRigEditorController.cpp b/Engine/Plugins/Animation/IKRig/Source/IKRigEditor/Private/RigEditor/IKRigEditorController.cpp
--- a/Engine/Plugins/Animation/IKRig/Source/IKRigEditor/Private/RigEditor/IKRigEditorController.cpp
+++ b/Engine/Plugins/Animation/IKRig/Source/IKRigEditor/Private/RigEditor/IKRigEditorController.cpp
@@ -433,7 +433,7 @@ void FIKRigEditorController::AddNewGoals(const TArray<FName>& GoalNames, const T
 		}
 
 		// ask user if they want to assign this goal to a chain (if there is one on this bone)
-		PromptToAddGoalToChain(NewGoal);
+		PromptToAssignGoalToChain(NewGoal);
 
 		LastCreatedGoalName = GoalName;
 	}
@@ -839,7 +839,7 @@ void FIKRigEditorController::SetDetailsView(const TSharedPtr<IDetailsView>& InDe
 	ShowEmptyDetails();
 }
 
-void FIKRigEditorController::PromptToAddGoalToChain(UIKRigEffectorGoal* NewGoal) const
+void FIKRigEditorController::PromptToAssignGoalToChain(UIKRigEffectorGoal* NewGoal) const
 {
 	const TArray<FBoneChain>& AllRetargetChains = AssetController->GetRetargetChains();
 	FName ChainToAddGoalTo = NAME_None;
@@ -857,15 +857,15 @@ void FIKRigEditorController::PromptToAddGoalToChain(UIKRigEffectorGoal* NewGoal)
 	}
 
 	TSharedRef<SCustomDialog> AddGoalToChainDialog = SNew(SCustomDialog)
-		.Title(FText(LOCTEXT("AddGoalToNewChainTitle", "Add Goal to Retarget Chain")))
+		.Title(FText(LOCTEXT("AssignGoalToNewChainTitle", "Assign Existing Goal to Retarget Chain")))
 		.Content()
 		[
 			SNew(STextBlock)
-			.Text(FText::Format(LOCTEXT("AddNewGoalToChainLabel", "Add new goal, {0} to retarget chain, {1}?"), FText::FromName(NewGoal->GoalName), FText::FromName(ChainToAddGoalTo)))
+			.Text(FText::Format(LOCTEXT("AssignGoalToChainLabel", "Assign goal, {0} to retarget chain, {1}?"), FText::FromName(NewGoal->GoalName), FText::FromName(ChainToAddGoalTo)))
 		]
 		.Buttons({
-			SCustomDialog::FButton(LOCTEXT("AddGoal", "Add Goal")),
-			SCustomDialog::FButton(LOCTEXT("Skip", "Skip"))
+			SCustomDialog::FButton(LOCTEXT("AssignGoal", "Assign Goal")),
+			SCustomDialog::FButton(LOCTEXT("SkipGoal", "Skip Goal"))
 	});
 
 	if (AddGoalToChainDialog->ShowModal() != 0)
@@ -876,25 +876,26 @@ void FIKRigEditorController::PromptToAddGoalToChain(UIKRigEffectorGoal* NewGoal)
 	AssetController->SetRetargetChainGoal(ChainToAddGoalTo, NewGoal->GoalName);
 }
 
-void FIKRigEditorController::PromptToAddNewRetargetChain(const FBoneChain& BoneChain) const
+FName FIKRigEditorController::PromptToAddNewRetargetChain(const FBoneChain& BoneChain) const
 {
 	const TSharedPtr<FStructOnScope> StructToDisplay = MakeShareable(new FStructOnScope(FIKRigRetargetChainSettings::StaticStruct(), (uint8*)&BoneChain));
 	TSharedRef<SKismetInspector> KismetInspector = SNew(SKismetInspector);
 	KismetInspector->ShowSingleStruct(StructToDisplay);
 
+	FName NewChainName = NAME_None;
 	SGenericDialogWidget::FArguments DialogArguments;
-	DialogArguments.OnOkPressed_Lambda([&BoneChain, this] ()
+	DialogArguments.OnOkPressed_Lambda([&BoneChain, &NewChainName, this] ()
 	{
 		// prompt to add a goal to the chain (skipped if already has a goal)
 		UIKRigEffectorGoal* NewGoal = PromptToAddGoalToNewChain(BoneChain);
 
 		// add the retarget chain
-		AssetController->AddRetargetChain(BoneChain);
+		NewChainName = AssetController->AddRetargetChain(BoneChain);
 
 		// ask user if they want to assign this goal to a chain (if there is one on this bone)
 		if (NewGoal)
 		{
-			PromptToAddGoalToChain(NewGoal);
+			PromptToAssignGoalToChain(NewGoal);
 		}
 		
 		RefreshAllViews();
@@ -905,6 +906,8 @@ void FIKRigEditorController::PromptToAddNewRetargetChain(const FBoneChain& BoneC
 		KismetInspector,
 		DialogArguments,
 		true);
+
+	return NewChainName;
 }
 
 UIKRigEffectorGoal* FIKRigEditorController::PromptToAddGoalToNewChain(const FBoneChain& BoneChain) const
@@ -924,12 +927,12 @@ UIKRigEffectorGoal* FIKRigEditorController::PromptToAddGoalToNewChain(const FBon
 			SNew(STextBlock)
 			.Text(FText::Format(
 				LOCTEXT("AddGoalToNewChainLabel", "Add goal to end bone, {0} of new chain, {1}?"),
-				FText::FromName(BoneChain.IKGoalName), 
-				FText::FromName(BoneChain.EndBone.BoneName)))
+				FText::FromName(BoneChain.EndBone.BoneName), 
+				FText::FromName(BoneChain.ChainName)))
 		]
 		.Buttons({
 			SCustomDialog::FButton(LOCTEXT("AddGoal", "Add Goal")),
-			SCustomDialog::FButton(LOCTEXT("Skip", "Skip"))
+			SCustomDialog::FButton(LOCTEXT("NoGoal", "No Goal"))
 	});
 
 	if (AddGoalToChainDialog->ShowModal() != 0)
@@ -945,42 +948,6 @@ UIKRigEffectorGoal* FIKRigEditorController::PromptToAddGoalToNewChain(const FBon
 	return AssetController->AddNewGoal(NewGoalName, BoneChain.EndBone.BoneName);
 }
 
-void FIKRigEditorController::PromptToAddGoalToNewlyMirroredChain(const FBoneChain& OldChain, const FBoneChain& MirroredChain) const
-{
-	const TSharedRef<SCustomDialog> AddGoalToChainDialog = SNew(SCustomDialog)
-		.Title(FText(LOCTEXT("AddGoalToMirroredChainTitleLabel", "Mirror Goal")))
-		.Content()
-		[
-			SNew(STextBlock)
-			.Text(FText::Format(LOCTEXT("AddGoalToMirroredChainLabel", "Mirror goal, {0} on chain, {1}?"), FText::FromName(OldChain.IKGoalName), FText::FromName(OldChain.IKGoalName)))
-		]
-		.Buttons({
-			SCustomDialog::FButton(LOCTEXT("AddGoal", "Add Goal")),
-			SCustomDialog::FButton(LOCTEXT("Skip", "Skip"))
-	});
-
-	if (AddGoalToChainDialog->ShowModal() != 0)
-	{
-		return; // cancel button pressed, or window closed
-	}
-
-	const UIKRigEffectorGoal* OldGoal = AssetController->GetGoal(OldChain.IKGoalName);
-	const FName NewGoalName = FName(OldChain.IKGoalName.ToString() + "_Mirrored");
-	// add a default solver if there isn't one already
-	PromptToAddDefaultSolver();
-	UIKRigEffectorGoal* NewGoal = AssetController->AddNewGoal(NewGoalName, MirroredChain.EndBone.BoneName);
-
-	// connect to the same solvers
-	const TArray<UIKRigSolver*>& AllSolvers =  AssetController->GetSolverArray();
-	for (int32 SolverIndex=0; SolverIndex<AllSolvers.Num(); ++SolverIndex)
-	{
-		if (AssetController->IsGoalConnectedToSolver(OldGoal->GoalName, SolverIndex))
-		{
-			AssetController->ConnectGoalToSolver(*NewGoal, SolverIndex);
-		}
-	}
-}
-
 void FIKRigEditorController::PlayAnimationAsset(UAnimationAsset* AssetToPlay)
 {
 	if (AssetToPlay && AnimInstance.IsValid())
diff --git a/Engine/Plugins/Animation/IKRig/Source/IKRigEditor/Private/RigEditor/SIKRigRetargetChainList.cpp b/Engine/Plugins/Animation/IKRig/Source/IKRigEditor/Private/RigEditor/SIKRigRetargetChainList.cpp
--- a/Engine/Plugins/Animation/IKRig/Source/IKRigEditor/Private/RigEditor/SIKRigRetargetChainList.cpp
+++ b/Engine/Plugins/Animation/IKRig/Source/IKRigEditor/Private/RigEditor/SIKRigRetargetChainList.cpp
@@ -16,6 +16,7 @@
 #include "SPositiveActionButton.h"
 #include "SSearchableComboBox.h"
 #include "BoneSelectionWidget.h"
+#include "Animation/MirrorDataTable.h"
 #include "Engine/SkeletalMesh.h"
 #include "Widgets/Input/SSearchBox.h"
 
@@ -743,7 +744,7 @@ void SIKRigRetargetChainList::MirrorSelectedChains() const
 		}
 
 		FBoneChain MirroredChain = *Chain;
-		MirroredChain.ChainName = FName(MirroredChain.ChainName.ToString() + "_Mirrored");
+		MirroredChain.ChainName = UMirrorDataTable::GetSettingsMirrorName(MirroredChain.ChainName);
 		MirroredChain.StartBone = IKRigSkeleton.BoneNames[MirroredIndices[0]];
 		MirroredChain.EndBone = IKRigSkeleton.BoneNames[MirroredIndices.Last()];
 		const UIKRigEffectorGoal* GoalOnMirroredBone = AssetController->GetGoalForBone(MirroredChain.EndBone.BoneName);
@@ -752,14 +753,26 @@ void SIKRigRetargetChainList::MirrorSelectedChains() const
 			MirroredChain.IKGoalName = GoalOnMirroredBone->GoalName;
 		}
 
-		Controller.PromptToAddNewRetargetChain(MirroredChain);
+		const FName NewChainName = Controller.PromptToAddNewRetargetChain(MirroredChain);
 
-		// old bone DOES have a goal, but the mirrored end bone does NOT have a goal,
-		// so lets ask the user if they want to add one...
-		const UIKRigEffectorGoal* GoalOnOldBone = AssetController->GetGoalForBone(Chain->EndBone.BoneName);
-		if (GoalOnOldBone && !GoalOnMirroredBone)
+		if (const FBoneChain* NewChain = AssetController->GetRetargetChainByName(NewChainName))
 		{
-			Controller.PromptToAddGoalToNewlyMirroredChain(*Chain, MirroredChain);
+			// old bone has a goal, and the new bone also has a goal, connect the new goal to the same solver(s)
+			const UIKRigEffectorGoal* GoalOnOldChain = AssetController->GetGoalForBone(Chain->EndBone.BoneName);
+			const UIKRigEffectorGoal* GoalOnNewChain = AssetController->GetGoal(NewChain->IKGoalName);
+			// connect to the same solvers
+			const TArray<UIKRigSolver*>& AllSolvers =  AssetController->GetSolverArray();
+			for (int32 SolverIndex=0; SolverIndex<AllSolvers.Num(); ++SolverIndex)
+			{
+				if (AssetController->IsGoalConnectedToSolver(GoalOnOldChain->GoalName, SolverIndex))
+				{
+					AssetController->ConnectGoalToSolver(*GoalOnNewChain, SolverIndex);
+				}
+				else
+				{
+					AssetController->DisconnectGoalFromSolver(GoalOnNewChain->GoalName, SolverIndex);
+				}
+			}
 		}
 	}
 }
diff --git a/Engine/Plugins/Animation/IKRig/Source/IKRigEditor/Public/RigEditor/IKRigController.h b/Engine/Plugins/Animation/IKRig/Source/IKRigEditor/Public/RigEditor/IKRigController.h
--- a/Engine/Plugins/Animation/IKRig/Source/IKRigEditor/Public/RigEditor/IKRigController.h
+++ b/Engine/Plugins/Animation/IKRig/Source/IKRigEditor/Public/RigEditor/IKRigController.h
@@ -148,8 +148,8 @@ public:
 	/** Retargeting Options and Retarget Bone Chains
 	 * 
 	 */
-	/** Add a Chain with the given Name and Start/End bones. Returns true if a new Chain was created. */
-	void AddRetargetChain(const FBoneChain& BoneChain) const;
+	/** Add a Chain with the given Name and Start/End bones. Returns newly created chain name. */
+	FName AddRetargetChain(const FBoneChain& BoneChain) const;
 	/** Remove a Chain with the given name. Returns true if a Chain was removed. */
 	bool RemoveRetargetChain(const FName& ChainName) const;
 	/** Renamed the given Chain. Returns the new name (same as old if unsuccessful). */
diff --git a/Engine/Plugins/Animation/IKRig/Source/IKRigEditor/Public/RigEditor/IKRigEditorController.h b/Engine/Plugins/Animation/IKRig/Source/IKRigEditor/Public/RigEditor/IKRigEditorController.h
--- a/Engine/Plugins/Animation/IKRig/Source/IKRigEditor/Public/RigEditor/IKRigEditorController.h
+++ b/Engine/Plugins/Animation/IKRig/Source/IKRigEditor/Public/RigEditor/IKRigEditorController.h
@@ -193,13 +193,11 @@ public:
 	/** right after importing a skeleton, we ask user what solver they want to use */
 	bool PromptToAddDefaultSolver() const;
 	/** show user the new retarget chain they are about to create (provides option to edit name) */
-	void PromptToAddNewRetargetChain(const FBoneChain& BoneChain) const;
+	FName PromptToAddNewRetargetChain(const FBoneChain& BoneChain) const;
 	/** prompt user if they want to add a goal to a newly created chain */
 	UIKRigEffectorGoal* PromptToAddGoalToNewChain(const FBoneChain& BoneChain) const;
-	/** prompt user if they want to add a goal to a newly mirrored chain */
-	void PromptToAddGoalToNewlyMirroredChain(const FBoneChain& OldChain, const FBoneChain& MirroredChain) const;
 	/** right after creating a goal, we ask user if they want it assigned to a retarget chain */
-	void PromptToAddGoalToChain(UIKRigEffectorGoal* NewGoal) const;
+	void PromptToAssignGoalToChain(UIKRigEffectorGoal* NewGoal) const;
 
 	/** play preview animation on running anim instance in editor (before IK) */
 	void PlayAnimationAsset(UAnimationAsset* AssetToPlay);
