Fixes retargeter crash when swapping between source/target on recently saved asset with custom poses.

diff --git a/Engine/Plugins/Animation/IKRig/Source/IKRigEditor/Private/RetargetEditor/SIKRetargetPoseEditor.cpp b/Engine/Plugins/Animation/IKRig/Source/IKRigEditor/Private/RetargetEditor/SIKRetargetPoseEditor.cpp
--- a/Engine/Plugins/Animation/IKRig/Source/IKRigEditor/Private/RetargetEditor/SIKRetargetPoseEditor.cpp
+++ b/Engine/Plugins/Animation/IKRig/Source/IKRigEditor/Private/RetargetEditor/SIKRetargetPoseEditor.cpp
@@ -22,9 +22,6 @@ void SIKRetargetPoseEditor::Construct(
 	// the commands for the menus
 	TSharedPtr<FUICommandList> Commands = Controller->Editor.Pin()->GetToolkitCommands();
 	
-	// build the list of poses
-	Refresh();
-	
 	ChildSlot
 	[
 		SNew(SVerticalBox)
@@ -54,8 +51,9 @@ void SIKRetargetPoseEditor::Construct(
 					// pose selection combobox
 					+SVerticalBox::Slot()
 					[
-						SNew(SComboBox<TSharedPtr<FName>>)
+						SAssignNew(PoseListComboBox, SComboBox<TSharedPtr<FName>>)
 						.OptionsSource(&PoseNames)
+						.OnComboBoxOpening(this, &SIKRetargetPoseEditor::Refresh)
 						.OnGenerateWidget_Lambda([](TSharedPtr<FName> InItem)
 						{
 							return SNew(STextBlock).Text(FText::FromName(*InItem.Get()));
@@ -191,6 +189,8 @@ void SIKRetargetPoseEditor::Refresh()
 	{
 		PoseNames.Add(MakeShareable(new FName(Pose.Key)));
 	}
+
+	PoseListComboBox->RefreshOptions();
 }
 
 TSharedRef<SWidget>  SIKRetargetPoseEditor::MakeToolbar(TSharedPtr<FUICommandList> Commands)
diff --git a/Engine/Plugins/Animation/IKRig/Source/IKRigEditor/Public/RetargetEditor/SIKRetargetPoseEditor.h b/Engine/Plugins/Animation/IKRig/Source/IKRigEditor/Public/RetargetEditor/SIKRetargetPoseEditor.h
--- a/Engine/Plugins/Animation/IKRig/Source/IKRigEditor/Public/RetargetEditor/SIKRetargetPoseEditor.h
+++ b/Engine/Plugins/Animation/IKRig/Source/IKRigEditor/Public/RetargetEditor/SIKRetargetPoseEditor.h
@@ -4,6 +4,7 @@
 
 #include "CoreMinimal.h"
 #include "DetailLayoutBuilder.h"
+#include "Widgets/Input/SComboBox.h"
 
 class UIKRetargeterController;
 class FIKRetargetEditorController;
@@ -31,4 +32,5 @@ private:
 	
 	TWeakPtr<FIKRetargetEditorController> EditorController;
 	TArray<TSharedPtr<FName>> PoseNames;
+	TSharedPtr<SComboBox<TSharedPtr<FName>>> PoseListComboBox;
 };
