Fix for framing chain selection in the retargeting viewport.

diff --git a/Engine/Plugins/Animation/IKRig/Source/IKRigEditor/Private/RetargetEditor/IKRetargetDefaultMode.cpp b/Engine/Plugins/Animation/IKRig/Source/IKRigEditor/Private/RetargetEditor/IKRetargetDefaultMode.cpp
--- a/Engine/Plugins/Animation/IKRig/Source/IKRigEditor/Private/RetargetEditor/IKRetargetDefaultMode.cpp
+++ b/Engine/Plugins/Animation/IKRig/Source/IKRigEditor/Private/RetargetEditor/IKRetargetDefaultMode.cpp
@@ -23,45 +23,118 @@ bool FIKRetargetDefaultMode::GetCameraTarget(FSphere& OutTarget) const
 	if (!Controller.IsValid())
 	{
 		return false;
-	}	
+	}
 
-	// target the selected bones
-	const TArray<FName>& SelectedBoneNames = Controller->GetSelectedBones();
-	if (!SelectedBoneNames.IsEmpty())
+	// center the view on the last selected item
+	switch (Controller->GetLastSelectedItemType())
 	{
-		TArray<FVector> TargetPoints;
-		const UDebugSkelMeshComponent* SkeletalMeshComponent = Controller->GetSkeletalMeshComponent(Controller->GetSourceOrTarget());
-		const FReferenceSkeleton& RefSkeleton = SkeletalMeshComponent->GetReferenceSkeleton();
-		TArray<int32> ChildrenIndices;
-		for (const FName& SelectedBoneName : SelectedBoneNames)
+	case ERetargetSelectionType::BONE:
 		{
-			const int32 BoneIndex = RefSkeleton.FindBoneIndex(SelectedBoneName);
-			if (BoneIndex == INDEX_NONE)
+			// target the selected bones
+			const TArray<FName>& SelectedBoneNames = Controller->GetSelectedBones();
+			if (SelectedBoneNames.IsEmpty())
 			{
-				continue;
+				return false;
 			}
 
-			TargetPoints.Add(SkeletalMeshComponent->GetBoneTransform(BoneIndex).GetLocation());
-			ChildrenIndices.Reset();
-			//RefSkeleton.GetDirectChildBones(BoneIndex, ChildrenIndices);
-			for (const int32 ChildIndex : ChildrenIndices)
+			TArray<FVector> TargetPoints;
+			const UDebugSkelMeshComponent* SkeletalMeshComponent = Controller->GetSkeletalMeshComponent(Controller->GetSourceOrTarget());
+			const FReferenceSkeleton& RefSkeleton = SkeletalMeshComponent->GetReferenceSkeleton();
+			TArray<int32> ChildrenIndices;
+			for (const FName& SelectedBoneName : SelectedBoneNames)
+			{
+				const int32 BoneIndex = RefSkeleton.FindBoneIndex(SelectedBoneName);
+				if (BoneIndex == INDEX_NONE)
+				{
+					continue;
+				}
+
+				TargetPoints.Add(SkeletalMeshComponent->GetBoneTransform(BoneIndex).GetLocation());
+				ChildrenIndices.Reset();
+				RefSkeleton.GetDirectChildBones(BoneIndex, ChildrenIndices);
+				for (const int32 ChildIndex : ChildrenIndices)
+				{
+					TargetPoints.Add(SkeletalMeshComponent->GetBoneTransform(ChildIndex).GetLocation());
+				}
+			}
+	
+			// create a sphere that contains all the target points
+			if (TargetPoints.Num() == 0)
 			{
-				TargetPoints.Add(SkeletalMeshComponent->GetBoneTransform(ChildIndex).GetLocation());
+				TargetPoints.Add(FVector::ZeroVector);
 			}
+			OutTarget = FSphere(&TargetPoints[0], TargetPoints.Num());
+			return true;
 		}
 		
-		// create a sphere that contains all the target points
-		OutTarget = FSphere(&TargetPoints[0], TargetPoints.Num());
-		return true;
-	}
+	case ERetargetSelectionType::CHAIN:
+		{
+			const ERetargetSourceOrTarget SourceOrTarget = Controller->GetSourceOrTarget();
+			const UIKRigDefinition* IKRig = Controller->AssetController->GetIKRig(SourceOrTarget);
+			if (!IKRig)
+			{
+				return false;
+			}
 
-	// target the selected mesh
-	if (const UPrimitiveComponent* SelectedMesh = Controller->GetSelectedMesh())
-	{
-		OutTarget = SelectedMesh->Bounds.GetSphere();
-		return true;
+			const UDebugSkelMeshComponent* SkeletalMeshComponent = Controller->GetSkeletalMeshComponent(Controller->GetSourceOrTarget());
+			const FReferenceSkeleton& RefSkeleton = SkeletalMeshComponent->GetReferenceSkeleton();
+
+			// get target points from start/end bone of all selected chains on the currently active skeleton (source or target)
+			TArray<FVector> TargetPoints;
+			const TArray<FName>& SelectedChains = Controller->GetSelectedChains();
+			for (const FName SelectedChainName : SelectedChains)
+			{
+				const URetargetChainSettings* ChainMap = Controller->AssetController->GetChainMappingByTargetChainName(SelectedChainName);
+				if (!ChainMap)
+				{
+					continue;
+				}
+
+				const FName& ChainName = SourceOrTarget == ERetargetSourceOrTarget::Target ? SelectedChainName : ChainMap->SourceChain;
+				if (ChainName == NAME_None)
+				{
+					continue;
+				}
+
+				const FBoneChain* BoneChain = IKRig->GetRetargetChainByName(ChainName);
+				if (!BoneChain)
+				{
+					continue;
+				}
+
+				const int32 StartBoneIndex = RefSkeleton.FindBoneIndex(BoneChain->StartBone.BoneName);
+				if (StartBoneIndex != INDEX_NONE)
+				{
+					TargetPoints.Add(SkeletalMeshComponent->GetBoneTransform(StartBoneIndex).GetLocation());
+				}
+
+				const int32 EndBoneIndex = RefSkeleton.FindBoneIndex(BoneChain->EndBone.BoneName);
+				if (EndBoneIndex != INDEX_NONE)
+				{
+					TargetPoints.Add(SkeletalMeshComponent->GetBoneTransform(EndBoneIndex).GetLocation());
+				}
+			}
+
+			// create a sphere that contains all the target points
+			if (TargetPoints.Num() == 0)
+			{
+				TargetPoints.Add(FVector::ZeroVector);
+			}
+			OutTarget = FSphere(&TargetPoints[0], TargetPoints.Num());
+			return true;	
+		}
+	case ERetargetSelectionType::ROOT:
+	case ERetargetSelectionType::MESH:
+	case ERetargetSelectionType::NONE:
+	default:
+		// target the current mesh
+		if (const UPrimitiveComponent* SelectedMesh = Controller->GetSkeletalMeshComponent(Controller->GetSourceOrTarget()))
+		{
+			OutTarget = SelectedMesh->Bounds.GetSphere();
+			return true;
+		}
 	}
-	
+
 	return false;
 }
 
diff --git a/Engine/Plugins/Animation/IKRig/Source/IKRigEditor/Private/RetargetEditor/IKRetargeterController.cpp b/Engine/Plugins/Animation/IKRig/Source/IKRigEditor/Private/RetargetEditor/IKRetargeterController.cpp
--- a/Engine/Plugins/Animation/IKRig/Source/IKRigEditor/Private/RetargetEditor/IKRetargeterController.cpp
+++ b/Engine/Plugins/Animation/IKRig/Source/IKRigEditor/Private/RetargetEditor/IKRetargeterController.cpp
@@ -341,6 +341,19 @@ const TArray<TObjectPtr<URetargetChainSettings>>& UIKRetargeterController::GetCh
 	return Asset->ChainSettings;
 }
 
+const URetargetChainSettings* UIKRetargeterController::GetChainMappingByTargetChainName(const FName& TargetChainName) const
+{
+	for (TObjectPtr<URetargetChainSettings> ChainMap : Asset->ChainSettings)
+	{
+		if (ChainMap->TargetChain == TargetChainName)
+		{
+			return ChainMap;
+		}
+	}
+
+	return nullptr;
+}
+
 FName UIKRetargeterController::GetChainGoal(const TObjectPtr<URetargetChainSettings> ChainSettings) const
 {
 	if (ChainSettings.IsNull())
diff --git a/Engine/Plugins/Animation/IKRig/Source/IKRigEditor/Public/RetargetEditor/IKRetargetEditorController.h b/Engine/Plugins/Animation/IKRig/Source/IKRigEditor/Public/RetargetEditor/IKRetargetEditorController.h
--- a/Engine/Plugins/Animation/IKRig/Source/IKRigEditor/Public/RetargetEditor/IKRetargetEditorController.h
+++ b/Engine/Plugins/Animation/IKRig/Source/IKRigEditor/Public/RetargetEditor/IKRetargetEditorController.h
@@ -210,6 +210,7 @@ public:
 	bool GetRootSelected() const { return bIsRootSelected; };
 	
 	void ClearSelection(const bool bKeepBoneSelection=false);
+	ERetargetSelectionType GetLastSelectedItemType() const { return LastSelectedItem; };
 
 	/** ------------------------- END SELECTION -----------------------------*/
 
diff --git a/Engine/Plugins/Animation/IKRig/Source/IKRigEditor/Public/RetargetEditor/IKRetargeterController.h b/Engine/Plugins/Animation/IKRig/Source/IKRigEditor/Public/RetargetEditor/IKRetargeterController.h
--- a/Engine/Plugins/Animation/IKRig/Source/IKRigEditor/Public/RetargetEditor/IKRetargeterController.h
+++ b/Engine/Plugins/Animation/IKRig/Source/IKRigEditor/Public/RetargetEditor/IKRetargeterController.h
@@ -65,6 +65,8 @@ public:
 	void SetSourceChainForTargetChain(URetargetChainSettings* ChainMap, FName SourceChainToMapTo) const;
 	/** Get read-only access to the list of chain mappings */
 	const TArray<TObjectPtr<URetargetChainSettings>>& GetChainMappings() const;
+	/** Get read-only acces to the chain settings associated with a given target chain */
+	const URetargetChainSettings* GetChainMappingByTargetChainName(const FName& TargetChainName) const;
 	/** Get the name of the goal assigned to a chain */
 	FName GetChainGoal(const TObjectPtr<URetargetChainSettings> ChainSettings) const;
 	/** Get whether the given chain's IK goal is connected to a solver */
diff --git a/Engine/Source/Runtime/Engine/Public/ReferenceSkeleton.h b/Engine/Source/Runtime/Engine/Public/ReferenceSkeleton.h
--- a/Engine/Source/Runtime/Engine/Public/ReferenceSkeleton.h
+++ b/Engine/Source/Runtime/Engine/Public/ReferenceSkeleton.h
@@ -407,7 +407,7 @@ public:
 	SIZE_T GetDataSize() const;
 
 	// very slow search function for all children
-	int32 GetDirectChildBones(int32 ParentBoneIndex, TArray<int32> & Children) const;
+	ENGINE_API int32 GetDirectChildBones(int32 ParentBoneIndex, TArray<int32> & Children) const;
 	ENGINE_API friend FArchive & operator<<(FArchive & Ar, FReferenceSkeleton & F);
 	friend FReferenceSkeletonModifier;
 };
