Adding check for zero-height root in retargeter.

diff --git a/Engine/Plugins/Animation/IKRig/Source/IKRig/Private/Retargeter/IKRetargetProcessor.cpp b/Engine/Plugins/Animation/IKRig/Source/IKRig/Private/Retargeter/IKRetargetProcessor.cpp
--- a/Engine/Plugins/Animation/IKRig/Source/IKRig/Private/Retargeter/IKRetargetProcessor.cpp
+++ b/Engine/Plugins/Animation/IKRig/Source/IKRig/Private/Retargeter/IKRetargetProcessor.cpp
@@ -1204,7 +1204,7 @@ bool FRootRetargeter::InitializeSource(
 	if (InitialHeight < KINDA_SMALL_NUMBER)
 	{
 		// warn user and push it up slightly to avoid divide by zero
-		Log.LogWarning(LOCTEXT("BadRootHeight", "IK Retargeter root bone is very near the ground plane. This is probably not intentional."));
+		Log.LogError(LOCTEXT("BadRootHeight", "The source retarget root bone is very near the ground plane. This will cause the target to be moved very far. To resolve this, please create a retarget pose with the retarget root at the correct height off the ground."));
 		InitialHeight = 1.0f;
 	}
 
diff --git a/Engine/Plugins/Animation/IKRig/Source/IKRig/Public/Retargeter/IKRetargeter.h b/Engine/Plugins/Animation/IKRig/Source/IKRig/Public/Retargeter/IKRetargeter.h
--- a/Engine/Plugins/Animation/IKRig/Source/IKRig/Public/Retargeter/IKRetargeter.h
+++ b/Engine/Plugins/Animation/IKRig/Source/IKRig/Public/Retargeter/IKRetargeter.h
@@ -422,6 +422,10 @@ private:
 	/** The controller responsible for managing this asset's data (all editor mutation goes through this) */
 	UPROPERTY(Transient)
 	TObjectPtr<UObject> Controller;
+
+	/** only ask to fix the root height once, then warn thereafter (don't nag) */
+	UPROPERTY()
+	TSet<USkeletalMesh*> MeshesAskedToFixRootHeightFor;
 #endif
 	
 private:
diff --git a/Engine/Plugins/Animation/IKRig/Source/IKRigEditor/Private/RetargetEditor/IKRetargetEditor.cpp b/Engine/Plugins/Animation/IKRig/Source/IKRigEditor/Private/RetargetEditor/IKRetargetEditor.cpp
--- a/Engine/Plugins/Animation/IKRig/Source/IKRigEditor/Private/RetargetEditor/IKRetargetEditor.cpp
+++ b/Engine/Plugins/Animation/IKRig/Source/IKRigEditor/Private/RetargetEditor/IKRetargetEditor.cpp
@@ -371,6 +371,9 @@ void FIKRetargetEditor::HandlePreviewSceneCreated(const TSharedRef<IPersonaPrevi
 	
 	InPersonaPreviewScene->AddComponent(EditorController->SourceSkelMeshComponent, FTransform::Identity);
 	InPersonaPreviewScene->AddComponent(EditorController->TargetSkelMeshComponent, FTransform::Identity);
+
+	EditorController->FixZeroHeightRetargetRoot(ERetargetSourceOrTarget::Source);
+	EditorController->FixZeroHeightRetargetRoot(ERetargetSourceOrTarget::Target);
 }
 
 void FIKRetargetEditor::SetupAnimInstance()
diff --git a/Engine/Plugins/Animation/IKRig/Source/IKRigEditor/Private/RetargetEditor/IKRetargetEditorController.cpp b/Engine/Plugins/Animation/IKRig/Source/IKRigEditor/Private/RetargetEditor/IKRetargetEditorController.cpp
--- a/Engine/Plugins/Animation/IKRig/Source/IKRigEditor/Private/RetargetEditor/IKRetargetEditorController.cpp
+++ b/Engine/Plugins/Animation/IKRig/Source/IKRigEditor/Private/RetargetEditor/IKRetargetEditorController.cpp
@@ -10,6 +10,7 @@
 #include "Widgets/Input/SButton.h"
 #include "Animation/DebugSkelMeshComponent.h"
 #include "Animation/PoseAsset.h"
+#include "Dialog/SCustomDialog.h"
 #include "Preferences/PersonaOptions.h"
 #include "RetargetEditor/IKRetargetAnimInstance.h"
 #include "RetargetEditor/IKRetargetDefaultMode.h"
@@ -141,6 +142,10 @@ void FIKRetargetEditorController::HandleRetargeterNeedsInitialized(UIKRetargeter
 {
 	// clear the output log
 	ClearOutputLog();
+
+	// check for "zero height" retarget roots, and prompt user to fix
+	FixZeroHeightRetargetRoot(ERetargetSourceOrTarget::Source);
+	FixZeroHeightRetargetRoot(ERetargetSourceOrTarget::Target);
 	
 	// force reinit the retarget processor (also inits the target IK Rig processor)
 	if (UIKRetargetProcessor* Processor = GetRetargetProcessor())
@@ -1253,6 +1258,82 @@ void FIKRetargetEditorController::RenderSkeleton(FPrimitiveDrawInterface* PDI, E
 	);
 }
 
+void FIKRetargetEditorController::FixZeroHeightRetargetRoot(ERetargetSourceOrTarget SourceOrTarget) const
+{
+	// is there a mesh to check?
+	USkeletalMesh* SkeletalMesh = GetSkeletalMesh(SourceOrTarget);
+	if (!SkeletalMesh)
+	{
+		return;
+	}
+
+	// have we already nagged the user about fixing this mesh?
+	if (AssetController->GetAskedToFixRootHeightForMesh(SkeletalMesh))
+	{
+		return;
+	}
+
+	const FName CurrentRetargetPoseName = AssetController->GetCurrentRetargetPoseName(SourceOrTarget);
+	FIKRetargetPose& CurrentRetargetPose = AssetController->GetCurrentRetargetPose(SourceOrTarget);
+	const FName RetargetRootBoneName = AssetController->GetRetargetRootBone(SourceOrTarget);
+	if (RetargetRootBoneName == NAME_None)
+	{
+		return;
+	}
+
+	FRetargetSkeleton DummySkeleton;
+	DummySkeleton.Initialize(
+		SkeletalMesh,
+		TArray<FBoneChain>(),
+		CurrentRetargetPoseName,
+		&CurrentRetargetPose,
+		RetargetRootBoneName);
+
+	const int32 RootBoneIndex = DummySkeleton.FindBoneIndexByName(RetargetRootBoneName);
+	if (RootBoneIndex == INDEX_NONE)
+	{
+		return;
+	}
+
+	const FTransform& RootTransform = DummySkeleton.RetargetGlobalPose[RootBoneIndex];
+	if (RootTransform.GetLocation().Z < 1.0f)
+	{
+		if (PromptToFixRootHeight(SourceOrTarget))
+		{
+			// move it up based on the height of the mesh
+			const float FixedHeight = FMath::Abs(SkeletalMesh->GetBounds().GetBoxExtrema(-1).Z);
+			// update the current retarget pose
+			CurrentRetargetPose.SetRootTranslationDelta(FVector(0.f, 0.f,FixedHeight));
+		}
+	}
+
+	AssetController->SetAskedToFixRootHeightForMesh(SkeletalMesh, true);
+}
+
+bool FIKRetargetEditorController::PromptToFixRootHeight(ERetargetSourceOrTarget SourceOrTarget) const
+{
+	const FText SourceOrTargetText = SourceOrTarget == ERetargetSourceOrTarget::Source ? FText::FromString("Source") : FText::FromString("Target");
+
+	TSharedRef<SCustomDialog> Dialog = SNew(SCustomDialog)
+		.Title(FText(LOCTEXT("FixRootHeightTitle", "Add Height to Retarget Root Pose")))
+		.Content()
+		[
+			SNew(STextBlock)
+			.Text(FText::Format(LOCTEXT("FixRootHeightLabel", "The {0} skeleton has a retarget root bone on the ground. Apply a vertical offset to root bone in the current retarget pose?"), SourceOrTargetText))
+		]
+		.Buttons({
+			SCustomDialog::FButton(LOCTEXT("ApplyOffset", "Apply Offset")),
+			SCustomDialog::FButton(LOCTEXT("No", "No"))
+	});
+
+	if (Dialog->ShowModal() != 0)
+	{
+		return false; // cancel button pressed, or window closed
+	}
+
+	return true;
+}
+
 FText FIKRetargetEditorController::GetCurrentPoseName() const
 {
 	return FText::FromName(AssetController->GetCurrentRetargetPoseName(GetSourceOrTarget()));
diff --git a/Engine/Plugins/Animation/IKRig/Source/IKRigEditor/Private/RetargetEditor/IKRetargeterController.cpp b/Engine/Plugins/Animation/IKRig/Source/IKRigEditor/Private/RetargetEditor/IKRetargeterController.cpp
--- a/Engine/Plugins/Animation/IKRig/Source/IKRigEditor/Private/RetargetEditor/IKRetargeterController.cpp
+++ b/Engine/Plugins/Animation/IKRig/Source/IKRigEditor/Private/RetargetEditor/IKRetargeterController.cpp
@@ -86,7 +86,26 @@ void UIKRetargeterController::OnIKRigChanged(const ERetargetSourceOrTarget& Sour
 		{
 			Asset->TargetPreviewMesh = IKRig->PreviewSkeletalMesh;
 		}
-		
+
+		// re-ask to fix root height for this mesh
+		SetAskedToFixRootHeightForMesh(IKRig->PreviewSkeletalMesh.Get(), false);
+	}
+}
+
+bool UIKRetargeterController::GetAskedToFixRootHeightForMesh(USkeletalMesh* Mesh) const
+{
+	return GetAsset()->MeshesAskedToFixRootHeightFor.Contains(Mesh);
+}
+
+void UIKRetargeterController::SetAskedToFixRootHeightForMesh(USkeletalMesh* Mesh, bool InAsked) const
+{
+	if (InAsked)
+	{
+		GetAsset()->MeshesAskedToFixRootHeightFor.Add(Mesh);
+	}
+	else
+	{
+		GetAsset()->MeshesAskedToFixRootHeightFor.Remove(Mesh);
 	}
 }
 
diff --git a/Engine/Plugins/Animation/IKRig/Source/IKRigEditor/Public/RetargetEditor/IKRetargetEditorController.h b/Engine/Plugins/Animation/IKRig/Source/IKRigEditor/Public/RetargetEditor/IKRetargetEditorController.h
--- a/Engine/Plugins/Animation/IKRig/Source/IKRigEditor/Public/RetargetEditor/IKRetargetEditorController.h
+++ b/Engine/Plugins/Animation/IKRig/Source/IKRigEditor/Public/RetargetEditor/IKRetargetEditorController.h
@@ -266,6 +266,9 @@ public:
 	bool CanRenamePose() const;
 	TSharedPtr<SWindow> RenamePoseWindow;
 	TSharedPtr<SEditableTextBox> NewNameEditableText;
+
+	/* detect and auto-fix retarget pose that causes root height to be on the ground */
+	void FixZeroHeightRetargetRoot(ERetargetSourceOrTarget SourceOrTarget) const;
 	
 	/** ------------------------- END RETARGET POSES -----------------------------*/
 
@@ -279,6 +282,10 @@ public:
 
 private:
 
+	/** modal dialog to ask user if they want to fix root bones that are "on the ground" */
+	bool PromptToFixRootHeight(ERetargetSourceOrTarget SourceOrTarget) const;
+	bool bAskedToFixRoot = false;
+
 	/** asset properties tab */
 	TSharedPtr<IDetailsView> DetailsView;
 	/** chain list view */
diff --git a/Engine/Plugins/Animation/IKRig/Source/IKRigEditor/Public/RetargetEditor/IKRetargeterController.h b/Engine/Plugins/Animation/IKRig/Source/IKRigEditor/Public/RetargetEditor/IKRetargeterController.h
--- a/Engine/Plugins/Animation/IKRig/Source/IKRigEditor/Public/RetargetEditor/IKRetargeterController.h
+++ b/Engine/Plugins/Animation/IKRig/Source/IKRigEditor/Public/RetargetEditor/IKRetargeterController.h
@@ -42,7 +42,10 @@ public:
 	UIKRigDefinition* GetIKRigWriteable(const ERetargetSourceOrTarget& SourceOrTarget) const;
 	/** Set the SOURCE or TARGET preview mesh based on the mesh in the corresponding IK Rig asset */
 	void OnIKRigChanged(const ERetargetSourceOrTarget& SourceOrTarget) const;
-
+	/** Get if we've already asked to fix the root height for the given skeletal mesh */
+	bool GetAskedToFixRootHeightForMesh(USkeletalMesh* Mesh) const;
+	/** Set if we've asked to fix the root height for the given skeletal mesh */
+	void SetAskedToFixRootHeightForMesh(USkeletalMesh* Mesh, bool InAsked) const;
 	/** Get name of the Root bone used for retargeting. */
 	FName GetRetargetRootBone(const ERetargetSourceOrTarget& SourceOrTarget) const;
 
