Fixing native pointer usage in IKRetargeter.h

diff --git a/Engine/Plugins/Animation/IKRig/Source/IKRig/Public/Retargeter/IKRetargeter.h b/Engine/Plugins/Animation/IKRig/Source/IKRig/Public/Retargeter/IKRetargeter.h
--- a/Engine/Plugins/Animation/IKRig/Source/IKRig/Public/Retargeter/IKRetargeter.h
+++ b/Engine/Plugins/Animation/IKRig/Source/IKRig/Public/Retargeter/IKRetargeter.h
@@ -425,7 +425,7 @@ private:
 
 	/** only ask to fix the root height once, then warn thereafter (don't nag) */
 	UPROPERTY()
-	TSet<USkeletalMesh*> MeshesAskedToFixRootHeightFor;
+	TSet<TObjectPtr<USkeletalMesh>> MeshesAskedToFixRootHeightFor;
 #endif
 	
 private:
