Fixes retargeter hierarchy panel not filtering retargeted bones correctly.

diff --git a/Engine/Plugins/Animation/IKRig/Source/IKRig/Private/Retargeter/IKRetargetProcessor.cpp b/Engine/Plugins/Animation/IKRig/Source/IKRig/Private/Retargeter/IKRetargetProcessor.cpp
--- a/Engine/Plugins/Animation/IKRig/Source/IKRig/Private/Retargeter/IKRetargetProcessor.cpp
+++ b/Engine/Plugins/Animation/IKRig/Source/IKRig/Private/Retargeter/IKRetargetProcessor.cpp
@@ -2124,7 +2124,21 @@ bool UIKRetargetProcessor::IsBoneRetargeted(
 	const bool bUseSourceSkeleton = SkeletonToCheck == 0;
 	const FRetargetSkeleton& Skeleton = bUseSourceSkeleton ? SourceSkeleton : TargetSkeleton;
 	const int32 RootBoneIndex = bUseSourceSkeleton ? GetSourceRetargetRoot() : GetTargetRetargetRoot();
-	return (Skeleton.ChainThatContainsBone[BoneIndex] != NAME_None) || BoneIndex==RootBoneIndex;
+	if (BoneIndex==RootBoneIndex)
+	{
+		// retarget root is always retargeted
+		return true;
+	}
+	const FName ChainThatContainsBone = Skeleton.ChainThatContainsBone[BoneIndex];
+	if (ChainThatContainsBone == NAME_None)
+	{
+		// bones not in a chain, therefore cannot be retargeted
+		return false;
+	}
+
+	// bone must be in a chain that is actually mapped to something
+	const TObjectPtr<URetargetChainSettings> ChainMap = RetargeterAsset->GetChainMapByName(ChainThatContainsBone);
+	return !ChainMap->SourceChain.IsNone();
 }
 
 FName UIKRetargetProcessor::GetChainNameForBone(const int32& BoneIndex, const int8& SkeletonToCheck) const
diff --git a/Engine/Plugins/Animation/IKRig/Source/IKRig/Private/Retargeter/IKRetargeter.cpp b/Engine/Plugins/Animation/IKRig/Source/IKRig/Private/Retargeter/IKRetargeter.cpp
--- a/Engine/Plugins/Animation/IKRig/Source/IKRig/Private/Retargeter/IKRetargeter.cpp
+++ b/Engine/Plugins/Animation/IKRig/Source/IKRig/Private/Retargeter/IKRetargeter.cpp
@@ -265,16 +265,23 @@ void FIKRetargetPose::SortHierarchically(const FIKRigSkeleton& Skeleton)
 	});
 }
 
-const FTargetChainSettings* UIKRetargeter::GetChainSettingsByName(const FName& TargetChainName) const
+const TObjectPtr<URetargetChainSettings> UIKRetargeter::GetChainMapByName(const FName& TargetChainName) const
 {
 	const TObjectPtr<URetargetChainSettings>* ChainMap = ChainSettings.FindByPredicate(
 		[TargetChainName](const TObjectPtr<URetargetChainSettings> ChainMap)
 		{
 			return ChainMap->TargetChain == TargetChainName;
 		});
+	
+	return ChainMap->IsNull() ? nullptr : ChainMap->Get();
+}
+
+const FTargetChainSettings* UIKRetargeter::GetChainSettingsByName(const FName& TargetChainName) const
+{
+	const TObjectPtr<URetargetChainSettings> ChainMap = GetChainMapByName(TargetChainName);
 	if (ChainMap)
 	{
-		return &ChainMap->Get()->Settings;
+		return &ChainMap->Settings;
 	}
 
 	return nullptr;
diff --git a/Engine/Plugins/Animation/IKRig/Source/IKRig/Public/Retargeter/IKRetargeter.h b/Engine/Plugins/Animation/IKRig/Source/IKRig/Public/Retargeter/IKRetargeter.h
--- a/Engine/Plugins/Animation/IKRig/Source/IKRig/Public/Retargeter/IKRetargeter.h
+++ b/Engine/Plugins/Animation/IKRig/Source/IKRig/Public/Retargeter/IKRetargeter.h
@@ -217,6 +217,8 @@ public:
 
 	/** Get read-only access to the chain mapping */
 	const TArray<TObjectPtr<URetargetChainSettings>>& GetAllChainSettings() const { return ChainSettings; };
+	/** Get read-only access to the chain map for a given chain (null if chain not in retargeter) */
+	const TObjectPtr<URetargetChainSettings> GetChainMapByName(const FName& TargetChainName) const;
 	/** Get read-only access to the chain settings for a given chain (null if chain not in retargeter) */
 	const FTargetChainSettings* GetChainSettingsByName(const FName& TargetChainName) const;
 	/** Get access to the root settings */
