Fix hierarchy filtering to support "unaffected bones" in IK Rig

diff --git a/Engine/Plugins/Animation/IKRig/Source/IKRigEditor/Private/RigEditor/SIKRigHierarchy.cpp b/Engine/Plugins/Animation/IKRig/Source/IKRigEditor/Private/RigEditor/SIKRigHierarchy.cpp
--- a/Engine/Plugins/Animation/IKRig/Source/IKRigEditor/Private/RigEditor/SIKRigHierarchy.cpp
+++ b/Engine/Plugins/Animation/IKRig/Source/IKRigEditor/Private/RigEditor/SIKRigHierarchy.cpp
@@ -1729,9 +1729,10 @@ void SIKRigHierarchy::RefreshTreeView(bool IsInitialSetup /*=false*/)
 			BoneElement->bIsHidden = true;
 		}
 
-		if (FilterOptions.bHideUnaffectedBones)
+		// apply filter to hide unaffected bones
+		if (FilterOptions.bHideUnaffectedBones && !IsElementConnectedToAnySolver(BoneElement))
 		{
-			
+			BoneElement->bIsHidden = true;	
 		}
 		
 		// apply text filter to bones
