Fixing blueprint compile warning when binding a retarget asset to an input pin.

diff --git a/Engine/Plugins/Animation/IKRig/Source/IKRigDeveloper/Private/AnimGraphNode_RetargetPoseFromMesh.cpp b/Engine/Plugins/Animation/IKRig/Source/IKRigDeveloper/Private/AnimGraphNode_RetargetPoseFromMesh.cpp
--- a/Engine/Plugins/Animation/IKRig/Source/IKRigDeveloper/Private/AnimGraphNode_RetargetPoseFromMesh.cpp
+++ b/Engine/Plugins/Animation/IKRig/Source/IKRigDeveloper/Private/AnimGraphNode_RetargetPoseFromMesh.cpp
@@ -54,7 +54,9 @@ void UAnimGraphNode_RetargetPoseFromMesh::ValidateAnimNodeDuringCompilation(USke
 	// validate source mesh component is not null
 	if (!Node.bUseAttachedParent)
 	{
-		if (!IsPinExposedAndLinked(GET_MEMBER_NAME_STRING_CHECKED(FAnimNode_RetargetPoseFromMesh, SourceMeshComponent)))
+		const bool bIsLinked = IsPinExposedAndLinked(GET_MEMBER_NAME_STRING_CHECKED(FAnimNode_RetargetPoseFromMesh, SourceMeshComponent));
+		const bool bIsBound = IsPinExposedAndBound(GET_MEMBER_NAME_STRING_CHECKED(FAnimNode_RetargetPoseFromMesh, SourceMeshComponent));
+		if (!(bIsLinked || bIsBound))
 		{
 			MessageLog.Error(TEXT("@@ is missing a Source Skeletal Mesh Component reference."), this);
 			return;
@@ -68,17 +70,8 @@ void UAnimGraphNode_RetargetPoseFromMesh::ValidateAnimNodeDuringCompilation(USke
 		if (Pin == nullptr)
 		{
 			// retarget asset unassigned
-			MessageLog.Warning(TEXT("@@ is missing an IKRetargeter asset."), this);
-			return;	
+			MessageLog.Error(TEXT("@@ does not have an IK Retargeter asset assigned."), this);
 		}
-		
-		if (!IsPinExposedAndLinked(GET_MEMBER_NAME_STRING_CHECKED(FAnimNode_RetargetPoseFromMesh, IKRetargeterAsset)))
-		{
-			// retarget asset pin not connected to an asset
-			MessageLog.Warning(TEXT("@@ has a IKRetargeter asset pin that is not connected to a valid asset."), this);
-			return;	
-		}
-
 		return;
 	}
 
