Fixing retargeter crash bug regression from

diff --git a/Engine/Plugins/Animation/IKRig/Source/IKRig/Private/Retargeter/IKRetargetProcessor.cpp b/Engine/Plugins/Animation/IKRig/Source/IKRig/Private/Retargeter/IKRetargetProcessor.cpp
--- a/Engine/Plugins/Animation/IKRig/Source/IKRig/Private/Retargeter/IKRetargetProcessor.cpp
+++ b/Engine/Plugins/Animation/IKRig/Source/IKRig/Private/Retargeter/IKRetargetProcessor.cpp
@@ -2137,8 +2137,31 @@ bool UIKRetargetProcessor::IsBoneRetargeted(
 	}
 
 	// bone must be in a chain that is actually mapped to something
-	const TObjectPtr<URetargetChainSettings> ChainMap = RetargeterAsset->GetChainMapByName(ChainThatContainsBone);
-	return !ChainMap->SourceChain.IsNone();
+	const TArray<TObjectPtr<URetargetChainSettings>>& ChainMaps = RetargeterAsset->GetAllChainSettings();
+	for (const TObjectPtr<URetargetChainSettings>& ChainMap : ChainMaps)
+	{
+		if (ChainMap->SourceChain == NAME_None || ChainMap->TargetChain == NAME_None)
+		{
+			continue;
+		}
+		
+		if (bUseSourceSkeleton)
+		{
+			if (ChainMap->SourceChain == ChainThatContainsBone)
+			{
+				return true;
+			}
+		}
+		else
+		{
+			if (ChainMap->TargetChain == ChainThatContainsBone)
+			{
+				return true;
+			}
+		}
+	}
+
+	return false;
 }
 
 FName UIKRetargetProcessor::GetChainNameForBone(const int32& BoneIndex, const int8& SkeletonToCheck) const
diff --git a/Engine/Plugins/Animation/IKRig/Source/IKRig/Private/Retargeter/IKRetargeter.cpp b/Engine/Plugins/Animation/IKRig/Source/IKRig/Private/Retargeter/IKRetargeter.cpp
--- a/Engine/Plugins/Animation/IKRig/Source/IKRig/Private/Retargeter/IKRetargeter.cpp
+++ b/Engine/Plugins/Animation/IKRig/Source/IKRig/Private/Retargeter/IKRetargeter.cpp
@@ -273,7 +273,7 @@ const TObjectPtr<URetargetChainSettings> UIKRetargeter::GetChainMapByName(const
 			return ChainMap->TargetChain == TargetChainName;
 		});
 	
-	return ChainMap->IsNull() ? nullptr : ChainMap->Get();
+	return !ChainMap ? nullptr : ChainMap->Get();
 }
 
 const FTargetChainSettings* UIKRetargeter::GetChainSettingsByName(const FName& TargetChainName) const
