Fixes retargeter crash bug when mirroring chain with goal (regression)

diff --git a/Engine/Plugins/Animation/IKRig/Source/IKRigEditor/Private/RigEditor/SIKRigRetargetChainList.cpp b/Engine/Plugins/Animation/IKRig/Source/IKRigEditor/Private/RigEditor/SIKRigRetargetChainList.cpp
--- a/Engine/Plugins/Animation/IKRig/Source/IKRigEditor/Private/RigEditor/SIKRigRetargetChainList.cpp
+++ b/Engine/Plugins/Animation/IKRig/Source/IKRigEditor/Private/RigEditor/SIKRigRetargetChainList.cpp
@@ -754,12 +754,20 @@ void SIKRigRetargetChainList::MirrorSelectedChains() const
 		}
 
 		const FName NewChainName = Controller.PromptToAddNewRetargetChain(MirroredChain);
-
-		if (const FBoneChain* NewChain = AssetController->GetRetargetChainByName(NewChainName))
+		const FBoneChain* NewChain = AssetController->GetRetargetChainByName(NewChainName);
+		
+		if (!NewChain)
+		{
+			// user cancelled mirroring the chain
+			continue;
+		}
+		
+		// check old bone has a goal, and the new bone also has a goal
+		// so we can connect the new goal to the same solver(s)
+		const UIKRigEffectorGoal* GoalOnOldChain = AssetController->GetGoalForBone(Chain->EndBone.BoneName);
+		const UIKRigEffectorGoal* GoalOnNewChain = AssetController->GetGoal(NewChain->IKGoalName);
+		if (GoalOnOldChain && GoalOnNewChain)
 		{
-			// old bone has a goal, and the new bone also has a goal, connect the new goal to the same solver(s)
-			const UIKRigEffectorGoal* GoalOnOldChain = AssetController->GetGoalForBone(Chain->EndBone.BoneName);
-			const UIKRigEffectorGoal* GoalOnNewChain = AssetController->GetGoal(NewChain->IKGoalName);
 			// connect to the same solvers
 			const TArray<UIKRigSolver*>& AllSolvers =  AssetController->GetSolverArray();
 			for (int32 SolverIndex=0; SolverIndex<AllSolvers.Num(); ++SolverIndex)
