IK Rig's new retarget chain dialog displays wrong bone names.

diff --git a/Engine/Plugins/Animation/IKRig/Source/IKRigEditor/Private/RigEditor/IKRigEditorController.cpp b/Engine/Plugins/Animation/IKRig/Source/IKRigEditor/Private/RigEditor/IKRigEditorController.cpp
--- a/Engine/Plugins/Animation/IKRig/Source/IKRigEditor/Private/RigEditor/IKRigEditorController.cpp
+++ b/Engine/Plugins/Animation/IKRig/Source/IKRigEditor/Private/RigEditor/IKRigEditorController.cpp
@@ -878,7 +878,7 @@ void FIKRigEditorController::PromptToAssignGoalToChain(UIKRigEffectorGoal* NewGo
 
 FName FIKRigEditorController::PromptToAddNewRetargetChain(const FBoneChain& BoneChain) const
 {
-	const TSharedPtr<FStructOnScope> StructToDisplay = MakeShareable(new FStructOnScope(FIKRigRetargetChainSettings::StaticStruct(), (uint8*)&BoneChain));
+	const TSharedPtr<FStructOnScope> StructToDisplay = MakeShareable(new FStructOnScope(FBoneChain::StaticStruct(), (uint8*)&BoneChain));
 	TSharedRef<SKismetInspector> KismetInspector = SNew(SKismetInspector);
 	KismetInspector->ShowSingleStruct(StructToDisplay);
 
