Fixes crash when reloading an edited IK Retargeter asset.

diff --git a/Engine/Plugins/Animation/IKRig/Source/IKRigEditor/Private/RetargetEditor/IKRetargetEditor.cpp b/Engine/Plugins/Animation/IKRig/Source/IKRigEditor/Private/RetargetEditor/IKRetargetEditor.cpp
--- a/Engine/Plugins/Animation/IKRig/Source/IKRigEditor/Private/RetargetEditor/IKRetargetEditor.cpp
+++ b/Engine/Plugins/Animation/IKRig/Source/IKRigEditor/Private/RetargetEditor/IKRetargetEditor.cpp
@@ -52,7 +52,7 @@ void FIKRetargetEditor::InitAssetEditor(
 	PersonaToolkitArgs.OnPreviewSceneSettingsCustomized = FOnPreviewSceneSettingsCustomized::FDelegate::CreateSP(this, &FIKRetargetEditor::HandleOnPreviewSceneSettingsCustomized);
 	
 	FPersonaModule& PersonaModule = FModuleManager::LoadModuleChecked<FPersonaModule>("Persona");
-	EditorController->PersonaToolkit = PersonaModule.CreatePersonaToolkit(InAsset, PersonaToolkitArgs);
+	PersonaToolkit = PersonaModule.CreatePersonaToolkit(InAsset, PersonaToolkitArgs);
 	
 	TSharedRef<IAssetFamily> AssetFamily = PersonaModule.CreatePersonaAssetFamily(InAsset);
 	AssetFamily->RecordAssetOpened(FAssetData(InAsset));
@@ -71,7 +71,7 @@ void FIKRetargetEditor::InitAssetEditor(
 	// this sets the application mode which defines the tab factory that builds the editor layout
 	AddApplicationMode(
 		IKRetargetApplicationModes::IKRetargetApplicationMode,
-		MakeShareable(new FIKRetargetApplicationMode(SharedThis(this),EditorController->PersonaToolkit->GetPreviewScene())));
+		MakeShareable(new FIKRetargetApplicationMode(SharedThis(this), PersonaToolkit->GetPreviewScene())));
 	SetCurrentMode(IKRetargetApplicationModes::IKRetargetApplicationMode);
 
 	// set the default editing mode to use in the editor
diff --git a/Engine/Plugins/Animation/IKRig/Source/IKRigEditor/Public/RetargetEditor/IKRetargetEditor.h b/Engine/Plugins/Animation/IKRig/Source/IKRigEditor/Public/RetargetEditor/IKRetargetEditor.h
--- a/Engine/Plugins/Animation/IKRig/Source/IKRigEditor/Public/RetargetEditor/IKRetargetEditor.h
+++ b/Engine/Plugins/Animation/IKRig/Source/IKRigEditor/Public/RetargetEditor/IKRetargetEditor.h
@@ -70,7 +70,7 @@ public:
 	//~ END FTickableEditorObject Interface
 
 	/** IHasPersonaToolkit interface */
-	virtual TSharedRef<IPersonaToolkit> GetPersonaToolkit() const override { return EditorController->PersonaToolkit.ToSharedRef(); }
+	virtual TSharedRef<IPersonaToolkit> GetPersonaToolkit() const override { return PersonaToolkit.ToSharedRef(); }
 	/** END IHasPersonaToolkit interface */
 
 	/** FSelfRegisteringEditorUndoClient interface */
@@ -100,6 +100,9 @@ private:
 	/** centralized management across all views */
 	TSharedRef<FIKRetargetEditorController> EditorController;
 
+	/** Preview scene to be supplied by IHasPersonaToolkit::GetPersonaToolkit */
+	TSharedPtr<IPersonaToolkit> PersonaToolkit;
+
 	/** record previous playback time of source anim instance to trigger reset when scrubbing / looping */
 	float PreviousTime;
 	
diff --git a/Engine/Plugins/Animation/IKRig/Source/IKRigEditor/Public/RetargetEditor/IKRetargetEditorController.h b/Engine/Plugins/Animation/IKRig/Source/IKRigEditor/Public/RetargetEditor/IKRetargetEditorController.h
--- a/Engine/Plugins/Animation/IKRig/Source/IKRigEditor/Public/RetargetEditor/IKRetargetEditorController.h
+++ b/Engine/Plugins/Animation/IKRig/Source/IKRigEditor/Public/RetargetEditor/IKRetargetEditorController.h
@@ -88,9 +88,6 @@ public:
 	/** all modifications to the data model should go through this controller */
 	TObjectPtr<UIKRetargeterController> AssetController;
 
-	/** Preview scene to be supplied by IHasPersonaToolkit::GetPersonaToolkit */
-	TSharedPtr<IPersonaToolkit> PersonaToolkit;
-
 	/** the persona toolkit */
 	TWeakPtr<FIKRetargetEditor> Editor;
 
