Fixed retarget phase settings not saving.

diff --git a/Engine/Plugins/Animation/IKRig/Source/IKRig/Private/Retargeter/IKRetargeter.cpp b/Engine/Plugins/Animation/IKRig/Source/IKRig/Private/Retargeter/IKRetargeter.cpp
--- a/Engine/Plugins/Animation/IKRig/Source/IKRig/Private/Retargeter/IKRetargeter.cpp
+++ b/Engine/Plugins/Animation/IKRig/Source/IKRig/Private/Retargeter/IKRetargeter.cpp
@@ -92,17 +92,17 @@ void UIKRetargeter::PostLoad()
 		}
 
 		// load deprecated global settings
-		if (bRetargetRoot_DEPRECATED)
+		if (!bRetargetRoot_DEPRECATED)
 		{
-			GlobalSettings->Settings.bEnableRoot = bRetargetRoot_DEPRECATED;
+			GlobalSettings->Settings.bEnableRoot = false;
 		}
-		if (bRetargetFK_DEPRECATED)
+		if (!bRetargetFK_DEPRECATED)
 		{
-			GlobalSettings->Settings.bEnableFK = bRetargetFK_DEPRECATED;
+			GlobalSettings->Settings.bEnableFK = false;
 		}
-		if (bRetargetIK_DEPRECATED)
+		if (!bRetargetIK_DEPRECATED)
 		{
-			GlobalSettings->Settings.bEnableIK = bRetargetIK_DEPRECATED;
+			GlobalSettings->Settings.bEnableIK = false;
 		}
 	#endif
 
@@ -110,14 +110,12 @@ void UIKRetargeter::PostLoad()
 	if (!RetargetPoses_DEPRECATED.IsEmpty())
 	{
 		TargetRetargetPoses = RetargetPoses_DEPRECATED;
-		RetargetPoses_DEPRECATED.Empty();
 	}
 
 	// load deprecated current retarget pose (pre adding retarget poses for source)
 	if (CurrentRetargetPose_DEPRECATED != NAME_None)
 	{
 		CurrentTargetRetargetPose = CurrentRetargetPose_DEPRECATED;
-		CurrentRetargetPose_DEPRECATED = NAME_None;
 	}
 	
 	PRAGMA_ENABLE_DEPRECATION_WARNINGS
