Fixes retargeter viewport focusing not working in edit pose mode.

diff --git a/Engine/Plugins/Animation/IKRig/Source/IKRigEditor/Private/RetargetEditor/IKRetargetDefaultMode.cpp b/Engine/Plugins/Animation/IKRig/Source/IKRigEditor/Private/RetargetEditor/IKRetargetDefaultMode.cpp
--- a/Engine/Plugins/Animation/IKRig/Source/IKRigEditor/Private/RetargetEditor/IKRetargetDefaultMode.cpp
+++ b/Engine/Plugins/Animation/IKRig/Source/IKRigEditor/Private/RetargetEditor/IKRetargetDefaultMode.cpp
@@ -24,118 +24,7 @@ bool FIKRetargetDefaultMode::GetCameraTarget(FSphere& OutTarget) const
 	{
 		return false;
 	}
-
-	// center the view on the last selected item
-	switch (Controller->GetLastSelectedItemType())
-	{
-	case ERetargetSelectionType::BONE:
-		{
-			// target the selected bones
-			const TArray<FName>& SelectedBoneNames = Controller->GetSelectedBones();
-			if (SelectedBoneNames.IsEmpty())
-			{
-				return false;
-			}
-
-			TArray<FVector> TargetPoints;
-			const UDebugSkelMeshComponent* SkeletalMeshComponent = Controller->GetSkeletalMeshComponent(Controller->GetSourceOrTarget());
-			const FReferenceSkeleton& RefSkeleton = SkeletalMeshComponent->GetReferenceSkeleton();
-			TArray<int32> ChildrenIndices;
-			for (const FName& SelectedBoneName : SelectedBoneNames)
-			{
-				const int32 BoneIndex = RefSkeleton.FindBoneIndex(SelectedBoneName);
-				if (BoneIndex == INDEX_NONE)
-				{
-					continue;
-				}
-
-				TargetPoints.Add(SkeletalMeshComponent->GetBoneTransform(BoneIndex).GetLocation());
-				ChildrenIndices.Reset();
-				RefSkeleton.GetDirectChildBones(BoneIndex, ChildrenIndices);
-				for (const int32 ChildIndex : ChildrenIndices)
-				{
-					TargetPoints.Add(SkeletalMeshComponent->GetBoneTransform(ChildIndex).GetLocation());
-				}
-			}
-	
-			// create a sphere that contains all the target points
-			if (TargetPoints.Num() == 0)
-			{
-				TargetPoints.Add(FVector::ZeroVector);
-			}
-			OutTarget = FSphere(&TargetPoints[0], TargetPoints.Num());
-			return true;
-		}
-		
-	case ERetargetSelectionType::CHAIN:
-		{
-			const ERetargetSourceOrTarget SourceOrTarget = Controller->GetSourceOrTarget();
-			const UIKRigDefinition* IKRig = Controller->AssetController->GetIKRig(SourceOrTarget);
-			if (!IKRig)
-			{
-				return false;
-			}
-
-			const UDebugSkelMeshComponent* SkeletalMeshComponent = Controller->GetSkeletalMeshComponent(Controller->GetSourceOrTarget());
-			const FReferenceSkeleton& RefSkeleton = SkeletalMeshComponent->GetReferenceSkeleton();
-
-			// get target points from start/end bone of all selected chains on the currently active skeleton (source or target)
-			TArray<FVector> TargetPoints;
-			const TArray<FName>& SelectedChains = Controller->GetSelectedChains();
-			for (const FName SelectedChainName : SelectedChains)
-			{
-				const URetargetChainSettings* ChainMap = Controller->AssetController->GetChainMappingByTargetChainName(SelectedChainName);
-				if (!ChainMap)
-				{
-					continue;
-				}
-
-				const FName& ChainName = SourceOrTarget == ERetargetSourceOrTarget::Target ? SelectedChainName : ChainMap->SourceChain;
-				if (ChainName == NAME_None)
-				{
-					continue;
-				}
-
-				const FBoneChain* BoneChain = IKRig->GetRetargetChainByName(ChainName);
-				if (!BoneChain)
-				{
-					continue;
-				}
-
-				const int32 StartBoneIndex = RefSkeleton.FindBoneIndex(BoneChain->StartBone.BoneName);
-				if (StartBoneIndex != INDEX_NONE)
-				{
-					TargetPoints.Add(SkeletalMeshComponent->GetBoneTransform(StartBoneIndex).GetLocation());
-				}
-
-				const int32 EndBoneIndex = RefSkeleton.FindBoneIndex(BoneChain->EndBone.BoneName);
-				if (EndBoneIndex != INDEX_NONE)
-				{
-					TargetPoints.Add(SkeletalMeshComponent->GetBoneTransform(EndBoneIndex).GetLocation());
-				}
-			}
-
-			// create a sphere that contains all the target points
-			if (TargetPoints.Num() == 0)
-			{
-				TargetPoints.Add(FVector::ZeroVector);
-			}
-			OutTarget = FSphere(&TargetPoints[0], TargetPoints.Num());
-			return true;	
-		}
-	case ERetargetSelectionType::ROOT:
-	case ERetargetSelectionType::MESH:
-	case ERetargetSelectionType::NONE:
-	default:
-		// target the current mesh
-		if (const UPrimitiveComponent* SelectedMesh = Controller->GetSkeletalMeshComponent(Controller->GetSourceOrTarget()))
-		{
-			OutTarget = SelectedMesh->Bounds.GetSphere();
-			return true;
-		}
-	}
-
-	return false;
+	return Controller->GetCameraTargetForSelection(OutTarget);
 }
 
 IPersonaPreviewScene& FIKRetargetDefaultMode::GetAnimPreviewScene() const
diff --git a/Engine/Plugins/Animation/IKRig/Source/IKRigEditor/Private/RetargetEditor/IKRetargetEditPoseMode.cpp b/Engine/Plugins/Animation/IKRig/Source/IKRigEditor/Private/RetargetEditor/IKRetargetEditPoseMode.cpp
--- a/Engine/Plugins/Animation/IKRig/Source/IKRigEditor/Private/RetargetEditor/IKRetargetEditPoseMode.cpp
+++ b/Engine/Plugins/Animation/IKRig/Source/IKRigEditor/Private/RetargetEditor/IKRetargetEditPoseMode.cpp
@@ -22,17 +22,12 @@ FName FIKRetargetEditPoseMode::ModeName("IKRetargetAssetEditMode");
 
 bool FIKRetargetEditPoseMode::GetCameraTarget(FSphere& OutTarget) const
 {
-	// target skeletal mesh
-	if (const TSharedPtr<FIKRetargetEditorController> Controller = EditorController.Pin())
+	const TSharedPtr<FIKRetargetEditorController> Controller = EditorController.Pin();
+	if (!Controller.IsValid())
 	{
-		if (Controller->TargetSkelMeshComponent)
-		{
-			OutTarget = Controller->TargetSkelMeshComponent->Bounds.GetSphere();
-			return true;
-		}
+		return false;
 	}
-	
-	return false;
+	return Controller->GetCameraTargetForSelection(OutTarget);
 }
 
 IPersonaPreviewScene& FIKRetargetEditPoseMode::GetAnimPreviewScene() const
diff --git a/Engine/Plugins/Animation/IKRig/Source/IKRigEditor/Private/RetargetEditor/IKRetargetEditorController.cpp b/Engine/Plugins/Animation/IKRig/Source/IKRigEditor/Private/RetargetEditor/IKRetargetEditorController.cpp
--- a/Engine/Plugins/Animation/IKRig/Source/IKRigEditor/Private/RetargetEditor/IKRetargetEditorController.cpp
+++ b/Engine/Plugins/Animation/IKRig/Source/IKRigEditor/Private/RetargetEditor/IKRetargetEditorController.cpp
@@ -199,6 +199,121 @@ void FIKRetargetEditorController::AddOffsetToMeshComponent(const FVector& Offset
 	MeshComponent->SetWorldScale3D(FVector(Scale,Scale,Scale));
 }
 
+bool FIKRetargetEditorController::GetCameraTargetForSelection(FSphere& OutTarget) const
+{
+	// center the view on the last selected item
+	switch (GetLastSelectedItemType())
+	{
+	case ERetargetSelectionType::BONE:
+		{
+			// target the selected bones
+			const TArray<FName>& SelectedBones = GetSelectedBones();
+			if (SelectedBones.IsEmpty())
+			{
+				return false;
+			}
+
+			TArray<FVector> TargetPoints;
+			const UDebugSkelMeshComponent* SkeletalMeshComponent = GetSkeletalMeshComponent(GetSourceOrTarget());
+			const FReferenceSkeleton& RefSkeleton = SkeletalMeshComponent->GetReferenceSkeleton();
+			TArray<int32> ChildrenIndices;
+			for (const FName& SelectedBoneName : SelectedBones)
+			{
+				const int32 BoneIndex = RefSkeleton.FindBoneIndex(SelectedBoneName);
+				if (BoneIndex == INDEX_NONE)
+				{
+					continue;
+				}
+
+				TargetPoints.Add(SkeletalMeshComponent->GetBoneTransform(BoneIndex).GetLocation());
+				ChildrenIndices.Reset();
+				RefSkeleton.GetDirectChildBones(BoneIndex, ChildrenIndices);
+				for (const int32 ChildIndex : ChildrenIndices)
+				{
+					TargetPoints.Add(SkeletalMeshComponent->GetBoneTransform(ChildIndex).GetLocation());
+				}
+			}
+	
+			// create a sphere that contains all the target points
+			if (TargetPoints.Num() == 0)
+			{
+				TargetPoints.Add(FVector::ZeroVector);
+			}
+			OutTarget = FSphere(&TargetPoints[0], TargetPoints.Num());
+			return true;
+		}
+		
+	case ERetargetSelectionType::CHAIN:
+		{
+			const ERetargetSourceOrTarget SourceOrTarget = GetSourceOrTarget();
+			const UIKRigDefinition* IKRig = AssetController->GetIKRig(SourceOrTarget);
+			if (!IKRig)
+			{
+				return false;
+			}
+
+			const UDebugSkelMeshComponent* SkeletalMeshComponent = GetSkeletalMeshComponent(GetSourceOrTarget());
+			const FReferenceSkeleton& RefSkeleton = SkeletalMeshComponent->GetReferenceSkeleton();
+
+			// get target points from start/end bone of all selected chains on the currently active skeleton (source or target)
+			TArray<FVector> TargetPoints;
+			const TArray<FName>& SelectedChainNames = GetSelectedChains();
+			for (const FName SelectedChainName : SelectedChainNames)
+			{
+				const URetargetChainSettings* ChainMap = AssetController->GetChainMappingByTargetChainName(SelectedChainName);
+				if (!ChainMap)
+				{
+					continue;
+				}
+
+				const FName& ChainName = SourceOrTarget == ERetargetSourceOrTarget::Target ? SelectedChainName : ChainMap->SourceChain;
+				if (ChainName == NAME_None)
+				{
+					continue;
+				}
+
+				const FBoneChain* BoneChain = IKRig->GetRetargetChainByName(ChainName);
+				if (!BoneChain)
+				{
+					continue;
+				}
+
+				const int32 StartBoneIndex = RefSkeleton.FindBoneIndex(BoneChain->StartBone.BoneName);
+				if (StartBoneIndex != INDEX_NONE)
+				{
+					TargetPoints.Add(SkeletalMeshComponent->GetBoneTransform(StartBoneIndex).GetLocation());
+				}
+
+				const int32 EndBoneIndex = RefSkeleton.FindBoneIndex(BoneChain->EndBone.BoneName);
+				if (EndBoneIndex != INDEX_NONE)
+				{
+					TargetPoints.Add(SkeletalMeshComponent->GetBoneTransform(EndBoneIndex).GetLocation());
+				}
+			}
+
+			// create a sphere that contains all the target points
+			if (TargetPoints.Num() == 0)
+			{
+				TargetPoints.Add(FVector::ZeroVector);
+			}
+			OutTarget = FSphere(&TargetPoints[0], TargetPoints.Num());
+			return true;	
+		}
+	case ERetargetSelectionType::ROOT:
+	case ERetargetSelectionType::MESH:
+	case ERetargetSelectionType::NONE:
+	default:
+		// target the current mesh
+		if (const UPrimitiveComponent* CurrentMesh = GetSkeletalMeshComponent(GetSourceOrTarget()))
+		{
+			OutTarget = CurrentMesh->Bounds.GetSphere();
+			return true;
+		}
+	}
+
+	return false;
+}
+
 bool FIKRetargetEditorController::IsBoneRetargeted(const FName& BoneName, ERetargetSourceOrTarget SourceOrTarget) const
 {
 	// get an initialized processor
diff --git a/Engine/Plugins/Animation/IKRig/Source/IKRigEditor/Public/RetargetEditor/IKRetargetEditorController.h b/Engine/Plugins/Animation/IKRig/Source/IKRigEditor/Public/RetargetEditor/IKRetargetEditorController.h
--- a/Engine/Plugins/Animation/IKRig/Source/IKRigEditor/Public/RetargetEditor/IKRetargetEditorController.h
+++ b/Engine/Plugins/Animation/IKRig/Source/IKRigEditor/Public/RetargetEditor/IKRetargetEditorController.h
@@ -209,6 +209,8 @@ public:
 	void ClearSelection(const bool bKeepBoneSelection=false);
 	ERetargetSelectionType GetLastSelectedItemType() const { return LastSelectedItem; };
 
+	bool GetCameraTargetForSelection(FSphere& OutTarget) const; 
+
 	/** ------------------------- END SELECTION -----------------------------*/
 
 	/** determine if bone in the specified skeleton is part of the retarget (in a mapped chain) */
