Fixing hard cap on retargeting Speed Threshold.

diff --git a/Engine/Plugins/Animation/IKRig/Source/IKRig/Public/Retargeter/IKRetargetSettings.h b/Engine/Plugins/Animation/IKRig/Source/IKRig/Public/Retargeter/IKRetargetSettings.h
--- a/Engine/Plugins/Animation/IKRig/Source/IKRig/Public/Retargeter/IKRetargetSettings.h
+++ b/Engine/Plugins/Animation/IKRig/Source/IKRig/Public/Retargeter/IKRetargetSettings.h
@@ -40,7 +40,7 @@ struct IKRIG_API FTargetChainSpeedPlantSettings
 
 	/** Range 0 to 1000. Default 15. The maximum speed a source bone can be moving while being considered 'planted'.
 	*  The target IK goal will not be allowed to move whenever the source bone speed drops below this threshold speed. */
-	UPROPERTY(EditAnywhere, BlueprintReadWrite, Category = "Plant IK by Speed", meta = (ClampMin = "0.0", ClampMax = "100.0", UIMin = "0.0", UIMax = "100.0"))
+	UPROPERTY(EditAnywhere, BlueprintReadWrite, Category = "Plant IK by Speed", meta = (ClampMin = "0.0", UIMin = "0.0", UIMax = "100.0"))
 	float SpeedThreshold = 15.0f;
 
 	// How stiff the spring model is that smoothly pulls the IK position after unplanting (more stiffness means more oscillation around the target value)
