Adding "auto-map" options menu to retargeter chain map list.

diff --git a/Engine/Plugins/Animation/IKRig/Source/IKRigEditor/Private/RetargetEditor/IKRetargetEditor.cpp b/Engine/Plugins/Animation/IKRig/Source/IKRigEditor/Private/RetargetEditor/IKRetargetEditor.cpp
--- a/Engine/Plugins/Animation/IKRig/Source/IKRigEditor/Private/RetargetEditor/IKRetargetEditor.cpp
+++ b/Engine/Plugins/Animation/IKRig/Source/IKRigEditor/Private/RetargetEditor/IKRetargetEditor.cpp
@@ -435,7 +435,8 @@ void FIKRetargetEditor::OnFinishedChangingDetails(const FPropertyChangedEvent& P
 		EditorController->ClearOutputLog();
 		EditorController->BindToIKRigAssets(AssetController->GetAsset());
 		EditorController->AssetController->CleanChainMapping();
-		EditorController->AssetController->AutoMapChains();
+		constexpr bool bForceRemap = false;
+		EditorController->AssetController->AutoMapChains(EAutoMapChainType::Fuzzy, bForceRemap);
 	}
 
 	// if either the source or target meshes are possibly modified, update scene components, anim instance and UI
diff --git a/Engine/Plugins/Animation/IKRig/Source/IKRigEditor/Private/RetargetEditor/IKRetargeterController.cpp b/Engine/Plugins/Animation/IKRig/Source/IKRigEditor/Private/RetargetEditor/IKRetargeterController.cpp
--- a/Engine/Plugins/Animation/IKRig/Source/IKRigEditor/Private/RetargetEditor/IKRetargeterController.cpp
+++ b/Engine/Plugins/Animation/IKRig/Source/IKRigEditor/Private/RetargetEditor/IKRetargeterController.cpp
@@ -237,40 +237,77 @@ void UIKRetargeterController::CleanPoseList(const ERetargetSourceOrTarget& Sourc
 	}
 }
 
-void UIKRetargeterController::AutoMapChains() const
+void UIKRetargeterController::AutoMapChains(const EAutoMapChainType AutoMapType, const bool bForceRemap) const
 {
+	FScopedTransaction Transaction(LOCTEXT("AutoMapRetargetChains", "Auto-Map Retarget Chains"));
+	
+	CleanChainMapping();
+
+	// get names of all the potential chains we could map to on the source
 	TArray<FName> SourceChainNames;
 	GetChainNames(ERetargetSourceOrTarget::Source, SourceChainNames);
 	
-	// auto-map any chains that have no value using a fuzzy string search
+	// iterate over all the chain mappings and find matching source chain
 	for (URetargetChainSettings* ChainMap : Asset->ChainSettings)
 	{
-		if (ChainMap->SourceChain != NAME_None)
+		const bool bIsMappedAlready = ChainMap->SourceChain != NAME_None;
+		if (bIsMappedAlready && !bForceRemap)
 		{
 			continue; // already set by user
 		}
 
-		// find "best match" automatically as a convenience for the user
-		FString TargetNameLowerCase = ChainMap->TargetChain.ToString().ToLower();
-		float HighestScore = 0.2f;
-		int32 HighestScoreIndex = -1;
-		for (int32 ChainIndex=0; ChainIndex<SourceChainNames.Num(); ++ChainIndex)
+		ChainMap->Modify();
+
+		// find a source chain to map to
+		int32 SourceChainIndexToMapTo = -1;
+
+		switch (AutoMapType)
 		{
-			FString SourceNameLowerCase = SourceChainNames[ChainIndex].ToString().ToLower();
-			float WorstCase = TargetNameLowerCase.Len() + SourceNameLowerCase.Len();
-			WorstCase = WorstCase < 1.0f ? 1.0f : WorstCase;
-			const float Score = 1.0f - (Algo::LevenshteinDistance(TargetNameLowerCase, SourceNameLowerCase) / WorstCase);
-			if (Score > HighestScore)
+		case EAutoMapChainType::Fuzzy:
 			{
-				HighestScore = Score;
-				HighestScoreIndex = ChainIndex;
+				// auto-map chains using a fuzzy string comparison
+				FString TargetNameLowerCase = ChainMap->TargetChain.ToString().ToLower();
+				float HighestScore = 0.2f;
+				for (int32 ChainIndex=0; ChainIndex<SourceChainNames.Num(); ++ChainIndex)
+				{
+					FString SourceNameLowerCase = SourceChainNames[ChainIndex].ToString().ToLower();
+					float WorstCase = TargetNameLowerCase.Len() + SourceNameLowerCase.Len();
+					WorstCase = WorstCase < 1.0f ? 1.0f : WorstCase;
+					const float Score = 1.0f - (Algo::LevenshteinDistance(TargetNameLowerCase, SourceNameLowerCase) / WorstCase);
+					if (Score > HighestScore)
+					{
+						HighestScore = Score;
+						SourceChainIndexToMapTo = ChainIndex;
+					}
+				}
+				break;
+			}
+		case EAutoMapChainType::Exact:
+			{
+				// auto-map chains with same name
+				for (int32 ChainIndex=0; ChainIndex<SourceChainNames.Num(); ++ChainIndex)
+				{
+					if (SourceChainNames[ChainIndex] == ChainMap->TargetChain)
+					{
+						SourceChainIndexToMapTo = ChainIndex;
+						break;
+					}
+				}
+				break;
+			}
+		case EAutoMapChainType::Clear:
+			{
+				ChainMap->SourceChain = NAME_None;
+				break;
 			}
+		default:
+			checkNoEntry();
 		}
 
 		// apply source if any decent matches were found
-		if (SourceChainNames.IsValidIndex(HighestScoreIndex))
+		if (SourceChainNames.IsValidIndex(SourceChainIndexToMapTo))
 		{
-			ChainMap->SourceChain = SourceChainNames[HighestScoreIndex];
+			ChainMap->SourceChain = SourceChainNames[SourceChainIndexToMapTo];
 		}
 	}
 
diff --git a/Engine/Plugins/Animation/IKRig/Source/IKRigEditor/Private/RetargetEditor/SIKRetargetChainMapList.cpp b/Engine/Plugins/Animation/IKRig/Source/IKRigEditor/Private/RetargetEditor/SIKRetargetChainMapList.cpp
--- a/Engine/Plugins/Animation/IKRig/Source/IKRigEditor/Private/RetargetEditor/SIKRetargetChainMapList.cpp
+++ b/Engine/Plugins/Animation/IKRig/Source/IKRigEditor/Private/RetargetEditor/SIKRetargetChainMapList.cpp
@@ -303,19 +303,6 @@ void SIKRetargetChainMapList::Construct(
 		.AutoHeight()
 		[
 			SNew(SHorizontalBox)
-
-			+SHorizontalBox::Slot()
-			.AutoWidth()
-			.VAlign(VAlign_Center)
-			.Padding(FMargin(6.f, 0.0))
-			[
-				SNew(SPositiveActionButton)
-				.Visibility(this, &SIKRetargetChainMapList::IsAutoMapButtonVisible)
-				.Icon(FAppStyle::Get().GetBrush("Icons.Refresh"))
-				.Text(LOCTEXT("AutoMapButtonLabel", "Auto-Map Chains"))
-				.ToolTipText(LOCTEXT("AutoMapButtonToolTip", "Automatically assign source chains based on fuzzy string match"))
-				.OnClicked(this, &SIKRetargetChainMapList::OnAutoMapButtonClicked)
-			]
 			
 			+SHorizontalBox::Slot()
 			.FillWidth(1.0f)
@@ -345,6 +332,25 @@ void SIKRetargetChainMapList::Construct(
 					.ColorAndOpacity(FSlateColor::UseForeground())
 				]
 			]
+
+			+SHorizontalBox::Slot()
+			.AutoWidth()
+			.VAlign(VAlign_Center)
+			.Padding(FMargin(6.f, 0.0))
+			[
+				SNew(SComboButton)
+				.ComboButtonStyle(&FAppStyle::Get().GetWidgetStyle<FComboButtonStyle>("SimpleComboButton"))
+				.ForegroundColor(FSlateColor::UseStyle())
+				.ContentPadding(2.0f)
+				.Visibility(this, &SIKRetargetChainMapList::IsAutoMapButtonVisible)
+				.ToolTipText(LOCTEXT("AutoMapButtonToolTip", "Automatically assign source chains based on matching rule."))
+				.OnGetMenuContent( this, &SIKRetargetChainMapList::CreateChainMapMenuWidget )
+				.HasDownArrow(true)
+				.ButtonContent()
+				[
+					SNew(STextBlock).Text(LOCTEXT("AutoMapButtonLabel", "Auto-Map Chains"))
+				]
+			]
 		]
 	    
         +SVerticalBox::Slot()
@@ -641,29 +647,129 @@ void SIKRetargetChainMapList::OnItemClicked(TSharedPtr<FRetargetChainMapElement>
 	EditorController.Pin()->EditChainSelection(SelectedChains, ESelectionEdit::Replace, bEditFromChainsView);
 }
 
+TSharedRef<SWidget> SIKRetargetChainMapList::CreateChainMapMenuWidget()
+{
+	const FUIAction MapAllByFuzzyNameAction = FUIAction(
+		FExecuteAction::CreateLambda([this]
+		{
+			constexpr bool bForceRemap = true;
+			AutoMapChains(EAutoMapChainType::Fuzzy, bForceRemap);
+		}),
+		FCanExecuteAction(), FIsActionChecked());
+	
+	const FUIAction MapAllByExactNameAction = FUIAction(
+		FExecuteAction::CreateLambda([this]
+		{
+			constexpr bool bForceRemap = true;
+			AutoMapChains(EAutoMapChainType::Exact, bForceRemap);
+		}),
+		FCanExecuteAction(), FIsActionChecked());
+
+	const FUIAction MapUnmappedByExactNameAction = FUIAction(
+		FExecuteAction::CreateLambda([this]
+		{
+			constexpr bool bForceRemap = false;
+			AutoMapChains(EAutoMapChainType::Exact, bForceRemap);
+		}),
+		FCanExecuteAction(), FIsActionChecked());
+
+	const FUIAction MapUnmappedByFuzzyNameAction = FUIAction(
+		FExecuteAction::CreateLambda([this]
+		{
+			constexpr bool bForceRemap = false;
+			AutoMapChains(EAutoMapChainType::Fuzzy, bForceRemap);
+		}),
+		FCanExecuteAction(), FIsActionChecked());
+
+	const FUIAction ClearAllMappingsAction = FUIAction(
+		FExecuteAction::CreateLambda([this]
+		{
+			constexpr bool bForceRemap = true;
+			AutoMapChains(EAutoMapChainType::Clear, bForceRemap);
+		}),
+		FCanExecuteAction(), FIsActionChecked());
+	
+	static constexpr bool CloseAfterSelection = true;
+	FMenuBuilder MenuBuilder(CloseAfterSelection, CommandList);
+
+	MenuBuilder.BeginSection("Auto-Map Chains Fuzzy", LOCTEXT("FuzzyNameSection", "Fuzzy Name Matching"));
+
+	MenuBuilder.AddMenuEntry(
+		LOCTEXT("MapAllByNameFuzzyLabel", "Map All (Fuzzy)"),
+		LOCTEXT("MapAllByNameFuzzyTooltip", "Map all chains to the source chain with the closest name (not necessarily exact)."),
+		FSlateIcon(),
+		MapAllByFuzzyNameAction,
+		NAME_None,
+		EUserInterfaceActionType::Button);
+
+	MenuBuilder.AddMenuEntry(
+		LOCTEXT("MapMissingByNameFuzzyLabel", "Map Only Empty (Fuzzy)"),
+		LOCTEXT("MapMissingByNameFuzzyTooltip", "Map all unmapped chains to the source chain with the closest name (not necessarily exact)."),
+		FSlateIcon(),
+		MapUnmappedByFuzzyNameAction,
+		NAME_None,
+		EUserInterfaceActionType::Button);
+
+	MenuBuilder.EndSection();
+	MenuBuilder.BeginSection("Auto-Map Chains Exact", LOCTEXT("ExactNameSection", "Exact Name Matching"));
+
+	MenuBuilder.AddMenuEntry(
+		LOCTEXT("MapAllByNameExactLabel", "Map All (Exact)"),
+		LOCTEXT("MapAllByNameExactTooltip", "Map all chains with identical name. If no match found, does not change mapping."),
+		FSlateIcon(),
+		MapAllByExactNameAction,
+		NAME_None,
+		EUserInterfaceActionType::Button);
+	
+	MenuBuilder.AddMenuEntry(
+		LOCTEXT("MapMissingByNameExactLabel", "Map Only Empty (Exact)"),
+		LOCTEXT("MapMissingByNameExactTooltip", "Map unmapped chains using identical name. If no match found, does not change mapping."),
+		FSlateIcon(),
+		MapUnmappedByExactNameAction,
+		NAME_None,
+		EUserInterfaceActionType::Button);
+
+	MenuBuilder.EndSection();
+	MenuBuilder.BeginSection("Clear", LOCTEXT("ClearMapSection", "Clear All"));
+	
+	MenuBuilder.AddMenuEntry(
+		LOCTEXT("ClearMapLabel", "Clear All Mappings"),
+		LOCTEXT("ClearMapTooltip", "Map all chains to None."),
+		FSlateIcon(),
+		ClearAllMappingsAction,
+		NAME_None,
+		EUserInterfaceActionType::Button);
+
+	MenuBuilder.EndSection();
+	
+	return MenuBuilder.MakeWidget();
+}
+
 EVisibility SIKRetargetChainMapList::IsAutoMapButtonVisible() const
 {
 	return IsChainMapEnabled() ? EVisibility::Visible : EVisibility::Hidden;
 }
 
-FReply SIKRetargetChainMapList::OnAutoMapButtonClicked() const
+void SIKRetargetChainMapList::AutoMapChains(const EAutoMapChainType AutoMapType, const bool bForceRemap)
 {
 	const TSharedPtr<FIKRetargetEditorController> Controller = EditorController.Pin();
 	if (!Controller.IsValid())
 	{
-		return FReply::Unhandled();
+		return;
 	}
 	
 	UIKRetargeterController* RetargeterController = GetRetargetController();
 	if (!RetargeterController)
 	{
-		return FReply::Unhandled();
+		return;
 	}
 
+	// clear output log
 	Controller->ClearOutputLog();
-	RetargeterController->CleanChainMapping();
-	RetargeterController->AutoMapChains();
-	return FReply::Handled();
+	// run the chain mapping
+	RetargeterController->AutoMapChains(AutoMapType, bForceRemap);
+	// refresh the UI to show new mapping
+	RefreshView();
 }
 
 FReply SIKRetargetChainMapList::OnGlobalSettingsButtonClicked() const
@@ -710,5 +816,4 @@ FReply SIKRetargetChainMapList::OnRootSettingsButtonClicked() const
 }
 
 
-
 #undef LOCTEXT_NAMESPACE
diff --git a/Engine/Plugins/Animation/IKRig/Source/IKRigEditor/Public/RetargetEditor/IKRetargeterController.h b/Engine/Plugins/Animation/IKRig/Source/IKRigEditor/Public/RetargetEditor/IKRetargeterController.h
--- a/Engine/Plugins/Animation/IKRig/Source/IKRigEditor/Public/RetargetEditor/IKRetargeterController.h
+++ b/Engine/Plugins/Animation/IKRig/Source/IKRigEditor/Public/RetargetEditor/IKRetargeterController.h
@@ -14,6 +14,13 @@ class UIKRigDefinition;
 class UIKRetargeter;
 class USkeletalMesh;
 
+enum class EAutoMapChainType : uint8
+{
+	Exact, // map chains that have exactly the same name (case insensitive)
+	Fuzzy, // map chains to the chain with the closest name (levenshtein distance)
+	Clear, // clear all mappings, set them to None
+};
+
 /** A singleton (per-asset) class used to make modifications to a UIKRetargeter asset.
 * Call the static GetController(UIKRetargeter* Asset) function to get the controller for the asset you want to modify. */ 
 UCLASS(config = Engine, hidecategories = UObject)
@@ -57,7 +64,7 @@ public:
 	/** Remove invalid chain mappings (no longer existing in currently referenced source/target IK Rig assets) */
 	void CleanChainMapping(const bool bForceReinitialization=true) const;
 	/** Use fuzzy string search to find "best" Source chain to map to each Target chain */
-	void AutoMapChains() const;
+	void AutoMapChains(const EAutoMapChainType AutoMapType, const bool bForceRemap) const;
 	/** Callback when IK Rig chain is added or removed. */
 	void OnRetargetChainAdded(UIKRigDefinition* IKRig) const;
 	/** Callback when IK Rig chain is renamed. Retains existing mappings using the new name */
@@ -153,7 +160,4 @@ private:
 
 	/** The actual asset that this Controller modifies. */
 	TObjectPtr<UIKRetargeter> Asset = nullptr;
-
-	/** The editor controller for this asset. */
-	TSharedPtr<FIKRetargetEditorController> EditorController = nullptr;
 };
diff --git a/Engine/Plugins/Animation/IKRig/Source/IKRigEditor/Public/RetargetEditor/SIKRetargetChainMapList.h b/Engine/Plugins/Animation/IKRig/Source/IKRigEditor/Public/RetargetEditor/SIKRetargetChainMapList.h
--- a/Engine/Plugins/Animation/IKRig/Source/IKRigEditor/Public/RetargetEditor/SIKRetargetChainMapList.h
+++ b/Engine/Plugins/Animation/IKRig/Source/IKRigEditor/Public/RetargetEditor/SIKRetargetChainMapList.h
@@ -128,8 +128,9 @@ private:
 	FReply OnRootSettingsButtonClicked() const;
 
 	/** auto-map chain button*/
+	TSharedRef<SWidget> CreateChainMapMenuWidget();
 	EVisibility IsAutoMapButtonVisible() const;
-	FReply OnAutoMapButtonClicked() const;
+	void AutoMapChains(const EAutoMapChainType AutoMapType, const bool bSetUnmatchedToNone);
 	/** END auto-map chain button*/
 
 	/** filtering the list with search box */
