Fixed issue, unable to retarget IK on a bone chain with less than 3 bones.

diff --git a/Engine/Plugins/Animation/IKRig/Source/IKRig/Private/Retargeter/IKRetargetProcessor.cpp b/Engine/Plugins/Animation/IKRig/Source/IKRig/Private/Retargeter/IKRetargetProcessor.cpp
--- a/Engine/Plugins/Animation/IKRig/Source/IKRig/Private/Retargeter/IKRetargetProcessor.cpp
+++ b/Engine/Plugins/Animation/IKRig/Source/IKRig/Private/Retargeter/IKRetargetProcessor.cpp
@@ -843,22 +843,21 @@ bool FChainRetargeterIK::InitializeSource(
 	const TArray<FTransform>& SourceInitialGlobalPose,
 	FIKRigLogger& Log)
 {
-	if (BoneIndices.Num() < 3)
+	if (BoneIndices.Num() < 2)
 	{
-		Log.LogWarning(LOCTEXT("SourceChainLessThanThree", "IK Retargeter trying to retarget source bone chain with IK but it has less than 3 joints."));
+		Log.LogWarning(LOCTEXT("SourceChainLessThanThree", "IK Retargeter trying to retarget source bone chain with IK but it has less than 2 joints."));
 		return false;
 	}
 	
-	Source.BoneIndexA = BoneIndices[0];
-	Source.BoneIndexB = BoneIndices[1];
-	Source.BoneIndexC = BoneIndices.Last();
+	Source.StartBoneIndex = BoneIndices[0];
+	Source.EndBoneIndex = BoneIndices.Last();
 	
-	const FTransform& End = SourceInitialGlobalPose[Source.BoneIndexC];
+	const FTransform& End = SourceInitialGlobalPose[Source.EndBoneIndex];
 	Source.PreviousEndPosition = End.GetTranslation();
 	Source.InitialEndPosition = End.GetTranslation();
 	Source.InitialEndRotation = End.GetRotation();
 
-	const FTransform& Start = SourceInitialGlobalPose[Source.BoneIndexA];
+	const FTransform& Start = SourceInitialGlobalPose[Source.StartBoneIndex];
 	const float Length = (Start.GetTranslation() - Source.InitialEndPosition).Size();
 
 	if (Length <= KINDA_SMALL_NUMBER)
@@ -874,23 +873,21 @@ bool FChainRetargeterIK::InitializeSource(
 
 void FChainRetargeterIK::EncodePose(const TArray<FTransform>& InSourceGlobalPose)
 {
-	const FVector A = InSourceGlobalPose[Source.BoneIndexA].GetTranslation();
-	//FVector B = InputGlobalPose[BoneIndexB].GetTranslation(); TODO use for pole vector 
-	const FVector C = InSourceGlobalPose[Source.BoneIndexC].GetTranslation();
+	const FVector Start = InSourceGlobalPose[Source.StartBoneIndex].GetTranslation();
+	const FVector End = InSourceGlobalPose[Source.EndBoneIndex].GetTranslation();
 
     // get the normalized direction / length of the IK limb (how extended it is as percentage of original length)
-    const FVector AC = C - A;
-	float ACLength;
-	FVector ACDirection;
-	AC.ToDirectionAndLength(ACDirection, ACLength);
-	const float NormalizedLimbLength = ACLength * Source.InvInitialLength;
+    const FVector ChainVector = End - Start;
+	float ChainLength;
+	FVector ChainDirection;
+	ChainVector.ToDirectionAndLength(ChainDirection, ChainLength);
+	const float NormalizedLimbLength = ChainLength * Source.InvInitialLength;
 
 	Source.PreviousEndPosition = Source.CurrentEndPosition;
-	Source.CurrentEndPosition = C;
-	Source.CurrentEndDirectionNormalized = ACDirection * NormalizedLimbLength;
-	Source.CurrentEndRotation = InSourceGlobalPose[Source.BoneIndexC].GetRotation();
-	Source.CurrentHeightFromGroundNormalized = (C.Z - Source.InitialEndPosition.Z)  * Source.InvInitialLength;
-	Source.PoleVectorDirection = FVector::OneVector; // TBD
+	Source.CurrentEndPosition = End;
+	Source.CurrentEndDirectionNormalized = ChainDirection * NormalizedLimbLength;
+	Source.CurrentEndRotation = InSourceGlobalPose[Source.EndBoneIndex].GetRotation();
+	Source.CurrentHeightFromGroundNormalized = (End.Z - Source.InitialEndPosition.Z)  * Source.InvInitialLength;
 }
 
 bool FChainRetargeterIK::InitializeTarget(
diff --git a/Engine/Plugins/Animation/IKRig/Source/IKRig/Public/Retargeter/IKRetargetProcessor.h b/Engine/Plugins/Animation/IKRig/Source/IKRig/Public/Retargeter/IKRetargetProcessor.h
--- a/Engine/Plugins/Animation/IKRig/Source/IKRig/Public/Retargeter/IKRetargetProcessor.h
+++ b/Engine/Plugins/Animation/IKRig/Source/IKRig/Public/Retargeter/IKRetargetProcessor.h
@@ -322,9 +322,8 @@ struct FDecodedIKChain
 
 struct FSourceChainIK
 {
-	int32 BoneIndexA = INDEX_NONE;
-	int32 BoneIndexB = INDEX_NONE;
-	int32 BoneIndexC = INDEX_NONE;
+	int32 StartBoneIndex = INDEX_NONE;
+	int32 EndBoneIndex = INDEX_NONE;
 	
 	FVector InitialEndPosition = FVector::ZeroVector;
 	FQuat InitialEndRotation = FQuat::Identity;
@@ -336,7 +335,6 @@ struct FSourceChainIK
 	FVector CurrentEndDirectionNormalized = FVector::ZeroVector;
 	FQuat CurrentEndRotation = FQuat::Identity;
 	float CurrentHeightFromGroundNormalized = 0.0f;
-	FVector PoleVectorDirection = FVector::ZeroVector;
 };
 
 struct FTargetChainIK
