Fixes rotation gizmo losing focus in viewport when selecting bone in retarget hierarchy view.

diff --git a/Engine/Plugins/Animation/IKRig/Source/IKRigEditor/Private/RetargetEditor/IKRetargetEditorController.cpp b/Engine/Plugins/Animation/IKRig/Source/IKRigEditor/Private/RetargetEditor/IKRetargetEditorController.cpp
--- a/Engine/Plugins/Animation/IKRig/Source/IKRigEditor/Private/RetargetEditor/IKRetargetEditorController.cpp
+++ b/Engine/Plugins/Animation/IKRig/Source/IKRigEditor/Private/RetargetEditor/IKRetargetEditorController.cpp
@@ -710,6 +710,12 @@ void FIKRetargetEditorController::EditBoneSelection(
 	{
 		RefreshHierarchyView();
 	}
+	else
+	{
+		// If selection was made from the hierarchy view, the viewport must be invalidated for the
+		// new widget hit proxies to be activated. Otherwise user has to click in the viewport first to gain focus.
+		Editor.Pin()->GetPersonaToolkit()->GetPreviewScene()->InvalidateViews();
+	}
 
 	// update details
 	if (SelectedBoneNames[CurrentlyEditingSourceOrTarget].IsEmpty())
