Fixing issue,undoing swapping prevew mesh in IK Rig doesn't update the viewport.

diff --git a/Engine/Plugins/Animation/IKRig/Source/IKRigEditor/Private/RigEditor/IKRigController.cpp b/Engine/Plugins/Animation/IKRig/Source/IKRigEditor/Private/RigEditor/IKRigController.cpp
--- a/Engine/Plugins/Animation/IKRig/Source/IKRigEditor/Private/RigEditor/IKRigController.cpp
+++ b/Engine/Plugins/Animation/IKRig/Source/IKRigEditor/Private/RigEditor/IKRigController.cpp
@@ -483,7 +483,7 @@ FName UIKRigController::GetRetargetChainFromGoal(const FName& GoalName) const
 // SKELETON
 //
 
-bool UIKRigController::SetSkeletalMesh(USkeletalMesh* SkeletalMesh, bool bTransact) const
+bool UIKRigController::SetSkeletalMesh(USkeletalMesh* SkeletalMesh) const
 {
 	if (!SkeletalMesh)
 	{
@@ -497,31 +497,15 @@ bool UIKRigController::SetSkeletalMesh(USkeletalMesh* SkeletalMesh, bool bTransa
 		UE_LOG(LogTemp, Warning, TEXT("Trying to initialize IKRig with a Skeleton that is missing required bones. See output log. {0}"), *Asset->GetName());
 		return false;
 	}
-
-	FScopedTransaction Transaction(LOCTEXT("SetSkeletalMesh_Label", "Set Skeletal Mesh"));
 	
-	if (bTransact)
-	{	
-		Asset->Modify();
-	}
-
 	// update stored skeletal mesh used for previewing results
 	Asset->PreviewSkeletalMesh = SkeletalMesh;
+
 	// copy skeleton data from the actual skeleton we want to run on
 	Asset->Skeleton.SetInputSkeleton(InputSkeleton, GetAsset()->Skeleton.ExcludedBones);
-	// update goal's initial transforms to reflect new
-	for (UIKRigEffectorGoal* Goal : Asset->Goals)
-	{
-		if (bTransact)
-		{
-			Goal->Modify();
-		}
-		
-		const FTransform InitialTransform = GetRefPoseTransformOfBone(Goal->BoneName);
-		Goal->InitialTransform = InitialTransform;
-	}
 
-	BroadcastNeedsReinitialized();
+	// update goal's initial transforms to reflect new skeleton
+	ResetInitialGoalTransforms();
 
 	return true;
 }
@@ -913,9 +897,11 @@ bool UIKRigController::SetGoalBone(const FName& GoalName, const FName& NewBoneNa
 	
 	FScopedTransaction Transaction(LOCTEXT("SetGoalBone_Label", "Set Goal Bone"));
 
-	// update goal
-	Asset->Goals[GoalIndex]->Modify();
-	Asset->Goals[GoalIndex]->BoneName = NewBoneName;
+	// update goal with new bone
+	TObjectPtr<UIKRigEffectorGoal> Goal = Asset->Goals[GoalIndex];
+	Goal->Modify();
+	Goal->BoneName = NewBoneName;
+	Goal->InitialTransform = GetRefPoseTransformOfBone(NewBoneName);
 	
 	// update in solvers
 	for (UIKRigSolver* Solver : Asset->Solvers)
@@ -924,9 +910,6 @@ bool UIKRigController::SetGoalBone(const FName& GoalName, const FName& NewBoneNa
 		Solver->SetGoalBone(GoalName, NewBoneName);
 	}
 
-	// update initial transforms
-	ResetGoalTransforms();
-
 	BroadcastNeedsReinitialized();
 	
 	return true;
@@ -1083,19 +1066,31 @@ void UIKRigController::SetGoalCurrentTransform(const FName& GoalName, const FTra
 	Goal->CurrentTransform = Transform;
 }
 
-void UIKRigController::ResetGoalTransforms() const
+void UIKRigController::ResetCurrentGoalTransforms() const
 {
-	FScopedTransaction Transaction(LOCTEXT("ResetGoalTransforms", "Reset All Goal Transforms"));
+	FScopedTransaction Transaction(LOCTEXT("ResetCurrentGoalTransforms", "Reset All Goal Transforms"));
 	
 	for (UIKRigEffectorGoal* Goal : Asset->Goals)
     {
 		Goal->Modify();
-    	const FTransform InitialTransform = GetRefPoseTransformOfBone(Goal->BoneName);
-    	Goal->InitialTransform = InitialTransform;
-    	Goal->CurrentTransform = InitialTransform;
+    	Goal->CurrentTransform = Goal->InitialTransform;
     }
 }
 
+void UIKRigController::ResetInitialGoalTransforms() const
+{
+	for (UIKRigEffectorGoal* Goal : Asset->Goals)
+	{
+		// record the current delta rotation
+		const FQuat DeltaRotation = Goal->CurrentTransform.GetRotation() * Goal->InitialTransform.GetRotation().Inverse();
+		// update the initial transform based on the new ref pose
+		const FTransform InitialTransform = GetRefPoseTransformOfBone(Goal->BoneName);
+		Goal->InitialTransform = InitialTransform;
+		// restore the current transform
+		Goal->CurrentTransform.SetRotation(Goal->InitialTransform.GetRotation() * DeltaRotation);
+	}
+}
+
 void UIKRigController::SanitizeGoalName(FString& InOutName)
 {
 	for (int32 i = 0; i < InOutName.Len(); ++i)
diff --git a/Engine/Plugins/Animation/IKRig/Source/IKRigEditor/Private/RigEditor/IKRigEditorController.cpp b/Engine/Plugins/Animation/IKRig/Source/IKRigEditor/Private/RigEditor/IKRigEditorController.cpp
--- a/Engine/Plugins/Animation/IKRig/Source/IKRigEditor/Private/RigEditor/IKRigEditorController.cpp
+++ b/Engine/Plugins/Animation/IKRig/Source/IKRigEditor/Private/RigEditor/IKRigEditorController.cpp
@@ -345,21 +345,33 @@ void FIKRigEditorController::OnIKRigNeedsInitialized(UIKRigDefinition* ModifiedI
 	}
 
 	ClearOutputLog();
-	
-	AnimInstance->SetProcessorNeedsInitialized();
 
-	// Initialize editor's instances on request
+	// anim instance needs reinitialized
 	InitializeSolvers();
+	AnimInstance->SetProcessorNeedsInitialized();
+	AnimInstance->InitializeAnimation();
+	
+	// in case the skeletal mesh was swapped out, we need to ensure the preview scene is up-to-date
+	{
+		USkeletalMesh* NewMesh = AssetController->GetAsset()->GetPreviewMesh();
+		const TSharedRef<IPersonaPreviewScene> PreviewScene = EditorToolkit.Pin()->GetPersonaToolkit()->GetPreviewScene();
+		PreviewScene->SetPreviewMesh(NewMesh);
+
+		// ensure goals are using initial transforms from the current mesh 
+		AssetController->ResetInitialGoalTransforms();
+	}
 
 	// update the bone details so it can pull on the current data
 	BoneDetails->AnimInstancePtr = AnimInstance;
 	BoneDetails->AssetPtr = ModifiedIKRig;
+
+	RefreshAllViews();
 }
 
 void FIKRigEditorController::Reset() const
 {
 	SkelMeshComponent->ShowReferencePose(true);
-	AssetController->ResetGoalTransforms();
+	AssetController->ResetCurrentGoalTransforms();
 }
 
 void FIKRigEditorController::RefreshAllViews() const
@@ -802,33 +814,15 @@ void FIKRigEditorController::ShowDetailsForElements(const TArray<TSharedPtr<FIKR
 	}
 }
 
-void FIKRigEditorController::OnFinishedChangingDetails(const FPropertyChangedEvent& PropertyChangedEvent)
+void FIKRigEditorController::OnFinishedChangingDetails(const FPropertyChangedEvent& PropertyChangedEvent) const
 {
+	// if user swapped in a different skeletal mesh, then trigger a total reinitialization
 	const bool bPreviewChanged = PropertyChangedEvent.GetPropertyName() == UIKRigDefinition::GetPreviewMeshPropertyName();
-	
 	if (bPreviewChanged)
 	{
-		// set the source and target skeletal meshes on the component
-		// NOTE: this must be done AFTER setting the AnimInstance so that the correct root anim node is loaded
 		USkeletalMesh* NewMesh = AssetController->GetAsset()->GetPreviewMesh();
-		if (NewMesh)
-		{
-			// apply the mesh to the preview scene
-			TSharedRef<IPersonaPreviewScene> PreviewScene = EditorToolkit.Pin()->GetPersonaToolkit()->GetPreviewScene();
-			if (PreviewScene->GetPreviewMesh() != NewMesh)
-			{
-				PreviewScene->SetPreviewMeshComponent(SkelMeshComponent);
-				PreviewScene->SetPreviewMesh(NewMesh);
-			}
-
-			ClearOutputLog();
-			AssetController->SetSkeletalMesh(NewMesh);
-			AnimInstance->SetProcessorNeedsInitialized();
-			AnimInstance->InitializeAnimation();
-			AssetController->BroadcastNeedsReinitialized();
-			AssetController->ResetGoalTransforms();
-			RefreshAllViews();
-		}
+		AssetController->SetSkeletalMesh(NewMesh);
+		AssetController->BroadcastNeedsReinitialized();
 	}
 }
 
diff --git a/Engine/Plugins/Animation/IKRig/Source/IKRigEditor/Private/RigEditor/IKRigToolkit.cpp b/Engine/Plugins/Animation/IKRig/Source/IKRigEditor/Private/RigEditor/IKRigToolkit.cpp
--- a/Engine/Plugins/Animation/IKRig/Source/IKRigEditor/Private/RigEditor/IKRigToolkit.cpp
+++ b/Engine/Plugins/Animation/IKRig/Source/IKRigEditor/Private/RigEditor/IKRigToolkit.cpp
@@ -220,7 +220,6 @@ void FIKRigEditorToolkit::HandlePreviewSceneCreated(const TSharedRef<IPersonaPre
 	EditorController->AnimInstance = AnimInstance;
 	AnimInstance->SetIKRigAsset(EditorController->AssetController->GetAsset());
 	EditorController->SkelMeshComponent->PreviewInstance = AnimInstance;
-	EditorController->OnIKRigNeedsInitialized(EditorController->AssetController->GetAsset());
 	AnimInstance->InitializeAnimation();
 
 	// set the skeletal mesh on the component
diff --git a/Engine/Plugins/Animation/IKRig/Source/IKRigEditor/Public/RigEditor/IKRigController.h b/Engine/Plugins/Animation/IKRig/Source/IKRigEditor/Public/RigEditor/IKRigController.h
--- a/Engine/Plugins/Animation/IKRig/Source/IKRigEditor/Public/RigEditor/IKRigController.h
+++ b/Engine/Plugins/Animation/IKRig/Source/IKRigEditor/Public/RigEditor/IKRigController.h
@@ -37,9 +37,9 @@ public:
 	/** SKELETON
 	 * 
 	 */
-	/** Sets the preview mesh to use, can optionally reinitialize the skeleton with bReImportBone=true.
+	/** Sets the preview mesh to use. Loads the hierarchy into the asset's IKRigSkeleton.
 	 * Returns true if the mesh was able to be set. False if it was incompatible for any reason. */
-	bool SetSkeletalMesh(USkeletalMesh* SkeletalMesh, bool bTransact=false) const;
+	bool SetSkeletalMesh(USkeletalMesh* SkeletalMesh) const;
 	/** Get read-access to the IKRig skeleton representation */
 	const FIKRigSkeleton& GetIKRigSkeleton() const;
 	/** Get the USkeleton asset this rig was initialized with */
@@ -120,7 +120,9 @@ public:
 	/** Set the Goal to the given transform. */
 	void SetGoalCurrentTransform(const FName& GoalName, const FTransform& Transform) const;
 	/** Reset all Goals back to their initial transforms. */
-	void ResetGoalTransforms() const;
+	void ResetCurrentGoalTransforms() const;
+	/** Update initial goal transforms (called after new mesh is applied). */
+	void ResetInitialGoalTransforms() const;
 	/** Ensure that the given name adheres to required standards for Goal names (no special characters etc..)*/
 	static void SanitizeGoalName(FString& InOutName);
 	/** Add a suffix as needed to ensure the Goal name is unique */
diff --git a/Engine/Plugins/Animation/IKRig/Source/IKRigEditor/Public/RigEditor/IKRigEditorController.h b/Engine/Plugins/Animation/IKRig/Source/IKRigEditor/Public/RigEditor/IKRigEditorController.h
--- a/Engine/Plugins/Animation/IKRig/Source/IKRigEditor/Public/RigEditor/IKRigEditorController.h
+++ b/Engine/Plugins/Animation/IKRig/Source/IKRigEditor/Public/RigEditor/IKRigEditorController.h
@@ -177,7 +177,7 @@ public:
 	/** show selected items in details view */
 	void ShowDetailsForElements(const TArray<TSharedPtr<FIKRigTreeElement>>& InItems) const;
 	/** callback when detail is edited */
-	void OnFinishedChangingDetails(const FPropertyChangedEvent& PropertyChangedEvent);
+	void OnFinishedChangingDetails(const FPropertyChangedEvent& PropertyChangedEvent) const;
 	
 	/** set details tab view */
 	void SetDetailsView(const TSharedPtr<class IDetailsView>& InDetailsView);
