Fixed bug, mirroring a retarget chain does not always find correct bone.

diff --git a/Engine/Plugins/Animation/IKRig/Source/IKRig/Private/IKRigSkeleton.cpp b/Engine/Plugins/Animation/IKRig/Source/IKRig/Private/IKRigSkeleton.cpp
--- a/Engine/Plugins/Animation/IKRig/Source/IKRig/Private/IKRigSkeleton.cpp
+++ b/Engine/Plugins/Animation/IKRig/Source/IKRig/Private/IKRigSkeleton.cpp
@@ -3,7 +3,9 @@
 #include "IKRigSkeleton.h"
 
 #include "IKRigDefinition.h"
+#include "Algo/LevenshteinDistance.h"
 #include "Engine/SkeletalMesh.h"
+#include "Math/Bounds.h"
 
 #include UE_INLINE_GENERATED_CPP_BY_NAME(IKRigSkeleton)
 
@@ -421,45 +423,120 @@ bool FIKRigSkeleton::ValidateChainAndGetBones(const FBoneChain& BoneChain, TArra
 
 bool FIKRigSkeleton::GetMirroredBoneIndices(
 	const TArray<int32>& BoneIndices,
-	TArray<int32>& OutMirroredIndices,
-	float MaxDistanceThreshold) const
+	TArray<int32>& OutMirroredIndices) const
 {
-	constexpr float MirrorDistanceThreshold = 5.0f;
+	auto GetDistancesToAllBones = [this](const FVector& InPoint, TArray<float>& OutDistances)
+	{
+		OutDistances.SetNumUninitialized(RefPoseGlobal.Num());
+		int32 Index = 0;
+		for (const FTransform& BoneTransform : RefPoseGlobal)
+		{
+			OutDistances[Index] = FVector::Distance(InPoint, BoneTransform.GetLocation());
+			++Index;
+		}
+	};
+	
+	auto GetBonesWithinDistance = [this](const float MaxDistance, const TArray<float>& InDistances, TArray<int32>& OutBoneIndices)
+	{
+		OutBoneIndices.Reset();
+		int32 Index = 0;
+		for (const float& Distance : InDistances)
+		{
+			if (Distance < MaxDistance)
+			{
+				OutBoneIndices.Add(Index);
+			}
+			++Index;
+		}
+		return !OutBoneIndices.IsEmpty();
+	};
+
+	auto GetBoneWithClosestName = [this](const FName BaseName, const TArray<int32>& InBoneIndices) -> int32
+	{
+		check(!InBoneIndices.IsEmpty());
+		int32 BoneWithClosestNameToBaseName = InBoneIndices[0];
+		const FString BaseNameStr = BaseName.ToString().ToLower();
+		float HighestScore = -1.f;
+		for (const int32 BoneIndex : InBoneIndices)
+		{
+			FString CurrentNameStr = BoneNames[BoneIndex].ToString().ToLower();
+			float WorstCase = BaseNameStr.Len() + CurrentNameStr.Len();
+			WorstCase = WorstCase < 1.0f ? 1.0f : WorstCase;
+			const float Score = 1.0f - (Algo::LevenshteinDistance(BaseNameStr, CurrentNameStr) / WorstCase);
+			if (Score > HighestScore)
+			{
+				HighestScore = Score;
+				BoneWithClosestNameToBaseName = BoneIndex;
+			}
+		}
+		
+		return BoneWithClosestNameToBaseName;
+	};
+
+	auto GetBoundingSphere = [this]()->FSphere
+	{
+		TArray<FVector> Points;
+		Points.AddUninitialized(RefPoseGlobal.Num());
+		int32 Index=0;
+		for (const FTransform& BoneTransform : RefPoseGlobal)
+		{
+			Points[Index] = BoneTransform.GetLocation();
+			++Index;
+		}
+		return FSphere(&Points[0], Points.Num());
+	};
 
+	// generate bounding sphere of skeleton to establish reasonable search radius
+	const FSphere BoundingSphere = GetBoundingSphere();
+	if (BoundingSphere.W < KINDA_SMALL_NUMBER)
+	{
+		// skeleton is degenerate
+		return false;
+	}
+		
+	const float MinSearchRadius = BoundingSphere.W * 0.01f;
+	const float MaxSearchRadius = BoundingSphere.W;
+	const float SearchStep = MinSearchRadius * 2.f;
+
+	// spin through all the bones in the chain:
+	// 1. iteratively expand a search radius to find bones near the mirrored location
+	// 2. pick the bone near the mirrored location with the most similar name
+	//
+	// NOTE: we can't just do a "find closest bone" here because production skeletons often have lots of "helper" bones
+	// that are located at or near the same position. So we have to use the name to help find the actual mirrored bone.
 	OutMirroredIndices.Reset();
 	for (const int32 BoneIndex : BoneIndices)
 	{
+		// get mirrored location to search around
 		const FVector BoneLocation = RefPoseGlobal[BoneIndex].GetLocation();
 		const FVector MirroredLocation = BoneLocation * FVector(-1.f,1.f,1.f);
-		int32 ClosestBoneIndex;
-		float Distance;
-		GetClosestBone(MirroredLocation, ClosestBoneIndex, Distance);
-		if (Distance > MirrorDistanceThreshold || ClosestBoneIndex == INDEX_NONE)
+		
+		// get distance from the mirrored location to all other bones in the skeleton
+		TArray<float> DistancesToAllBones;
+		GetDistancesToAllBones(MirroredLocation, DistancesToAllBones);
+		
+		// do an expanding search around the mirrored location until we find bones near it
+		TArray<int32> NearbyBones;
+		float SearchRadius = MinSearchRadius;
+		while(SearchRadius <= MaxSearchRadius)
+		{
+			if (GetBonesWithinDistance(SearchRadius, DistancesToAllBones, NearbyBones))
+			{
+				break;
+			}
+			SearchRadius += SearchStep;
+		};
+		
+		if (NearbyBones.IsEmpty())
 		{
 			OutMirroredIndices.Reset();
 			return false;
 		}
-		OutMirroredIndices.Add(ClosestBoneIndex);
-	}
-	return true;
-}
 
-void FIKRigSkeleton::GetClosestBone(
-	const FVector& InPoint,
-	int32& OutBoneIndex,
-	float& OutDistance) const
-{
-	int32 Index = 0;
-	OutBoneIndex = INDEX_NONE;
-	OutDistance = TNumericLimits<float>::Max();
-	for (const FTransform& BonePose : RefPoseGlobal)
-	{
-		const float Distance = FVector::Distance(InPoint, BonePose.GetLocation());
-		if (Distance < OutDistance)
-		{
-			OutBoneIndex = Index;
-			OutDistance = Distance;
-		}
-		++Index;
+		// found the nearby bone with the best matching name
+		int32 BestMatchBoneIndex = GetBoneWithClosestName(BoneNames[BoneIndex], NearbyBones);
+		OutMirroredIndices.Add(BestMatchBoneIndex);
 	}
+	
+	return true;
 }
diff --git a/Engine/Plugins/Animation/IKRig/Source/IKRig/Public/IKRigSkeleton.h b/Engine/Plugins/Animation/IKRig/Source/IKRig/Public/IKRigSkeleton.h
--- a/Engine/Plugins/Animation/IKRig/Source/IKRig/Public/IKRigSkeleton.h
+++ b/Engine/Plugins/Animation/IKRig/Source/IKRig/Public/IKRigSkeleton.h
@@ -116,13 +116,7 @@ struct IKRIG_API FIKRigSkeleton
 
 	// given a set of bones, returns a set of indices of bones that are mirrored across X axis (YZ plane)
 	// return false if any of the bones do not have a mirrored pair within the MaxDistanceThreshold
-	bool GetMirroredBoneIndices(
-		const TArray<int32>& BoneIndices,
-		TArray<int32>& OutMirroredIndices,
-		float MaxDistanceThreshold=5.0f) const;
-	
-	// get the bone that is closest to the given point (sets bone index and distance)
-	void GetClosestBone(const FVector& InPoint, int32& OutBoneIndex, float& OutDistance) const;
+	bool GetMirroredBoneIndices(const TArray<int32>& BoneIndices, TArray<int32>& OutMirroredIndices) const;
 
 private:
 
