Retargeter preview mesh not updated if swapped and undone. Also fixes animation playback not continuing after swapping preview mesh.

diff --git a/Engine/Plugins/Animation/IKRig/Source/IKRigEditor/Private/RetargetEditor/IKRetargetEditor.cpp b/Engine/Plugins/Animation/IKRig/Source/IKRigEditor/Private/RetargetEditor/IKRetargetEditor.cpp
--- a/Engine/Plugins/Animation/IKRig/Source/IKRigEditor/Private/RetargetEditor/IKRetargetEditor.cpp
+++ b/Engine/Plugins/Animation/IKRig/Source/IKRigEditor/Private/RetargetEditor/IKRetargetEditor.cpp
@@ -382,7 +382,7 @@ void FIKRetargetEditor::SetupAnimInstance()
 	
 	// configure SOURCE anim instance (will only output retarget pose)
 	EditorController->SourceAnimInstance->ConfigureAnimInstance(ERetargetSourceOrTarget::Source, Asset, nullptr);
-	// configure TARGET anim instance (will output retarget pose AND retarget pose from source skel mesh component)
+	// configure TARGET anim instance (will output retarget pose AND retargeted pose from source skel mesh component)
 	EditorController->TargetAnimInstance->ConfigureAnimInstance(ERetargetSourceOrTarget::Target, Asset, EditorController->SourceSkelMeshComponent);
 
 	EditorController->SourceSkelMeshComponent->PreviewInstance = EditorController->SourceAnimInstance.Get();
@@ -442,29 +442,6 @@ void FIKRetargetEditor::OnFinishedChangingDetails(const FPropertyChangedEvent& P
 	// if either the source or target meshes are possibly modified, update scene components, anim instance and UI
 	if (bTargetIKRigChanged || bSourceIKRigChanged || bTargetPreviewChanged || bSourcePreviewChanged)
 	{
-		EditorController->ClearOutputLog();
-		
-		// set the source and target skeletal meshes on the component
-		// NOTE: this must be done AFTER setting the AnimInstance so that the correct root anim node is loaded
-		USkeletalMesh* SourceMesh = EditorController->GetSkeletalMesh(ERetargetSourceOrTarget::Source);
-		USkeletalMesh* TargetMesh = EditorController->GetSkeletalMesh(ERetargetSourceOrTarget::Target);
-		EditorController->SourceSkelMeshComponent->SetSkeletalMesh(SourceMesh);
-		EditorController->TargetSkelMeshComponent->SetSkeletalMesh(TargetMesh);
-
-		// reset bone selections in case of incompatible indices
-		EditorController->ClearSelection();
-	
-		// apply mesh to the preview scene
-		TSharedRef<IPersonaPreviewScene> PreviewScene = GetPersonaToolkit()->GetPreviewScene();
-		if (PreviewScene->GetPreviewMesh() != SourceMesh)
-		{
-			PreviewScene->SetPreviewMeshComponent(EditorController->SourceSkelMeshComponent);
-			PreviewScene->SetPreviewMesh(SourceMesh);
-			EditorController->SourceSkelMeshComponent->bCanHighlightSelectedSections = false;
-		}
-	
-		SetupAnimInstance();
-
 		AssetController->BroadcastNeedsReinitialized();
 	}
 }
diff --git a/Engine/Plugins/Animation/IKRig/Source/IKRigEditor/Private/RetargetEditor/IKRetargetEditorController.cpp b/Engine/Plugins/Animation/IKRig/Source/IKRigEditor/Private/RetargetEditor/IKRetargetEditorController.cpp
--- a/Engine/Plugins/Animation/IKRig/Source/IKRigEditor/Private/RetargetEditor/IKRetargetEditorController.cpp
+++ b/Engine/Plugins/Animation/IKRig/Source/IKRigEditor/Private/RetargetEditor/IKRetargetEditorController.cpp
@@ -27,7 +27,7 @@
 #define LOCTEXT_NAMESPACE "IKRetargetEditorController"
 
 
-FBoundIKRig::FBoundIKRig(UIKRigDefinition* InIKRig, const FIKRetargetEditorController& InController)
+FBoundIKRig::FBoundIKRig(UIKRigDefinition* InIKRig, FIKRetargetEditorController& InController)
 {
 	check(InIKRig);
 	IKRig = InIKRig;
@@ -105,7 +105,7 @@ void FIKRetargetEditorController::BindToIKRigAssets(UIKRetargeter* InAsset)
 	}
 }
 
-void FIKRetargetEditorController::HandleIKRigNeedsInitialized(UIKRigDefinition* ModifiedIKRig) const
+void FIKRetargetEditorController::HandleIKRigNeedsInitialized(UIKRigDefinition* ModifiedIKRig)
 {
 	UIKRetargeter* Retargeter = AssetController->GetAsset();
 	check(ModifiedIKRig && Retargeter)
@@ -138,11 +138,38 @@ void FIKRetargetEditorController::HandleRetargetChainRemoved(UIKRigDefinition* M
 	RefreshAllViews();
 }
 
-void FIKRetargetEditorController::HandleRetargeterNeedsInitialized(UIKRetargeter* Retargeter) const
+void FIKRetargetEditorController::HandleRetargeterNeedsInitialized(UIKRetargeter* Retargeter)
 {
+	// record status of animation playback so we can restore it after initialization
+	const bool bWasPlayingAnimation = SourceAnimInstance->IsPlaying();
+	UAnimationAsset* AnimToRestore = SourceAnimInstance->GetAnimationAsset();
+	const float PlaybackTimeToRestore = SourceAnimInstance->GetCurrentTime();
+	
 	// clear the output log
 	ClearOutputLog();
 
+	// ensure source and target meshes are setup in case they were modified
+	{
+		// set the source and target skeletal meshes on the component
+		// NOTE: this must be done AFTER setting the AnimInstance so that the correct root anim node is loaded
+		USkeletalMesh* SourceMesh = GetSkeletalMesh(ERetargetSourceOrTarget::Source);
+		USkeletalMesh* TargetMesh = GetSkeletalMesh(ERetargetSourceOrTarget::Target);
+		SourceSkelMeshComponent->SetSkeletalMesh(SourceMesh);
+		TargetSkelMeshComponent->SetSkeletalMesh(TargetMesh);
+	
+		// apply mesh to the preview scene
+		TSharedRef<IPersonaPreviewScene> PreviewScene = Editor.Pin()->GetPersonaToolkit()->GetPreviewScene();
+		if (PreviewScene->GetPreviewMesh() != SourceMesh)
+		{
+			PreviewScene->SetPreviewMeshComponent(SourceSkelMeshComponent);
+			PreviewScene->SetPreviewMesh(SourceMesh);
+			SourceSkelMeshComponent->bCanHighlightSelectedSections = false;
+		}
+	
+		SourceAnimInstance->InitializeAnimation();
+		TargetAnimInstance->InitializeAnimation();
+	}
+
 	// check for "zero height" retarget roots, and prompt user to fix
 	FixZeroHeightRetargetRoot(ERetargetSourceOrTarget::Source);
 	FixZeroHeightRetargetRoot(ERetargetSourceOrTarget::Target);
@@ -157,9 +184,26 @@ void FIKRetargetEditorController::HandleRetargeterNeedsInitialized(UIKRetargeter
 			Retargeter,
 			bSuppressWarnings);
 	}
+
+	// reset bone selections in case of incompatible indices
+	CleanSelection(ERetargetSourceOrTarget::Source);
+	CleanSelection(ERetargetSourceOrTarget::Target);
 	
 	// refresh all the UI views
 	RefreshAllViews();
+
+	// restore animation playback
+	if (bWasPlayingAnimation && AnimToRestore)
+	{
+		// verify that animation is compatible with (potentially) new source mesh
+		USkeletalMesh* NewSourceMesh = GetSkeletalMesh(ERetargetSourceOrTarget::Source);
+		if (AnimThatWasPlaying->GetSkeleton()->IsCompatible(NewSourceMesh->GetSkeleton()))
+		{
+			SourceAnimInstance->SetAnimationAsset(AnimToRestore);
+			SourceAnimInstance->SetPlaying(true);
+			SourceAnimInstance->SetPosition(PlaybackTimeToRestore);	
+		}
+	}
 }
 
 UDebugSkelMeshComponent* FIKRetargetEditorController::GetSkeletalMeshComponent(
@@ -822,6 +866,28 @@ void FIKRetargetEditorController::SetRootSelected(const bool bIsSelected)
 	LastSelectedItem = ERetargetSelectionType::ROOT;
 }
 
+void FIKRetargetEditorController::CleanSelection(ERetargetSourceOrTarget SourceOrTarget)
+{
+	USkeletalMesh* SkeletalMesh = GetSkeletalMesh(SourceOrTarget);
+	if (!SkeletalMesh)
+	{
+		SelectedBoneNames[SourceOrTarget].Empty();
+		return;
+	}
+
+	TArray<FName> CleanedSelection;
+	const FReferenceSkeleton& RefSkeleton = SkeletalMesh->GetRefSkeleton();
+	for (const FName& SelectedBone : SelectedBoneNames[SourceOrTarget])
+	{
+		if (RefSkeleton.FindBoneIndex(SelectedBone))
+		{
+			CleanedSelection.Add(SelectedBone);
+		}
+	}
+
+	SelectedBoneNames[SourceOrTarget] = CleanedSelection;
+}
+
 void FIKRetargetEditorController::ClearSelection(const bool bKeepBoneSelection)
 {
 	// clear root and mesh selection
diff --git a/Engine/Plugins/Animation/IKRig/Source/IKRigEditor/Public/RetargetEditor/IKRetargetEditorController.h b/Engine/Plugins/Animation/IKRig/Source/IKRigEditor/Public/RetargetEditor/IKRetargetEditorController.h
--- a/Engine/Plugins/Animation/IKRig/Source/IKRigEditor/Public/RetargetEditor/IKRetargetEditorController.h
+++ b/Engine/Plugins/Animation/IKRig/Source/IKRigEditor/Public/RetargetEditor/IKRetargetEditorController.h
@@ -50,7 +50,7 @@ enum class ERetargetSelectionType : uint8
 
 struct FBoundIKRig
 {
-	FBoundIKRig(UIKRigDefinition* InIKRig, const FIKRetargetEditorController& InController);
+	FBoundIKRig(UIKRigDefinition* InIKRig, FIKRetargetEditorController& InController);
 	void UnBind() const;
 	
 	UIKRigDefinition* IKRig;
@@ -74,7 +74,7 @@ public:
 	/** Bind callbacks to this IK Rig */
 	void BindToIKRigAssets(UIKRetargeter* InAsset);
 	/** callback when IK Rig asset requires reinitialization */
-	void HandleIKRigNeedsInitialized(UIKRigDefinition* ModifiedIKRig) const;
+	void HandleIKRigNeedsInitialized(UIKRigDefinition* ModifiedIKRig);
 	/** callback when IK Rig asset's retarget chain's have been added or removed*/
 	void HandleRetargetChainAdded(UIKRigDefinition* ModifiedIKRig) const;
 	/** callback when IK Rig asset's retarget chain has been renamed */
@@ -82,7 +82,7 @@ public:
 	/** callback when IK Rig asset's retarget chain has been removed */
 	void HandleRetargetChainRemoved(UIKRigDefinition* ModifiedIKRig, const FName& InChainRemoved) const;
 	/** callback when IK Retargeter asset requires reinitialization */
-	void HandleRetargeterNeedsInitialized(UIKRetargeter* Retargeter) const;
+	void HandleRetargeterNeedsInitialized(UIKRetargeter* Retargeter);
 	FDelegateHandle RetargeterReInitDelegateHandle;
 	
 	/** all modifications to the data model should go through this controller */
@@ -205,7 +205,10 @@ public:
 
 	void SetRootSelected(const bool bIsSelected);
 	bool GetRootSelected() const { return bIsRootSelected; };
-	
+
+	/* remove selected elements that no longer exist after swapping input assets */
+	void CleanSelection(ERetargetSourceOrTarget SourceOrTarget);
+	/* set selection to nothing */
 	void ClearSelection(const bool bKeepBoneSelection=false);
 	ERetargetSelectionType GetLastSelectedItemType() const { return LastSelectedItem; };
 
