Changing default bone size in retargeter from 8 to 1

diff --git a/Engine/Plugins/Animation/IKRig/Source/IKRig/Public/Retargeter/IKRetargeter.h b/Engine/Plugins/Animation/IKRig/Source/IKRig/Public/Retargeter/IKRetargeter.h
--- a/Engine/Plugins/Animation/IKRig/Source/IKRig/Public/Retargeter/IKRetargeter.h
+++ b/Engine/Plugins/Animation/IKRig/Source/IKRig/Public/Retargeter/IKRetargeter.h
@@ -421,7 +421,7 @@ public:
 	
 	/** The visual size of the bones in the viewport (saved between sessions). This is set from the viewport Character>Bones menu*/
 	UPROPERTY()
-	float BoneDrawSize = 8.0f;
+	float BoneDrawSize = 1.0f;
 	
 private:
 
