Added ability to live-edit retarget settings in PIE

diff --git a/Engine/Plugins/Animation/IKRig/Source/IKRig/Private/AnimNodes/AnimNode_RetargetPoseFromMesh.cpp b/Engine/Plugins/Animation/IKRig/Source/IKRig/Private/AnimNodes/AnimNode_RetargetPoseFromMesh.cpp
--- a/Engine/Plugins/Animation/IKRig/Source/IKRig/Private/AnimNodes/AnimNode_RetargetPoseFromMesh.cpp
+++ b/Engine/Plugins/Animation/IKRig/Source/IKRig/Private/AnimNodes/AnimNode_RetargetPoseFromMesh.cpp
@@ -98,18 +98,11 @@ void FAnimNode_RetargetPoseFromMesh::Evaluate_AnyThread(FPoseContext& Output)
 		return;
 	}
 
-#if WITH_EDITOR
+	#if WITH_EDITOR
 	// live preview source asset settings in the retarget, editor only
 	// NOTE: this copies goal targets as well, but these are overwritten by IK chain goals
-	if (bDriveWithAsset)
-	{
-		Processor->ApplySettingsFromAsset();
-		if (const FRetargetProfile* CurrentProfile = IKRetargeterAsset->GetCurrentProfile())
-		{
-			Processor->ApplySettingsFromProfile(*CurrentProfile);
-		}
-	}
-#endif
+	Processor->ApplySettingsFromAsset();
+	#endif
 
 	// apply custom profile settings to the processor
 	Processor->ApplySettingsFromProfile(CustomRetargetProfile);
diff --git a/Engine/Plugins/Animation/IKRig/Source/IKRig/Private/Retargeter/IKRetargetProcessor.cpp b/Engine/Plugins/Animation/IKRig/Source/IKRig/Private/Retargeter/IKRetargetProcessor.cpp
--- a/Engine/Plugins/Animation/IKRig/Source/IKRig/Private/Retargeter/IKRetargetProcessor.cpp
+++ b/Engine/Plugins/Animation/IKRig/Source/IKRig/Private/Retargeter/IKRetargetProcessor.cpp
@@ -2210,6 +2210,12 @@ void UIKRetargetProcessor::ApplySettingsFromAsset()
 
 	// copy global settings
 	GlobalSettings = RetargeterAsset->GetGlobalSettings();
+
+	// apply the current profile
+	if (const FRetargetProfile* CurrentProfile = RetargeterAsset->GetCurrentProfile())
+	{
+		ApplySettingsFromProfile(*CurrentProfile);
+	}
 }
 
 void UIKRetargetProcessor::ApplySettingsFromProfile(const FRetargetProfile& Profile)
diff --git a/Engine/Plugins/Animation/IKRig/Source/IKRig/Public/AnimNodes/AnimNode_RetargetPoseFromMesh.h b/Engine/Plugins/Animation/IKRig/Source/IKRig/Public/AnimNodes/AnimNode_RetargetPoseFromMesh.h
--- a/Engine/Plugins/Animation/IKRig/Source/IKRig/Public/AnimNodes/AnimNode_RetargetPoseFromMesh.h
+++ b/Engine/Plugins/Animation/IKRig/Source/IKRig/Public/AnimNodes/AnimNode_RetargetPoseFromMesh.h
@@ -41,11 +41,6 @@ struct IKRIG_API FAnimNode_RetargetPoseFromMesh : public FAnimNode_Base
 	/* Copy curves from SouceMeshComponent. This will copy any curves the source/target Skeleton have in common. */
 	UPROPERTY(EditAnywhere, BlueprintReadWrite, Category = Settings, meta = (NeverAsPin))
 	bool bCopyCurves = true;
-
-#if WITH_EDITOR
-	/** when true, will copy all setting from target IK Rig asset each tick (for live preview) */
-	bool bDriveWithAsset = false;
-#endif
 	
 	// FAnimNode_Base interface
 	virtual void Initialize_AnyThread(const FAnimationInitializeContext& Context) override;
diff --git a/Engine/Plugins/Animation/IKRig/Source/IKRigEditor/Private/RetargetEditor/IKRetargetAnimInstanceProxy.cpp b/Engine/Plugins/Animation/IKRig/Source/IKRigEditor/Private/RetargetEditor/IKRetargetAnimInstanceProxy.cpp
--- a/Engine/Plugins/Animation/IKRig/Source/IKRigEditor/Private/RetargetEditor/IKRetargetAnimInstanceProxy.cpp
+++ b/Engine/Plugins/Animation/IKRig/Source/IKRigEditor/Private/RetargetEditor/IKRetargetAnimInstanceProxy.cpp
@@ -100,7 +100,6 @@ void FIKRetargetAnimInstanceProxy::ConfigureAnimInstance(
 		RetargetNode->IKRetargeterAsset = InIKRetargetAsset;
 		RetargetNode->bUseAttachedParent = false;
 		RetargetNode->SourceMeshComponent = InSourceMeshComponent;
-		RetargetNode->bDriveWithAsset = true;
 		if (UIKRetargetProcessor* Processor = RetargetNode->GetRetargetProcessor())
 		{
 			Processor->SetNeedsInitialized();
