Adding "..." suffix to retargeter new chain menu item

diff --git a/Engine/Plugins/Animation/IKRig/Source/IKRigEditor/Private/RigEditor/IKRigSkeletonCommands.cpp b/Engine/Plugins/Animation/IKRig/Source/IKRigEditor/Private/RigEditor/IKRigSkeletonCommands.cpp
--- a/Engine/Plugins/Animation/IKRig/Source/IKRigEditor/Private/RigEditor/IKRigSkeletonCommands.cpp
+++ b/Engine/Plugins/Animation/IKRig/Source/IKRigEditor/Private/RigEditor/IKRigSkeletonCommands.cpp
@@ -16,7 +16,7 @@ void FIKRigSkeletonCommands::RegisterCommands()
 	UI_COMMAND(RemoveBoneSettings, "Remove Settings on Selected Bone", "Remove all settings on the selected bone for the selected solver.", EUserInterfaceActionType::Button, FInputChord());
 	UI_COMMAND(ExcludeBone, "Exclude Selected Bone From Solve", "Ignore bone in all solvers.", EUserInterfaceActionType::Button, FInputChord());
 	UI_COMMAND(IncludeBone, "Include Selected Bone In Solve", "Include bone in all solvers.", EUserInterfaceActionType::Button, FInputChord());
-	UI_COMMAND(NewRetargetChain, "New Retarget Chain from Selected Bones", "Create a new retarget bone chain from the selected bones.", EUserInterfaceActionType::Button, FInputChord());
+	UI_COMMAND(NewRetargetChain, "New Retarget Chain...", "Create a new retarget bone chain from the selected bones.", EUserInterfaceActionType::Button, FInputChord());
 	UI_COMMAND(SetRetargetRoot, "Set Retarget Root", "Set the Root Bone used for retargeting. Usually 'Pelvis'.", EUserInterfaceActionType::Button, FInputChord());
 	UI_COMMAND(ClearRetargetRoot, "Clear Retarget Root", "Clear the Root Bone used for retargeting.", EUserInterfaceActionType::Button, FInputChord());
 	UI_COMMAND(RenameGoal, "Rename Goal", "Rename the selected goal.", EUserInterfaceActionType::Button, FInputChord(EKeys::F2));
