Fixed issue, importing retarget pose from animation sequence didn't filter out pose assets.

diff --git a/Engine/Plugins/Animation/IKRig/Source/IKRigEditor/Private/RetargetEditor/IKRetargetPoseExporter.cpp b/Engine/Plugins/Animation/IKRig/Source/IKRigEditor/Private/RetargetEditor/IKRetargetPoseExporter.cpp
--- a/Engine/Plugins/Animation/IKRig/Source/IKRigEditor/Private/RetargetEditor/IKRetargetPoseExporter.cpp
+++ b/Engine/Plugins/Animation/IKRig/Source/IKRigEditor/Private/RetargetEditor/IKRetargetPoseExporter.cpp
@@ -303,8 +303,6 @@ void FIKRetargetPoseExporter::HandleImportFromSequenceAsset()
 	// the asset picker will only show animation sequences compatible with the preview mesh
 	FAssetPickerConfig AssetPickerConfig;
 	AssetPickerConfig.Filter.ClassPaths.Add(UAnimSequence::StaticClass()->GetClassPathName());
-	AssetPickerConfig.Filter.ClassPaths.Add(UAnimMontage::StaticClass()->GetClassPathName());
-	AssetPickerConfig.Filter.ClassPaths.Add(UPoseAsset::StaticClass()->GetClassPathName());
 	AssetPickerConfig.InitialAssetViewType = EAssetViewType::Column;
 	AssetPickerConfig.bAddFilterUI = true;
 	AssetPickerConfig.bShowPathInColumnView = true;
