Added ability to scale vertical motion of IK goals.

diff --git a/Engine/Plugins/Animation/IKRig/Source/IKRig/Private/Retargeter/IKRetargetProcessor.cpp b/Engine/Plugins/Animation/IKRig/Source/IKRig/Private/Retargeter/IKRetargetProcessor.cpp
--- a/Engine/Plugins/Animation/IKRig/Source/IKRig/Private/Retargeter/IKRetargetProcessor.cpp
+++ b/Engine/Plugins/Animation/IKRig/Source/IKRig/Private/Retargeter/IKRetargetProcessor.cpp
@@ -960,6 +960,9 @@ void FChainRetargeterIK::DecodePose(
 	// apply local static offset
 	GoalPosition += GoalRotation.RotateVector(Settings.IK.StaticLocalOffset);
 
+	// apply vertical scale
+	GoalPosition.Z *= Settings.IK.ScaleVertical;
+	
 	// apply extension
 	if (!FMath::IsNearlyEqual(Settings.IK.Extension, 1.0f))
 	{
diff --git a/Engine/Plugins/Animation/IKRig/Source/IKRig/Private/Retargeter/IKRetargetSettings.cpp b/Engine/Plugins/Animation/IKRig/Source/IKRig/Private/Retargeter/IKRetargetSettings.cpp
--- a/Engine/Plugins/Animation/IKRig/Source/IKRig/Private/Retargeter/IKRetargetSettings.cpp
+++ b/Engine/Plugins/Animation/IKRig/Source/IKRig/Private/Retargeter/IKRetargetSettings.cpp
@@ -31,6 +31,7 @@ bool FTargetChainIKSettings::operator==(const FTargetChainIKSettings& Other) con
 	&& StaticOffset.Equals(Other.StaticOffset)
 	&& StaticLocalOffset.Equals(Other.StaticLocalOffset)
 	&& StaticRotationOffset.Equals(Other.StaticRotationOffset)
+	&& FMath::IsNearlyEqualByULP(ScaleVertical, Other.ScaleVertical)
 	&& FMath::IsNearlyEqualByULP(Extension, Other.Extension)
 	&& bAffectedByIKWarping == Other.bAffectedByIKWarping;
 }
diff --git a/Engine/Plugins/Animation/IKRig/Source/IKRig/Public/Retargeter/IKRetargetSettings.h b/Engine/Plugins/Animation/IKRig/Source/IKRig/Public/Retargeter/IKRetargetSettings.h
--- a/Engine/Plugins/Animation/IKRig/Source/IKRig/Public/Retargeter/IKRetargetSettings.h
+++ b/Engine/Plugins/Animation/IKRig/Source/IKRig/Public/Retargeter/IKRetargetSettings.h
@@ -138,6 +138,10 @@ struct IKRIG_API FTargetChainIKSettings
 	/** Default 0, 0, 0. Apply a static local-space offset to IK goal rotation. */
 	UPROPERTY(EditAnywhere, BlueprintReadWrite, Category = "IK Adjustments")
 	FRotator StaticRotationOffset = FRotator::ZeroRotator;
+
+	/** Range +-Infinity. Default 1. Scales the vertical component of the IK goal's position.*/
+	UPROPERTY(EditAnywhere, BlueprintReadWrite, Category = "IK Adjustments", meta = (UIMin = "0.0", UIMax = "5.0"))
+	float ScaleVertical = 1.0f;
 	
 	/** Range 0 to 5. Default 1. Brings IK goal closer (0) or further (1+) from origin of chain.
 	*  At 0 the effector is placed at the origin of the chain (ie Shoulder, Hip etc).
