Making IK Rig plugin compatible in multi-project environments

diff --git a/Engine/Plugins/Animation/IKRig/Source/IKRigEditor/Private/RetargetEditor/IKRetargetDetails.cpp b/Engine/Plugins/Animation/IKRig/Source/IKRigEditor/Private/RetargetEditor/IKRetargetDetails.cpp
--- a/Engine/Plugins/Animation/IKRig/Source/IKRigEditor/Private/RetargetEditor/IKRetargetDetails.cpp
+++ b/Engine/Plugins/Animation/IKRig/Source/IKRigEditor/Private/RetargetEditor/IKRetargetDetails.cpp
@@ -1004,7 +1004,7 @@ void FRetargetChainSettingsCustomization::AddSettingsSection(
 				.OnCheckStateChanged_Lambda([this, IsEnabledProperty](ECheckBoxState State)
 				{
 					IsEnabledProperty->SetValue(State == ECheckBoxState::Checked);
-					Controller->HandleRetargeterNeedsInitialized(Controller->AssetController->GetAsset());
+					Controller->HandleRetargeterNeedsInitialized();
 				})
 			];
 	}
diff --git a/Engine/Plugins/Animation/IKRig/Source/IKRigEditor/Private/RetargetEditor/IKRetargetEditor.cpp b/Engine/Plugins/Animation/IKRig/Source/IKRigEditor/Private/RetargetEditor/IKRetargetEditor.cpp
--- a/Engine/Plugins/Animation/IKRig/Source/IKRigEditor/Private/RetargetEditor/IKRetargetEditor.cpp
+++ b/Engine/Plugins/Animation/IKRig/Source/IKRigEditor/Private/RetargetEditor/IKRetargetEditor.cpp
@@ -88,6 +88,10 @@ void FIKRetargetEditor::InitAssetEditor(
 
 	ExtendToolbar();
 	RegenerateMenusAndToolbars();
+	
+	// initial setup, ignored if IK Rig is already assigned
+	EditorController->PromptUserToAssignIKRig(ERetargetSourceOrTarget::Source);
+	EditorController->PromptUserToAssignIKRig(ERetargetSourceOrTarget::Target);
 }
 
 void FIKRetargetEditor::OnClose()
@@ -269,12 +273,14 @@ TStatId FIKRetargetEditor::GetStatId() const
 
 void FIKRetargetEditor::PostUndo(bool bSuccess)
 {
-	EditorController->HandleRetargeterNeedsInitialized(EditorController->AssetController->GetAsset());
+	EditorController->HandlePreviewMeshReplaced(ERetargetSourceOrTarget::Source);
+	EditorController->HandleRetargeterNeedsInitialized();
 }
 
 void FIKRetargetEditor::PostRedo(bool bSuccess)
 {
-	EditorController->HandleRetargeterNeedsInitialized(EditorController->AssetController->GetAsset());
+	EditorController->HandlePreviewMeshReplaced(ERetargetSourceOrTarget::Source);
+	EditorController->HandleRetargeterNeedsInitialized();
 }
 
 void FIKRetargetEditor::HandleViewportCreated(const TSharedRef<class IPersonaViewport>& InViewport)
@@ -332,7 +338,7 @@ void FIKRetargetEditor::HandleViewportCreated(const TSharedRef<class IPersonaVie
 }
 
 void FIKRetargetEditor::HandlePreviewSceneCreated(const TSharedRef<IPersonaPreviewScene>& InPersonaPreviewScene)
-{
+{	
 	AAnimationEditorPreviewActor* Actor = InPersonaPreviewScene->GetWorld()->SpawnActor<AAnimationEditorPreviewActor>(AAnimationEditorPreviewActor::StaticClass(), FTransform::Identity);
 	Actor->SetFlags(RF_Transient);
 	InPersonaPreviewScene->SetActor(Actor);
@@ -368,7 +374,7 @@ void FIKRetargetEditor::HandlePreviewSceneCreated(const TSharedRef<IPersonaPrevi
 	
 	InPersonaPreviewScene->AddComponent(EditorController->SourceSkelMeshComponent, FTransform::Identity);
 	InPersonaPreviewScene->AddComponent(EditorController->TargetSkelMeshComponent, FTransform::Identity);
-
+	
 	EditorController->FixZeroHeightRetargetRoot(ERetargetSourceOrTarget::Source);
 	EditorController->FixZeroHeightRetargetRoot(ERetargetSourceOrTarget::Target);
 }
@@ -379,7 +385,7 @@ void FIKRetargetEditor::SetupAnimInstance()
 	
 	// configure SOURCE anim instance (will only output retarget pose)
 	EditorController->SourceAnimInstance->ConfigureAnimInstance(ERetargetSourceOrTarget::Source, Asset, nullptr);
-	// configure TARGET anim instance (will output retarget pose AND retargeted pose from source skel mesh component)
+	// configure TARGET anim instance (will output retarget pose AND retarget pose from source skel mesh component)
 	EditorController->TargetAnimInstance->ConfigureAnimInstance(ERetargetSourceOrTarget::Target, Asset, EditorController->SourceSkelMeshComponent);
 
 	EditorController->SourceSkelMeshComponent->PreviewInstance = EditorController->SourceAnimInstance.Get();
@@ -406,7 +412,7 @@ void FIKRetargetEditor::HandleDetailsCreated(const TSharedRef<class IDetailsView
 
 void FIKRetargetEditor::OnFinishedChangingDetails(const FPropertyChangedEvent& PropertyChangedEvent)
 {
-	const UIKRetargeterController* AssetController = EditorController->AssetController;
+	UIKRetargeterController* AssetController = EditorController->AssetController;
 
 	// determine which properties were modified
 	const bool bSourceIKRigChanged = PropertyChangedEvent.GetPropertyName() == UIKRetargeter::GetSourceIKRigPropertyName();
@@ -417,29 +423,28 @@ void FIKRetargetEditor::OnFinishedChangingDetails(const FPropertyChangedEvent& P
 	// if no override target mesh has been specified, update the override to reflect the mesh in the ik rig asset
 	if (bTargetIKRigChanged)
 	{
-		AssetController->OnIKRigChanged(ERetargetSourceOrTarget::Target);
+		UIKRigDefinition* NewIKRig = AssetController->GetIKRigWriteable(ERetargetSourceOrTarget::Target);
+		AssetController->SetIKRig(ERetargetSourceOrTarget::Target, NewIKRig);
 	}
 
 	// if no override source mesh has been specified, update the override to reflect the mesh in the ik rig asset
 	if (bSourceIKRigChanged)
 	{
-		AssetController->OnIKRigChanged(ERetargetSourceOrTarget::Source);
+		UIKRigDefinition* NewIKRig = AssetController->GetIKRigWriteable(ERetargetSourceOrTarget::Source);
+		AssetController->SetIKRig(ERetargetSourceOrTarget::Source, NewIKRig);
 	}
 
-	// if either IK Rig asset has been modified, rebind and refresh UI
-	if (bTargetIKRigChanged || bSourceIKRigChanged)
+	// if either the source or target meshes are possibly modified, update scene components, anim instance and UI
+	if (bSourcePreviewChanged)
 	{
-		EditorController->ClearOutputLog();
-		EditorController->BindToIKRigAssets(AssetController->GetAsset());
-		EditorController->AssetController->CleanChainMapping();
-		constexpr bool bForceRemap = false;
-		EditorController->AssetController->AutoMapChains(EAutoMapChainType::Fuzzy, bForceRemap);
+		USkeletalMesh* Mesh = AssetController->GetPreviewMesh(ERetargetSourceOrTarget::Source);
+		AssetController->SetPreviewMesh(ERetargetSourceOrTarget::Source, Mesh);
 	}
-
-	// if either the source or target meshes are possibly modified, update scene components, anim instance and UI
-	if (bTargetIKRigChanged || bSourceIKRigChanged || bTargetPreviewChanged || bSourcePreviewChanged)
+	
+	if (bTargetPreviewChanged)
 	{
-		AssetController->BroadcastNeedsReinitialized();
+		USkeletalMesh* Mesh = AssetController->GetPreviewMesh(ERetargetSourceOrTarget::Target);
+		AssetController->SetPreviewMesh(ERetargetSourceOrTarget::Target, Mesh);
 	}
 }
 
diff --git a/Engine/Plugins/Animation/IKRig/Source/IKRigEditor/Private/RetargetEditor/IKRetargetEditorController.cpp b/Engine/Plugins/Animation/IKRig/Source/IKRigEditor/Private/RetargetEditor/IKRetargetEditorController.cpp
--- a/Engine/Plugins/Animation/IKRig/Source/IKRigEditor/Private/RetargetEditor/IKRetargetEditorController.cpp
+++ b/Engine/Plugins/Animation/IKRig/Source/IKRigEditor/Private/RetargetEditor/IKRetargetEditorController.cpp
@@ -4,6 +4,7 @@
 
 #include "AnimPose.h"
 #include "AssetToolsModule.h"
+#include "ContentBrowserModule.h"
 #include "EditorModeManager.h"
 #include "IContentBrowserSingleton.h"
 #include "SkeletalDebugRendering.h"
@@ -27,7 +28,7 @@
 #define LOCTEXT_NAMESPACE "IKRetargetEditorController"
 
 
-FBoundIKRig::FBoundIKRig(UIKRigDefinition* InIKRig, FIKRetargetEditorController& InController)
+FBoundIKRig::FBoundIKRig(UIKRigDefinition* InIKRig, const FIKRetargetEditorController& InController)
 {
 	check(InIKRig);
 	IKRig = InIKRig;
@@ -48,6 +49,80 @@ void FBoundIKRig::UnBind() const
 	IKRigController->OnRetargetChainRenamed().Remove(RenameChainDelegateHandle);
 }
 
+FRetargetPlaybackManager::FRetargetPlaybackManager(TWeakPtr<FIKRetargetEditorController> InEditorController)
+{
+	check(InEditorController.Pin())
+	EditorController = InEditorController;
+}
+
+void FRetargetPlaybackManager::PlayAnimationAsset(UAnimationAsset* AssetToPlay)
+{
+	UIKRetargetAnimInstance* AnimInstance = EditorController.Pin()->SourceAnimInstance.Get();
+	if (!AnimInstance)
+	{
+		return;
+	}
+	
+	if (AssetToPlay)
+	{
+		AnimInstance->SetAnimationAsset(AssetToPlay);
+		AnimInstance->SetPlaying(true);
+		AnimThatWasPlaying = AssetToPlay;
+		// ensure we are running the retargeter so you can see the animation
+		if (FIKRetargetEditorController* Controller = EditorController.Pin().Get())
+		{
+			Controller->SetRetargeterMode(ERetargeterOutputMode::RunRetarget);
+		}
+	}
+}
+
+void FRetargetPlaybackManager::StopPlayback() const
+{
+	UIKRetargetAnimInstance* AnimInstance = EditorController.Pin()->SourceAnimInstance.Get();
+	if (!AnimInstance)
+	{
+		return;
+	}
+	
+	AnimInstance->SetPlaying(false);
+	AnimInstance->SetAnimationAsset(nullptr);
+}
+
+void FRetargetPlaybackManager::PausePlayback()
+{
+	UIKRetargetAnimInstance* AnimInstance = EditorController.Pin()->SourceAnimInstance.Get();
+	if (!AnimInstance)
+	{
+		return;
+	}
+	
+	AnimThatWasPlaying = AnimInstance->GetAnimationAsset();
+	
+	if (AnimThatWasPlaying)
+	{
+		TimeWhenPaused = AnimInstance->GetCurrentTime();
+	}
+	
+	AnimInstance->SetPlaying(false);
+	AnimInstance->SetAnimationAsset(nullptr);
+}
+
+void FRetargetPlaybackManager::ResumePlayback() const
+{
+	UIKRetargetAnimInstance* AnimInstance = EditorController.Pin()->SourceAnimInstance.Get();
+	if (!AnimInstance)
+	{
+		return;
+	}
+	
+	if (AnimThatWasPlaying)
+	{
+		AnimInstance->SetAnimationAsset(AnimThatWasPlaying);
+		AnimInstance->SetPlaying(true);
+		AnimInstance->SetPosition(TimeWhenPaused);	
+	}
+}
+
 void FIKRetargetEditorController::Initialize(TSharedPtr<FIKRetargetEditor> InEditor, UIKRetargeter* InAsset)
 {
 	Editor = InEditor;
@@ -58,25 +133,31 @@ void FIKRetargetEditorController::Initialize(TSharedPtr<FIKRetargetEditor> InEdi
 	PoseExporter = MakeShared<FIKRetargetPoseExporter>();
 	PoseExporter->Initialize(SharedThis(this));
 
+	PlaybackManager = MakeUnique<FRetargetPlaybackManager>(SharedThis(this));
+
 	SelectedBoneNames.Add(ERetargetSourceOrTarget::Source);
 	SelectedBoneNames.Add(ERetargetSourceOrTarget::Target);
 	LastSelectedItem = ERetargetSelectionType::NONE;
 
 	// clean the asset before editing
-	const bool bForceReinitialization = false;
-	AssetController->CleanChainMapping(bForceReinitialization);
-	AssetController->CleanPoseLists(bForceReinitialization);
+	AssetController->CleanChainMapping();
+	AssetController->CleanPoseLists();
 
 	// bind callbacks when SOURCE or TARGET IK Rigs are modified
-	BindToIKRigAssets(AssetController->GetAsset());
+	BindToIKRigAssets();
 
 	// bind callback when retargeter needs reinitialized
 	RetargeterReInitDelegateHandle = AssetController->OnRetargeterNeedsInitialized().AddSP(this, &FIKRetargetEditorController::HandleRetargeterNeedsInitialized);
+	// bind callback when IK Rig asset is replaced with a different asset
+	IKRigReplacedDelegateHandle = AssetController->OnIKRigReplaced().AddSP(this, &FIKRetargetEditorController::HandleIKRigReplaced);
+	// bind callback when Preview Mesh asset is replaced with a different asset
+	PreviewMeshReplacedDelegateHandle = AssetController->OnPreviewMeshReplaced().AddSP(this, &FIKRetargetEditorController::HandlePreviewMeshReplaced);
 }
 
 void FIKRetargetEditorController::Close()
 {
 	AssetController->OnRetargeterNeedsInitialized().Remove(RetargeterReInitDelegateHandle);
+	AssetController->OnIKRigReplaced().Remove(IKRigReplacedDelegateHandle);
 
 	for (const FBoundIKRig& BoundIKRig : BoundIKRigs)
 	{
@@ -84,8 +165,72 @@ void FIKRetargetEditorController::Close()
 	}
 }
 
-void FIKRetargetEditorController::BindToIKRigAssets(UIKRetargeter* InAsset)
+void FIKRetargetEditorController::PromptUserToAssignIKRig(const ERetargetSourceOrTarget SourceOrTarget)
 {
+	// early out if we already have an IK Rig assigned
+	if (AssetController->GetIKRig(SourceOrTarget))
+	{
+		return;
+	}
+	
+	// Load the content browser module to display an asset picker
+	FContentBrowserModule& ContentBrowserModule = FModuleManager::LoadModuleChecked<FContentBrowserModule>("ContentBrowser");
+	FAssetPickerConfig AssetPickerConfig;
+	// must set the parent UObject so that the resulting list filters correctly in multi-project environments
+	AssetPickerConfig.AdditionalReferencingAssets.Add(AssetController->GetAsset());
+	/** The asset picker will only show skeletal meshes */
+	AssetPickerConfig.Filter.ClassPaths.Add(UIKRigDefinition::StaticClass()->GetClassPathName());
+	AssetPickerConfig.Filter.bRecursiveClasses = true;
+	/** The delegate that fires when an asset was selected */
+	AssetPickerConfig.OnAssetSelected =  FOnAssetSelected::CreateLambda([this, &SourceOrTarget](const FAssetData& AssetData)
+	{
+		IKRigPickerWindow->RequestDestroyWindow();
+
+		// set the IK Rig
+		if (const TObjectPtr<UIKRigDefinition> IKRigAsset = Cast<UIKRigDefinition>(AssetData.GetAsset()))
+		{
+			AssetController->SetIKRig(SourceOrTarget, IKRigAsset.Get());
+		}
+	});
+	/** The default view mode should be a list view */
+	AssetPickerConfig.InitialAssetViewType = EAssetViewType::List;
+
+	// change title depending on if we're picking a source or target IK Rig
+	FText TitleText;
+	if (SourceOrTarget == ERetargetSourceOrTarget::Source)
+	{
+		TitleText = LOCTEXT("SelectSourceIKRig", "Select Source IK Rig (To Copy Animation From)");
+	}
+	else
+	{
+		TitleText = LOCTEXT("SelectTargetIKRig", "Select Target IK Rig (To Copy Animation To)");
+	}
+
+	// show an asset browser in a pop-up modal dialog
+	IKRigPickerWindow = SNew(SWindow)
+	.Title(TitleText)
+	.ClientSize(FVector2D(500, 600))
+	.SupportsMinimize(false).SupportsMaximize(false)
+	[
+		SNew(SBorder)
+		.BorderImage( FAppStyle::GetBrush("Menu.Background") )
+		[
+			ContentBrowserModule.Get().CreateAssetPicker(AssetPickerConfig)
+		]
+	];
+
+	GEditor->EditorAddModalWindow(IKRigPickerWindow.ToSharedRef());
+	IKRigPickerWindow.Reset();
+}
+
+void FIKRetargetEditorController::BindToIKRigAssets()
+{
+	const UIKRetargeter* Asset = AssetController->GetAsset();
+	if (!Asset)
+	{
+		return;
+	}
+	
 	// unbind previously bound IK Rigs
 	for (const FBoundIKRig& BoundIKRig : BoundIKRigs)
 	{
@@ -93,19 +238,19 @@ void FIKRetargetEditorController::BindToIKRigAssets(UIKRetargeter* InAsset)
 	}
 
 	BoundIKRigs.Empty();
-
-	if (InAsset->GetSourceIKRigWriteable())
+	
+	if (Asset->GetSourceIKRigWriteable())
 	{
-		BoundIKRigs.Emplace(FBoundIKRig(InAsset->GetSourceIKRigWriteable(), *this));
+		BoundIKRigs.Emplace(FBoundIKRig(Asset->GetSourceIKRigWriteable(), *this));
 	}
 
-	if (InAsset->GetTargetIKRigWriteable())
+	if (Asset->GetTargetIKRigWriteable())
 	{
-		BoundIKRigs.Emplace(FBoundIKRig(InAsset->GetTargetIKRigWriteable(), *this));
+		BoundIKRigs.Emplace(FBoundIKRig(Asset->GetTargetIKRigWriteable(), *this));
 	}
 }
 
-void FIKRetargetEditorController::HandleIKRigNeedsInitialized(UIKRigDefinition* ModifiedIKRig)
+void FIKRetargetEditorController::HandleIKRigNeedsInitialized(UIKRigDefinition* ModifiedIKRig) const
 {
 	UIKRetargeter* Retargeter = AssetController->GetAsset();
 	check(ModifiedIKRig && Retargeter)
@@ -115,7 +260,7 @@ void FIKRetargetEditorController::HandleIKRigNeedsInitialized(UIKRigDefinition*
 	check(bIsSource || bIsTarget);
 
 	// the target anim instance has the RetargetPoseFromMesh node which needs reinitialized with new asset version
-	HandleRetargeterNeedsInitialized(Retargeter);
+	HandleRetargeterNeedsInitialized();
 }
 
 void FIKRetargetEditorController::HandleRetargetChainAdded(UIKRigDefinition* ModifiedIKRig) const
@@ -138,38 +283,11 @@ void FIKRetargetEditorController::HandleRetargetChainRemoved(UIKRigDefinition* M
 	RefreshAllViews();
 }
 
-void FIKRetargetEditorController::HandleRetargeterNeedsInitialized(UIKRetargeter* Retargeter)
+void FIKRetargetEditorController::HandleRetargeterNeedsInitialized() const
 {
-	// record status of animation playback so we can restore it after initialization
-	const bool bWasPlayingAnimation = SourceAnimInstance->IsPlaying();
-	UAnimationAsset* AnimToRestore = SourceAnimInstance->GetAnimationAsset();
-	const float PlaybackTimeToRestore = SourceAnimInstance->GetCurrentTime();
-	
 	// clear the output log
 	ClearOutputLog();
 
-	// ensure source and target meshes are setup in case they were modified
-	{
-		// set the source and target skeletal meshes on the component
-		// NOTE: this must be done AFTER setting the AnimInstance so that the correct root anim node is loaded
-		USkeletalMesh* SourceMesh = GetSkeletalMesh(ERetargetSourceOrTarget::Source);
-		USkeletalMesh* TargetMesh = GetSkeletalMesh(ERetargetSourceOrTarget::Target);
-		SourceSkelMeshComponent->SetSkeletalMesh(SourceMesh);
-		TargetSkelMeshComponent->SetSkeletalMesh(TargetMesh);
-	
-		// apply mesh to the preview scene
-		TSharedRef<IPersonaPreviewScene> PreviewScene = Editor.Pin()->GetPersonaToolkit()->GetPreviewScene();
-		if (PreviewScene->GetPreviewMesh() != SourceMesh)
-		{
-			PreviewScene->SetPreviewMeshComponent(SourceSkelMeshComponent);
-			PreviewScene->SetPreviewMesh(SourceMesh);
-			SourceSkelMeshComponent->bCanHighlightSelectedSections = false;
-		}
-	
-		SourceAnimInstance->InitializeAnimation();
-		TargetAnimInstance->InitializeAnimation();
-	}
-
 	// check for "zero height" retarget roots, and prompt user to fix
 	FixZeroHeightRetargetRoot(ERetargetSourceOrTarget::Source);
 	FixZeroHeightRetargetRoot(ERetargetSourceOrTarget::Target);
@@ -181,39 +299,59 @@ void FIKRetargetEditorController::HandleRetargeterNeedsInitialized(UIKRetargeter
 		Processor->Initialize(
 			GetSkeletalMesh(ERetargetSourceOrTarget::Source),
 			GetSkeletalMesh(ERetargetSourceOrTarget::Target),
-			Retargeter,
+			AssetController->GetAsset(),
 			bSuppressWarnings);
 	}
-
-	// reset bone selections in case of incompatible indices
-	CleanSelection(ERetargetSourceOrTarget::Source);
-	CleanSelection(ERetargetSourceOrTarget::Target);
 	
 	// refresh all the UI views
 	RefreshAllViews();
+}
+
+void FIKRetargetEditorController::HandleIKRigReplaced(ERetargetSourceOrTarget SourceOrTarget)
+{
+	BindToIKRigAssets();
+}
 
-	// restore animation playback
-	if (bWasPlayingAnimation && AnimToRestore)
+void FIKRetargetEditorController::HandlePreviewMeshReplaced(ERetargetSourceOrTarget SourceOrTarget)
+{
+	// pause playback so we can resume after mesh swapped out
+	PlaybackManager->PausePlayback();
+	
+	// set the source and target skeletal meshes on the component
+	// NOTE: this must be done AFTER setting the AnimInstance so that the correct root anim node is loaded
+	USkeletalMesh* SourceMesh = GetSkeletalMesh(ERetargetSourceOrTarget::Source);
+	USkeletalMesh* TargetMesh = GetSkeletalMesh(ERetargetSourceOrTarget::Target);
+	SourceSkelMeshComponent->SetSkeletalMesh(SourceMesh);
+	TargetSkelMeshComponent->SetSkeletalMesh(TargetMesh);
+
+	// clean bone selections in case of incompatible indices
+	CleanSelection(ERetargetSourceOrTarget::Source);
+	CleanSelection(ERetargetSourceOrTarget::Target);
+
+	// apply mesh to the preview scene
+	TSharedRef<IPersonaPreviewScene> PreviewScene = Editor.Pin()->GetPersonaToolkit()->GetPreviewScene();
+	if (PreviewScene->GetPreviewMesh() != SourceMesh)
 	{
-		// verify that animation is compatible with (potentially) new source mesh
-		USkeletalMesh* NewSourceMesh = GetSkeletalMesh(ERetargetSourceOrTarget::Source);
-		if (AnimThatWasPlaying->GetSkeleton()->IsCompatibleForEditor(NewSourceMesh->GetSkeleton()))
-		{
-			SourceAnimInstance->SetAnimationAsset(AnimToRestore);
-			SourceAnimInstance->SetPlaying(true);
-			SourceAnimInstance->SetPosition(PlaybackTimeToRestore);	
-		}
+		PreviewScene->SetPreviewMeshComponent(SourceSkelMeshComponent);
+		PreviewScene->SetPreviewMesh(SourceMesh);
+		SourceSkelMeshComponent->bCanHighlightSelectedSections = false;
 	}
+
+	// re-initializes the anim instances running in the viewport
+	Editor.Pin()->SetupAnimInstance();
+
+	// continue playing where we left off
+	PlaybackManager->ResumePlayback();
 }
 
 UDebugSkelMeshComponent* FIKRetargetEditorController::GetSkeletalMeshComponent(
-	const ERetargetSourceOrTarget& SourceOrTarget) const
+	const ERetargetSourceOrTarget SourceOrTarget) const
 {
 	return SourceOrTarget == ERetargetSourceOrTarget::Source ? SourceSkelMeshComponent : TargetSkelMeshComponent;
 }
 
 UIKRetargetAnimInstance* FIKRetargetEditorController::GetAnimInstance(
-	const ERetargetSourceOrTarget& SourceOrTarget) const
+	const ERetargetSourceOrTarget SourceOrTarget) const
 {
 	return SourceOrTarget == ERetargetSourceOrTarget::Source ? SourceAnimInstance.Get() : TargetAnimInstance.Get();
 }
@@ -601,43 +739,6 @@ void FIKRetargetEditorController::SetDetailsObjects(const TArray<UObject*>& Deta
 	}
 }
 
-void FIKRetargetEditorController::PlayAnimationAsset(UAnimationAsset* AssetToPlay)
-{
-	if (AssetToPlay && SourceAnimInstance.IsValid())
-	{
-		SourceAnimInstance->SetAnimationAsset(AssetToPlay);
-		SourceAnimInstance->SetPlaying(true);
-		AnimThatWasPlaying = AssetToPlay;
-		// ensure we are running the retargeter so you can see the animation
-		SetRetargeterMode(ERetargeterOutputMode::RunRetarget);
-	}
-}
-
-void FIKRetargetEditorController::StopPlayback()
-{
-	SourceAnimInstance->SetPlaying(false);
-	SourceAnimInstance->SetAnimationAsset(nullptr);
-}
-
-void FIKRetargetEditorController::PausePlayback()
-{
-	if (UAnimationAsset* CurrentAnim = SourceAnimInstance->GetAnimationAsset())
-	{
-		AnimThatWasPlaying = CurrentAnim;
-		TimeWhenPaused = SourceAnimInstance->GetCurrentTime();
-	}
-	
-	SourceAnimInstance->SetPlaying(false);
-	SourceAnimInstance->SetAnimationAsset(nullptr);
-}
-
-void FIKRetargetEditorController::ResumePlayback()
-{
-	SourceAnimInstance->SetAnimationAsset(AnimThatWasPlaying);
-	SourceAnimInstance->SetPlaying(true);
-	SourceAnimInstance->SetPosition(TimeWhenPaused);
-}
-
 float FIKRetargetEditorController::GetRetargetPoseAmount() const
 {
 	return RetargetPosePreviewBlend;
@@ -938,7 +1039,7 @@ void FIKRetargetEditorController::SetRetargeterMode(ERetargeterOutputMode Mode)
 			OutputMode = ERetargeterOutputMode::EditRetargetPose;
 			SourceAnimInstance->SetRetargetMode(ERetargeterOutputMode::EditRetargetPose);
 			TargetAnimInstance->SetRetargetMode(ERetargeterOutputMode::EditRetargetPose);
-			PausePlayback();
+			PlaybackManager->PausePlayback();
 			SetRetargetPoseAmount(1.0f);
 			break;
 		
@@ -950,7 +1051,7 @@ void FIKRetargetEditorController::SetRetargeterMode(ERetargeterOutputMode Mode)
 			TargetAnimInstance->SetRetargetMode(ERetargeterOutputMode::RunRetarget);
 			// force a reinitialization in case retarget pose was edited
 			AssetController->BroadcastNeedsReinitialized();
-			ResumePlayback();
+			PlaybackManager->ResumePlayback();
 			break;
 
 		case ERetargeterOutputMode::ShowRetargetPose:
@@ -960,7 +1061,7 @@ void FIKRetargetEditorController::SetRetargeterMode(ERetargeterOutputMode Mode)
 			// show retarget pose
 			SourceAnimInstance->SetRetargetMode(ERetargeterOutputMode::ShowRetargetPose);
 			TargetAnimInstance->SetRetargetMode(ERetargeterOutputMode::ShowRetargetPose);
-			PausePlayback();
+			PlaybackManager->PausePlayback();
 			SetRetargetPoseAmount(1.0f);
 			break;
 
diff --git a/Engine/Plugins/Animation/IKRig/Source/IKRigEditor/Private/RetargetEditor/IKRetargetFactory.cpp b/Engine/Plugins/Animation/IKRig/Source/IKRigEditor/Private/RetargetEditor/IKRetargetFactory.cpp
--- a/Engine/Plugins/Animation/IKRig/Source/IKRigEditor/Private/RetargetEditor/IKRetargetFactory.cpp
+++ b/Engine/Plugins/Animation/IKRig/Source/IKRigEditor/Private/RetargetEditor/IKRetargetFactory.cpp
@@ -33,22 +33,7 @@ UObject* UIKRetargetFactory::FactoryCreateNew(
 	UObject* Context,
 	FFeedbackContext* Warn)
 {
-	if (!SourceIKRig.IsValid())
-	{
-		UE_LOG(LogIKRigEditor, Warning, TEXT("Unable to create IK Retargter. No source IK Rig asset supplied."));
-		return nullptr;
-	}
-	
-	if (!IsValid(SourceIKRig->GetPreviewMesh()))
-	{
-		UE_LOG(LogIKRigEditor, Warning, TEXT("Unable to create IK Retargter. Source IK Rig does not have an associated skeletal mesh."));
-		return nullptr;
-	}
-	
-	UIKRetargeter* Retargeter = NewObject<UIKRetargeter>(InParent, Class, Name, Flags);
-	UIKRetargeterController* Controller = UIKRetargeterController::GetController(Retargeter);
-	Controller->SetSourceIKRig(SourceIKRig.Get());
-	return Retargeter;
+	return NewObject<UIKRetargeter>(InParent, Class, Name, Flags);
 }
 
 bool UIKRetargetFactory::ShouldShowInNewMenu() const
@@ -56,50 +41,6 @@ bool UIKRetargetFactory::ShouldShowInNewMenu() const
 	return true;
 }
 
-void UIKRetargetFactory::OnTargetIKRigSelected(const FAssetData& SelectedAsset)
-{
-	SourceIKRig = Cast<UIKRigDefinition>(SelectedAsset.GetAsset());
-	PickerWindow->RequestDestroyWindow();
-}
-
-bool UIKRetargetFactory::ConfigureProperties()
-{
-	// Null the parent class so we can check for selection later
-	SourceIKRig = nullptr;
-
-	// Load the content browser module to display an asset picker
-	FContentBrowserModule& ContentBrowserModule = FModuleManager::LoadModuleChecked<FContentBrowserModule>("ContentBrowser");
-
-	FAssetPickerConfig AssetPickerConfig;
-
-	/** The asset picker will only show skeletal meshes */
-	AssetPickerConfig.Filter.ClassPaths.Add(UIKRigDefinition::StaticClass()->GetClassPathName());
-	AssetPickerConfig.Filter.bRecursiveClasses = true;
-
-	/** The delegate that fires when an asset was selected */
-	AssetPickerConfig.OnAssetSelected = FOnAssetSelected::CreateUObject(this, &UIKRetargetFactory::OnTargetIKRigSelected);
-
-	/** The default view mode should be a list view */
-	AssetPickerConfig.InitialAssetViewType = EAssetViewType::List;
-
-	PickerWindow = SNew(SWindow)
-	.Title(LOCTEXT("CreateRetargeterOptions", "Pick IK Rig To Copy Animation From"))
-	.ClientSize(FVector2D(500, 600))
-	.SupportsMinimize(false) .SupportsMaximize(false)
-	[
-		SNew(SBorder)
-		.BorderImage( FAppStyle::GetBrush("Menu.Background") )
-		[
-			ContentBrowserModule.Get().CreateAssetPicker(AssetPickerConfig)
-		]
-	];
-
-	GEditor->EditorAddModalWindow(PickerWindow.ToSharedRef());
-	PickerWindow.Reset();
-
-	return SourceIKRig != nullptr;
-}
-
 FText UIKRetargetFactory::GetDisplayName() const
 {
 	return LOCTEXT("IKRetargeter_DisplayName", "IK Retargeter");
diff --git a/Engine/Plugins/Animation/IKRig/Source/IKRigEditor/Private/RetargetEditor/IKRetargeterController.cpp b/Engine/Plugins/Animation/IKRig/Source/IKRigEditor/Private/RetargetEditor/IKRetargeterController.cpp
--- a/Engine/Plugins/Animation/IKRig/Source/IKRigEditor/Private/RetargetEditor/IKRetargeterController.cpp
+++ b/Engine/Plugins/Animation/IKRig/Source/IKRigEditor/Private/RetargetEditor/IKRetargeterController.cpp
@@ -38,58 +38,85 @@ UIKRetargeter* UIKRetargeterController::GetAsset() const
 	return Asset;
 }
 
-void UIKRetargeterController::SetSourceIKRig(UIKRigDefinition* SourceIKRig)
+void UIKRetargeterController::SetIKRig(const ERetargetSourceOrTarget SourceOrTarget, UIKRigDefinition* IKRig) const
 {
-	Asset->SourceIKRigAsset = SourceIKRig;
-	Asset->SourcePreviewMesh = Asset->SourceIKRigAsset->PreviewSkeletalMesh;
-}
-
-USkeletalMesh* UIKRetargeterController::GetPreviewMesh(const ERetargetSourceOrTarget& SourceOrTarget) const
-{
-	// can't preview anything if target IK Rig is null
-	const UIKRigDefinition* IKRig = GetIKRig(SourceOrTarget);
-	if (!IKRig)
+	if (SourceOrTarget == ERetargetSourceOrTarget::Source)
 	{
-		return nullptr;
+		Asset->SourceIKRigAsset = IKRig;
+		Asset->SourcePreviewMesh = IKRig ? IKRig->GetPreviewMesh() : nullptr;
 	}
-
-	// optionally prefer override if one is provided
-	const TSoftObjectPtr<USkeletalMesh> OverrideMesh = SourceOrTarget == ERetargetSourceOrTarget::Source ? Asset->SourcePreviewMesh : Asset->TargetPreviewMesh;
-	if (!OverrideMesh.IsNull())
+	else
 	{
-		return OverrideMesh.LoadSynchronous();
+		Asset->TargetIKRigAsset = IKRig;
+		Asset->TargetPreviewMesh = IKRig ? IKRig->GetPreviewMesh() : nullptr;
+	}
+	
+	// re-ask to fix root height for this mesh
+	if (IKRig)
+	{
+		SetAskedToFixRootHeightForMesh(IKRig->PreviewSkeletalMesh.Get(), false);
 	}
 
-	// fallback to preview mesh from IK Rig asset
-	return IKRig->GetPreviewMesh();
+	CleanChainMapping();
+	
+	constexpr bool bForceRemap = false;
+	AutoMapChains(EAutoMapChainType::Fuzzy, bForceRemap);
+
+	// update any editors attached to this asset
+	BroadcastIKRigReplaced(SourceOrTarget);
+	BroadcastPreviewMeshReplaced(SourceOrTarget);
+	BroadcastNeedsReinitialized();
 }
 
-const UIKRigDefinition* UIKRetargeterController::GetIKRig(const ERetargetSourceOrTarget& SourceOrTarget) const
+const UIKRigDefinition* UIKRetargeterController::GetIKRig(const ERetargetSourceOrTarget SourceOrTarget) const
 {
 	return SourceOrTarget == ERetargetSourceOrTarget::Source ? Asset->GetSourceIKRig() : Asset->GetTargetIKRig();
 }
 
-UIKRigDefinition* UIKRetargeterController::GetIKRigWriteable(const ERetargetSourceOrTarget& SourceOrTarget) const
+UIKRigDefinition* UIKRetargeterController::GetIKRigWriteable(const ERetargetSourceOrTarget SourceOrTarget) const
 {
 	return SourceOrTarget == ERetargetSourceOrTarget::Source ? Asset->GetSourceIKRigWriteable() : Asset->GetTargetIKRigWriteable();
 }
 
-void UIKRetargeterController::OnIKRigChanged(const ERetargetSourceOrTarget& SourceOrTarget) const
+void UIKRetargeterController::SetPreviewMesh(
+	const ERetargetSourceOrTarget SourceOrTarget,
+	USkeletalMesh* PreviewMesh)
+{	
+	if (SourceOrTarget == ERetargetSourceOrTarget::Source)
+	{
+		Asset->SourcePreviewMesh = PreviewMesh;
+	}
+	else
+	{
+		Asset->TargetPreviewMesh = PreviewMesh;
+	}
+	
+	// re-ask to fix root height for this mesh
+	SetAskedToFixRootHeightForMesh(PreviewMesh, false);
+	
+	// update any editors attached to this asset
+	BroadcastPreviewMeshReplaced(SourceOrTarget);
+	BroadcastNeedsReinitialized();
+}
+
+USkeletalMesh* UIKRetargeterController::GetPreviewMesh(const ERetargetSourceOrTarget SourceOrTarget) const
 {
-	if (const UIKRigDefinition* IKRig = GetIKRig(SourceOrTarget))
+	// can't preview anything if target IK Rig is null
+	const UIKRigDefinition* IKRig = GetIKRig(SourceOrTarget);
+	if (!IKRig)
 	{
-		if (SourceOrTarget == ERetargetSourceOrTarget::Source)
-		{
-			Asset->SourcePreviewMesh = IKRig->PreviewSkeletalMesh;
-		}
-		else
-		{
-			Asset->TargetPreviewMesh = IKRig->PreviewSkeletalMesh;
-		}
+		return nullptr;
+	}
 
-		// re-ask to fix root height for this mesh
-		SetAskedToFixRootHeightForMesh(IKRig->PreviewSkeletalMesh.Get(), false);
+	// optionally prefer override if one is provided
+	const TSoftObjectPtr<USkeletalMesh> OverrideMesh = SourceOrTarget == ERetargetSourceOrTarget::Source ? Asset->SourcePreviewMesh : Asset->TargetPreviewMesh;
+	if (!OverrideMesh.IsNull())
+	{
+		return OverrideMesh.LoadSynchronous();
 	}
+
+	// fallback to preview mesh from IK Rig asset
+	return IKRig->GetPreviewMesh();
 }
 
 bool UIKRetargeterController::GetAskedToFixRootHeightForMesh(USkeletalMesh* Mesh) const
@@ -110,13 +137,13 @@ void UIKRetargeterController::SetAskedToFixRootHeightForMesh(USkeletalMesh* Mesh
 }
 
 FName UIKRetargeterController::GetRetargetRootBone(
-	const ERetargetSourceOrTarget& SourceOrTarget) const
+	const ERetargetSourceOrTarget SourceOrTarget) const
 {
 	const UIKRigDefinition* IKRig = GetIKRig(SourceOrTarget);
 	return IKRig ? IKRig->GetRetargetRoot() : FName("None");
 }
 
-void UIKRetargeterController::GetChainNames(const ERetargetSourceOrTarget& SourceOrTarget, TArray<FName>& OutNames) const
+void UIKRetargeterController::GetChainNames(const ERetargetSourceOrTarget SourceOrTarget, TArray<FName>& OutNames) const
 {
 	if (const UIKRigDefinition* IKRig = GetIKRig(SourceOrTarget))
 	{
@@ -128,7 +155,7 @@ void UIKRetargeterController::GetChainNames(const ERetargetSourceOrTarget& Sourc
 	}
 }
 
-void UIKRetargeterController::CleanChainMapping(const bool bForceReinitialization) const
+void UIKRetargeterController::CleanChainMapping() const
 {
 	if (IsValid(Asset->GetTargetIKRig()))
 	{
@@ -186,25 +213,15 @@ void UIKRetargeterController::CleanChainMapping(const bool bForceReinitializatio
 
 	// enforce the chain order based on the StartBone index
 	SortChainMapping();
-
-	if (bForceReinitialization)
-	{
-		BroadcastNeedsReinitialized();
-	}
 }
 
-void UIKRetargeterController::CleanPoseLists(const bool bForceReinitialization) const
+void UIKRetargeterController::CleanPoseLists() const
 {
 	CleanPoseList(ERetargetSourceOrTarget::Source);
 	CleanPoseList(ERetargetSourceOrTarget::Target);
-
-	if (bForceReinitialization)
-	{
-		BroadcastNeedsReinitialized();	
-	}
 }
 
-void UIKRetargeterController::CleanPoseList(const ERetargetSourceOrTarget& SourceOrTarget) const
+void UIKRetargeterController::CleanPoseList(const ERetargetSourceOrTarget SourceOrTarget) const
 {
 	TMap<FName, FIKRetargetPose>& RetargetPoses = GetRetargetPoses(SourceOrTarget);
 
@@ -313,9 +330,6 @@ void UIKRetargeterController::AutoMapChains(const EAutoMapChainType AutoMapType,
 
 	// sort them
 	SortChainMapping();
-
-	// force update with latest mapping
-	BroadcastNeedsReinitialized();
 }
 
 void UIKRetargeterController::OnRetargetChainAdded(UIKRigDefinition* IKRig) const
@@ -328,8 +342,8 @@ void UIKRetargeterController::OnRetargetChainAdded(UIKRigDefinition* IKRig) cons
 	}
 
 	// add the new chain to the mapping data
-	constexpr bool bForceReinitialization = true;
-	CleanChainMapping(bForceReinitialization);
+	CleanChainMapping();
+	BroadcastNeedsReinitialized();
 }
 
 void UIKRetargeterController::OnRetargetChainRenamed(UIKRigDefinition* IKRig, FName OldChainName, FName NewChainName) const
@@ -442,7 +456,7 @@ bool UIKRetargeterController::IsChainGoalConnectedToASolver(const FName& GoalNam
 void UIKRetargeterController::AddRetargetPose(
 	const FName& NewPoseName,
 	const FIKRetargetPose* ToDuplicate,
-	const ERetargetSourceOrTarget& SourceOrTarget) const
+	const ERetargetSourceOrTarget SourceOrTarget) const
 {
 	FScopedTransaction Transaction(LOCTEXT("AddRetargetPose", "Add Retarget Pose"));
 	Asset->Modify();
@@ -464,7 +478,7 @@ void UIKRetargeterController::AddRetargetPose(
 
 void UIKRetargeterController::RenameCurrentRetargetPose(
 	const FName& NewPoseName,
-	const ERetargetSourceOrTarget& SourceOrTarget) const
+	const ERetargetSourceOrTarget SourceOrTarget) const
 {
 	// do we already have a retarget pose with this name?
 	
@@ -492,7 +506,7 @@ void UIKRetargeterController::RenameCurrentRetargetPose(
 
 void UIKRetargeterController::RemoveRetargetPose(
 	const FName& PoseToRemove,
-	const ERetargetSourceOrTarget& SourceOrTarget) const
+	const ERetargetSourceOrTarget SourceOrTarget) const
 {
 	if (PoseToRemove == Asset->GetDefaultPoseName())
 	{
@@ -522,7 +536,7 @@ void UIKRetargeterController::RemoveRetargetPose(
 void UIKRetargeterController::ResetRetargetPose(
 	const FName& PoseToReset,
 	const TArray<FName>& BonesToReset,
-	const ERetargetSourceOrTarget& SourceOrTarget) const
+	const ERetargetSourceOrTarget SourceOrTarget) const
 {
 	TMap<FName, FIKRetargetPose>& Poses = GetRetargetPoses(SourceOrTarget);
 	if (!Poses.Contains(PoseToReset))
@@ -563,12 +577,12 @@ void UIKRetargeterController::ResetRetargetPose(
 	BroadcastNeedsReinitialized();
 }
 
-FName UIKRetargeterController::GetCurrentRetargetPoseName(const ERetargetSourceOrTarget& SourceOrTarget) const
+FName UIKRetargeterController::GetCurrentRetargetPoseName(const ERetargetSourceOrTarget SourceOrTarget) const
 {
 	return SourceOrTarget == ERetargetSourceOrTarget::Source ? GetAsset()->CurrentSourceRetargetPose : GetAsset()->CurrentTargetRetargetPose;
 }
 
-void UIKRetargeterController::SetCurrentRetargetPose(FName NewCurrentPose, const ERetargetSourceOrTarget& SourceOrTarget) const
+void UIKRetargeterController::SetCurrentRetargetPose(FName NewCurrentPose, const ERetargetSourceOrTarget SourceOrTarget) const
 {
 	const TMap<FName, FIKRetargetPose>& Poses = GetRetargetPoses(SourceOrTarget);
 	check(Poses.Contains(NewCurrentPose));
@@ -581,12 +595,12 @@ void UIKRetargeterController::SetCurrentRetargetPose(FName NewCurrentPose, const
 	BroadcastNeedsReinitialized();
 }
 
-TMap<FName, FIKRetargetPose>& UIKRetargeterController::GetRetargetPoses(const ERetargetSourceOrTarget& SourceOrTarget) const
+TMap<FName, FIKRetargetPose>& UIKRetargeterController::GetRetargetPoses(const ERetargetSourceOrTarget SourceOrTarget) const
 {
 	return SourceOrTarget == ERetargetSourceOrTarget::Source ? GetAsset()->SourceRetargetPoses : GetAsset()->TargetRetargetPoses;
 }
 
-FIKRetargetPose& UIKRetargeterController::GetCurrentRetargetPose(const ERetargetSourceOrTarget& SourceOrTarget) const
+FIKRetargetPose& UIKRetargeterController::GetCurrentRetargetPose(const ERetargetSourceOrTarget SourceOrTarget) const
 {
 	return GetRetargetPoses(SourceOrTarget)[GetCurrentRetargetPoseName(SourceOrTarget)];
 }
@@ -594,7 +608,7 @@ FIKRetargetPose& UIKRetargeterController::GetCurrentRetargetPose(const ERetarget
 void UIKRetargeterController::SetRotationOffsetForRetargetPoseBone(
 	const FName& BoneName,
 	const FQuat& RotationOffset,
-	const ERetargetSourceOrTarget& SourceOrTarget) const
+	const ERetargetSourceOrTarget SourceOrTarget) const
 {
 	const UIKRigDefinition* IKRig = SourceOrTarget == ERetargetSourceOrTarget::Source ? GetAsset()->GetSourceIKRig() : GetAsset()->GetTargetIKRig();
 	FIKRetargetPose& Pose = GetCurrentRetargetPose(SourceOrTarget);
@@ -604,7 +618,7 @@ void UIKRetargeterController::SetRotationOffsetForRetargetPoseBone(
 
 FQuat UIKRetargeterController::GetRotationOffsetForRetargetPoseBone(
 	const FName& BoneName,
-	const ERetargetSourceOrTarget& SourceOrTarget) const
+	const ERetargetSourceOrTarget SourceOrTarget) const
 {
 	TMap<FName, FQuat>& BoneOffsets = GetCurrentRetargetPose(SourceOrTarget).BoneRotationOffsets;
 	if (!BoneOffsets.Contains(BoneName))
@@ -617,12 +631,12 @@ FQuat UIKRetargeterController::GetRotationOffsetForRetargetPoseBone(
 
 void UIKRetargeterController::AddTranslationOffsetToRetargetRootBone(
 	const FVector& TranslationOffset,
-	const ERetargetSourceOrTarget& SourceOrTarget) const
+	const ERetargetSourceOrTarget SourceOrTarget) const
 {
 	GetCurrentRetargetPose(SourceOrTarget).AddToRootTranslationDelta(TranslationOffset);
 }
 
-FName UIKRetargeterController::MakePoseNameUnique(const FString& PoseName, const ERetargetSourceOrTarget& SourceOrTarget) const
+FName UIKRetargeterController::MakePoseNameUnique(const FString& PoseName, const ERetargetSourceOrTarget SourceOrTarget) const
 {
 	FString UniqueName = PoseName;
 	int32 Suffix = 1;
diff --git a/Engine/Plugins/Animation/IKRig/Source/IKRigEditor/Private/RetargetEditor/SIKRetargetAssetBrowser.cpp b/Engine/Plugins/Animation/IKRig/Source/IKRigEditor/Private/RetargetEditor/SIKRetargetAssetBrowser.cpp
--- a/Engine/Plugins/Animation/IKRig/Source/IKRigEditor/Private/RetargetEditor/SIKRetargetAssetBrowser.cpp
+++ b/Engine/Plugins/Animation/IKRig/Source/IKRigEditor/Private/RetargetEditor/SIKRetargetAssetBrowser.cpp
@@ -55,6 +55,13 @@ void SIKRetargetAssetBrowser::Construct(
 void SIKRetargetAssetBrowser::RefreshView()
 {
 	FAssetPickerConfig AssetPickerConfig;
+	
+	// assign "referencer" asset for project filtering
+    if (EditorController.IsValid())
+    {
+    	const TObjectPtr<UObject> Referencer = EditorController.Pin()->AssetController->GetAsset();
+    	AssetPickerConfig.AdditionalReferencingAssets.Add(FAssetData(Referencer));
+    }
 
 	// setup filtering
 	AssetPickerConfig.Filter.ClassPaths.Add(UAnimSequence::StaticClass()->GetClassPathName());
@@ -206,7 +213,7 @@ void SIKRetargetAssetBrowser::OnAssetDoubleClicked(const FAssetData& AssetData)
 	UAnimationAsset* NewAnimationAsset = Cast<UAnimationAsset>(AssetData.GetAsset());
 	if (NewAnimationAsset && EditorController.Pin().IsValid())
 	{
-		EditorController.Pin()->PlayAnimationAsset(NewAnimationAsset);
+		EditorController.Pin()->PlaybackManager->PlayAnimationAsset(NewAnimationAsset);
 	}
 }
 
diff --git a/Engine/Plugins/Animation/IKRig/Source/IKRigEditor/Private/RetargetEditor/SIKRetargetChainMapList.cpp b/Engine/Plugins/Animation/IKRig/Source/IKRigEditor/Private/RetargetEditor/SIKRetargetChainMapList.cpp
--- a/Engine/Plugins/Animation/IKRig/Source/IKRigEditor/Private/RetargetEditor/SIKRetargetChainMapList.cpp
+++ b/Engine/Plugins/Animation/IKRig/Source/IKRigEditor/Private/RetargetEditor/SIKRetargetChainMapList.cpp
@@ -764,12 +764,10 @@ void SIKRetargetChainMapList::AutoMapChains(const EAutoMapChainType AutoMapType,
 		return;
 	}
 
-	// clear output log
 	Controller->ClearOutputLog();
-	// run the chain mapping
+	RetargeterController->CleanChainMapping();
 	RetargeterController->AutoMapChains(AutoMapType, bForceRemap);
-	// refresh the UI to show new mapping
-	RefreshView();
+	Controller->HandleRetargeterNeedsInitialized();
 }
 
 FReply SIKRetargetChainMapList::OnGlobalSettingsButtonClicked() const
@@ -817,3 +815,4 @@ FReply SIKRetargetChainMapList::OnRootSettingsButtonClicked() const
 
 
 #undef LOCTEXT_NAMESPACE
+
diff --git a/Engine/Plugins/Animation/IKRig/Source/IKRigEditor/Private/RetargetEditor/SIKRetargetPoseEditor.cpp b/Engine/Plugins/Animation/IKRig/Source/IKRigEditor/Private/RetargetEditor/SIKRetargetPoseEditor.cpp
--- a/Engine/Plugins/Animation/IKRig/Source/IKRigEditor/Private/RetargetEditor/SIKRetargetPoseEditor.cpp
+++ b/Engine/Plugins/Animation/IKRig/Source/IKRigEditor/Private/RetargetEditor/SIKRetargetPoseEditor.cpp
@@ -130,7 +130,7 @@ void SIKRetargetPoseEditor::Construct(
 							SNew(SButton)
 							.OnClicked_Lambda([Controller]
 							{
-								Controller->StopPlayback();
+								Controller->PlaybackManager->StopPlayback();
 								return FReply::Handled();
 							})
 							.IsEnabled_Lambda([Controller]
diff --git a/Engine/Plugins/Animation/IKRig/Source/IKRigEditor/Private/RigEditor/IKRigController.cpp b/Engine/Plugins/Animation/IKRig/Source/IKRigEditor/Private/RigEditor/IKRigController.cpp
--- a/Engine/Plugins/Animation/IKRig/Source/IKRigEditor/Private/RigEditor/IKRigController.cpp
+++ b/Engine/Plugins/Animation/IKRig/Source/IKRigEditor/Private/RigEditor/IKRigController.cpp
@@ -483,7 +483,7 @@ FName UIKRigController::GetRetargetChainFromGoal(const FName& GoalName) const
 // SKELETON
 //
 
-bool UIKRigController::SetSkeletalMesh(USkeletalMesh* SkeletalMesh) const
+bool UIKRigController::SetSkeletalMesh(USkeletalMesh* SkeletalMesh, bool bTransact) const
 {
 	if (!SkeletalMesh)
 	{
@@ -497,16 +497,22 @@ bool UIKRigController::SetSkeletalMesh(USkeletalMesh* SkeletalMesh) const
 		UE_LOG(LogTemp, Warning, TEXT("Trying to initialize IKRig with a Skeleton that is missing required bones. See output log. {0}"), *Asset->GetName());
 		return false;
 	}
-	
+
+	FScopedTransaction Transaction(LOCTEXT("SetSkeletalMesh_Label", "Set Skeletal Mesh"));
+	if (bTransact)
+	{	
+		Asset->Modify();
+	}
+
 	// update stored skeletal mesh used for previewing results
 	Asset->PreviewSkeletalMesh = SkeletalMesh;
-
 	// copy skeleton data from the actual skeleton we want to run on
 	Asset->Skeleton.SetInputSkeleton(InputSkeleton, GetAsset()->Skeleton.ExcludedBones);
-
 	// update goal's initial transforms to reflect new skeleton
 	ResetInitialGoalTransforms();
 
+	BroadcastNeedsReinitialized();
+
 	return true;
 }
 
@@ -897,8 +903,8 @@ bool UIKRigController::SetGoalBone(const FName& GoalName, const FName& NewBoneNa
 	
 	FScopedTransaction Transaction(LOCTEXT("SetGoalBone_Label", "Set Goal Bone"));
 
-	// update goal with new bone
-	TObjectPtr<UIKRigEffectorGoal> Goal = Asset->Goals[GoalIndex];
+	// update goal's bone name and it's initial transform
+	TObjectPtr<UIKRigEffectorGoal> Goal = Asset->Goals[GoalIndex]; 
 	Goal->Modify();
 	Goal->BoneName = NewBoneName;
 	Goal->InitialTransform = GetRefPoseTransformOfBone(NewBoneName);
@@ -1066,14 +1072,16 @@ void UIKRigController::SetGoalCurrentTransform(const FName& GoalName, const FTra
 	Goal->CurrentTransform = Transform;
 }
 
-void UIKRigController::ResetCurrentGoalTransforms() const
+void UIKRigController::ResetGoalTransforms() const
 {
-	FScopedTransaction Transaction(LOCTEXT("ResetCurrentGoalTransforms", "Reset All Goal Transforms"));
+	FScopedTransaction Transaction(LOCTEXT("ResetGoalTransforms", "Reset All Goal Transforms"));
 	
 	for (UIKRigEffectorGoal* Goal : Asset->Goals)
     {
 		Goal->Modify();
-    	Goal->CurrentTransform = Goal->InitialTransform;
+    	const FTransform InitialTransform = GetRefPoseTransformOfBone(Goal->BoneName);
+    	Goal->InitialTransform = InitialTransform;
+		Goal->CurrentTransform = InitialTransform;
     }
 }
 
diff --git a/Engine/Plugins/Animation/IKRig/Source/IKRigEditor/Private/RigEditor/IKRigDefinitionFactory.cpp b/Engine/Plugins/Animation/IKRig/Source/IKRigEditor/Private/RigEditor/IKRigDefinitionFactory.cpp
--- a/Engine/Plugins/Animation/IKRig/Source/IKRigEditor/Private/RigEditor/IKRigDefinitionFactory.cpp
+++ b/Engine/Plugins/Animation/IKRig/Source/IKRigEditor/Private/RigEditor/IKRigDefinitionFactory.cpp
@@ -1,15 +1,10 @@
 // Copyright Epic Games, Inc. All Rights Reserved.
 
 #include "RigEditor/IKRigDefinitionFactory.h"
-#include "IKRigEditor.h"
 #include "IKRigDefinition.h"
 #include "AssetTypeCategories.h"
 #include "ContentBrowserModule.h"
 #include "IContentBrowserSingleton.h"
-#include "Editor.h"
-#include "Engine/SkeletalMesh.h"
-#include "RigEditor/IKRigController.h"
-#include "Widgets/SWindow.h"
 
 #include UE_INLINE_GENERATED_CPP_BY_NAME(IKRigDefinitionFactory)
 
@@ -26,24 +21,13 @@ UIKRigDefinitionFactory::UIKRigDefinitionFactory()
 UObject* UIKRigDefinitionFactory::FactoryCreateNew(
 	UClass* Class,
 	UObject* InParent,
-	FName Name,
-	EObjectFlags Flags,
+	FName InName,
+	EObjectFlags InFlags,
 	UObject* Context, 
 	FFeedbackContext* Warn)
 {
-	if (!SkeletalMesh.IsValid())
-	{
-		UE_LOG(LogIKRigEditor, Warning, TEXT("Unable to create IK Rig. No Skeletal Mesh asset supplied."));
-		return nullptr;
-	}
-	
-	UIKRigDefinition* IKRig = NewObject<UIKRigDefinition>(InParent, Name, Flags | RF_Transactional);
-	
-	// imports the skeleton data into the IK Rig
-	UIKRigController* Controller = UIKRigController::GetIKRigController(IKRig);
-	Controller->SetSkeletalMesh(SkeletalMesh.Get());
-	
-	return IKRig;
+	// create the IK Rig asset
+	return NewObject<UIKRigDefinition>(InParent, InName, InFlags | RF_Transactional);
 }
 
 bool UIKRigDefinitionFactory::ShouldShowInNewMenu() const
@@ -51,49 +35,6 @@ bool UIKRigDefinitionFactory::ShouldShowInNewMenu() const
 	return true;
 }
 
-void UIKRigDefinitionFactory::OnSkeletalMeshSelected(const FAssetData& SelectedAsset)
-{
-	SkeletalMesh = Cast<USkeletalMesh>(SelectedAsset.GetAsset());
-	PickerWindow->RequestDestroyWindow();
-}
-
-bool UIKRigDefinitionFactory::ConfigureProperties()
-{
-	// Null so we can check for selection later
-	SkeletalMesh = nullptr;
-
-	// Load the content browser module to display an asset picker
-	FContentBrowserModule& ContentBrowserModule = FModuleManager::LoadModuleChecked<FContentBrowserModule>("ContentBrowser");
-
-	FAssetPickerConfig AssetPickerConfig;
-
-	/** The asset picker will only show skeletal meshes */
-	AssetPickerConfig.Filter.ClassPaths.Add(USkeletalMesh::StaticClass()->GetClassPathName());
-
-	/** The delegate that fires when an asset was selected */
-	AssetPickerConfig.OnAssetSelected = FOnAssetSelected::CreateUObject(this, &UIKRigDefinitionFactory::OnSkeletalMeshSelected);
-
-	/** The default view mode should be a list view */
-	AssetPickerConfig.InitialAssetViewType = EAssetViewType::List;
-
-	PickerWindow = SNew(SWindow)
-	.Title(LOCTEXT("CreateIKRigOptions", "Pick Skeletal Mesh"))
-	.ClientSize(FVector2D(500, 600))
-	.SupportsMinimize(false) .SupportsMaximize(false)
-	[
-		SNew(SBorder)
-		.BorderImage( FAppStyle::GetBrush("Menu.Background") )
-		[
-			ContentBrowserModule.Get().CreateAssetPicker(AssetPickerConfig)
-		]
-	];
-
-	GEditor->EditorAddModalWindow(PickerWindow.ToSharedRef());
-	PickerWindow.Reset();
-
-	return SkeletalMesh != nullptr;
-}
-
 FText UIKRigDefinitionFactory::GetDisplayName() const
 {
 	return LOCTEXT("IKRigDefinition_DisplayName", "IK Rig");
diff --git a/Engine/Plugins/Animation/IKRig/Source/IKRigEditor/Private/RigEditor/IKRigEditorController.cpp b/Engine/Plugins/Animation/IKRig/Source/IKRigEditor/Private/RigEditor/IKRigEditorController.cpp
--- a/Engine/Plugins/Animation/IKRig/Source/IKRigEditor/Private/RigEditor/IKRigEditorController.cpp
+++ b/Engine/Plugins/Animation/IKRig/Source/IKRigEditor/Private/RigEditor/IKRigEditorController.cpp
@@ -2,6 +2,8 @@
 
 #include "RigEditor/IKRigEditorController.h"
 
+#include "ContentBrowserModule.h"
+#include "IContentBrowserSingleton.h"
 #include "IPersonaToolkit.h"
 #include "SKismetInspector.h"
 #include "Animation/DebugSkelMeshComponent.h"
@@ -303,7 +305,7 @@ void FIKRigEditorController::Initialize(TSharedPtr<FIKRigEditorToolkit> Toolkit,
 	if (!AssetController->OnIKRigNeedsInitialized().IsBoundToObject(this))
 	{
 		ReinitializeDelegateHandle = AssetController->OnIKRigNeedsInitialized().AddSP(
-			this, &FIKRigEditorController::OnIKRigNeedsInitialized);
+			this, &FIKRigEditorController::HandleIKRigNeedsInitialized);
 
 		// Initialize editor's instances at first initialization
 		InitializeSolvers();
@@ -315,6 +317,60 @@ void FIKRigEditorController::Close() const
 	AssetController->OnIKRigNeedsInitialized().Remove(ReinitializeDelegateHandle);
 }
 
+void FIKRigEditorController::PromptUserToAssignMesh()
+{
+	// do we already have a skeletal mesh assigned?
+	if (AssetController->GetAsset()->PreviewSkeletalMesh)
+	{
+		return;
+	}
+
+	// is there already an imported hierarchy of bones?
+	if (!AssetController->GetIKRigSkeleton().BoneNames.IsEmpty())
+	{
+		return;
+	}
+
+	// no skeletal mesh imported yet... so let's prompt the user to pick one...
+	
+	// Load the content browser module to display an asset picker
+	FContentBrowserModule& ContentBrowserModule = FModuleManager::LoadModuleChecked<FContentBrowserModule>("ContentBrowser");
+	FAssetPickerConfig AssetPickerConfig;
+	// must set the parent UObject so that the resulting list filters correctly in multi-project environments
+	AssetPickerConfig.AdditionalReferencingAssets.Add(AssetController->GetAsset());
+	// the asset picker will only show skeletal meshes
+	AssetPickerConfig.Filter.ClassPaths.Add(USkeletalMesh::StaticClass()->GetClassPathName());
+	// the delegate that fires when an asset is selected
+	AssetPickerConfig.OnAssetSelected = FOnAssetSelected::CreateLambda([this](const FAssetData& AssetData)
+	{
+		MeshPickerWindow->RequestDestroyWindow();
+		
+		if (const TObjectPtr<USkeletalMesh> SkeletalMesh = Cast<USkeletalMesh>(AssetData.GetAsset()))
+		{
+			// import the skeleton data into the IK Rig
+			AssetController->SetSkeletalMesh(SkeletalMesh.Get());
+		}
+		
+	});
+	// the default view mode should be a list view
+	AssetPickerConfig.InitialAssetViewType = EAssetViewType::List;
+
+	MeshPickerWindow = SNew(SWindow)
+	.Title(LOCTEXT("CreateIKRigOptions", "Assign Skeletal Mesh to IK Rig"))
+	.ClientSize(FVector2D(500, 600))
+	.SupportsMinimize(false) .SupportsMaximize(false)
+	[
+		SNew(SBorder)
+		.BorderImage( FAppStyle::GetBrush("Menu.Background") )
+		[
+			ContentBrowserModule.Get().CreateAssetPicker(AssetPickerConfig)
+		]
+	];
+
+	GEditor->EditorAddModalWindow(MeshPickerWindow.ToSharedRef());
+	MeshPickerWindow.Reset();
+}
+
 UIKRigProcessor* FIKRigEditorController::GetIKRigProcessor() const
 {
 	if (AnimInstance.IsValid())
@@ -338,7 +394,7 @@ const FIKRigSkeleton* FIKRigEditorController::GetCurrentIKRigSkeleton() const
 	return nullptr;
 }
 
-void FIKRigEditorController::OnIKRigNeedsInitialized(UIKRigDefinition* ModifiedIKRig)
+void FIKRigEditorController::HandleIKRigNeedsInitialized(UIKRigDefinition* ModifiedIKRig) const
 {
 	if (ModifiedIKRig != AssetController->GetAsset())
 	{
@@ -346,11 +402,6 @@ void FIKRigEditorController::OnIKRigNeedsInitialized(UIKRigDefinition* ModifiedI
 	}
 
 	ClearOutputLog();
-
-	// anim instance needs reinitialized
-	InitializeSolvers();
-	AnimInstance->SetProcessorNeedsInitialized();
-	AnimInstance->InitializeAnimation();
 	
 	// in case the skeletal mesh was swapped out, we need to ensure the preview scene is up-to-date
 	{
@@ -362,6 +413,11 @@ void FIKRigEditorController::OnIKRigNeedsInitialized(UIKRigDefinition* ModifiedI
 		AssetController->ResetInitialGoalTransforms();
 	}
 
+	// anim instance needs reinitialized
+	InitializeSolvers();
+	AnimInstance->SetProcessorNeedsInitialized();
+	AnimInstance->InitializeAnimation();
+
 	// update the bone details so it can pull on the current data
 	BoneDetails->AnimInstancePtr = AnimInstance;
 	BoneDetails->AssetPtr = ModifiedIKRig;
@@ -372,7 +428,7 @@ void FIKRigEditorController::OnIKRigNeedsInitialized(UIKRigDefinition* ModifiedI
 void FIKRigEditorController::Reset() const
 {
 	SkelMeshComponent->ShowReferencePose(true);
-	AssetController->ResetCurrentGoalTransforms();
+	AssetController->ResetGoalTransforms();
 }
 
 void FIKRigEditorController::RefreshAllViews() const
@@ -817,13 +873,10 @@ void FIKRigEditorController::ShowDetailsForElements(const TArray<TSharedPtr<FIKR
 
 void FIKRigEditorController::OnFinishedChangingDetails(const FPropertyChangedEvent& PropertyChangedEvent) const
 {
-	// if user swapped in a different skeletal mesh, then trigger a total reinitialization
-	const bool bPreviewChanged = PropertyChangedEvent.GetPropertyName() == UIKRigDefinition::GetPreviewMeshPropertyName();
-	if (bPreviewChanged)
+	const bool bPreviewMeshChanged = PropertyChangedEvent.GetPropertyName() == UIKRigDefinition::GetPreviewMeshPropertyName();
+	if (bPreviewMeshChanged)
 	{
-		USkeletalMesh* NewMesh = AssetController->GetAsset()->GetPreviewMesh();
-		AssetController->SetSkeletalMesh(NewMesh);
-		AssetController->BroadcastNeedsReinitialized();
+		HandleIKRigNeedsInitialized(AssetController->GetAsset());
 	}
 }
 
diff --git a/Engine/Plugins/Animation/IKRig/Source/IKRigEditor/Private/RigEditor/IKRigToolkit.cpp b/Engine/Plugins/Animation/IKRig/Source/IKRigEditor/Private/RigEditor/IKRigToolkit.cpp
--- a/Engine/Plugins/Animation/IKRig/Source/IKRigEditor/Private/RigEditor/IKRigToolkit.cpp
+++ b/Engine/Plugins/Animation/IKRig/Source/IKRigEditor/Private/RigEditor/IKRigToolkit.cpp
@@ -71,18 +71,21 @@ void FIKRigEditorToolkit::InitAssetEditor(
 		bCreateDefaultToolbar, 
 		IKRigAsset);
 
-	AddApplicationMode(
-		IKRigEditorModes::IKRigEditorMode,
-		MakeShareable(new FIKRigMode(SharedThis(this), PersonaToolkit->GetPreviewScene())));
+	TSharedRef<FIKRigMode> IKRigEditMode = MakeShareable(new FIKRigMode(SharedThis(this), PersonaToolkit->GetPreviewScene()));
+	AddApplicationMode(IKRigEditorModes::IKRigEditorMode, IKRigEditMode);
 
 	SetCurrentMode(IKRigEditorModes::IKRigEditorMode);
 
 	GetEditorModeManager().SetDefaultMode(FIKRigEditMode::ModeName);
-	GetEditorModeManager().ActivateMode(FIKRigEditMode::ModeName);
-	static_cast<FIKRigEditMode*>(GetEditorModeManager().GetActiveMode(FIKRigEditMode::ModeName))->SetEditorController(EditorController);
+	GetEditorModeManager().ActivateDefaultMode();
+	FIKRigEditMode* EditMode = GetEditorModeManager().GetActiveModeTyped<FIKRigEditMode>(FIKRigEditMode::ModeName);
+	EditMode->SetEditorController(EditorController);
 
 	ExtendToolbar();
 	RegenerateMenusAndToolbars();
+
+	// ignored if hierarchy / mesh already imported into IK Rig asset
+	EditorController->PromptUserToAssignMesh();
 }
 
 void FIKRigEditorToolkit::OnClose()
diff --git a/Engine/Plugins/Animation/IKRig/Source/IKRigEditor/Private/RigEditor/SIKRigAssetBrowser.cpp b/Engine/Plugins/Animation/IKRig/Source/IKRigEditor/Private/RigEditor/SIKRigAssetBrowser.cpp
--- a/Engine/Plugins/Animation/IKRig/Source/IKRigEditor/Private/RigEditor/SIKRigAssetBrowser.cpp
+++ b/Engine/Plugins/Animation/IKRig/Source/IKRigEditor/Private/RigEditor/SIKRigAssetBrowser.cpp
@@ -35,6 +35,13 @@ void SIKRigAssetBrowser::Construct(
 void SIKRigAssetBrowser::RefreshView()
 {
 	FAssetPickerConfig AssetPickerConfig;
+
+	// assign "referencer" asset for project filtering
+	if (EditorController.IsValid())
+	{
+		const TObjectPtr<UObject> Referencer = EditorController.Pin()->AssetController->GetAsset();
+		AssetPickerConfig.AdditionalReferencingAssets.Add(FAssetData(Referencer));
+	}
 	
 	// setup filtering
 	AssetPickerConfig.Filter.ClassPaths.Add(UAnimSequence::StaticClass()->GetClassPathName());
diff --git a/Engine/Plugins/Animation/IKRig/Source/IKRigEditor/Public/RetargetEditor/IKRetargetEditorController.h b/Engine/Plugins/Animation/IKRig/Source/IKRigEditor/Public/RetargetEditor/IKRetargetEditorController.h
--- a/Engine/Plugins/Animation/IKRig/Source/IKRigEditor/Public/RetargetEditor/IKRetargetEditorController.h
+++ b/Engine/Plugins/Animation/IKRig/Source/IKRigEditor/Public/RetargetEditor/IKRetargetEditorController.h
@@ -51,7 +51,7 @@ enum class ERetargetSelectionType : uint8
 
 struct FBoundIKRig
 {
-	FBoundIKRig(UIKRigDefinition* InIKRig, FIKRetargetEditorController& InController);
+	FBoundIKRig(UIKRigDefinition* InIKRig, const FIKRetargetEditorController& InController);
 	void UnBind() const;
 	
 	UIKRigDefinition* IKRig;
@@ -61,6 +61,22 @@ struct FBoundIKRig
 	FDelegateHandle RemoveChainDelegateHandle;
 };
 
+struct FRetargetPlaybackManager : public TSharedFromThis<FRetargetPlaybackManager>
+{
+	FRetargetPlaybackManager(TWeakPtr<FIKRetargetEditorController> InEditorController);
+	void PlayAnimationAsset(UAnimationAsset* AssetToPlay);
+	void StopPlayback() const;
+	void PausePlayback();
+	void ResumePlayback() const;
+	
+private:
+
+	TWeakPtr<FIKRetargetEditorController> EditorController;
+	UAnimationAsset* AnimThatWasPlaying = nullptr;
+	float TimeWhenPaused = 0.0f;
+	bool bWasPlayingAnim = false;
+};
+
 /** a home for cross-widget communication to synchronize state across all tabs and viewport */
 class FIKRetargetEditorController : public TSharedFromThis<FIKRetargetEditorController>, FGCObject
 {
@@ -72,10 +88,14 @@ public:
 	void Initialize(TSharedPtr<FIKRetargetEditor> InEditor, UIKRetargeter* InAsset);
 	/** Close the editor */
 	void Close();
+	
+	// prompt user to assign an IK Rig to this retargeter (either source or target), returns true if an IK Rig was assigned
+	void PromptUserToAssignIKRig(const ERetargetSourceOrTarget SourceOrTarget);
+	
 	/** Bind callbacks to this IK Rig */
-	void BindToIKRigAssets(UIKRetargeter* InAsset);
+	void BindToIKRigAssets();
 	/** callback when IK Rig asset requires reinitialization */
-	void HandleIKRigNeedsInitialized(UIKRigDefinition* ModifiedIKRig);
+	void HandleIKRigNeedsInitialized(UIKRigDefinition* ModifiedIKRig) const;
 	/** callback when IK Rig asset's retarget chain's have been added or removed*/
 	void HandleRetargetChainAdded(UIKRigDefinition* ModifiedIKRig) const;
 	/** callback when IK Rig asset's retarget chain has been renamed */
@@ -83,8 +103,14 @@ public:
 	/** callback when IK Rig asset's retarget chain has been removed */
 	void HandleRetargetChainRemoved(UIKRigDefinition* ModifiedIKRig, const FName& InChainRemoved) const;
 	/** callback when IK Retargeter asset requires reinitialization */
-	void HandleRetargeterNeedsInitialized(UIKRetargeter* Retargeter);
+	void HandleRetargeterNeedsInitialized() const;
 	FDelegateHandle RetargeterReInitDelegateHandle;
+	/** callback when IK Rig asset has been swapped out */
+	void HandleIKRigReplaced(ERetargetSourceOrTarget SourceOrTarget);
+	FDelegateHandle IKRigReplacedDelegateHandle;
+	/** callback when Preview Mesh asset has been swapped out */
+	void HandlePreviewMeshReplaced(ERetargetSourceOrTarget SourceOrTarget);
+	FDelegateHandle PreviewMeshReplacedDelegateHandle;
 	
 	/** all modifications to the data model should go through this controller */
 	TObjectPtr<UIKRetargeterController> AssetController;
@@ -95,13 +121,16 @@ public:
 	/** import / export retarget poses */
 	TSharedPtr<FIKRetargetPoseExporter> PoseExporter;
 
+	/** manage playback of animation in the editor */
+	TUniquePtr<FRetargetPlaybackManager> PlaybackManager;
+
 	/** viewport skeletal mesh */
-	UDebugSkelMeshComponent* GetSkeletalMeshComponent(const ERetargetSourceOrTarget& SourceOrTarget) const;
+	UDebugSkelMeshComponent* GetSkeletalMeshComponent(const ERetargetSourceOrTarget SourceOrTarget) const;
 	UDebugSkelMeshComponent* SourceSkelMeshComponent;
 	UDebugSkelMeshComponent* TargetSkelMeshComponent;
 
 	/** viewport anim instance */
-	UIKRetargetAnimInstance* GetAnimInstance(const ERetargetSourceOrTarget& SourceOrTarget) const;
+	UIKRetargetAnimInstance* GetAnimInstance(const ERetargetSourceOrTarget SourceOrTarget) const;
 	UPROPERTY(transient, NonTransactional)
 	TWeakObjectPtr<class UIKRetargetAnimInstance> SourceAnimInstance;
 	UPROPERTY(transient, NonTransactional)
@@ -158,16 +187,6 @@ public:
 	/** Reset the planting state of the IK (when scrubbing or animation loops over) */
 	void ResetIKPlantingState() const;
 
-	/** Sequence Browser**/
-	void PlayAnimationAsset(UAnimationAsset* AssetToPlay);
-	void StopPlayback();
-	void PausePlayback();
-	void ResumePlayback();
-	UAnimationAsset* AnimThatWasPlaying = nullptr;
-	float TimeWhenPaused = 0.0f;
-	bool bWasPlayingAnim = false;
-	/** END Sequence Browser */
-
 	/** Set viewport / editor tool mode */
 	void SetRetargeterMode(ERetargeterOutputMode Mode);
 	ERetargeterOutputMode GetRetargeterMode() const { return OutputMode; }
@@ -207,13 +226,12 @@ public:
 	void SetRootSelected(const bool bIsSelected);
 	bool GetRootSelected() const { return bIsRootSelected; };
 
-	/* remove selected elements that no longer exist after swapping input assets */
 	void CleanSelection(ERetargetSourceOrTarget SourceOrTarget);
-	/* set selection to nothing */
 	void ClearSelection(const bool bKeepBoneSelection=false);
 	ERetargetSelectionType GetLastSelectedItemType() const { return LastSelectedItem; };
 
-	bool GetCameraTargetForSelection(FSphere& OutTarget) const; 
+	/** to frame selection when pressing "f" in viewport */
+	bool GetCameraTargetForSelection(FSphere& OutTarget) const;
 
 	/** ------------------------- END SELECTION -----------------------------*/
 
@@ -300,6 +318,9 @@ private:
 	/** hierarchy view */
 	TSharedPtr<SIKRetargetHierarchy> HierarchyView;
 
+	/** when prompting user to assign an IK Rig */
+	TSharedPtr<SWindow> IKRigPickerWindow;
+
 	/** the current output mode of the retargeter */
 	ERetargeterOutputMode OutputMode;
 	/** slider value to blend between reference pose and retarget pose */
diff --git a/Engine/Plugins/Animation/IKRig/Source/IKRigEditor/Public/RetargetEditor/IKRetargetFactory.h b/Engine/Plugins/Animation/IKRig/Source/IKRigEditor/Public/RetargetEditor/IKRetargetFactory.h
--- a/Engine/Plugins/Animation/IKRig/Source/IKRigEditor/Public/RetargetEditor/IKRetargetFactory.h
+++ b/Engine/Plugins/Animation/IKRig/Source/IKRigEditor/Public/RetargetEditor/IKRetargetFactory.h
@@ -27,12 +27,5 @@ public:
 	virtual FText GetToolTip() const override;
 	virtual FString GetDefaultNewAssetName() const override;
 	virtual UObject* FactoryCreateNew(UClass* Class, UObject* InParent, FName Name, EObjectFlags Flags, UObject* Context, FFeedbackContext* Warn) override;
-	virtual bool ConfigureProperties() override;
 	virtual bool ShouldShowInNewMenu() const override;
-
-private:
-	
-	void OnTargetIKRigSelected(const FAssetData& SelectedAsset);
-
-	TSharedPtr<SWindow> PickerWindow;
 };
diff --git a/Engine/Plugins/Animation/IKRig/Source/IKRigEditor/Public/RetargetEditor/IKRetargeterController.h b/Engine/Plugins/Animation/IKRig/Source/IKRigEditor/Public/RetargetEditor/IKRetargeterController.h
--- a/Engine/Plugins/Animation/IKRig/Source/IKRigEditor/Public/RetargetEditor/IKRetargeterController.h
+++ b/Engine/Plugins/Animation/IKRig/Source/IKRigEditor/Public/RetargetEditor/IKRetargeterController.h
@@ -2,6 +2,7 @@
 #pragma once
 
 #include "CoreMinimal.h"
+#include "IKRetargetAnimInstanceProxy.h"
 #include "UObject/Object.h"
 
 #include "IKRetargeterController.generated.h"
@@ -35,34 +36,36 @@ public:
 	/** Get access to the retargeter asset.
 	 *@warning Do not make modifications to the asset directly. Use this API instead. */
 	UIKRetargeter* GetAsset() const;
+	
 
 	/** SOURCE / TARGET
 	* 
 	*/
-	/** Set the IK Rig to use as the source (to copy animation FROM) */
-	void SetSourceIKRig(UIKRigDefinition* SourceIKRig);
-	/** Get the preview skeletal mesh */
-	USkeletalMesh* GetPreviewMesh(const ERetargetSourceOrTarget& SourceOrTarget) const;
+	/** Set the IK Rig to use as the source or target (to copy animation FROM/TO) */
+	void SetIKRig(const ERetargetSourceOrTarget SourceOrTarget, UIKRigDefinition* IKRig) const;
 	/** Get either source or target IK Rig */
-	const UIKRigDefinition* GetIKRig(const ERetargetSourceOrTarget& SourceOrTarget) const;
+	const UIKRigDefinition* GetIKRig(const ERetargetSourceOrTarget SourceOrTarget) const;
 	/** Get either source or target IK Rig */
-	UIKRigDefinition* GetIKRigWriteable(const ERetargetSourceOrTarget& SourceOrTarget) const;
-	/** Set the SOURCE or TARGET preview mesh based on the mesh in the corresponding IK Rig asset */
-	void OnIKRigChanged(const ERetargetSourceOrTarget& SourceOrTarget) const;
+	UIKRigDefinition* GetIKRigWriteable(const ERetargetSourceOrTarget SourceOrTarget) const;
+	/** Set the preview skeletal mesh for either source or target */
+	void SetPreviewMesh(const ERetargetSourceOrTarget SourceOrTarget, USkeletalMesh* PreviewMesh);
+	/** Get the preview skeletal mesh */
+	USkeletalMesh* GetPreviewMesh(const ERetargetSourceOrTarget SourceOrTarget) const;
 	/** Get if we've already asked to fix the root height for the given skeletal mesh */
 	bool GetAskedToFixRootHeightForMesh(USkeletalMesh* Mesh) const;
 	/** Set if we've asked to fix the root height for the given skeletal mesh */
 	void SetAskedToFixRootHeightForMesh(USkeletalMesh* Mesh, bool InAsked) const;
 	/** Get name of the Root bone used for retargeting. */
-	FName GetRetargetRootBone(const ERetargetSourceOrTarget& SourceOrTarget) const;
+	FName GetRetargetRootBone(const ERetargetSourceOrTarget SourceOrTarget) const;
 
 	/** RETARGET CHAIN MAPPING
 	* 
 	*/
 	/** Get names of all the bone chains. */
-	void GetChainNames(const ERetargetSourceOrTarget& SourceOrTarget, TArray<FName>& OutNames) const;
-	/** Remove invalid chain mappings (no longer existing in currently referenced source/target IK Rig assets) */
-	void CleanChainMapping(const bool bForceReinitialization=true) const;
+	void GetChainNames(const ERetargetSourceOrTarget SourceOrTarget, TArray<FName>& OutNames) const;
+	/** Remove invalid chain mappings (no longer existing in currently referenced source/target IK Rig assets)
+	 * NOTE: be sure to reinitialize any running processors after chain mappings are cleaned! */
+	void CleanChainMapping() const;
 	/** Use fuzzy string search to find "best" Source chain to map to each Target chain */
 	void AutoMapChains(const EAutoMapChainType AutoMapType, const bool bForceRemap) const;
 	/** Callback when IK Rig chain is added or removed. */
@@ -87,68 +90,87 @@ public:
 	 * 
 	 */
 	/** Remove bones from retarget poses that are no longer in skeleton */
-	void CleanPoseLists(const bool bForceReinitialization=true) const;
+	void CleanPoseLists() const;
 	/** Remove bones from retarget poses that are no longer in skeleton */
-	void CleanPoseList(const ERetargetSourceOrTarget& SourceOrTarget) const;
+	void CleanPoseList(const ERetargetSourceOrTarget SourceOrTarget) const;
 	/** Add new retarget pose. */
 	void AddRetargetPose(
 		const FName& NewPoseName,
 		const FIKRetargetPose* ToDuplicate,
-		const ERetargetSourceOrTarget& SourceOrTarget) const;
+		const ERetargetSourceOrTarget SourceOrTarget) const;
 	/** Rename current retarget pose. */
 	void RenameCurrentRetargetPose(
 		const FName& NewPoseName,
-		const ERetargetSourceOrTarget& SourceOrTarget) const;
+		const ERetargetSourceOrTarget SourceOrTarget) const;
 	/** Remove a retarget pose. */
 	void RemoveRetargetPose(
 	const FName& PoseToRemove, 
-	const ERetargetSourceOrTarget& SourceOrTarget) const;
+	const ERetargetSourceOrTarget SourceOrTarget) const;
 	/** Reset a retarget pose for the specified bones.
 	 *If BonesToReset is Empty, will removes all stored deltas, returning pose to reference pose */
 	void ResetRetargetPose(
 		const FName& PoseToReset,
 		const TArray<FName>& BonesToReset,
-		const ERetargetSourceOrTarget& SourceOrTarget) const;
+		const ERetargetSourceOrTarget SourceOrTarget) const;
 	/** Get the current retarget pose */
-    FName GetCurrentRetargetPoseName(const ERetargetSourceOrTarget& SourceOrTarget) const;
+    FName GetCurrentRetargetPoseName(const ERetargetSourceOrTarget SourceOrTarget) const;
 	/** Change which retarget pose is used by the retargeter at runtime */
-	void SetCurrentRetargetPose(FName CurrentPose, const ERetargetSourceOrTarget& SourceOrTarget) const;
+	void SetCurrentRetargetPose(FName CurrentPose, const ERetargetSourceOrTarget SourceOrTarget) const;
 	/** Get access to array of retarget poses */
-	TMap<FName, FIKRetargetPose>& GetRetargetPoses(const ERetargetSourceOrTarget& SourceOrTarget) const;
+	TMap<FName, FIKRetargetPose>& GetRetargetPoses(const ERetargetSourceOrTarget SourceOrTarget) const;
 	/** Get the current retarget pose */
-	FIKRetargetPose& GetCurrentRetargetPose(const ERetargetSourceOrTarget& SourceOrTarget) const;
+	FIKRetargetPose& GetCurrentRetargetPose(const ERetargetSourceOrTarget SourceOrTarget) const;
 	/** Set a delta rotation for a given bone for the current retarget pose (used in Edit Mode in the retarget editor) */
 	void SetRotationOffsetForRetargetPoseBone(
 		const FName& BoneName,
 		const FQuat& RotationOffset,
-		const ERetargetSourceOrTarget& SkeletonMode) const;
+		const ERetargetSourceOrTarget SkeletonMode) const;
 	/** Get a delta rotation for a given bone for the current retarget pose (used in Edit Mode in the retarget editor) */
 	FQuat GetRotationOffsetForRetargetPoseBone(
 		const FName& BoneName,
-		const ERetargetSourceOrTarget& SourceOrTarget) const;
+		const ERetargetSourceOrTarget SourceOrTarget) const;
 	/** Add a delta translation to the root bone (used in Edit Mode in the retarget editor) */
 	void AddTranslationOffsetToRetargetRootBone(
 		const FVector& TranslationOffset,
-		const ERetargetSourceOrTarget& SourceOrTarget) const;
+		const ERetargetSourceOrTarget SourceOrTarget) const;
 	/** Add a numbered suffix to the given pose name to make it unique. */
 	FName MakePoseNameUnique(
 		const FString& PoseName,
-		const ERetargetSourceOrTarget& SourceOrTarget) const;
+		const ERetargetSourceOrTarget SourceOrTarget) const;
 	/** END RETARGET POSE EDITING */
 
 private:
 	
 	/** Called whenever the retargeter is modified in such a way that would require re-initialization by dependent systems.*/
-	DECLARE_MULTICAST_DELEGATE_OneParam(FOnRetargeterNeedsInitialized, UIKRetargeter*);
+	DECLARE_MULTICAST_DELEGATE(FOnRetargeterNeedsInitialized);
 	FOnRetargeterNeedsInitialized RetargeterNeedsInitialized;
+
+	/** Called whenever either the source or target IK Rig asset's are swapped out.*/
+	DECLARE_MULTICAST_DELEGATE_OneParam(FOnIKRigReplaced, ERetargetSourceOrTarget);
+	FOnIKRigReplaced IKRigReplaced;
+	
+	/** Called whenever either the source or target Preview Mesh asset's are swapped out.*/
+	DECLARE_MULTICAST_DELEGATE_OneParam(FOnPreviewMeshReplaced, ERetargetSourceOrTarget);
+	FOnPreviewMeshReplaced PreviewMeshReplaced;
 	
 public:
 	
-	FOnRetargeterNeedsInitialized& OnRetargeterNeedsInitialized(){ return RetargeterNeedsInitialized; };
+	FOnPreviewMeshReplaced& OnPreviewMeshReplaced(){ return PreviewMeshReplaced; };
+	void BroadcastPreviewMeshReplaced(ERetargetSourceOrTarget SourceOrTarget) const
+	{
+		PreviewMeshReplaced.Broadcast(SourceOrTarget);
+	}
 
+	FOnIKRigReplaced& OnIKRigReplaced(){ return IKRigReplaced; };
+	void BroadcastIKRigReplaced(ERetargetSourceOrTarget SourceOrTarget) const
+	{
+		IKRigReplaced.Broadcast(SourceOrTarget);
+	}
+	
+	FOnRetargeterNeedsInitialized& OnRetargeterNeedsInitialized(){ return RetargeterNeedsInitialized; };
 	void BroadcastNeedsReinitialized() const
 	{
-		RetargeterNeedsInitialized.Broadcast(GetAsset());
+		RetargeterNeedsInitialized.Broadcast();
 	}
 	
 private:
diff --git a/Engine/Plugins/Animation/IKRig/Source/IKRigEditor/Public/RigEditor/IKRigController.h b/Engine/Plugins/Animation/IKRig/Source/IKRigEditor/Public/RigEditor/IKRigController.h
--- a/Engine/Plugins/Animation/IKRig/Source/IKRigEditor/Public/RigEditor/IKRigController.h
+++ b/Engine/Plugins/Animation/IKRig/Source/IKRigEditor/Public/RigEditor/IKRigController.h
@@ -39,7 +39,7 @@ public:
 	 */
 	/** Sets the preview mesh to use. Loads the hierarchy into the asset's IKRigSkeleton.
 	 * Returns true if the mesh was able to be set. False if it was incompatible for any reason. */
-	bool SetSkeletalMesh(USkeletalMesh* SkeletalMesh) const;
+	bool SetSkeletalMesh(USkeletalMesh* SkeletalMesh, bool bTransact=false) const;
 	/** Get read-access to the IKRig skeleton representation */
 	const FIKRigSkeleton& GetIKRigSkeleton() const;
 	/** Get the USkeleton asset this rig was initialized with */
@@ -120,8 +120,8 @@ public:
 	/** Set the Goal to the given transform. */
 	void SetGoalCurrentTransform(const FName& GoalName, const FTransform& Transform) const;
 	/** Reset all Goals back to their initial transforms. */
-	void ResetCurrentGoalTransforms() const;
-	/** Update initial goal transforms (called after new mesh is applied). */
+	void ResetGoalTransforms() const;
+	/** Reset initial goal transform */
 	void ResetInitialGoalTransforms() const;
 	/** Ensure that the given name adheres to required standards for Goal names (no special characters etc..)*/
 	static void SanitizeGoalName(FString& InOutName);
diff --git a/Engine/Plugins/Animation/IKRig/Source/IKRigEditor/Public/RigEditor/IKRigDefinitionFactory.h b/Engine/Plugins/Animation/IKRig/Source/IKRigEditor/Public/RigEditor/IKRigDefinitionFactory.h
--- a/Engine/Plugins/Animation/IKRig/Source/IKRigEditor/Public/RigEditor/IKRigDefinitionFactory.h
+++ b/Engine/Plugins/Animation/IKRig/Source/IKRigEditor/Public/RigEditor/IKRigDefinitionFactory.h
@@ -14,9 +14,6 @@ class UIKRigDefinitionFactory : public UFactory
 {
 	GENERATED_BODY()
 
-	UPROPERTY()
-	TWeakObjectPtr<class USkeletalMesh>	SkeletalMesh;
-
 public:
 
 	UIKRigDefinitionFactory();
@@ -25,21 +22,14 @@ public:
 	virtual UObject* FactoryCreateNew(
         UClass* Class,
         UObject* InParent,
-        FName Name,
-        EObjectFlags Flags,
+        FName InName,
+        EObjectFlags InFlags,
         UObject* Context,
         FFeedbackContext* Warn) override;
 	virtual FText GetDisplayName() const override;
 	virtual uint32 GetMenuCategories() const override;
 	virtual FText GetToolTip() const override;
 	virtual FString GetDefaultNewAssetName() const override;
-	virtual bool ConfigureProperties() override;
 	virtual bool ShouldShowInNewMenu() const override;
 	// END UFactory interface
-
-private:
-	
-	void OnSkeletalMeshSelected(const FAssetData& SelectedAsset);
-
-	TSharedPtr<SWindow> PickerWindow;
 };
diff --git a/Engine/Plugins/Animation/IKRig/Source/IKRigEditor/Public/RigEditor/IKRigEditorController.h b/Engine/Plugins/Animation/IKRig/Source/IKRigEditor/Public/RigEditor/IKRigEditorController.h
--- a/Engine/Plugins/Animation/IKRig/Source/IKRigEditor/Public/RigEditor/IKRigEditorController.h
+++ b/Engine/Plugins/Animation/IKRig/Source/IKRigEditor/Public/RigEditor/IKRigEditorController.h
@@ -113,13 +113,16 @@ public:
 	/** cleanup when editor closed */
 	void Close() const;
 
+	// initial setup, used when asset is empty (no hierarchy, no mesh assigned)
+	void PromptUserToAssignMesh();
+
 	/** get the currently active processor running the IK Rig in the editor */
 	UIKRigProcessor* GetIKRigProcessor() const;
 	/** get the currently running IKRig skeleton (if there is a running processor) */
 	const FIKRigSkeleton* GetCurrentIKRigSkeleton() const;
 
 	/** callback when IK Rig requires re-initialization */
-	void OnIKRigNeedsInitialized(UIKRigDefinition* ModifiedIKRig);
+	void HandleIKRigNeedsInitialized(UIKRigDefinition* ModifiedIKRig) const;
 
 	/** create goals */
 	void AddNewGoals(const TArray<FName>& GoalNames, const TArray<FName>& BoneNames);
@@ -260,6 +263,9 @@ private:
 	/** output log view */
 	TSharedPtr<SIKRigOutputLog> OutputLogView;
 
+	/** when prompting user to assign a skeletal mesh */
+	TSharedPtr<SWindow> MeshPickerWindow;
+
 	EIKRigSelectionType LastSelectedType;
 
 	UPROPERTY()
