ensure that the bone actually exists before trying to create a new goal

diff --git a/Engine/Plugins/Animation/IKRig/Source/IKRigEditor/Private/RigEditor/IKRigController.cpp b/Engine/Plugins/Animation/IKRig/Source/IKRigEditor/Private/RigEditor/IKRigController.cpp
--- a/Engine/Plugins/Animation/IKRig/Source/IKRigEditor/Private/RigEditor/IKRigController.cpp
+++ b/Engine/Plugins/Animation/IKRig/Source/IKRigEditor/Private/RigEditor/IKRigController.cpp
@@ -732,6 +732,11 @@ UIKRigSolver* UIKRigController::GetSolver(int32 Index) const
 
 UIKRigEffectorGoal* UIKRigController::AddNewGoal(const FName& GoalName, const FName& BoneName) const
 {
+	if (GetIKRigSkeleton().GetBoneIndexFromName(BoneName) == INDEX_NONE)
+	{
+		return nullptr; // bone does not exist in the skeleton
+	}
+	
 	if (GetGoalIndex(GoalName) != INDEX_NONE)
 	{
 		return nullptr; // goal already exists!
diff --git a/Engine/Plugins/Animation/IKRig/Source/IKRigEditor/Private/RigEditor/IKRigEditorController.cpp b/Engine/Plugins/Animation/IKRig/Source/IKRigEditor/Private/RigEditor/IKRigEditorController.cpp
--- a/Engine/Plugins/Animation/IKRig/Source/IKRigEditor/Private/RigEditor/IKRigEditorController.cpp
+++ b/Engine/Plugins/Animation/IKRig/Source/IKRigEditor/Private/RigEditor/IKRigEditorController.cpp
@@ -967,6 +967,13 @@ UIKRigEffectorGoal* FIKRigEditorController::PromptToAddGoalToNewChain(const FBon
 		return ExistingGoal;
 	}
 
+	// check if there is a corresponding bone, if there is none, then no need to create a goal
+	const FIKRigSkeleton& IKRigSkeleton = AssetController->GetIKRigSkeleton();
+	if (IKRigSkeleton.GetBoneIndexFromName(BoneChain.EndBone.BoneName) == INDEX_NONE)
+	{
+		return nullptr;
+	}
+
 	// ask user if they want to add a goal to the end of this chain
 	const TSharedRef<SCustomDialog> AddGoalToChainDialog = SNew(SCustomDialog)
 		.Title(FText(LOCTEXT("AddGoalToNewChainTitleLabel", "Add Goal to New Chain")))
