Do not prompt ik rig asset selection when unattended

#rnx

diff --git a/Engine/Plugins/Animation/IKRig/Source/IKRigEditor/Private/RetargetEditor/IKRetargetEditorController.cpp b/Engine/Plugins/Animation/IKRig/Source/IKRigEditor/Private/RetargetEditor/IKRetargetEditorController.cpp
--- a/Engine/Plugins/Animation/IKRig/Source/IKRigEditor/Private/RetargetEditor/IKRetargetEditorController.cpp
+++ b/Engine/Plugins/Animation/IKRig/Source/IKRigEditor/Private/RetargetEditor/IKRetargetEditorController.cpp
@@ -167,8 +167,8 @@ void FIKRetargetEditorController::Close()
 
 void FIKRetargetEditorController::PromptUserToAssignIKRig(const ERetargetSourceOrTarget SourceOrTarget)
 {
-	// early out if we already have an IK Rig assigned
-	if (AssetController->GetIKRig(SourceOrTarget))
+	// early out if we already have an IK Rig assigned or is unattended
+	if (AssetController->GetIKRig(SourceOrTarget) || FApp::IsUnattended())
 	{
 		return;
 	}
