Adding python API for creating/editing IK Rig assets

diff --git a/Engine/Plugins/Animation/AnimationData/Source/AnimationData/Private/AnimSequencerDataModel.cpp b/Engine/Plugins/Animation/AnimationData/Source/AnimationData/Private/AnimSequencerDataModel.cpp
--- a/Engine/Plugins/Animation/AnimationData/Source/AnimationData/Private/AnimSequencerDataModel.cpp
+++ b/Engine/Plugins/Animation/AnimationData/Source/AnimationData/Private/AnimSequencerDataModel.cpp
@@ -785,9 +785,12 @@ void UAnimationSequencerDataModel::GenerateLegacyCurveData()
 									const FAnimationCurveIdentifier CurveId(FloatCurve.Name, ERawCurveTrackTypes::RCT_Float);
 									if (!RetainFloatCurves || !FloatCurve.Name.DisplayName.ToString().Contains(SequencerSuffix))
 									{
-										const FAnimationCurveMetaData& CurveMetaData = CurveIdentifierToMetaData.FindChecked(CurveId);
-										FloatCurve.SetCurveTypeFlags(CurveMetaData.Flags);
-										FloatCurve.Color = CurveMetaData.Color;
+										if (CurveIdentifierToMetaData.Contains(CurveId))
+										{
+											const FAnimationCurveMetaData& CurveMetaData = CurveIdentifierToMetaData.FindChecked(CurveId);
+											FloatCurve.SetCurveTypeFlags(CurveMetaData.Flags);
+											FloatCurve.Color = CurveMetaData.Color;
+										}
 									}							
 
 									if (const FScalarParameterNameAndCurve* ScalarCurve = ScalarCurves.FindByPredicate([Element](FScalarParameterNameAndCurve Curve)
diff --git a/Engine/Plugins/Animation/IKRig/Source/IKRig/Private/IKRigDefinition.cpp b/Engine/Plugins/Animation/IKRig/Source/IKRig/Private/IKRigDefinition.cpp
--- a/Engine/Plugins/Animation/IKRig/Source/IKRig/Private/IKRigDefinition.cpp
+++ b/Engine/Plugins/Animation/IKRig/Source/IKRig/Private/IKRigDefinition.cpp
@@ -287,8 +287,9 @@ const FBoneChain* UIKRigDefinition::GetRetargetChainByName(FName ChainName) cons
 	return nullptr;
 }
 
+/** IInterface_PreviewMeshProvider interface */
 void UIKRigDefinition::SetPreviewMesh(USkeletalMesh* PreviewMesh, bool bMarkAsDirty)
-{	
+{
 	PreviewSkeletalMesh = PreviewMesh;
 }
 
@@ -296,6 +297,7 @@ USkeletalMesh* UIKRigDefinition::GetPreviewMesh() const
 {
 	return PreviewSkeletalMesh.LoadSynchronous();
 }
+/** END IInterface_PreviewMeshProvider interface */
 
 #undef LOCTEXT_NAMESPACE
 
diff --git a/Engine/Plugins/Animation/IKRig/Source/IKRig/Private/IKRigProcessor.cpp b/Engine/Plugins/Animation/IKRig/Source/IKRig/Private/IKRigProcessor.cpp
--- a/Engine/Plugins/Animation/IKRig/Source/IKRig/Private/IKRigProcessor.cpp
+++ b/Engine/Plugins/Animation/IKRig/Source/IKRig/Private/IKRigProcessor.cpp
@@ -37,9 +37,9 @@ void UIKRigProcessor::Initialize(
 	bTriedToInitialize = true;
 
 	// copy skeleton data from the actual skeleton we want to run on
-	Skeleton.SetInputSkeleton(SkeletalMesh, InRigAsset->Skeleton.ExcludedBones);
+	Skeleton.SetInputSkeleton(SkeletalMesh, InRigAsset->GetSkeleton().ExcludedBones);
 	
-	if (InRigAsset->Skeleton.BoneNames.IsEmpty())
+	if (InRigAsset->GetSkeleton().BoneNames.IsEmpty())
 	{
 		Log.LogError(FText::Format(
 			LOCTEXT("NoSkeleton", "Trying to initialize IK Rig, '{0}' that has no skeleton."),
@@ -213,20 +213,20 @@ bool UIKRigProcessor::IsIKRigCompatibleWithSkeleton(
 
 	// now we validate that hierarchy matches for all required bones...
 	bool bAllParentsValid = true;
-	
+	const FIKRigSkeleton& AssetSkeleton = InRigAsset->GetSkeleton();
 	for (const FName& RequiredBone : RequiredBones)
 	{
 		const int32 InputBoneIndex = InputSkeleton.BoneNames.Find(RequiredBone);
-		const int32 AssetBoneIndex = InRigAsset->Skeleton.BoneNames.Find(RequiredBone);
+		const int32 AssetBoneIndex = AssetSkeleton.BoneNames.Find(RequiredBone);
 
 		// we shouldn't get this far otherwise due to early return above...
 		check(InputBoneIndex != INDEX_NONE && AssetBoneIndex != INDEX_NONE)
 
 		// validate that input skeleton hierarchy is as expected
-		const int32 AssetParentIndex = InRigAsset->Skeleton.ParentIndices[AssetBoneIndex];
-		if (InRigAsset->Skeleton.BoneNames.IsValidIndex(AssetParentIndex)) // root bone has no parent
+		const int32 AssetParentIndex = AssetSkeleton.ParentIndices[AssetBoneIndex];
+		if (AssetSkeleton.BoneNames.IsValidIndex(AssetParentIndex)) // root bone has no parent
 		{
-			const FName& AssetParentName = InRigAsset->Skeleton.BoneNames[AssetParentIndex];
+			const FName& AssetParentName = AssetSkeleton.BoneNames[AssetParentIndex];
 			const int32 InputParentIndex = InputSkeleton.ParentIndices[InputBoneIndex];
 			if (!InputSkeleton.BoneNames.IsValidIndex(InputParentIndex))
 			{
diff --git a/Engine/Plugins/Animation/IKRig/Source/IKRig/Private/IKRigSkeleton.cpp b/Engine/Plugins/Animation/IKRig/Source/IKRig/Private/IKRigSkeleton.cpp
--- a/Engine/Plugins/Animation/IKRig/Source/IKRig/Private/IKRigSkeleton.cpp
+++ b/Engine/Plugins/Animation/IKRig/Source/IKRig/Private/IKRigSkeleton.cpp
@@ -21,6 +21,11 @@ void FIKRigInputSkeleton::Initialize(const USkeletalMesh* SkeletalMesh)
 {
 	Reset();
 
+	if (!SkeletalMesh)
+	{
+		return;
+	}
+	
 	SkeletalMeshName = FName(SkeletalMesh->GetName());
 	const FReferenceSkeleton& RefSkeleton = SkeletalMesh->GetRefSkeleton();
 	const TArray<FMeshBoneInfo>& BoneInfo = RefSkeleton.GetRefBoneInfo();
@@ -368,9 +373,11 @@ void FIKRigSkeleton::GetChainsInList(const TArray<int32>& SelectedBones, TArray<
     		}
     	}
 
+    	const FName ChainName = NAME_None;
     	const FName StartBoneName = GetBoneNameFromIndex(ChainStart);
     	const FName EndBoneName = GetBoneNameFromIndex(ChainEnd);
-    	OutChains.Add(FBoneChain(NAME_None, StartBoneName, EndBoneName));	
+    	const FName GoalName = NAME_None;
+    	OutChains.Add(FBoneChain(ChainName, StartBoneName, EndBoneName, GoalName));	
     }
 }
 
diff --git a/Engine/Plugins/Animation/IKRig/Source/IKRig/Private/Solvers/IKRig_PoleSolver.cpp b/Engine/Plugins/Animation/IKRig/Source/IKRig/Private/Solvers/IKRig_PoleSolver.cpp
--- a/Engine/Plugins/Animation/IKRig/Source/IKRig/Private/Solvers/IKRig_PoleSolver.cpp
+++ b/Engine/Plugins/Animation/IKRig/Source/IKRig/Private/Solvers/IKRig_PoleSolver.cpp
@@ -260,6 +260,11 @@ void UIKRig_PoleSolver::SetEndBone(const FName& EndBoneName)
 	EndName = EndBoneName;
 }
 
+FName UIKRig_PoleSolver::GetEndBone() const
+{
+	return EndName;
+}
+
 #endif
 
 void UIKRig_PoleSolver::GatherChildren(const int32 BoneIndex, const FIKRigSkeleton& InSkeleton, TArray<int32>& OutChildren)
diff --git a/Engine/Plugins/Animation/IKRig/Source/IKRig/Public/IKRigDefinition.h b/Engine/Plugins/Animation/IKRig/Source/IKRig/Public/IKRigDefinition.h
--- a/Engine/Plugins/Animation/IKRig/Source/IKRig/Public/IKRigDefinition.h
+++ b/Engine/Plugins/Animation/IKRig/Source/IKRig/Public/IKRigDefinition.h
@@ -10,6 +10,7 @@
 #if WITH_EDITOR
 #include "SAdvancedTransformInputBox.h"
 #endif
+
 #include "IKRigDefinition.generated.h"
 
 class IPropertyHandle;
@@ -128,18 +129,22 @@ public:
 #endif
 };
 
-USTRUCT()
+USTRUCT(BlueprintType)
 struct IKRIG_API FBoneChain
 {
 	GENERATED_BODY()
 
 	FBoneChain() = default;
 	
-	FBoneChain(FName InName, const FName& InStartBone, const FName& InEndBone)
+	FBoneChain(
+		FName InName, 
+		const FName InStartBone,
+		const FName InEndBone,
+		const FName InGoalName)
 		: ChainName(InName)
 		, StartBone(InStartBone)
 		, EndBone(InEndBone)
-		, IKGoalName(NAME_None)
+		, IKGoalName(InGoalName)
 	{}
 
 	UPROPERTY(EditAnywhere, Category = BoneChain)
@@ -188,8 +193,39 @@ public:
 	UPROPERTY(AssetRegistrySearchable, EditAnywhere, Category = PreviewMesh)
 	TSoftObjectPtr<USkeletalMesh> PreviewSkeletalMesh;
 
-#if WITH_EDITORONLY_DATA
+	/** runtime, read-only access to skeleton data */
+	const FIKRigSkeleton& GetSkeleton() const { return Skeleton; };
 	
+	/** runtime, read-only access to array of Goals, all modifications must go through UIKRigController */
+	const TArray<UIKRigEffectorGoal*>& GetGoalArray() const { return Goals; };
+
+	/** runtime, read-only access to array of Solvers, all modifications must go through UIKRigController */
+	const TArray<UIKRigSolver*>& GetSolverArray() const { return Solvers; };
+
+	/** runtime, read-only access to Retarget Definition, all modifications must go through UIKRigController */
+	const TArray<FBoneChain>& GetRetargetChains() const { return RetargetDefinition.BoneChains; };
+
+	/** runtime, read-only access to Retarget Root, all modifications must go through UIKRigController */
+	const FName& GetRetargetRoot() const { return RetargetDefinition.RootBone; };
+
+	/** runtime, read-only access to read a bone chain, all modifications must go through UIKRigController */
+	const FBoneChain* GetRetargetChainByName(FName ChainName) const;
+
+	/** UObject */
+	virtual void Serialize(FArchive& Ar) override;
+	/** END UObject */
+	
+	/** IInterface_PreviewMeshProvider interface */
+	virtual void SetPreviewMesh(USkeletalMesh* PreviewMesh, bool bMarkAsDirty = true) override;
+	virtual USkeletalMesh* GetPreviewMesh() const override;
+	/** END IInterface_PreviewMeshProvider interface */
+
+#if WITH_EDITOR
+	/* Get name of Preview Mesh property */
+	static const FName GetPreviewMeshPropertyName();
+#endif
+
+#if WITH_EDITORONLY_DATA
 	/**The size of the Bones in the editor viewport. This is set by callbacks from the viewport Character>Bones menu*/
 	UPROPERTY()
 	float BoneSize = 2.0f;
@@ -206,46 +242,16 @@ public:
 	UPROPERTY(EditAnywhere, Category = "Viewport Goal Settings",  meta = (ClampMin = "0.01", UIMin = "0.0", UIMax = "10.0"))
 	float GoalThickness = 0.7f;
 
-	/** The controller responsible for managing this asset's data (all editor mutation goes through this) */
+	/** The controller responsible for managing this asset's data (all mutation goes through this) */
 	UPROPERTY(Transient)
-	TObjectPtr<UObject> Controller;
-
+	mutable TObjectPtr<UObject> Controller;
 #endif
 
-	/** UObject */
-	virtual void Serialize(FArchive& Ar) override;
-	/** END UObject */
+private:
 	
 	/** hierarchy and bone-pose transforms */
 	UPROPERTY()
 	FIKRigSkeleton Skeleton;
-	
-	/** read-only access to array of Goals, all modifications must go through UIKRigController */
-	const TArray<UIKRigEffectorGoal*>& GetGoalArray() const { return Goals; };
-
-	/** read-only access to array of Solvers, all modifications must go through UIKRigController */
-	const TArray<UIKRigSolver*>& GetSolverArray() const { return Solvers; };
-
-	/** runtime, read-only access to Retarget Definition, all modifications must go through UIKRigController */
-	const TArray<FBoneChain>& GetRetargetChains() const { return RetargetDefinition.BoneChains; };
-
-	/** runtime, read-only access to a FBoneChain by name. Returns nullptr if not found. */
-	const FBoneChain* GetRetargetChainByName(FName ChainName) const;
-
-	/** runtime, read-only access to Retarget Root */
-	const FName& GetRetargetRoot() const { return RetargetDefinition.RootBone; };
-	
-	/** IInterface_PreviewMeshProvider interface */
-	virtual void SetPreviewMesh(USkeletalMesh* PreviewMesh, bool bMarkAsDirty = true) override;
-	virtual USkeletalMesh* GetPreviewMesh() const override;
-	/** END IInterface_PreviewMeshProvider interface */
-
-#if WITH_EDITOR
-	/* Get name of Preview Mesh property */
-	static const FName GetPreviewMeshPropertyName();
-#endif
-	
-private:
 
 	/** goals, used as effectors by solvers that support them */
 	UPROPERTY()
diff --git a/Engine/Plugins/Animation/IKRig/Source/IKRig/Public/IKRigSolver.h b/Engine/Plugins/Animation/IKRig/Source/IKRig/Public/IKRigSolver.h
--- a/Engine/Plugins/Animation/IKRig/Source/IKRig/Public/IKRigSolver.h
+++ b/Engine/Plugins/Animation/IKRig/Source/IKRig/Public/IKRigSolver.h
@@ -85,6 +85,9 @@ public:
 	//** ROOT BONE (optional, implement if your solver requires a end bone) */
 	/** override to support SETTING END BONE for the solver */
 	virtual void SetEndBone(const FName& EndBoneName){};
+	/** override to support GETTING END BONE of the solver */
+	virtual FName GetEndBone() const { return NAME_None; };
+	/** override and return true if this solver requires an end bone to be specified */
 	virtual bool RequiresEndBone() const { return false; };
 	//** END ROOT BONE */
 
diff --git a/Engine/Plugins/Animation/IKRig/Source/IKRig/Public/Retargeter/IKRetargeter.h b/Engine/Plugins/Animation/IKRig/Source/IKRig/Public/Retargeter/IKRetargeter.h
--- a/Engine/Plugins/Animation/IKRig/Source/IKRig/Public/Retargeter/IKRetargeter.h
+++ b/Engine/Plugins/Animation/IKRig/Source/IKRig/Public/Retargeter/IKRetargeter.h
@@ -42,24 +42,24 @@ public:
 	
 	URetargetChainSettings(const FName& TargetChain) : TargetChain(TargetChain){}
 	
-	/** UObject */
+	// UObject 
 	virtual void Serialize(FArchive& Ar) override;
-	/** END UObject */
+	// END UObject 
 	
-	/** The chain on the Source IK Rig asset to copy animation FROM. */
+	// The chain on the Source IK Rig asset to copy animation FROM. 
 	UPROPERTY(VisibleAnywhere, Category = "Chain Mapping")
 	FName SourceChain;
 
-	/** The chain on the Target IK Rig asset to copy animation TO. */
+	// The chain on the Target IK Rig asset to copy animation TO. 
 	UPROPERTY(VisibleAnywhere, Category = "Chain Mapping")
 	FName TargetChain;
 
-	/** The settings used to control the motion on this target chain. */
+	// The settings used to control the motion on this target chain. 
 	UPROPERTY(EditAnywhere, Category = "Settings")
 	FTargetChainSettings Settings;
 
 #if WITH_EDITORONLY_DATA
-	/** Deprecated properties from before FTargetChainSettings / profile refactor  (July 2022)*/
+	// Deprecated properties from before FTargetChainSettings / profile refactor  (July 2022)
 	UPROPERTY()
 	bool CopyPoseUsingFK_DEPRECATED = true;
 	UPROPERTY()
@@ -94,7 +94,7 @@ public:
 	float UnplantStiffness_DEPRECATED = 250.0f;
 	UPROPERTY()
 	float UnplantCriticalDamping_DEPRECATED = 1.0f;
-	/** END deprecated properties */
+	// END deprecated properties 
 #endif
 
 	// pointer to editor for details customization
@@ -103,18 +103,18 @@ public:
 	#endif
 };
 
-UCLASS()
+UCLASS(BlueprintType)
 class IKRIG_API URetargetRootSettings: public UObject
 {
 	GENERATED_BODY()
 	
 public:
 	
-	/** UObject */
+	// UObject 
 	virtual void Serialize(FArchive& Ar) override;
-	/** END UObject */
+	// END UObject 
 
-	/** The settings used to control the motion of the target root bone. */
+	// The settings used to control the motion of the target root bone. 
 	UPROPERTY(EditAnywhere, Category = "Settings")
 	FTargetRootSettings Settings;
 
@@ -125,7 +125,7 @@ public:
 
 private:
 #if WITH_EDITORONLY_DATA
-	/** Deprecated properties from before FTargetRootSettings / profile refactor */
+	// Deprecated properties from before FTargetRootSettings / profile refactor 
 	UPROPERTY()
 	bool RetargetRootTranslation_DEPRECATED = true;
 	UPROPERTY()
@@ -138,19 +138,19 @@ private:
 	FVector StaticOffset_DEPRECATED = FVector::ZeroVector;
 	UPROPERTY()
 	FRotator StaticRotationOffset_DEPRECATED = FRotator::ZeroRotator;
-	/** END deprecated properties */
+	// END deprecated properties 
 #endif
 };
 
 
-UCLASS()
+UCLASS(BlueprintType)
 class IKRIG_API UIKRetargetGlobalSettings: public UObject
 {
 	GENERATED_BODY()
 	
 public:
 
-	/** Global retargeter settings. */
+	// Global retargeter settings. 
 	UPROPERTY(EditAnywhere, Category = "Settings")
 	FRetargetGlobalSettings Settings;
 
@@ -161,7 +161,7 @@ public:
 };
 
 
-USTRUCT()
+USTRUCT(BlueprintType)
 struct IKRIG_API FIKRetargetPose
 {
 	GENERATED_BODY()
@@ -193,13 +193,14 @@ private:
 };
 
 // which skeleton are we referring to?
+UENUM()
 enum class ERetargetSourceOrTarget : uint8
 {
 	Source,	// the SOURCE skeleton (to copy FROM)
 	Target, // the TARGET skeleton (to copy TO)
 };
 
-UCLASS(Blueprintable)
+UCLASS(BlueprintType)
 class IKRIG_API UIKRetargeter : public UObject
 {
 	GENERATED_BODY()
@@ -208,148 +209,147 @@ public:
 	
 	UIKRetargeter(const FObjectInitializer& ObjectInitializer);
 
-	/** Get read-only access to the source IK Rig asset */
+	// Get read-only access to the source IK Rig asset 
 	const UIKRigDefinition* GetSourceIKRig() const { return SourceIKRigAsset.LoadSynchronous(); };
-	/** Get read-only access to the target IK Rig asset */
+	// Get read-only access to the target IK Rig asset 
 	const UIKRigDefinition* GetTargetIKRig() const { return TargetIKRigAsset.LoadSynchronous(); };
-	/** Get read-write access to the source IK Rig asset.
-	 * WARNING: do not use for editing the data model. Use Controller class instead. */
+	// Get read-write access to the source IK Rig asset.
+	// WARNING: do not use for editing the data model. Use Controller class instead. 
 	UIKRigDefinition* GetSourceIKRigWriteable() const { return SourceIKRigAsset.LoadSynchronous(); };
-	/** Get read-write access to the target IK Rig asset.
-	 * WARNING: do not use for editing the data model. Use Controller class instead. */
+	// Get read-write access to the target IK Rig asset.
+	// WARNING: do not use for editing the data model. Use Controller class instead. 
 	UIKRigDefinition* GetTargetIKRigWriteable() const { return TargetIKRigAsset.LoadSynchronous(); };
 
-	/** Get read-only access to the chain mapping */
+	// Get read-only access to the chain mapping 
 	const TArray<TObjectPtr<URetargetChainSettings>>& GetAllChainSettings() const { return ChainSettings; };
-	/** Get read-only access to the chain map for a given chain (null if chain not in retargeter) */
+	// Get read-only access to the chain map for a given chain (null if chain not in retargeter) 
 	const TObjectPtr<URetargetChainSettings> GetChainMapByName(const FName& TargetChainName) const;
-	/** Get read-only access to the chain settings for a given chain (null if chain not in retargeter) */
+	// Get read-only access to the chain settings for a given chain (null if chain not in retargeter) 
 	const FTargetChainSettings* GetChainSettingsByName(const FName& TargetChainName) const;
-	/** Get access to the root settings */
+	// Get access to the root settings 
 	URetargetRootSettings* GetRootSettingsUObject() const { return RootSettings; };
-	/** Get access to the global settings uobject*/
+	// Get access to the global settings uobject
 	UIKRetargetGlobalSettings* GetGlobalSettingsUObject() const { return GlobalSettings; };
-	/** Get access to the global settings itself */
+	// Get access to the global settings itself 
 	const FRetargetGlobalSettings& GetGlobalSettings() const { return GlobalSettings->Settings; };
 
-	/** Get read-only access to a retarget pose */
+	// Get read-only access to a retarget pose 
 	const FIKRetargetPose* GetCurrentRetargetPose(const ERetargetSourceOrTarget& SourceOrTarget) const;
-	/** Get name of the current retarget pose */
+	// Get name of the current retarget pose 
 	FName GetCurrentRetargetPoseName(const ERetargetSourceOrTarget& SourceOrTarget) const;
-	/** Get read-only access to a retarget pose */
+	// Get read-only access to a retarget pose 
 	const FIKRetargetPose* GetRetargetPoseByName(const ERetargetSourceOrTarget& SourceOrTarget, const FName PoseName) const;
-	/* Get name of default pose */
+	// Get name of default pose 
 	static const FName GetDefaultPoseName();
 	
-	/* Get the current retarget profile (may be null) */
+	// Get the current retarget profile (may be null) 
 	const FRetargetProfile* GetCurrentProfile() const;
-	/* Get the retarget profile by name (may be null) */
+	// Get the retarget profile by name (may be null) 
 	const FRetargetProfile* GetProfileByName(const FName& ProfileName) const;
 
-	/** BLUEPRINT GETTERS */
+	// BLUEPRINT GETTERS 
 
-	/** Returns the chain settings associated with a given Goal in an IK Retargeter Asset using the given profile name (optional) */
+	// Returns the chain settings associated with a given Goal in an IK Retargeter Asset using the given profile name (optional) 
 	UFUNCTION(BlueprintCallable, BlueprintPure, Category=RetargetAsset)
 	static FTargetChainSettings GetChainUsingGoalFromRetargetAsset(
 		const UIKRetargeter* RetargetAsset,
 		const FName IKGoalName);
 	
-	/** Returns the chain settings associated with a given target chain in an IK Retargeter Asset using the given profile name (optional) */
+	// Returns the chain settings associated with a given target chain in an IK Retargeter Asset using the given profile name (optional) 
 	UFUNCTION(BlueprintCallable, BlueprintPure, Category=RetargetAsset)
 	static FTargetChainSettings GetChainSettingsFromRetargetAsset(
 		const UIKRetargeter* RetargetAsset,
 		const FName TargetChainName,
 		const FName OptionalProfileName);
 
-	/** Returns the chain settings associated with a given target chain in the supplied Retarget Profile. */
+	// Returns the chain settings associated with a given target chain in the supplied Retarget Profile. 
 	UFUNCTION(BlueprintCallable, BlueprintPure, Category=RetargetProfile)
 	static FTargetChainSettings GetChainSettingsFromRetargetProfile(
 		UPARAM(ref) FRetargetProfile& RetargetProfile,
 		const FName TargetChainName);
 
-	/** Returns the root settings in an IK Retargeter Asset using the given profile name (optional) */
+	// Returns the root settings in an IK Retargeter Asset using the given profile name (optional) 
 	UFUNCTION(BlueprintCallable, BlueprintPure, Category=RetargetAsset)
 	static void GetRootSettingsFromRetargetAsset(
 		const UIKRetargeter* RetargetAsset,
 		const FName OptionalProfileName,
 		FTargetRootSettings& OutSettings);
 
-	/** Returns the root settings in the supplied Retarget Profile. */
+	// Returns the root settings in the supplied Retarget Profile. 
 	UFUNCTION(BlueprintCallable, BlueprintPure, Category=RetargetProfile)
 	static FTargetRootSettings GetRootSettingsFromRetargetProfile(UPARAM(ref) FRetargetProfile& RetargetProfile);
 
-	/** Returns the global settings in an IK Retargeter Asset using the given profile name (optional) */
+	// Returns the global settings in an IK Retargeter Asset using the given profile name (optional) 
 	UFUNCTION(BlueprintCallable, BlueprintPure, Category=RetargetAsset)
 	static void GetGlobalSettingsFromRetargetAsset(
 		const UIKRetargeter* RetargetAsset,
 		const FName OptionalProfileName,
 		FRetargetGlobalSettings& OutSettings);
 
-	/** Returns the global settings in the supplied Retarget Profile. */
+	// Returns the global settings in the supplied Retarget Profile. 
 	UFUNCTION(BlueprintCallable, BlueprintPure, Category=RetargetProfile)
 	static FRetargetGlobalSettings GetGlobalSettingsFromRetargetProfile(UPARAM(ref) FRetargetProfile& RetargetProfile);
 
-	/** BLUEPRINT SETTERS */
+	// BLUEPRINT SETTERS 
 
-	/** Set the global settings in a retarget profile (will set bApplyGlobalSettings to true). */
+	// Set the global settings in a retarget profile (will set bApplyGlobalSettings to true). 
 	UFUNCTION(BlueprintCallable, Category=RetargetProfile)
 	static void SetGlobalSettingsInRetargetProfile(
 		UPARAM(ref) FRetargetProfile& RetargetProfile,
 		const FRetargetGlobalSettings& GlobalSettings);
 
-	/** Set the root settings in a retarget profile (will set bApplyRootSettings to true). */
+	// Set the root settings in a retarget profile (will set bApplyRootSettings to true). 
 	UFUNCTION(BlueprintCallable, Category=RetargetProfile)
 	static void SetRootSettingsInRetargetProfile(
 		UPARAM(ref) FRetargetProfile& RetargetProfile,
 		const FTargetRootSettings& RootSettings);
 	
-	/** Set the chain settings in a retarget profile (will set bApplyChainSettings to true). */
+	// Set the chain settings in a retarget profile (will set bApplyChainSettings to true). 
 	UFUNCTION(BlueprintCallable, Category=RetargetProfile)
 	static void SetChainSettingsInRetargetProfile(
 		UPARAM(ref) FRetargetProfile& RetargetProfile,
 		const FTargetChainSettings& ChainSettings,
 		const FName TargetChainName);
 
-	/** Set the chain FK settings in a retarget profile (will set bApplyChainSettings to true). */
+	// Set the chain FK settings in a retarget profile (will set bApplyChainSettings to true). 
 	UFUNCTION(BlueprintCallable, Category=RetargetProfile)
 	static void SetChainFKSettingsInRetargetProfile(
 		UPARAM(ref) FRetargetProfile& RetargetProfile,
 		const FTargetChainFKSettings& FKSettings,
 		const FName TargetChainName);
 
-	/** Set the chain IK settings in a retarget profile (will set bApplyChainSettings to true). */
+	// Set the chain IK settings in a retarget profile (will set bApplyChainSettings to true). 
 	UFUNCTION(BlueprintCallable, Category=RetargetProfile)
 	static void SetChainIKSettingsInRetargetProfile(
 		UPARAM(ref) FRetargetProfile& RetargetProfile,
 		const FTargetChainIKSettings& IKSettings,
 		const FName TargetChainName);
 
-	/** Set the chain Speed Plant settings in a retarget profile (will set bApplyChainSettings to true). */
+	// Set the chain Speed Plant settings in a retarget profile (will set bApplyChainSettings to true). 
 	UFUNCTION(BlueprintCallable, Category=RetargetProfile)
 	static void SetChainSpeedPlantSettingsInRetargetProfile(
 		UPARAM(ref) FRetargetProfile& RetargetProfile,
 		const FTargetChainSpeedPlantSettings& SpeedPlantSettings,
 		const FName TargetChainName);
 
-	/** UObject */
+	// UObject 
 	virtual void Serialize(FArchive& Ar) override;
 	virtual void PostLoad() override;
-#if WITH_EDITORONLY_DATA
-	static void DeclareConstructClasses(TArray<FTopLevelAssetPath>& OutConstructClasses, const UClass* SpecificSubclass);
-#endif
-
-	/** END UObject */
+	// END UObject
+	#if WITH_EDITORONLY_DATA
+    	static void DeclareConstructClasses(TArray<FTopLevelAssetPath>& OutConstructClasses, const UClass* SpecificSubclass);
+    #endif
 
 #if WITH_EDITOR
-	/* Get name of Source IK Rig property */
+	// Get name of Source IK Rig property 
 	static const FName GetSourceIKRigPropertyName();
-	/* Get name of Target IK Rig property */
+	// Get name of Target IK Rig property 
 	static const FName GetTargetIKRigPropertyName();
-	/* Get name of Source Preview Mesh property */
+	// Get name of Source Preview Mesh property 
 	static const FName GetSourcePreviewMeshPropertyName();
-	/* Get name of Target Preview Mesh property */
+	// Get name of Target Preview Mesh property 
 	static const FName GetTargetPreviewMeshPropertyName();
-	/** Get the names of the all the speed curves the retargeter will be looking for */
+	// Get the names of the all the speed curves the retargeter will be looking for 
 	void GetSpeedCurveNames(TArray<FName>& OutSpeedCurveNames) const;
 #endif
 
@@ -357,22 +357,22 @@ private:
 
 	void CleanAndInitialize();
 
-	/** The rig to copy animation FROM.*/
+	// The rig to copy animation FROM.
 	UPROPERTY(EditAnywhere, Category = Source)
 	TSoftObjectPtr<UIKRigDefinition> SourceIKRigAsset = nullptr;
 
 #if WITH_EDITORONLY_DATA
-	/** Optional. Override the Skeletal Mesh to copy animation from. Uses the preview mesh from the Source IK Rig asset by default. */
+	// Optional. Override the Skeletal Mesh to copy animation from. Uses the preview mesh from the Source IK Rig asset by default. 
 	UPROPERTY(EditAnywhere, Category = Source)
 	TSoftObjectPtr<USkeletalMesh> SourcePreviewMesh = nullptr;
 #endif
 	
-	/** The rig to copy animation TO.*/
+	// The rig to copy animation TO.
 	UPROPERTY(EditAnywhere, Category = Target)
 	TSoftObjectPtr<UIKRigDefinition> TargetIKRigAsset = nullptr;
 
 #if WITH_EDITORONLY_DATA
-	/** Optional. Override the Skeletal Mesh to preview the retarget on. Uses the preview mesh from the Target IK Rig asset by default. */
+	// Optional. Override the Skeletal Mesh to preview the retarget on. Uses the preview mesh from the Target IK Rig asset by default. 
 	UPROPERTY(EditAnywhere, Category = Target)
 	TSoftObjectPtr<USkeletalMesh> TargetPreviewMesh = nullptr;
 #endif
@@ -391,94 +391,94 @@ public:
 	UPROPERTY()
 	float TargetActorScale_DEPRECATED = 0.0f;
 
-	/** The offset applied to the target mesh in the editor viewport. */
+	// The offset applied to the target mesh in the editor viewport. 
 	UPROPERTY(EditAnywhere, Category = PreviewSettings)
 	FVector TargetMeshOffset;
 
-	/** Scale the target mesh in the viewport for easier visualization next to the source.*/
+	// Scale the target mesh in the viewport for easier visualization next to the source.
 	UPROPERTY(EditAnywhere, Category = PreviewSettings, meta = (UIMin = "0.01", UIMax = "10.0"))
 	float TargetMeshScale = 1.0f;
 
-	/** The offset applied to the source mesh in the editor viewport. */
+	// The offset applied to the source mesh in the editor viewport. 
 	UPROPERTY(EditAnywhere, Category = PreviewSettings)
 	FVector SourceMeshOffset;
 
-	/** Toggle debug drawing for retargeting in the viewport. */
+	// Toggle debug drawing for retargeting in the viewport. 
 	UPROPERTY(EditAnywhere, Category = DebugSettings)
 	bool bDebugDraw = true;
 
-	/** Draw final IK goal locations. */
+	// Draw final IK goal locations. 
 	UPROPERTY(EditAnywhere, Category = DebugSettings)
 	bool bDrawFinalGoals = true;
 
-	/** Draw goal locations from source skeleton. */
+	// Draw goal locations from source skeleton. 
 	UPROPERTY(EditAnywhere, Category = DebugSettings)
 	bool bDrawSourceLocations = true;
 	
-	/** The visual size of the IK goals in the viewport. */
+	// The visual size of the IK goals in the viewport. 
 	UPROPERTY(EditAnywhere, Category = DebugSettings)
 	float ChainDrawSize = 5.0f;
 
-	/** The thickness of lines on the IK goals in the viewport. */
+	// The thickness of lines on the IK goals in the viewport. 
 	UPROPERTY(EditAnywhere, Category = DebugSettings)
 	float ChainDrawThickness = 0.5f;
 	
-	/** The visual size of the bones in the viewport (saved between sessions). This is set from the viewport Character>Bones menu*/
+	// The visual size of the bones in the viewport (saved between sessions). This is set from the viewport Character>Bones menu
 	UPROPERTY()
 	float BoneDrawSize = 1.0f;
 	
 private:
 
-	/** The controller responsible for managing this asset's data (all editor mutation goes through this) */
+	/** The controller responsible for managing this asset's data (all mutation goes through this) */
 	UPROPERTY(Transient)
-	TObjectPtr<UObject> Controller;
+	mutable TObjectPtr<UObject> Controller;
 
-	/** only ask to fix the root height once, then warn thereafter (don't nag) */
+	// only ask to fix the root height once, then warn thereafter (don't nag) 
 	UPROPERTY()
 	TSet<TObjectPtr<USkeletalMesh>> MeshesAskedToFixRootHeightFor;
 #endif
 	
 private:
 
-	/** (OLD VERSION) Mapping of chains to copy animation between source and target rigs.*/
+	// (OLD VERSION) Mapping of chains to copy animation between source and target rigs.
 	PRAGMA_DISABLE_DEPRECATION_WARNINGS
 	UPROPERTY()
 	TArray<FRetargetChainMap> ChainMapping_DEPRECATED;
 	PRAGMA_ENABLE_DEPRECATION_WARNINGS
 	
-	/** Settings for how to map source chains to target chains.*/
+	// Settings for how to map source chains to target chains.
 	UPROPERTY()
 	TArray<TObjectPtr<URetargetChainSettings>> ChainSettings;
 	
-	/** the retarget root settings */
+	// the retarget root settings 
 	UPROPERTY()
 	TObjectPtr<URetargetRootSettings> RootSettings;
 
-	/** the retarget root settings */
+	// the retarget root settings 
 	UPROPERTY()
 	TObjectPtr<UIKRetargetGlobalSettings> GlobalSettings;
 
-	/** settings profiles stored in this asset */
+	// settings profiles stored in this asset 
 	UPROPERTY()
 	TMap<FName, FRetargetProfile> Profiles;
 	UPROPERTY()
 	FName CurrentProfile = NAME_None;
 	
-	/** The set of retarget poses for the SOURCE skeleton.*/
+	// The set of retarget poses for the SOURCE skeleton.
 	UPROPERTY()
 	TMap<FName, FIKRetargetPose> SourceRetargetPoses;
-	/** The set of retarget poses for the TARGET skeleton.*/
+	// The set of retarget poses for the TARGET skeleton.
 	UPROPERTY()
 	TMap<FName, FIKRetargetPose> TargetRetargetPoses;
 	
-	/** The current retarget pose to use for the SOURCE.*/
+	// The current retarget pose to use for the SOURCE.
 	UPROPERTY()
 	FName CurrentSourceRetargetPose;
-	/** The current retarget pose to use for the TARGET.*/
+	// The current retarget pose to use for the TARGET.
 	UPROPERTY()
 	FName CurrentTargetRetargetPose;
 
-	/** (OLD VERSION) Before retarget poses were stored for target AND source.*/
+	// (OLD VERSION) Before retarget poses were stored for target AND source.
 	PRAGMA_DISABLE_DEPRECATION_WARNINGS
 	UPROPERTY()
 	TMap<FName, FIKRetargetPose> RetargetPoses_DEPRECATED;
diff --git a/Engine/Plugins/Animation/IKRig/Source/IKRig/Public/Solvers/IKRig_PoleSolver.h b/Engine/Plugins/Animation/IKRig/Source/IKRig/Public/Solvers/IKRig_PoleSolver.h
--- a/Engine/Plugins/Animation/IKRig/Source/IKRig/Public/Solvers/IKRig_PoleSolver.h
+++ b/Engine/Plugins/Animation/IKRig/Source/IKRig/Public/Solvers/IKRig_PoleSolver.h
@@ -64,6 +64,7 @@ public:
 	virtual void SetRootBone(const FName& RootBoneName) override;
 	// end bone can be set on this solver
 	virtual void SetEndBone(const FName& EndBoneName) override;
+	virtual FName GetEndBone() const override;
 	virtual bool RequiresEndBone() const override { return true; };
 	/** END UIKRigSolver interface */
 #endif
diff --git a/Engine/Plugins/Animation/IKRig/Source/IKRigEditor/Private/RetargetEditor/IKRetargetBatchOperation.cpp b/Engine/Plugins/Animation/IKRig/Source/IKRigEditor/Private/RetargetEditor/IKRetargetBatchOperation.cpp
--- a/Engine/Plugins/Animation/IKRig/Source/IKRigEditor/Private/RetargetEditor/IKRetargetBatchOperation.cpp
+++ b/Engine/Plugins/Animation/IKRig/Source/IKRigEditor/Private/RetargetEditor/IKRetargetBatchOperation.cpp
@@ -10,6 +10,7 @@
 #include "ContentBrowserModule.h"
 #include "EditorReimportHandler.h"
 #include "IContentBrowserSingleton.h"
+#include "RigEditor/IKRigController.h"
 #include "ObjectEditorUtils.h"
 #include "SSkeletonWidget.h"
 #include "PropertyCustomizationHelpers.h"
@@ -27,7 +28,7 @@
 
 #define LOCTEXT_NAMESPACE "RetargetBatchOperation"
 
-int32 FIKRetargetBatchOperation::GenerateAssetLists(const FIKRetargetBatchOperationContext& Context)
+int32 UIKRetargetBatchOperation::GenerateAssetLists(const FIKRetargetBatchOperationContext& Context)
 {
 	// re-generate lists of selected and referenced assets
 	AnimationAssetsToRetarget.Reset();
@@ -97,7 +98,7 @@ int32 FIKRetargetBatchOperation::GenerateAssetLists(const FIKRetargetBatchOperat
 	return AnimationAssetsToRetarget.Num();
 }
 
-void FIKRetargetBatchOperation::DuplicateRetargetAssets(
+void UIKRetargetBatchOperation::DuplicateRetargetAssets(
 	const FIKRetargetBatchOperationContext& Context,
 	FScopedSlowTask& Progress)
 {
@@ -161,7 +162,7 @@ void FIKRetargetBatchOperation::DuplicateRetargetAssets(
 	DuplicatedBlueprints.GenerateValueArray(AnimBlueprintsToRetarget);
 }
 
-void FIKRetargetBatchOperation::RetargetAssets(
+void UIKRetargetBatchOperation::RetargetAssets(
 	const FIKRetargetBatchOperationContext& Context,
 	FScopedSlowTask& Progress)
 {
@@ -247,7 +248,7 @@ void FIKRetargetBatchOperation::RetargetAssets(
 	}
 }
 
-void FIKRetargetBatchOperation::ConvertAnimation(
+void UIKRetargetBatchOperation::ConvertAnimation(
 	const FIKRetargetBatchOperationContext& Context,
 	FScopedSlowTask& Progress)
 {
@@ -403,7 +404,7 @@ void FIKRetargetBatchOperation::ConvertAnimation(
 	}
 }
 
-void FIKRetargetBatchOperation::NotifyUserOfResults(
+void UIKRetargetBatchOperation::NotifyUserOfResults(
 	const FIKRetargetBatchOperationContext& Context,
 	FScopedSlowTask& Progress) const
 {
@@ -440,7 +441,7 @@ void FIKRetargetBatchOperation::NotifyUserOfResults(
 	ContentBrowserModule.Get().SyncBrowserToAssets(CurrentSelection);
 }
 
-void FIKRetargetBatchOperation::GetNewAssets(TArray<UObject*>& NewAssets) const
+void UIKRetargetBatchOperation::GetNewAssets(TArray<UObject*>& NewAssets) const
 {
 	TArray<UAnimationAsset*> NewAnims;
 	DuplicatedAnimAssets.GenerateValueArray(NewAnims);
@@ -458,10 +459,118 @@ void FIKRetargetBatchOperation::GetNewAssets(TArray<UObject*>& NewAssets) const
 }
 
 
+TArray<FAssetData> UIKRetargetBatchOperation::RunRetarget(
+	const TArray<FAssetData>& AssetsToRetarget,
+	USkeletalMesh* SourceMesh,
+	USkeletalMesh* TargetMesh,
+	UIKRetargeter* IKRetargetAsset,
+	const FString& Search,
+	const FString& Replace,
+	const FString& Prefix,
+	const FString& Suffix,
+	const bool bRemapReferencedAssets)
+{
+	// fill the context with all the data needed to run a batch retarget
+	FIKRetargetBatchOperationContext Context;
+	for (const FAssetData& Asset : AssetsToRetarget)
+	{
+		Context.AssetsToRetarget.Add(Asset.GetAsset()); // convert asset data to soft refs
+	}
+	Context.SourceMesh = SourceMesh;
+	Context.TargetMesh = TargetMesh;
+	Context.IKRetargetAsset = IKRetargetAsset;
+	Context.NameRule.Prefix = Prefix;
+	Context.NameRule.Suffix = Suffix;
+	Context.NameRule.ReplaceFrom = Search;
+	Context.NameRule.ReplaceTo = Replace;
+	Context.bRemapReferencedAssets = bRemapReferencedAssets;
+
+	// actually run the batch operation
+	RunRetarget(Context);
+
+	// create array of FAssetData to return
+	TArray<FAssetData> Results;
+	for (const UAnimationAsset* RetargetedAsset : AnimationAssetsToRetarget)
+	{
+		Results.Add(FAssetData(RetargetedAsset));
+	}
+	return Results;
+}
 
-void FIKRetargetBatchOperation::RunRetarget(FIKRetargetBatchOperationContext& Context)
-{	
+void UIKRetargetBatchOperation::RunRetarget(FIKRetargetBatchOperationContext& Context)
+{
+	Reset();
+	
+	// validate animation assets were provided
 	const int32 NumAssets = GenerateAssetLists(Context);
+	if (NumAssets == 0)
+	{
+		UE_LOG(LogTemp, Warning, TEXT("Batch retarget aborted. No animation assets were specified."));
+		return;
+	}
+
+	// validate a retarget asset was provided
+	if (!Context.IKRetargetAsset)
+	{
+		UE_LOG(LogTemp, Warning, TEXT("Batch retarget aborted. No IK Retargeter asset was specified."));
+		return;
+	}
+
+	// validate a source IK rig was provided
+	const UIKRigDefinition* SrcIKRig = Context.IKRetargetAsset->GetSourceIKRig();
+	if (!SrcIKRig)
+	{
+		UE_LOG(LogTemp, Warning, TEXT("Batch retarget aborted. Specified IK Retargeter does not reference a source IK Rig."));
+		return;
+	}
+
+	// validate a target IK rig was provided
+	const UIKRigDefinition* TgtIKRig = Context.IKRetargetAsset->GetTargetIKRig();
+	if (!TgtIKRig)
+	{
+		UE_LOG(LogTemp, Warning, TEXT("Batch retarget aborted. Specified IK Retargeter does not reference a target IK Rig."));
+		return;
+	}
+
+	// validate a source mesh was provided
+	if (!Context.SourceMesh)
+	{
+		// fallback to mesh provided by IK Rig
+		Context.SourceMesh = SrcIKRig->GetPreviewMesh();
+		if (!Context.SourceMesh)
+		{
+			UE_LOG(LogTemp, Warning, TEXT("Batch retarget aborted. No source mesh was specified and the source IK Rig did not have one. "));
+			return;	
+		}
+	}
+
+	// validate a target mesh was provided
+	if (!Context.TargetMesh)
+	{
+		// fallback to mesh provided by IK Rig
+		Context.TargetMesh = TgtIKRig->GetPreviewMesh();
+		if (!Context.TargetMesh)
+		{
+			UE_LOG(LogTemp, Warning, TEXT("Batch retarget aborted. No target mesh was specified and the target IK Rig did not have one. "));
+			return;
+		}
+	}
+
+	// validate source mesh is compatible with source IK Rig
+	const UIKRigController* SrcRigController = UIKRigController::GetController(SrcIKRig);
+	if (!SrcRigController->IsSkeletalMeshCompatible(Context.SourceMesh))
+	{
+		UE_LOG(LogTemp, Warning, TEXT("Batch retarget aborted. The provided source mesh was not compatible with the source IK Rig."));
+		return;
+	}
+	
+	// validate target mesh is compatible with target IK Rig
+	const UIKRigController* TgtRigController = UIKRigController::GetController(TgtIKRig);
+	if (!TgtRigController->IsSkeletalMeshCompatible(Context.TargetMesh))
+	{
+		UE_LOG(LogTemp, Warning, TEXT("Batch retarget aborted. The provided target mesh was not compatible with the target IK Rig."));
+		return;
+	}
 	
 	// show progress bar
 	FScopedSlowTask Progress(NumAssets + 2, LOCTEXT("GatheringBatchRetarget", "Gathering animation assets..."));
@@ -472,6 +581,15 @@ void FIKRetargetBatchOperation::RunRetarget(FIKRetargetBatchOperationContext& Co
 	NotifyUserOfResults(Context, Progress);
 }
 
+void UIKRetargetBatchOperation::Reset()
+{
+	AnimationAssetsToRetarget.Reset();
+	AnimBlueprintsToRetarget.Reset();
+	DuplicatedAnimAssets.Reset();
+	DuplicatedBlueprints.Reset();
+	RemappedAnimAssets.Reset();
+}
+
 /**
 * Duplicates the supplied AssetsToDuplicate and returns a map of original asset to duplicate. Templated wrapper that calls DuplicateAssetInternal.
 *
@@ -482,7 +600,7 @@ void FIKRetargetBatchOperation::RunRetarget(FIKRetargetBatchOperationContext& Co
 * @return	TMap of original animation to duplicate
 */
 template<class AssetType>
-TMap<AssetType*, AssetType*> FIKRetargetBatchOperation::DuplicateAssets(
+TMap<AssetType*, AssetType*> UIKRetargetBatchOperation::DuplicateAssets(
 	const TArray<AssetType*>& AssetsToDuplicate,
 	UPackage* DestinationPackage,
 	const FNameDuplicationRule* NameRule)
diff --git a/Engine/Plugins/Animation/IKRig/Source/IKRigEditor/Private/RetargetEditor/IKRetargetDetails.cpp b/Engine/Plugins/Animation/IKRig/Source/IKRigEditor/Private/RetargetEditor/IKRetargetDetails.cpp
--- a/Engine/Plugins/Animation/IKRig/Source/IKRigEditor/Private/RetargetEditor/IKRetargetDetails.cpp
+++ b/Engine/Plugins/Animation/IKRig/Source/IKRigEditor/Private/RetargetEditor/IKRetargetDetails.cpp
@@ -888,7 +888,7 @@ void FRetargetChainSettingsCustomization::CustomizeDetails(IDetailLayoutBuilder&
 			const FName SourceChainName = FName(*InString.Get());
 			for (TWeakObjectPtr<URetargetChainSettings> ChainMapSettings : ChainSettingsObjects)
 			{
-				Controller->AssetController->SetSourceChainForTargetChain(ChainMapSettings.Get(), SourceChainName);
+				Controller->AssetController->SetSourceChain(SourceChainName, ChainMapSettings->TargetChain);
 			}
 		})
 		[
diff --git a/Engine/Plugins/Animation/IKRig/Source/IKRigEditor/Private/RetargetEditor/IKRetargetEditorController.cpp b/Engine/Plugins/Animation/IKRig/Source/IKRigEditor/Private/RetargetEditor/IKRetargetEditorController.cpp
--- a/Engine/Plugins/Animation/IKRig/Source/IKRigEditor/Private/RetargetEditor/IKRetargetEditorController.cpp
+++ b/Engine/Plugins/Animation/IKRig/Source/IKRigEditor/Private/RetargetEditor/IKRetargetEditorController.cpp
@@ -32,7 +32,7 @@ FBoundIKRig::FBoundIKRig(UIKRigDefinition* InIKRig, const FIKRetargetEditorContr
 {
 	check(InIKRig);
 	IKRig = InIKRig;
-	UIKRigController* IKRigController = UIKRigController::GetIKRigController(InIKRig);
+	UIKRigController* IKRigController = UIKRigController::GetController(InIKRig);
 	ReInitIKDelegateHandle = IKRigController->OnIKRigNeedsInitialized().AddSP(&InController, &FIKRetargetEditorController::HandleIKRigNeedsInitialized);
 	AddedChainDelegateHandle = IKRigController->OnRetargetChainAdded().AddSP(&InController, &FIKRetargetEditorController::HandleRetargetChainAdded);
 	RemoveChainDelegateHandle = IKRigController->OnRetargetChainRemoved().AddSP(&InController, &FIKRetargetEditorController::HandleRetargetChainRemoved);
@@ -42,7 +42,7 @@ FBoundIKRig::FBoundIKRig(UIKRigDefinition* InIKRig, const FIKRetargetEditorContr
 void FBoundIKRig::UnBind() const
 {
 	check(IKRig);
-	UIKRigController* IKRigController = UIKRigController::GetIKRigController(IKRig);
+	UIKRigController* IKRigController = UIKRigController::GetController(IKRig);
 	IKRigController->OnIKRigNeedsInitialized().Remove(ReInitIKDelegateHandle);
 	IKRigController->OnRetargetChainAdded().Remove(AddedChainDelegateHandle);
 	IKRigController->OnRetargetChainRemoved().Remove(RemoveChainDelegateHandle);
@@ -140,8 +140,7 @@ void FIKRetargetEditorController::Initialize(TSharedPtr<FIKRetargetEditor> InEdi
 	LastSelectedItem = ERetargetSelectionType::NONE;
 
 	// clean the asset before editing
-	AssetController->CleanChainMapping();
-	AssetController->CleanPoseLists();
+	AssetController->CleanAsset();
 
 	// bind callbacks when SOURCE or TARGET IK Rigs are modified
 	BindToIKRigAssets();
@@ -266,20 +265,20 @@ void FIKRetargetEditorController::HandleIKRigNeedsInitialized(UIKRigDefinition*
 void FIKRetargetEditorController::HandleRetargetChainAdded(UIKRigDefinition* ModifiedIKRig) const
 {
 	check(ModifiedIKRig)
-	AssetController->OnRetargetChainAdded(ModifiedIKRig);
+	AssetController->HandleRetargetChainAdded(ModifiedIKRig);
 	RefreshAllViews();
 }
 
 void FIKRetargetEditorController::HandleRetargetChainRenamed(UIKRigDefinition* ModifiedIKRig, FName OldName, FName NewName) const
 {
 	check(ModifiedIKRig)
-	AssetController->OnRetargetChainRenamed(ModifiedIKRig, OldName, NewName);
+	AssetController->HandleRetargetChainRenamed(ModifiedIKRig, OldName, NewName);
 }
 
-void FIKRetargetEditorController::HandleRetargetChainRemoved(UIKRigDefinition* ModifiedIKRig, const FName& InChainRemoved) const
+void FIKRetargetEditorController::HandleRetargetChainRemoved(UIKRigDefinition* ModifiedIKRig, const FName InChainRemoved) const
 {
 	check(ModifiedIKRig)
-	AssetController->OnRetargetChainRemoved(ModifiedIKRig, InChainRemoved);
+	AssetController->HandleRetargetChainRemoved(ModifiedIKRig, InChainRemoved);
 	RefreshAllViews();
 }
 
@@ -338,7 +337,10 @@ void FIKRetargetEditorController::HandlePreviewMeshReplaced(ERetargetSourceOrTar
 	}
 
 	// re-initializes the anim instances running in the viewport
-	Editor.Pin()->SetupAnimInstance();
+	if (SourceAnimInstance)
+	{
+		Editor.Pin()->SetupAnimInstance();	
+	}
 
 	// continue playing where we left off
 	PlaybackManager->ResumePlayback();
@@ -442,19 +444,20 @@ bool FIKRetargetEditorController::GetCameraTargetForSelection(FSphere& OutTarget
 			const TArray<FName>& SelectedChainNames = GetSelectedChains();
 			for (const FName SelectedChainName : SelectedChainNames)
 			{
-				const URetargetChainSettings* ChainMap = AssetController->GetChainMappingByTargetChainName(SelectedChainName);
-				if (!ChainMap)
+				const FName SourceChain = AssetController->GetSourceChain(SelectedChainName);
+				if (SourceChain == NAME_None)
 				{
 					continue;
 				}
 
-				const FName& ChainName = SourceOrTarget == ERetargetSourceOrTarget::Target ? SelectedChainName : ChainMap->SourceChain;
+				const FName& ChainName = SourceOrTarget == ERetargetSourceOrTarget::Target ? SelectedChainName : SourceChain;
 				if (ChainName == NAME_None)
 				{
 					continue;
 				}
 
-				const FBoneChain* BoneChain = IKRig->GetRetargetChainByName(ChainName);
+				const UIKRigController* RigController = UIKRigController::GetController(IKRig);
+				const FBoneChain* BoneChain = RigController->GetRetargetChainByName(ChainName);
 				if (!BoneChain)
 				{
 					continue;
@@ -1049,8 +1052,6 @@ void FIKRetargetEditorController::SetRetargeterMode(ERetargeterOutputMode Mode)
 			OutputMode = ERetargeterOutputMode::RunRetarget;
 			SourceAnimInstance->SetRetargetMode(ERetargeterOutputMode::RunRetarget);
 			TargetAnimInstance->SetRetargetMode(ERetargeterOutputMode::RunRetarget);
-			// force a reinitialization in case retarget pose was edited
-			AssetController->BroadcastNeedsReinitialized();
 			PlaybackManager->ResumePlayback();
 			break;
 
@@ -1198,7 +1199,7 @@ bool FIKRetargetEditorController::CanCreatePose() const
 FReply FIKRetargetEditorController::CreateNewPose() const
 {
 	const FName NewPoseName = FName(NewPoseEditableText.Get()->GetText().ToString());
-	AssetController->AddRetargetPose(NewPoseName, nullptr, GetSourceOrTarget());
+	AssetController->CreateRetargetPose(NewPoseName, GetSourceOrTarget());
 	NewPoseWindow->RequestDestroyWindow();
 	RefreshPoseList();
 	return FReply::Handled();
@@ -1277,9 +1278,9 @@ void FIKRetargetEditorController::HandleDuplicatePose()
 
 FReply FIKRetargetEditorController::CreateDuplicatePose() const
 {
-	const FIKRetargetPose& PoseToDuplicate = AssetController->GetCurrentRetargetPose(CurrentlyEditingSourceOrTarget);
+	const FName PoseToDuplicate = AssetController->GetCurrentRetargetPoseName(CurrentlyEditingSourceOrTarget);
 	const FName NewPoseName = FName(NewPoseEditableText.Get()->GetText().ToString());
-	AssetController->AddRetargetPose(NewPoseName, &PoseToDuplicate, GetSourceOrTarget());
+	AssetController->DuplicateRetargetPose(PoseToDuplicate, NewPoseName, GetSourceOrTarget());
 	NewPoseWindow->RequestDestroyWindow();
 	RefreshPoseList();
 	return FReply::Handled();
@@ -1427,8 +1428,9 @@ FReply FIKRetargetEditorController::RenamePose() const
 {
 	const FName NewPoseName = FName(NewNameEditableText.Get()->GetText().ToString());
 	RenamePoseWindow->RequestDestroyWindow();
-	
-	AssetController->RenameCurrentRetargetPose(NewPoseName, GetSourceOrTarget());
+
+	const FName CurrentPoseName = AssetController->GetCurrentRetargetPoseName(GetSourceOrTarget());
+	AssetController->RenameRetargetPose(CurrentPoseName, NewPoseName, GetSourceOrTarget());
 	RefreshPoseList();
 	return FReply::Handled();
 }
@@ -1447,6 +1449,9 @@ void FIKRetargetEditorController::AddReferencedObjects(FReferenceCollector& Coll
 	{
 		Collector.AddReferencedObject(Pair.Value);
 	}
+
+	Collector.AddReferencedObject(SourceAnimInstance);
+	Collector.AddReferencedObject(TargetAnimInstance);
 }
 
 void FIKRetargetEditorController::RenderSkeleton(FPrimitiveDrawInterface* PDI, ERetargetSourceOrTarget InSourceOrTarget) const
diff --git a/Engine/Plugins/Animation/IKRig/Source/IKRigEditor/Private/RetargetEditor/IKRetargetFactory.cpp b/Engine/Plugins/Animation/IKRig/Source/IKRigEditor/Private/RetargetEditor/IKRetargetFactory.cpp
--- a/Engine/Plugins/Animation/IKRig/Source/IKRigEditor/Private/RetargetEditor/IKRetargetFactory.cpp
+++ b/Engine/Plugins/Animation/IKRig/Source/IKRigEditor/Private/RetargetEditor/IKRetargetFactory.cpp
@@ -2,16 +2,9 @@
 
 #include "RetargetEditor/IKRetargetFactory.h"
 #include "Retargeter/IKRetargeter.h"
-#include "IKRigEditor.h"
 #include "AssetTypeCategories.h"
 #include "ContentBrowserModule.h"
 #include "IContentBrowserSingleton.h"
-#include "Editor.h"
-#include "Styling/AppStyle.h"
-#include "RetargetEditor/IKRetargeterController.h"
-#include "Widgets/SWindow.h"
-#include "Widgets/Layout/SBorder.h"
-#include "Engine/SkeletalMesh.h"
 
 #include UE_INLINE_GENERATED_CPP_BY_NAME(IKRetargetFactory)
 
diff --git a/Engine/Plugins/Animation/IKRig/Source/IKRigEditor/Private/RetargetEditor/IKRetargetPoseExporter.cpp b/Engine/Plugins/Animation/IKRig/Source/IKRigEditor/Private/RetargetEditor/IKRetargetPoseExporter.cpp
--- a/Engine/Plugins/Animation/IKRig/Source/IKRigEditor/Private/RetargetEditor/IKRetargetPoseExporter.cpp
+++ b/Engine/Plugins/Animation/IKRig/Source/IKRigEditor/Private/RetargetEditor/IKRetargetPoseExporter.cpp
@@ -241,7 +241,8 @@ FReply FIKRetargetPoseExporter::ImportPoseAsset() const
 	PoseAsset->GetFullPose(IndexOfPoseToImport, LocalBoneTransformFromPose);
 
 	// create a new retarget pose to store the data from the selected retarget pose asset
-	FIKRetargetPose NewPose;
+	const FName NewPoseName = ControllerPtr->AssetController->CreateRetargetPose(FName(PoseAsset->GetName()), SourceOrTarget);
+	FIKRetargetPose& NewPose = ControllerPtr->AssetController->GetRetargetPoses(SourceOrTarget)[NewPoseName];
 	
 	// iterate over all bones in the reference skeleton and compare against bone in the pose asset
 	// store sparse set of deltas in the new retarget pose 
@@ -276,9 +277,6 @@ FReply FIKRetargetPoseExporter::ImportPoseAsset() const
 		}
 	}
 
-	// store the retarget pose in the retarget asset
-	ControllerPtr->AssetController->AddRetargetPose(FName(PoseAsset->GetName()), &NewPose, SourceOrTarget);
-
 	// update view with new pose
 	ControllerPtr->RefreshAllViews();
 	
@@ -512,8 +510,11 @@ FReply FIKRetargetPoseExporter::OnImportPoseFromSequence()
 	FrameOfSequenceToImport = FMath::Clamp(FrameOfSequenceToImport, 0, AnimSequence->GetNumberOfSampledKeys());
 	UAnimPoseExtensions::GetAnimPoseAtFrame(AnimSequence, FrameOfSequenceToImport, EvaluationOptions, ImportedPose);
 
+	// create a new retarget pose to hold imported data
+	const FName NewPose = ControllerPtr->AssetController->CreateRetargetPose(FName(ImportedPoseName.ToString()), SourceOrTarget);
+
 	// record delta pose for all bones being retargeted
-	FIKRetargetPose ImportedRetargetPose;
+	FIKRetargetPose& ImportedRetargetPose = ControllerPtr->AssetController->GetRetargetPoses(SourceOrTarget)[NewPose];
 	
 	// get all imported bone transforms and record them in the retarget pose
 	FReferenceSkeleton& RefSkeleton = Mesh->GetRefSkeleton();
@@ -572,11 +573,8 @@ FReply FIKRetargetPoseExporter::OnImportPoseFromSequence()
 		}
 	}
 
-	// store the newly imported retarget pose in the asset
-	ControllerPtr->AssetController->AddRetargetPose( FName(ImportedPoseName.ToString()), &ImportedRetargetPose, SourceOrTarget);
-
 	// notify user of new pose
-	const FText Message = FText::Format(LOCTEXT("ImportSuccess", "Imported pose from animation sequence: {0}"), ImportedPoseName);
+	const FText Message = FText::Format(LOCTEXT("ImportSuccess", "Imported new retarget pose: {0}"), FText::FromName(NewPose));
 	NotifyUser(Message, SNotificationItem::CS_Success);
 	return FReply::Unhandled();
 }
diff --git a/Engine/Plugins/Animation/IKRig/Source/IKRigEditor/Private/RetargetEditor/IKRetargeterController.cpp b/Engine/Plugins/Animation/IKRig/Source/IKRigEditor/Private/RetargetEditor/IKRetargeterController.cpp
--- a/Engine/Plugins/Animation/IKRig/Source/IKRigEditor/Private/RetargetEditor/IKRetargeterController.cpp
+++ b/Engine/Plugins/Animation/IKRig/Source/IKRigEditor/Private/RetargetEditor/IKRetargeterController.cpp
@@ -6,15 +6,14 @@
 #include "Algo/LevenshteinDistance.h"
 #include "Engine/AssetManager.h"
 #include "Engine/SkeletalMesh.h"
-#include "RetargetEditor/IKRetargetEditorController.h"
 #include "Retargeter/IKRetargeter.h"
 #include "RigEditor/IKRigController.h"
+#include "IKRigEditor.h"
 
 #include UE_INLINE_GENERATED_CPP_BY_NAME(IKRetargeterController)
 
 #define LOCTEXT_NAMESPACE "IKRetargeterController"
 
-
 UIKRetargeterController* UIKRetargeterController::GetController(UIKRetargeter* InRetargeterAsset)
 {
 	if (!InRetargeterAsset)
@@ -38,6 +37,13 @@ UIKRetargeter* UIKRetargeterController::GetAsset() const
 	return Asset;
 }
 
+void UIKRetargeterController::CleanAsset() const
+{
+	CleanChainMapping();
+	CleanPoseList(ERetargetSourceOrTarget::Source);
+	CleanPoseList(ERetargetSourceOrTarget::Target);
+}
+
 void UIKRetargeterController::SetIKRig(const ERetargetSourceOrTarget SourceOrTarget, UIKRigDefinition* IKRig) const
 {
 	if (SourceOrTarget == ERetargetSourceOrTarget::Source)
@@ -54,7 +60,7 @@ void UIKRetargeterController::SetIKRig(const ERetargetSourceOrTarget SourceOrTar
 	// re-ask to fix root height for this mesh
 	if (IKRig)
 	{
-		SetAskedToFixRootHeightForMesh(IKRig->PreviewSkeletalMesh.Get(), false);
+		SetAskedToFixRootHeightForMesh(GetPreviewMesh(SourceOrTarget), false);
 	}
 
 	CleanChainMapping();
@@ -80,7 +86,7 @@ UIKRigDefinition* UIKRetargeterController::GetIKRigWriteable(const ERetargetSour
 
 void UIKRetargeterController::SetPreviewMesh(
 	const ERetargetSourceOrTarget SourceOrTarget,
-	USkeletalMesh* PreviewMesh)
+	USkeletalMesh* PreviewMesh) const
 {	
 	if (SourceOrTarget == ERetargetSourceOrTarget::Source)
 	{
@@ -109,10 +115,10 @@ USkeletalMesh* UIKRetargeterController::GetPreviewMesh(const ERetargetSourceOrTa
 	}
 
 	// optionally prefer override if one is provided
-	const TSoftObjectPtr<USkeletalMesh> OverrideMesh = SourceOrTarget == ERetargetSourceOrTarget::Source ? Asset->SourcePreviewMesh : Asset->TargetPreviewMesh;
-	if (!OverrideMesh.IsNull())
+	const TSoftObjectPtr<USkeletalMesh> PreviewMesh = SourceOrTarget == ERetargetSourceOrTarget::Source ? Asset->SourcePreviewMesh : Asset->TargetPreviewMesh;
+	if (!PreviewMesh.IsNull())
 	{
-		return OverrideMesh.LoadSynchronous();
+		return PreviewMesh.LoadSynchronous();
 	}
 
 	// fallback to preview mesh from IK Rig asset
@@ -215,42 +221,37 @@ void UIKRetargeterController::CleanChainMapping() const
 	SortChainMapping();
 }
 
-void UIKRetargeterController::CleanPoseLists() const
-{
-	CleanPoseList(ERetargetSourceOrTarget::Source);
-	CleanPoseList(ERetargetSourceOrTarget::Target);
-}
-
 void UIKRetargeterController::CleanPoseList(const ERetargetSourceOrTarget SourceOrTarget) const
 {
-	TMap<FName, FIKRetargetPose>& RetargetPoses = GetRetargetPoses(SourceOrTarget);
-
-	// remove all bone offsets that are no longer part of the skeleton
 	const UIKRigDefinition* IKRig = GetIKRig(SourceOrTarget);
-	if (IKRig)
+	if (!IKRig)
 	{
-		const TArray<FName> AllowedBoneNames = IKRig->Skeleton.BoneNames;
-		for (TTuple<FName, FIKRetargetPose>& Pose : RetargetPoses)
+		return;
+	}
+	
+	// remove all bone offsets that are no longer part of the skeleton
+	const TArray<FName> AllowedBoneNames = IKRig->GetSkeleton().BoneNames;
+	TMap<FName, FIKRetargetPose>& RetargetPoses = GetRetargetPoses(SourceOrTarget);
+	for (TTuple<FName, FIKRetargetPose>& Pose : RetargetPoses)
+	{
+		// find bone offsets no longer in target skeleton
+		TArray<FName> BonesToRemove;
+		for (TTuple<FName, FQuat>& BoneOffset : Pose.Value.BoneRotationOffsets)
 		{
-			// find bone offsets no longer in target skeleton
-			TArray<FName> BonesToRemove;
-			for (TTuple<FName, FQuat>& BoneOffset : Pose.Value.BoneRotationOffsets)
-			{
-				if (!AllowedBoneNames.Contains(BoneOffset.Key))
-				{
-					BonesToRemove.Add(BoneOffset.Key);
-				}
-			}
-			
-			// remove bone offsets
-			for (const FName& BoneToRemove : BonesToRemove)
+			if (!AllowedBoneNames.Contains(BoneOffset.Key))
 			{
-				Pose.Value.BoneRotationOffsets.Remove(BoneToRemove);
+				BonesToRemove.Add(BoneOffset.Key);
 			}
-
-			// sort the pose offset from leaf to root
-			Pose.Value.SortHierarchically(IKRig->Skeleton);
 		}
+		
+		// remove bone offsets
+		for (const FName& BoneToRemove : BonesToRemove)
+		{
+			Pose.Value.BoneRotationOffsets.Remove(BoneToRemove);
+		}
+
+		// sort the pose offset from leaf to root
+		Pose.Value.SortHierarchically(IKRig->GetSkeleton());
 	}
 }
 
@@ -332,7 +333,7 @@ void UIKRetargeterController::AutoMapChains(const EAutoMapChainType AutoMapType,
 	SortChainMapping();
 }
 
-void UIKRetargeterController::OnRetargetChainAdded(UIKRigDefinition* IKRig) const
+void UIKRetargeterController::HandleRetargetChainAdded(UIKRigDefinition* IKRig) const
 {
 	const bool bIsTargetRig = IKRig == Asset->GetTargetIKRig();
 	if (!bIsTargetRig)
@@ -346,7 +347,7 @@ void UIKRetargeterController::OnRetargetChainAdded(UIKRigDefinition* IKRig) cons
 	BroadcastNeedsReinitialized();
 }
 
-void UIKRetargeterController::OnRetargetChainRenamed(UIKRigDefinition* IKRig, FName OldChainName, FName NewChainName) const
+void UIKRetargeterController::HandleRetargetChainRenamed(UIKRigDefinition* IKRig, FName OldChainName, FName NewChainName) const
 {
 	const bool bIsSourceRig = IKRig == Asset->GetSourceIKRig();
 	check(bIsSourceRig || IKRig == Asset->GetTargetIKRig())
@@ -362,7 +363,7 @@ void UIKRetargeterController::OnRetargetChainRenamed(UIKRigDefinition* IKRig, FN
 	}
 }
 
-void UIKRetargeterController::OnRetargetChainRemoved(UIKRigDefinition* IKRig, const FName& InChainRemoved) const
+void UIKRetargeterController::HandleRetargetChainRemoved(UIKRigDefinition* IKRig, const FName& InChainRemoved) const
 {
 	const bool bIsSourceRig = IKRig == Asset->GetSourceIKRig();
 	check(bIsSourceRig || IKRig == Asset->GetTargetIKRig())
@@ -395,33 +396,39 @@ void UIKRetargeterController::OnRetargetChainRemoved(UIKRigDefinition* IKRig, co
 	}
 }
 
-void UIKRetargeterController::SetSourceChainForTargetChain(URetargetChainSettings* ChainMap, FName SourceChainToMapTo) const
+bool UIKRetargeterController::SetSourceChain(FName SourceChainName, FName TargetChainName) const
 {
-	check(ChainMap)
+	URetargetChainSettings* ChainSettings = GetChainSettings(TargetChainName);
+	if (!ChainSettings)
+	{
+		UE_LOG(LogIKRigEditor, Warning, TEXT("Could not map target chain. Target chain not found, %s."), *TargetChainName.ToString());
+		return false;
+	}
 	
 	FScopedTransaction Transaction(LOCTEXT("SetRetargetChainSource", "Set Retarget Chain Source"));
-	ChainMap->Modify();
-	ChainMap->SourceChain = SourceChainToMapTo;
+	ChainSettings->Modify();
+	ChainSettings->SourceChain = SourceChainName;
 	
 	BroadcastNeedsReinitialized();
-}
 
-const TArray<TObjectPtr<URetargetChainSettings>>& UIKRetargeterController::GetChainMappings() const
-{
-	return Asset->ChainSettings;
+	return true;
 }
 
-const URetargetChainSettings* UIKRetargeterController::GetChainMappingByTargetChainName(const FName& TargetChainName) const
+FName UIKRetargeterController::GetSourceChain(const FName& TargetChainName) const
 {
-	for (TObjectPtr<URetargetChainSettings> ChainMap : Asset->ChainSettings)
+	const URetargetChainSettings* ChainSettings = GetChainSettings(TargetChainName);
+	if (!ChainSettings)
 	{
-		if (ChainMap->TargetChain == TargetChainName)
-		{
-			return ChainMap;
-		}
+		UE_LOG(LogIKRigEditor, Warning, TEXT("Could not map target chain. Target chain not found, %s."), *TargetChainName.ToString());
+		return NAME_None;
 	}
+	
+	return ChainSettings->SourceChain;
+}
 
-	return nullptr;
+const TArray<URetargetChainSettings*>& UIKRetargeterController::GetAllChainSettings() const
+{
+	return Asset->ChainSettings;
 }
 
 FName UIKRetargeterController::GetChainGoal(const TObjectPtr<URetargetChainSettings> ChainSettings) const
@@ -431,106 +438,138 @@ FName UIKRetargeterController::GetChainGoal(const TObjectPtr<URetargetChainSetti
 		return NAME_None;
 	}
 	
-	UIKRigDefinition* TargetIKRig = GetIKRigWriteable(ERetargetSourceOrTarget::Target);
+	const UIKRigDefinition* TargetIKRig = GetIKRig(ERetargetSourceOrTarget::Target);
 	if (!TargetIKRig)
 	{
 		return NAME_None;
 	}
-	
-	const UIKRigController* IKRigController = UIKRigController::GetIKRigController(TargetIKRig);
-	return IKRigController->GetRetargetChainGoal(ChainSettings->TargetChain);
+
+	const UIKRigController* RigController = UIKRigController::GetController(TargetIKRig);
+	return RigController->GetRetargetChainGoal(ChainSettings->TargetChain);
 }
 
 bool UIKRetargeterController::IsChainGoalConnectedToASolver(const FName& GoalName) const
 {
-	UIKRigDefinition* TargetIKRig = GetIKRigWriteable(ERetargetSourceOrTarget::Target);
+	const UIKRigDefinition* TargetIKRig = GetIKRig(ERetargetSourceOrTarget::Target);
 	if (!TargetIKRig)
 	{
 		return false;
 	}
-	
-	const UIKRigController* IKRigController = UIKRigController::GetIKRigController(TargetIKRig);
-	return IKRigController->IsGoalConnectedToAnySolver(GoalName);
+
+	const UIKRigController* RigController = UIKRigController::GetController(TargetIKRig);
+	return RigController->IsGoalConnectedToAnySolver(GoalName);
 }
 
-void UIKRetargeterController::AddRetargetPose(
-	const FName& NewPoseName,
-	const FIKRetargetPose* ToDuplicate,
-	const ERetargetSourceOrTarget SourceOrTarget) const
+FName UIKRetargeterController::CreateRetargetPose(const FName& NewPoseName, const ERetargetSourceOrTarget SourceOrTarget) const
 {
-	FScopedTransaction Transaction(LOCTEXT("AddRetargetPose", "Add Retarget Pose"));
+	FScopedTransaction Transaction(LOCTEXT("CreateRetargetPose", "Create Retarget Pose"));
 	Asset->Modify();
-	
+
+	// create a new pose with a unique name 
 	const FName UniqueNewPoseName = MakePoseNameUnique(NewPoseName.ToString(), SourceOrTarget);
-	TMap<FName, FIKRetargetPose>& Poses = GetRetargetPoses(SourceOrTarget);
-	FIKRetargetPose& NewPose = Poses.Add(UniqueNewPoseName);
-	if (ToDuplicate)
-	{
-		NewPose.RootTranslationOffset = ToDuplicate->RootTranslationOffset;
-		NewPose.BoneRotationOffsets = ToDuplicate->BoneRotationOffsets;
-	}
-	
+	GetRetargetPoses(SourceOrTarget).Add(UniqueNewPoseName);
+
+	// set new pose as the current pose
 	FName& CurrentRetargetPoseName = SourceOrTarget == ERetargetSourceOrTarget::Source ? Asset->CurrentSourceRetargetPose : Asset->CurrentTargetRetargetPose;
 	CurrentRetargetPoseName = UniqueNewPoseName;
 
 	BroadcastNeedsReinitialized();
+
+	return UniqueNewPoseName;
 }
 
-void UIKRetargeterController::RenameCurrentRetargetPose(
-	const FName& NewPoseName,
-	const ERetargetSourceOrTarget SourceOrTarget) const
+bool UIKRetargeterController::RemoveRetargetPose(const FName& PoseToRemove, const ERetargetSourceOrTarget SourceOrTarget) const
 {
-	// do we already have a retarget pose with this name?
-	
-	if (GetRetargetPoses(SourceOrTarget).Contains(NewPoseName))
+	if (PoseToRemove == Asset->GetDefaultPoseName())
 	{
-		return;
+		return false; // cannot remove default pose
 	}
-	
-	FScopedTransaction Transaction(LOCTEXT("RenameRetargetPose", "Rename Retarget Pose"));
-	Asset->Modify();
 
-	// replace key in the map
 	TMap<FName, FIKRetargetPose>& Poses = GetRetargetPoses(SourceOrTarget);
-	const FName CurrentPoseName = GetCurrentRetargetPoseName(SourceOrTarget);
-	const FIKRetargetPose CurrentPoseData = GetCurrentRetargetPose(SourceOrTarget);
-	Poses.Remove(CurrentPoseName);
-	Poses.Shrink();
-	Poses.Add(NewPoseName, CurrentPoseData);
+	if (!Poses.Contains(PoseToRemove))
+	{
+		return false; // cannot remove pose that doesn't exist
+	}
+
+	FScopedTransaction Transaction(LOCTEXT("RemoveRetargetPose", "Remove Retarget Pose"));
+	Asset->Modify();
 
-	// update current pose name
-	SetCurrentRetargetPose(NewPoseName, SourceOrTarget);
+	Poses.Remove(PoseToRemove);
+
+	// did we remove the currently used pose?
+	if (GetCurrentRetargetPoseName(SourceOrTarget) == PoseToRemove)
+	{
+		SetCurrentRetargetPose(UIKRetargeter::GetDefaultPoseName(), SourceOrTarget);
+	}
 
 	BroadcastNeedsReinitialized();
+
+	return true;
 }
 
-void UIKRetargeterController::RemoveRetargetPose(
-	const FName& PoseToRemove,
-	const ERetargetSourceOrTarget SourceOrTarget) const
+FName UIKRetargeterController::DuplicateRetargetPose( const FName PoseToDuplicate, const FName NewPoseName, const ERetargetSourceOrTarget SourceOrTarget) const
 {
-	if (PoseToRemove == Asset->GetDefaultPoseName())
+	TMap<FName, FIKRetargetPose>& Poses = GetRetargetPoses(SourceOrTarget);
+	if (!Poses.Contains(PoseToDuplicate))
 	{
-		return; // cannot remove default pose
+		UE_LOG(LogIKRigEditor, Warning, TEXT("Trying to duplicate pose that does not exist, %s."), *PoseToDuplicate.ToString());
+		return NAME_None; // cannot duplicate pose that doesn't exist
 	}
+	
+	FScopedTransaction Transaction(LOCTEXT("DuplicateRetargetPose", "Duplicate Retarget Pose"));
+	Asset->Modify();
 
-	TMap<FName, FIKRetargetPose>& Poses = GetRetargetPoses(SourceOrTarget);
-	if (!Poses.Contains(PoseToRemove))
+	// create a new pose with a unique name
+	const FName UniqueNewPoseName = MakePoseNameUnique(NewPoseName.ToString(), SourceOrTarget);
+	FIKRetargetPose& NewPose = Poses.Add(UniqueNewPoseName);
+	// duplicate the pose data
+	NewPose.RootTranslationOffset = Poses[PoseToDuplicate].RootTranslationOffset;
+	NewPose.BoneRotationOffsets = Poses[PoseToDuplicate].BoneRotationOffsets;
+
+	// set duplicate to be the current pose
+	FName& CurrentRetargetPoseName = SourceOrTarget == ERetargetSourceOrTarget::Source ? Asset->CurrentSourceRetargetPose : Asset->CurrentTargetRetargetPose;
+	CurrentRetargetPoseName = UniqueNewPoseName;
+
+	BroadcastNeedsReinitialized();
+
+	return UniqueNewPoseName;
+}
+
+bool UIKRetargeterController::RenameRetargetPose(const FName OldPoseName, const FName NewPoseName, const ERetargetSourceOrTarget SourceOrTarget) const
+{
+	// does the old pose exist?
+	if (!GetRetargetPoses(SourceOrTarget).Contains(OldPoseName))
 	{
-		return; // cannot remove pose that doesn't exist
+		UE_LOG(LogIKRigEditor, Warning, TEXT("Trying to rename pose that does not exist, %s."), *OldPoseName.ToString());
+		return false;
 	}
 
-	FScopedTransaction Transaction(LOCTEXT("RemoveRetargetPose", "Remove Retarget Pose"));
+	// check if we're renaming the current pose
+	const bool bWasCurrentPose = GetCurrentRetargetPoseName(SourceOrTarget) == OldPoseName;
+	
+	FScopedTransaction Transaction(LOCTEXT("RenameRetargetPose", "Rename Retarget Pose"));
 	Asset->Modify();
 
-	Poses.Remove(PoseToRemove);
+	// make sure new name is unique
+	const FName UniqueNewPoseName = MakePoseNameUnique(NewPoseName.ToString(), SourceOrTarget);
 
-	// did we remove the currently used pose?
-	if (GetCurrentRetargetPoseName(SourceOrTarget) == PoseToRemove)
+	// replace key in the map
+	TMap<FName, FIKRetargetPose>& Poses = GetRetargetPoses(SourceOrTarget);
+	const FName CurrentPoseName = GetCurrentRetargetPoseName(SourceOrTarget);
+	const FIKRetargetPose CurrentPoseData = GetCurrentRetargetPose(SourceOrTarget);
+	Poses.Remove(CurrentPoseName);
+	Poses.Shrink();
+	Poses.Add(UniqueNewPoseName, CurrentPoseData);
+
+	// make this the current retarget pose, iff the old one was
+	if (bWasCurrentPose)
 	{
-		SetCurrentRetargetPose(UIKRetargeter::GetDefaultPoseName(), SourceOrTarget);
+		SetCurrentRetargetPose(UniqueNewPoseName, SourceOrTarget);
 	}
 
 	BroadcastNeedsReinitialized();
+
+	return true;
 }
 
 void UIKRetargeterController::ResetRetargetPose(
@@ -582,10 +621,14 @@ FName UIKRetargeterController::GetCurrentRetargetPoseName(const ERetargetSourceO
 	return SourceOrTarget == ERetargetSourceOrTarget::Source ? GetAsset()->CurrentSourceRetargetPose : GetAsset()->CurrentTargetRetargetPose;
 }
 
-void UIKRetargeterController::SetCurrentRetargetPose(FName NewCurrentPose, const ERetargetSourceOrTarget SourceOrTarget) const
+bool UIKRetargeterController::SetCurrentRetargetPose(FName NewCurrentPose, const ERetargetSourceOrTarget SourceOrTarget) const
 {
 	const TMap<FName, FIKRetargetPose>& Poses = GetRetargetPoses(SourceOrTarget);
-	check(Poses.Contains(NewCurrentPose));
+	if (!Poses.Contains(NewCurrentPose))
+	{
+		UE_LOG(LogIKRigEditor, Warning, TEXT("Trying to set current pose to a pose that does not exist, %s."), *NewCurrentPose.ToString());
+		return false;
+	}
 
 	FScopedTransaction Transaction(LOCTEXT("SetCurrentPose", "Set Current Pose"));
 	Asset->Modify();
@@ -593,6 +636,8 @@ void UIKRetargeterController::SetCurrentRetargetPose(FName NewCurrentPose, const
 	CurrentPose = NewCurrentPose;
 	
 	BroadcastNeedsReinitialized();
+
+	return true;
 }
 
 TMap<FName, FIKRetargetPose>& UIKRetargeterController::GetRetargetPoses(const ERetargetSourceOrTarget SourceOrTarget) const
@@ -613,7 +658,7 @@ void UIKRetargeterController::SetRotationOffsetForRetargetPoseBone(
 	const UIKRigDefinition* IKRig = SourceOrTarget == ERetargetSourceOrTarget::Source ? GetAsset()->GetSourceIKRig() : GetAsset()->GetTargetIKRig();
 	FIKRetargetPose& Pose = GetCurrentRetargetPose(SourceOrTarget);
 	Pose.SetDeltaRotationForBone(BoneName, RotationOffset);
-	Pose.SortHierarchically(IKRig->Skeleton);
+	Pose.SortHierarchically(IKRig->GetSkeleton());
 }
 
 FQuat UIKRetargeterController::GetRotationOffsetForRetargetPoseBone(
@@ -629,16 +674,28 @@ FQuat UIKRetargeterController::GetRotationOffsetForRetargetPoseBone(
 	return BoneOffsets[BoneName];
 }
 
-void UIKRetargeterController::AddTranslationOffsetToRetargetRootBone(
+void UIKRetargeterController::SetRootOffsetInRetargetPose(
 	const FVector& TranslationOffset,
 	const ERetargetSourceOrTarget SourceOrTarget) const
 {
 	GetCurrentRetargetPose(SourceOrTarget).AddToRootTranslationDelta(TranslationOffset);
 }
 
+FVector UIKRetargeterController::GetRootOffsetInRetargetPose(
+	const ERetargetSourceOrTarget SourceOrTarget) const
+{
+	return GetCurrentRetargetPose(SourceOrTarget).GetRootTranslationDelta();
+}
+
 FName UIKRetargeterController::MakePoseNameUnique(const FString& PoseName, const ERetargetSourceOrTarget SourceOrTarget) const
 {
 	FString UniqueName = PoseName;
+	
+	if (UniqueName.IsEmpty())
+	{
+		UniqueName = Asset->GetDefaultPoseName().ToString();
+	}
+	
 	int32 Suffix = 1;
 	const TMap<FName, FIKRetargetPose>& Poses = GetRetargetPoses(SourceOrTarget);
 	while (Poses.Contains(FName(UniqueName)))
@@ -649,7 +706,7 @@ FName UIKRetargeterController::MakePoseNameUnique(const FString& PoseName, const
 	return FName(UniqueName);
 }
 
-URetargetChainSettings* UIKRetargeterController::GetChainMap(const FName& TargetChainName) const
+URetargetChainSettings* UIKRetargeterController::GetChainSettings(const FName& TargetChainName) const
 {
 	for (URetargetChainSettings* ChainMap : Asset->ChainSettings)
 	{
@@ -673,7 +730,7 @@ void UIKRetargeterController::SortChainMapping() const
 	Asset->ChainSettings.Sort([TargetIKRig](const URetargetChainSettings& A, const URetargetChainSettings& B)
 	{
 		const TArray<FBoneChain>& BoneChains = TargetIKRig->GetRetargetChains();
-		const FIKRigSkeleton& TargetSkeleton = TargetIKRig->Skeleton;
+		const FIKRigSkeleton& TargetSkeleton = TargetIKRig->GetSkeleton();
 
 		// look for chains
 		const int32 IndexA = BoneChains.IndexOfByPredicate([&A](const FBoneChain& Chain)
diff --git a/Engine/Plugins/Animation/IKRig/Source/IKRigEditor/Private/RetargetEditor/SIKRetargetAssetBrowser.cpp b/Engine/Plugins/Animation/IKRig/Source/IKRigEditor/Private/RetargetEditor/SIKRetargetAssetBrowser.cpp
--- a/Engine/Plugins/Animation/IKRig/Source/IKRigEditor/Private/RetargetEditor/SIKRetargetAssetBrowser.cpp
+++ b/Engine/Plugins/Animation/IKRig/Source/IKRigEditor/Private/RetargetEditor/SIKRetargetAssetBrowser.cpp
@@ -170,9 +170,9 @@ FReply SIKRetargetAssetBrowser::OnExportButtonClicked()
 	}
 
 	// actually run the retarget
-	FIKRetargetBatchOperation BatchOperation;
-	BatchOperation.RunRetarget(BatchContext);
-
+	const TStrongObjectPtr<UIKRetargetBatchOperation> BatchOperation(NewObject<UIKRetargetBatchOperation>());
+	BatchOperation->RunRetarget(BatchContext);
+	
 	return FReply::Handled();
 }
 
diff --git a/Engine/Plugins/Animation/IKRig/Source/IKRigEditor/Private/RetargetEditor/SIKRetargetChainMapList.cpp b/Engine/Plugins/Animation/IKRig/Source/IKRigEditor/Private/RetargetEditor/SIKRetargetChainMapList.cpp
--- a/Engine/Plugins/Animation/IKRig/Source/IKRigEditor/Private/RetargetEditor/SIKRetargetChainMapList.cpp
+++ b/Engine/Plugins/Animation/IKRig/Source/IKRigEditor/Private/RetargetEditor/SIKRetargetChainMapList.cpp
@@ -17,6 +17,7 @@
 #include "Misc/TextFilterExpressionEvaluator.h"
 #include "Widgets/Input/SSearchBox.h"
 #include "Retargeter/IKRetargetSettings.h"
+#include "RigEditor/IKRigController.h"
 
 #define LOCTEXT_NAMESPACE "SIKRigRetargetChains"
 
@@ -170,7 +171,7 @@ void SIKRetargetChainMapRow::OnSourceChainComboSelectionChanged(TSharedPtr<FStri
 	}
 
 	const FName SourceChainName = FName(*InName.Get());
-	RetargeterController->SetSourceChainForTargetChain(ChainMapElement.Pin()->ChainMap.Get(), SourceChainName);
+	RetargeterController->SetSourceChain(SourceChainName, ChainMapElement.Pin()->ChainMap->TargetChain);
 }
 
 FText SIKRetargetChainMapRow::GetSourceChainName() const
@@ -198,7 +199,8 @@ FText SIKRetargetChainMapRow::GetTargetIKGoalName() const
 		return FText(); 
 	}
 
-	const FBoneChain* Chain = IKRig->GetRetargetChainByName(ChainMapElement.Pin()->ChainMap->TargetChain);
+	const UIKRigController* RigController = UIKRigController::GetController(IKRig);
+	const FBoneChain* Chain = RigController->GetRetargetChainByName(ChainMapElement.Pin()->ChainMap->TargetChain);
 	if (!Chain)
 	{
 		return FText(); 
@@ -475,7 +477,8 @@ void SIKRetargetChainMapList::RefreshView()
 			return false;
 		}
 
-		const FBoneChain* Chain = IKRig->GetRetargetChainByName(TargetChainName);
+		const UIKRigController* RigController = UIKRigController::GetController(IKRig);
+		const FBoneChain* Chain = RigController->GetRetargetChainByName(TargetChainName);
 		if (!Chain)
 		{
 			return false;
@@ -489,7 +492,7 @@ void SIKRetargetChainMapList::RefreshView()
 
 	// refresh items
 	ListViewItems.Reset();
-	const TArray<TObjectPtr<URetargetChainSettings>>& ChainMappings = AssetController->GetChainMappings();
+	const TArray<URetargetChainSettings*>& ChainMappings = AssetController->GetAllChainSettings();
 	for (const TObjectPtr<URetargetChainSettings> ChainMap : ChainMappings)
 	{
 		// apply text filter to items
@@ -765,7 +768,6 @@ void SIKRetargetChainMapList::AutoMapChains(const EAutoMapChainType AutoMapType,
 	}
 
 	Controller->ClearOutputLog();
-	RetargeterController->CleanChainMapping();
 	RetargeterController->AutoMapChains(AutoMapType, bForceRemap);
 	Controller->HandleRetargeterNeedsInitialized();
 }
diff --git a/Engine/Plugins/Animation/IKRig/Source/IKRigEditor/Private/RetargetEditor/SRetargetAnimAssetsWindow.cpp b/Engine/Plugins/Animation/IKRig/Source/IKRigEditor/Private/RetargetEditor/SRetargetAnimAssetsWindow.cpp
--- a/Engine/Plugins/Animation/IKRig/Source/IKRigEditor/Private/RetargetEditor/SRetargetAnimAssetsWindow.cpp
+++ b/Engine/Plugins/Animation/IKRig/Source/IKRigEditor/Private/RetargetEditor/SRetargetAnimAssetsWindow.cpp
@@ -576,8 +576,8 @@ bool SRetargetAnimAssetsWindow::CanApply() const
 FReply SRetargetAnimAssetsWindow::OnApply()
 {
 	CloseWindow();
-	FIKRetargetBatchOperation BatchOperation;
-	BatchOperation.RunRetarget(BatchContext);
+	const TStrongObjectPtr<UIKRetargetBatchOperation> BatchOperation(NewObject<UIKRetargetBatchOperation>());
+	BatchOperation->RunRetarget(BatchContext);
 	return FReply::Handled();
 }
 
diff --git a/Engine/Plugins/Animation/IKRig/Source/IKRigEditor/Private/RigEditor/IKRigController.cpp b/Engine/Plugins/Animation/IKRig/Source/IKRigEditor/Private/RigEditor/IKRigController.cpp
--- a/Engine/Plugins/Animation/IKRig/Source/IKRigEditor/Private/RigEditor/IKRigController.cpp
+++ b/Engine/Plugins/Animation/IKRig/Source/IKRigEditor/Private/RigEditor/IKRigController.cpp
@@ -2,10 +2,12 @@
 
 #include "RigEditor/IKRigController.h"
 
+#include "IKRigEditor.h"
 #include "IKRigDefinition.h"
 #include "IKRigProcessor.h"
 #include "IKRigSolver.h"
 
+#include "UObject/Package.h"
 #include "Engine/SkeletalMesh.h"
 #include "ScopedTransaction.h"
 
@@ -13,7 +15,7 @@
 
 #define LOCTEXT_NAMESPACE "IKRigController"
 
-UIKRigController* UIKRigController::GetIKRigController(UIKRigDefinition* InIKRigDefinition)
+UIKRigController* UIKRigController::GetController(const UIKRigDefinition* InIKRigDefinition)
 {
 	if (!InIKRigDefinition)
 	{
@@ -23,7 +25,7 @@ UIKRigController* UIKRigController::GetIKRigController(UIKRigDefinition* InIKRig
 	if (!InIKRigDefinition->Controller)
 	{
 		UIKRigController* Controller = NewObject<UIKRigController>();
-		Controller->Asset = InIKRigDefinition;
+		Controller->Asset = const_cast<UIKRigDefinition*>(InIKRigDefinition);
 		InIKRigDefinition->Controller = Controller;
 	}
 
@@ -35,112 +37,125 @@ UIKRigDefinition* UIKRigController::GetAsset() const
 	return Asset;
 }
 
-FName UIKRigController::GetAssetIDAsName() const
-{
-	if (!Asset)
-	{
-		return NAME_None;
-	}
-
-	return FName(FString::FromInt(Asset->GetUniqueID()));
-}
-
-void UIKRigController::AddBoneSetting(const FName& BoneName, int32 SolverIndex) const
+bool UIKRigController::AddBoneSetting(const FName BoneName, int32 SolverIndex) const
 {
 	check(Asset->Solvers.IsValidIndex(SolverIndex))
 
 	if (!CanAddBoneSetting(BoneName, SolverIndex))
 	{
-		return; // prerequisites not met
+		UE_LOG(LogIKRigEditor, Warning, TEXT("Bone settings refused by solver. See output for details."));
+		return false; // prerequisites not met
 	}
-
+	
 	FScopedTransaction Transaction(LOCTEXT("AddBoneSetting_Label", "Add Bone Setting"));
 	UIKRigSolver* Solver = Asset->Solvers[SolverIndex];
 	Solver->Modify();
 	Solver->AddBoneSetting(BoneName);
 	BroadcastNeedsReinitialized();
+
+	return true;
 }
 
-bool UIKRigController::CanAddBoneSetting(const FName& BoneName, int32 SolverIndex) const
+bool UIKRigController::CanAddBoneSetting(const FName BoneName, int32 SolverIndex) const
 {
-	const UIKRigSolver* Solver = GetSolver(SolverIndex);
+	const UIKRigSolver* Solver = GetSolverAtIndex(SolverIndex);
 	if (!Solver)
 	{
+		UE_LOG(LogIKRigEditor, Warning, TEXT("Solver does not exist at index: %d."), SolverIndex);
 		return false; // solver doesn't exist
 	}
 
 	if (Asset->Skeleton.GetBoneIndexFromName(BoneName) == INDEX_NONE)
 	{
+		UE_LOG(LogIKRigEditor, Warning, TEXT("Bone does not exist, %s."), *BoneName.ToString());
 		return false; // bone doesn't exist
 	}
 
 	if (!Asset->Solvers[SolverIndex]->UsesBoneSettings())
 	{
+		UE_LOG(LogIKRigEditor, Warning, TEXT("Solver does not support bone settings."));
 		return false; // solver doesn't support per-bone settings
 	}
 
-	// returns true if the solver in question does NOT already have a settings object for this bone 
-	return !static_cast<bool>(Asset->Solvers[SolverIndex]->GetBoneSetting(BoneName));
+	if (Asset->Solvers[SolverIndex]->GetBoneSetting(BoneName))
+	{
+		UE_LOG(LogIKRigEditor, Warning, TEXT("Solver already has settings on bone, %s."), *BoneName.ToString());
+		return false; // solve already has settings on this bone
+	}
+	
+	return true;
 }
 
-void UIKRigController::RemoveBoneSetting(const FName& BoneName, int32 SolverIndex) const
+bool UIKRigController::RemoveBoneSetting(const FName BoneName, int32 SolverIndex) const
 {
-	check(Asset->Solvers.IsValidIndex(SolverIndex))
+	UIKRigSolver* Solver = GetSolverAtIndex(SolverIndex);
+	if (!Solver)
+	{
+		UE_LOG(LogIKRigEditor, Warning, TEXT("Solver does not exist at index: %d."), SolverIndex);
+		return false; // solver doesn't exist
+	}
 
 	if (Asset->Skeleton.GetBoneIndexFromName(BoneName) == INDEX_NONE)
 	{
-		return; // bone doesn't exist
+		UE_LOG(LogIKRigEditor, Warning, TEXT("Cannot remove setting on unknown bone, %s."), *BoneName.ToString());
+		return false; // bone doesn't exist
 	}
 
 	FScopedTransaction Transaction(LOCTEXT("RemoveBoneSetting_Label", "Remove Bone Setting"));
-	UIKRigSolver* Solver = Asset->Solvers[SolverIndex];
 	Solver->Modify();
 	Solver->RemoveBoneSetting(BoneName);
 	BroadcastNeedsReinitialized();
+
+	return true;
 }
 
-bool UIKRigController::CanRemoveBoneSetting(const FName& BoneName, int32 SolverIndex) const
+bool UIKRigController::CanRemoveBoneSetting(const FName BoneName, int32 SolverIndex) const
 {
-	const UIKRigSolver* Solver = GetSolver(SolverIndex);
+	const UIKRigSolver* Solver = GetSolverAtIndex(SolverIndex);
 	if (!Solver)
 	{
+		UE_LOG(LogIKRigEditor, Warning, TEXT("Solver does not exist at index: %d."), SolverIndex);
 		return false; // solver doesn't exist
 	}
 
 	if (!Solver->UsesBoneSettings())
 	{
+		UE_LOG(LogIKRigEditor, Warning, TEXT("Solver does not support bone settings."));
 		return false; // solver doesn't use bone settings
 	}
 
 	if (Asset->Skeleton.GetBoneIndexFromName(BoneName) == INDEX_NONE)
 	{
+		UE_LOG(LogIKRigEditor, Warning, TEXT("Bone does not exist, %s."), *BoneName.ToString());
 		return false; // bone doesn't exist
 	}
 
 	if (!Solver->GetBoneSetting(BoneName))
 	{
+		UE_LOG(LogIKRigEditor, Warning, TEXT("Solver does not have settings on bone, %s."), *BoneName.ToString());
 		return false; // solver doesn't have any settings for this bone
 	}
 
 	return true;
 }
 
-UObject* UIKRigController::GetSettingsForBone(const FName& BoneName, int32 SolverIndex) const
+UObject* UIKRigController::GetSettingsForBone(const FName BoneName, int32 SolverIndex) const
 {
-	const UIKRigSolver* Solver = GetSolver(SolverIndex);
+	const UIKRigSolver* Solver = GetSolverAtIndex(SolverIndex);
 	if (!Solver)
 	{
+		UE_LOG(LogIKRigEditor, Warning, TEXT("Solver does not exist at index: %d."), SolverIndex);
 		return nullptr; // solver doesn't exist
 	}
 	
 	return Solver->GetBoneSetting(BoneName);
 }
 
-bool UIKRigController::DoesBoneHaveSettings(const FName& BoneName) const
+bool UIKRigController::DoesBoneHaveSettings(const FName BoneName) const
 {
 	if (Asset->Skeleton.GetBoneIndexFromName(BoneName) == INDEX_NONE)
 	{
-		return false; // bone doesn't exist
+		return false; // bone doesn't exist (do not spam here)
 	}
 	
 	for (UIKRigSolver* Solver : Asset->Solvers)
@@ -154,11 +169,36 @@ bool UIKRigController::DoesBoneHaveSettings(const FName& BoneName) const
 	return false;
 }
 
+FName UIKRigController::AddRetargetChain(const FName ChainName, const FName StartBoneName, const FName EndBoneName, const FName GoalName) const
+{
+	return AddRetargetChain(FBoneChain(ChainName, StartBoneName, EndBoneName, GoalName));
+}
+
 FName UIKRigController::AddRetargetChain(const FBoneChain& BoneChain) const
 {
+	if (BoneChain.StartBone.BoneName != NAME_None && Asset->Skeleton.GetBoneIndexFromName(BoneChain.StartBone.BoneName) == INDEX_NONE)
+	{
+		UE_LOG(LogIKRigEditor, Warning, TEXT("Could not create retarget chain. Start Bone does not exist, %s."), *BoneChain.StartBone.BoneName.ToString());
+		return NAME_None; // bone doesn't exist
+	}
+
+	if (BoneChain.EndBone.BoneName != NAME_None && Asset->Skeleton.GetBoneIndexFromName(BoneChain.EndBone.BoneName) == INDEX_NONE)
+	{
+		UE_LOG(LogIKRigEditor, Warning, TEXT("Could not create retarget chain. End Bone does not exist, %s."), *BoneChain.EndBone.BoneName.ToString());
+		return NAME_None; // bone doesn't exist
+	}
+	
 	FBoneChain ChainToAdd = BoneChain;
-	ChainToAdd.ChainName = GetUniqueRetargetChainName(BoneChain.ChainName);
 
+	// if no name specified, use a default
+	if (ChainToAdd.ChainName == NAME_None)
+	{
+		const FName DefaultChainName = FName("DefaultChainName");
+		ChainToAdd.ChainName = DefaultChainName;
+	}
+	
+	ChainToAdd.ChainName = GetUniqueRetargetChainName(BoneChain.ChainName);
+	
 	FScopedTransaction Transaction(LOCTEXT("AddRetargetChain_Label", "Add Retarget Chain"));
 	Asset->Modify();
 	
@@ -170,7 +210,7 @@ FName UIKRigController::AddRetargetChain(const FBoneChain& BoneChain) const
 	return Asset->RetargetDefinition.BoneChains[NewChainIndex].ChainName;
 }
 
-bool UIKRigController::RemoveRetargetChain(const FName& ChainName) const
+bool UIKRigController::RemoveRetargetChain(const FName ChainName) const
 {
 	FScopedTransaction Transaction(LOCTEXT("RemoveRetargetChain_Label", "Remove Retarget Chain"));
 	Asset->Modify();
@@ -192,14 +232,16 @@ bool UIKRigController::RemoveRetargetChain(const FName& ChainName) const
 		return true;
 	}
 
+	UE_LOG(LogIKRigEditor, Warning, TEXT("Retarget chain not found: %s."), *ChainName.ToString());
 	return false;
 }
 
-FName UIKRigController::RenameRetargetChain(const FName& ChainName, const FName& NewChainName) const
+FName UIKRigController::RenameRetargetChain(const FName ChainName, const FName NewChainName) const
 {
 	FBoneChain* Chain = Asset->RetargetDefinition.GetEditableBoneChainByName(ChainName);
 	if (!Chain)
 	{
+		UE_LOG(LogIKRigEditor, Warning, TEXT("Retarget chain not found: %s."), *ChainName.ToString());
 		return ChainName; // chain doesn't exist to rename
 	}
 
@@ -214,8 +256,14 @@ FName UIKRigController::RenameRetargetChain(const FName& ChainName, const FName&
 	return UniqueChainName;
 }
 
-bool UIKRigController::SetRetargetChainStartBone(const FName& ChainName, const FName& StartBoneName) const
+bool UIKRigController::SetRetargetChainStartBone(const FName ChainName, const FName StartBoneName) const
 {
+	if (Asset->Skeleton.GetBoneIndexFromName(StartBoneName) == INDEX_NONE)
+	{
+		UE_LOG(LogIKRigEditor, Warning, TEXT("Bone does not exist, %s."), *StartBoneName.ToString());
+		return false; // bone doesn't exist
+	}
+	
 	if (FBoneChain* BoneChain = Asset->RetargetDefinition.GetEditableBoneChainByName(ChainName))
 	{
 		FScopedTransaction Transaction(LOCTEXT("SetRetargetChainStartBone_Label", "Set Retarget Chain Start Bone"));
@@ -225,11 +273,18 @@ bool UIKRigController::SetRetargetChainStartBone(const FName& ChainName, const F
 		return true;
 	}
 
+	UE_LOG(LogIKRigEditor, Warning, TEXT("Retarget chain not found: %s."), *ChainName.ToString());
 	return false; // no bone chain with that name
 }
 
-bool UIKRigController::SetRetargetChainEndBone(const FName& ChainName, const FName& EndBoneName) const
+bool UIKRigController::SetRetargetChainEndBone(const FName ChainName, const FName EndBoneName) const
 {
+	if (Asset->Skeleton.GetBoneIndexFromName(EndBoneName) == INDEX_NONE)
+	{
+		UE_LOG(LogIKRigEditor, Warning, TEXT("Bone does not exist, %s."), *EndBoneName.ToString());
+		return false; // bone doesn't exist
+	}
+	
 	if (FBoneChain* BoneChain = Asset->RetargetDefinition.GetEditableBoneChainByName(ChainName))
 	{
 		FScopedTransaction Transaction(LOCTEXT("SetRetargetChainEndBone_Label", "Set Retarget Chain End Bone"));
@@ -239,68 +294,83 @@ bool UIKRigController::SetRetargetChainEndBone(const FName& ChainName, const FNa
 		return true;
 	}
 
+	UE_LOG(LogIKRigEditor, Warning, TEXT("Retarget chain not found: %s."), *ChainName.ToString());
 	return false; // no bone chain with that name
 }
 
-bool UIKRigController::SetRetargetChainGoal(const FName& ChainName, const FName& GoalName) const
+bool UIKRigController::SetRetargetChainGoal(const FName ChainName, const FName GoalName) const
 {
-	FName GoalNameToUse = GoalName;
-	if (!GetGoal(GoalName))
+	FBoneChain* BoneChain = Asset->RetargetDefinition.GetEditableBoneChainByName(ChainName);
+	if (!BoneChain)
 	{
-		GoalNameToUse = NAME_None; // no goal with that name	
+		UE_LOG(LogIKRigEditor, Warning, TEXT("Retarget chain not found: %s."), *ChainName.ToString());
+		return false; // no bone chain with that name
 	}
 	
-	if (FBoneChain* BoneChain = Asset->RetargetDefinition.GetEditableBoneChainByName(ChainName))
+	FName GoalNameToUse = GoalName;
+	if (!GetGoal(GoalName))
 	{
-		FScopedTransaction Transaction(LOCTEXT("SetRetargetChainGoal_Label", "Set Retarget Chain Goal"));
-		Asset->Modify();
-		BoneChain->IKGoalName = GoalNameToUse;
-		BroadcastNeedsReinitialized();
-		return true;
+		GoalNameToUse = NAME_None; // no goal with that name, that's ok we set it to None
 	}
 
-	return false; // no bone chain with that name
+	FScopedTransaction Transaction(LOCTEXT("SetRetargetChainGoal_Label", "Set Retarget Chain Goal"));
+	Asset->Modify();
+	BoneChain->IKGoalName = GoalNameToUse;
+	BroadcastNeedsReinitialized();
+	return true;
 }
 
-FName UIKRigController::GetRetargetChainGoal(const FName& ChainName) const
+FName UIKRigController::GetRetargetChainGoal(const FName ChainName) const
 {
 	check(Asset)
-	const FBoneChain* Chain = Asset->GetRetargetChainByName(ChainName);
+	const FBoneChain* Chain = GetRetargetChainByName(ChainName);
 	if (!Chain)
 	{
+		UE_LOG(LogIKRigEditor, Warning, TEXT("Retarget chain not found: %s."), *ChainName.ToString());
 		return NAME_None;
 	}
 	
 	return Chain->IKGoalName;
 }
 
-FName UIKRigController::GetRetargetChainStartBone(const FName& ChainName) const
+FName UIKRigController::GetRetargetChainStartBone(const FName ChainName) const
 {
 	check(Asset)
-	const FBoneChain* Chain = Asset->GetRetargetChainByName(ChainName);
+	const FBoneChain* Chain = GetRetargetChainByName(ChainName);
 	if (!Chain)
 	{
+		UE_LOG(LogIKRigEditor, Warning, TEXT("Retarget chain not found: %s."), *ChainName.ToString());
 		return NAME_None;
 	}
 	
 	return Chain->StartBone.BoneName;
 }
 
-FName UIKRigController::GetRetargetChainEndBone(const FName& ChainName) const
+FName UIKRigController::GetRetargetChainEndBone(const FName ChainName) const
 {
 	check(Asset)
-	const FBoneChain* Chain = Asset->GetRetargetChainByName(ChainName);
+	const FBoneChain* Chain = GetRetargetChainByName(ChainName);
 	if (!Chain)
 	{
+		UE_LOG(LogIKRigEditor, Warning, TEXT("Retarget chain not found: %s."), *ChainName.ToString());
 		return NAME_None;
 	}
 	
 	return Chain->EndBone.BoneName;
 }
 
-const FBoneChain* UIKRigController::GetRetargetChainByName(const FName& ChainName) const
+const FBoneChain* UIKRigController::GetRetargetChainByName(const FName ChainName) const
 {
-	return Asset->GetRetargetChainByName(ChainName);
+	for (const FBoneChain& Chain : Asset->RetargetDefinition.BoneChains)
+	{
+		if (Chain.ChainName == ChainName)
+		{
+			return &Chain;
+		}
+	}
+
+	UE_LOG(LogIKRigEditor, Warning, TEXT("Retarget chain not found: %s."), *ChainName.ToString());
+	return nullptr;
 }
 
 const TArray<FBoneChain>& UIKRigController::GetRetargetChains() const
@@ -309,16 +379,24 @@ const TArray<FBoneChain>& UIKRigController::GetRetargetChains() const
 	return Asset->GetRetargetChains();
 }
 
-void UIKRigController::SetRetargetRoot(const FName& RootBoneName) const
+bool UIKRigController::SetRetargetRoot(const FName RootBoneName) const
 {
 	check(Asset)
 
+	if (Asset->Skeleton.GetBoneIndexFromName(RootBoneName) == INDEX_NONE)
+	{
+		UE_LOG(LogIKRigEditor, Warning, TEXT("Bone does not exist, %s."), *RootBoneName.ToString());
+		return false; // bone doesn't exist
+	}
+
 	FScopedTransaction Transaction(LOCTEXT("SetRetargetRootBone_Label", "Set Retarget Root Bone"));
 	Asset->Modify();
 	
 	Asset->RetargetDefinition.RootBone = RootBoneName;
 
 	BroadcastNeedsReinitialized();
+
+	return true;
 }
 
 FName UIKRigController::GetRetargetRoot() const
@@ -343,9 +421,9 @@ void UIKRigController::SortRetargetChains() const
 	});
 }
 
-FName UIKRigController::GetUniqueRetargetChainName(const FName& NameToMakeUnique) const
+FName UIKRigController::GetUniqueRetargetChainName(const FName NameToMakeUnique) const
 {
-	auto IsNameBeingUsed = [this](const FName& NameToTry)->bool
+	auto IsNameBeingUsed = [this](const FName NameToTry)->bool
 	{
 		for (const FBoneChain& Chain : Asset->RetargetDefinition.BoneChains)
 		{
@@ -374,11 +452,11 @@ FName UIKRigController::GetUniqueRetargetChainName(const FName& NameToMakeUnique
 }
 
 bool UIKRigController::ValidateChain(
-	const FName& ChainName,
+	const FName ChainName,
 	const FIKRigSkeleton* OptionalSkeleton,
 	TSet<int32>& OutChainIndices) const
 {
-	const FBoneChain* Chain = Asset->GetRetargetChainByName(ChainName);
+	const FBoneChain* Chain = GetRetargetChainByName(ChainName);
 	if (!Chain)
 	{
 		return false; // chain doesn't exist
@@ -433,7 +511,7 @@ bool UIKRigController::ValidateChain(
 	}
 }
 
-FName UIKRigController::GetRetargetChainFromBone(const FName& BoneName, const FIKRigSkeleton* OptionalSkeleton) const
+FName UIKRigController::GetRetargetChainFromBone(const FName BoneName, const FIKRigSkeleton* OptionalSkeleton) const
 {
 	const FIKRigSkeleton& Skeleton = OptionalSkeleton ? *OptionalSkeleton : GetIKRigSkeleton();
 	const int32 BoneIndex = Skeleton.GetBoneIndexFromName(BoneName);
@@ -460,7 +538,7 @@ FName UIKRigController::GetRetargetChainFromBone(const FName& BoneName, const FI
 	return NAME_None;
 }
 
-FName UIKRigController::GetRetargetChainFromGoal(const FName& GoalName) const
+FName UIKRigController::GetRetargetChainFromGoal(const FName GoalName) const
 {
 	if (GoalName == NAME_None)
 	{
@@ -485,16 +563,10 @@ FName UIKRigController::GetRetargetChainFromGoal(const FName& GoalName) const
 
 bool UIKRigController::SetSkeletalMesh(USkeletalMesh* SkeletalMesh, bool bTransact) const
 {
-	if (!SkeletalMesh)
-	{
-		return false;
-	}
-	
 	// first determine runtime compatibility between the IK Rig asset and the skeleton we're trying to run it on
-	const FIKRigInputSkeleton InputSkeleton = FIKRigInputSkeleton(SkeletalMesh);
-	if (!UIKRigProcessor::IsIKRigCompatibleWithSkeleton(Asset, InputSkeleton, nullptr))
+	if (!IsSkeletalMeshCompatible(SkeletalMesh))
 	{
-		UE_LOG(LogTemp, Warning, TEXT("Trying to initialize IKRig with a Skeleton that is missing required bones. See output log. {0}"), *Asset->GetName());
+		UE_LOG(LogIKRigEditor, Warning, TEXT("Trying to initialize IKRig with a Skeleton that is missing required bones. See output log. {0}"), *Asset->GetName());
 		return false;
 	}
 
@@ -507,6 +579,7 @@ bool UIKRigController::SetSkeletalMesh(USkeletalMesh* SkeletalMesh, bool bTransa
 	// update stored skeletal mesh used for previewing results
 	Asset->PreviewSkeletalMesh = SkeletalMesh;
 	// copy skeleton data from the actual skeleton we want to run on
+	const FIKRigInputSkeleton InputSkeleton = FIKRigInputSkeleton(SkeletalMesh);
 	Asset->Skeleton.SetInputSkeleton(InputSkeleton, GetAsset()->Skeleton.ExcludedBones);
 	// update goal's initial transforms to reflect new skeleton
 	ResetInitialGoalTransforms();
@@ -516,27 +589,37 @@ bool UIKRigController::SetSkeletalMesh(USkeletalMesh* SkeletalMesh, bool bTransa
 	return true;
 }
 
-const FIKRigSkeleton& UIKRigController::GetIKRigSkeleton() const
+bool UIKRigController::IsSkeletalMeshCompatible(USkeletalMesh* SkeletalMeshToCheck) const
 {
-	return Asset->Skeleton;
+	const FIKRigInputSkeleton InputSkeleton = FIKRigInputSkeleton(SkeletalMeshToCheck);
+	return UIKRigProcessor::IsIKRigCompatibleWithSkeleton(Asset, InputSkeleton, nullptr);
 }
 
-USkeleton* UIKRigController::GetSkeleton() const
+USkeletalMesh* UIKRigController::GetSkeletalMesh() const
 {
-	if (!Asset->PreviewSkeletalMesh)
-    {
-        return nullptr;
-    }
-    
-    return Asset->PreviewSkeletalMesh->GetSkeleton();
+	return Asset->PreviewSkeletalMesh.LoadSynchronous();
 }
 
-void UIKRigController::SetBoneExcluded(const FName& BoneName, const bool bExclude) const
+const FIKRigSkeleton& UIKRigController::GetIKRigSkeleton() const
 {
+	return Asset->Skeleton;
+}
+
+bool UIKRigController::SetBoneExcluded(const FName BoneName, const bool bExclude) const
+{
+	// does bone exist?
+	const int32 BoneIndex = Asset->Skeleton.GetBoneIndexFromName(BoneName);
+	if (BoneIndex == INDEX_NONE)
+	{
+		UE_LOG(LogIKRigEditor, Warning, TEXT("Trying to exclude non-existant bone, %s."), *BoneName.ToString());
+		return false;
+	}
+
+	// already excluded?
 	const bool bIsExcluded = Asset->Skeleton.ExcludedBones.Contains(BoneName);
 	if (bIsExcluded == bExclude)
 	{
-		return; // already in the requested state of exclusion
+		return false; // (don't spam warning)
 	}
 
 	FScopedTransaction Transaction(LOCTEXT("SetBoneExcluded_Label", "Set Bone Excluded"));
@@ -552,17 +635,24 @@ void UIKRigController::SetBoneExcluded(const FName& BoneName, const bool bExclud
 	}
 
 	BroadcastNeedsReinitialized();
+
+	return true;
 }
 
-bool UIKRigController::GetBoneExcluded(const FName& BoneName) const
+bool UIKRigController::GetBoneExcluded(const FName BoneName) const
 {
 	return Asset->Skeleton.ExcludedBones.Contains(BoneName);
 }
 
-FTransform UIKRigController::GetRefPoseTransformOfBone(const FName& BoneName) const
+FTransform UIKRigController::GetRefPoseTransformOfBone(const FName BoneName) const
 {	
 	const int32 BoneIndex = Asset->Skeleton.GetBoneIndexFromName(BoneName);
-	check(BoneIndex != INDEX_NONE) // must initialize IK Rig before getting here
+	if (BoneIndex == INDEX_NONE)
+	{
+		UE_LOG(LogIKRigEditor, Warning, TEXT("Tried to get the ref pose of bone that is not loaded into this rig."));
+		return FTransform::Identity;
+	}
+	
 	return Asset->Skeleton.RefPoseGlobal[BoneIndex];
 }
 
@@ -574,6 +664,12 @@ int32 UIKRigController::AddSolver(TSubclassOf<UIKRigSolver> InIKRigSolverClass)
 {
 	check(Asset)
 
+	if (!InIKRigSolverClass)
+	{
+		UE_LOG(LogIKRigEditor, Warning, TEXT("Could not add solver to IK Rig. Invalid solver class specified."));
+		return INDEX_NONE;
+	}
+
 	FScopedTransaction Transaction(LOCTEXT("AddSolver_Label", "Add Solver"));
 	Asset->Modify();
 
@@ -587,32 +683,73 @@ int32 UIKRigController::AddSolver(TSubclassOf<UIKRigSolver> InIKRigSolverClass)
 	return SolverIndex;
 }
 
-void UIKRigController::RemoveSolver(const int32 SolverIndex) const
+int32 UIKRigController::AddSolver(const FString& InSolverClassPath) const
+{
+	if (InSolverClassPath.IsEmpty() || InSolverClassPath == FName(NAME_None).ToString())
+	{
+		UE_LOG(LogIKRigEditor, Warning, TEXT("Could not add solver to IK Rig. No solver class specified."));
+		return INDEX_NONE;
+	}
+
+	// account for class paths coming in fully qualified or just with class name
+	UPackage* Package = nullptr;
+	FString PackageName;
+	FString CPPTypeObjectName = InSolverClassPath;
+	if (InSolverClassPath.Split(TEXT("."), &PackageName, &CPPTypeObjectName))
+	{
+		Package = FindPackage(nullptr, *PackageName);
+	}
+
+	// first attempt to find object in specified package
+	UObject* Object = FindObject<UObject>(Package, *CPPTypeObjectName);
+	if (!Object)
+	{
+		// attempt to find object in ANY package
+		Object = FindFirstObject<UObject>(*InSolverClassPath, EFindFirstObjectOptions::NativeFirst | EFindFirstObjectOptions::EnsureIfAmbiguous);
+	}
+
+	// assign to subclass of base solver (null if cast fails)
+	//UClass* ObjClass = Object->GetClass();
+	const TSubclassOf<UIKRigSolver> SolverClass = CastChecked<UClass>(Object);
+	
+	return AddSolver(SolverClass);
+}
+
+bool UIKRigController::RemoveSolver(const int32 SolverIndex) const
 {
 	check(Asset)
 	
 	if (!Asset->Solvers.IsValidIndex(SolverIndex))
 	{
-		return;
+		UE_LOG(LogIKRigEditor, Warning, TEXT("Solver not removed. Invalid index, %d."), SolverIndex);
+		return false;
 	}
 
 	FScopedTransaction Transaction(LOCTEXT("RemoveSolver_Label", "Remove Solver"));
 	Asset->Modify();
-
 	Asset->Solvers.RemoveAt(SolverIndex);
 
 	BroadcastNeedsReinitialized();
+	return true;
 }
 
 bool UIKRigController::MoveSolverInStack(int32 SolverToMoveIndex, int32 TargetSolverIndex) const
 {
 	if (!Asset->Solvers.IsValidIndex(SolverToMoveIndex))
 	{
+		UE_LOG(LogIKRigEditor, Warning, TEXT("Solver not moved. Invalid source index, %d."), SolverToMoveIndex);
 		return false;
 	}
 
 	if (!Asset->Solvers.IsValidIndex(TargetSolverIndex))
 	{
+		UE_LOG(LogIKRigEditor, Warning, TEXT("Solver not moved. Invalid target index, %d."), TargetSolverIndex);
+		return false;
+	}
+
+	if (SolverToMoveIndex == TargetSolverIndex)
+	{
+		UE_LOG(LogIKRigEditor, Warning, TEXT("Solver not moved. Source and target index cannot be the same."));
 		return false;
 	}
 
@@ -633,6 +770,7 @@ bool UIKRigController::SetSolverEnabled(int32 SolverIndex, bool bIsEnabled) cons
 {
 	if (!Asset->Solvers.IsValidIndex(SolverIndex))
 	{
+		UE_LOG(LogIKRigEditor, Warning, TEXT("Solver not enabled or disabled. Invalid index, %d."), SolverIndex);
 		return false;
 	}
 	
@@ -648,16 +786,29 @@ bool UIKRigController::SetSolverEnabled(int32 SolverIndex, bool bIsEnabled) cons
 	return true;
 }
 
-void UIKRigController::SetRootBone(const FName& RootBoneName, int32 SolverIndex) const
+bool UIKRigController::GetSolverEnabled(int32 SolverIndex) const
+{
+	if (!Asset->Solvers.IsValidIndex(SolverIndex))
+	{
+		UE_LOG(LogIKRigEditor, Warning, TEXT("Invalid solver index, %d."), SolverIndex);
+		return false;
+	}
+
+	return Asset->Solvers[SolverIndex]->IsEnabled();
+}
+
+bool UIKRigController::SetRootBone(const FName RootBoneName, int32 SolverIndex) const
 {
 	if (!Asset->Solvers.IsValidIndex(SolverIndex))
 	{
-		return; // solver doesn't exist
+		UE_LOG(LogIKRigEditor, Warning, TEXT("Root bone not set. Invalid solver index, %d."), SolverIndex);
+		return false; // solver doesn't exist
 	}
 
 	if (Asset->Skeleton.GetBoneIndexFromName(RootBoneName) == INDEX_NONE)
 	{
-		return; // bone doesn't exist
+		UE_LOG(LogIKRigEditor, Warning, TEXT("Root bone not set. Invalid bone specified, %s."), *RootBoneName.ToString());
+		return false; // bone doesn't exist
 	}
 
 	FScopedTransaction Transaction(LOCTEXT("SetRootBone_Label", "Set Root Bone"));
@@ -667,18 +818,39 @@ void UIKRigController::SetRootBone(const FName& RootBoneName, int32 SolverIndex)
 	Solver->SetRootBone(RootBoneName);
 
 	BroadcastNeedsReinitialized();
+
+	return true;
 }
 
-void UIKRigController::SetEndBone(const FName& EndBoneName, int32 SolverIndex) const
+FName UIKRigController::GetRootBone(int32 SolverIndex) const
 {
 	if (!Asset->Solvers.IsValidIndex(SolverIndex))
 	{
-		return; // solver doesn't exist
+		UE_LOG(LogIKRigEditor, Warning, TEXT("Could not query root bone. Invalid solver index, %d."), SolverIndex);
+		return NAME_None; // solver doesn't exist
+	}
+
+	return Asset->Solvers[SolverIndex]->GetRootBone();
+}
+
+bool UIKRigController::SetEndBone(const FName EndBoneName, int32 SolverIndex) const
+{
+	if (!Asset->Solvers.IsValidIndex(SolverIndex))
+	{
+		UE_LOG(LogIKRigEditor, Warning, TEXT("End bone not set. Invalid solver index, %d."), SolverIndex);
+		return false; // solver doesn't exist
 	}
 
 	if (Asset->Skeleton.GetBoneIndexFromName(EndBoneName) == INDEX_NONE)
 	{
-		return; // bone doesn't exist
+		UE_LOG(LogIKRigEditor, Warning, TEXT("End bone not set. Invalid bone specified, %s."), *EndBoneName.ToString());
+		return false; // bone doesn't exist
+	}
+
+	if (!Asset->Solvers[SolverIndex]->RequiresEndBone())
+	{
+		UE_LOG(LogIKRigEditor, Warning, TEXT("End bone not set. Specified solver does not support end bones."));
+		return false; //
 	}
 
 	FScopedTransaction Transaction(LOCTEXT("SetEndBone_Label", "Set End Bone"));
@@ -688,6 +860,19 @@ void UIKRigController::SetEndBone(const FName& EndBoneName, int32 SolverIndex) c
 	Solver->SetEndBone(EndBoneName);
 
 	BroadcastNeedsReinitialized();
+
+	return true;
+}
+
+FName UIKRigController::GetEndBone(int32 SolverIndex) const
+{
+	if (!Asset->Solvers.IsValidIndex(SolverIndex))
+	{
+		UE_LOG(LogIKRigEditor, Warning, TEXT("End bone not queried. Invalid solver index, %d."), SolverIndex);
+		return NAME_None; // solver doesn't exist
+	}
+
+	return Asset->Solvers[SolverIndex]->GetEndBone();
 }
 
 const TArray<UIKRigSolver*>& UIKRigController::GetSolverArray() const
@@ -695,7 +880,7 @@ const TArray<UIKRigSolver*>& UIKRigController::GetSolverArray() const
 	return Asset->Solvers;
 }
 
-FString UIKRigController::GetSolverUniqueName(int32 SolverIndex)
+FString UIKRigController::GetSolverUniqueName(int32 SolverIndex) const
 {
 	check(Asset)
 	if (!Asset->Solvers.IsValidIndex(SolverIndex))
@@ -714,7 +899,7 @@ int32 UIKRigController::GetNumSolvers() const
 	return Asset->Solvers.Num();
 }
 
-UIKRigSolver* UIKRigController::GetSolver(int32 Index) const
+UIKRigSolver* UIKRigController::GetSolverAtIndex(int32 Index) const
 {
 	check(Asset)
 
@@ -726,20 +911,31 @@ UIKRigSolver* UIKRigController::GetSolver(int32 Index) const
 	return nullptr;
 }
 
+int32 UIKRigController::GetIndexOfSolver(UIKRigSolver* Solver) const
+{
+	check(Asset)
+	return Asset->Solvers.Find(Solver);
+}
+
 // -------------------------------------------------------
 // GOALS
 //
 
-UIKRigEffectorGoal* UIKRigController::AddNewGoal(const FName& GoalName, const FName& BoneName) const
+FName UIKRigController::AddNewGoal(const FName GoalName, const FName BoneName) const
 {
-	if (GetIKRigSkeleton().GetBoneIndexFromName(BoneName) == INDEX_NONE)
+	// does goal already exist?
+	if (GetGoalIndex(GoalName) != INDEX_NONE)
 	{
-		return nullptr; // bone does not exist in the skeleton
+		UE_LOG(LogIKRigEditor, Error, TEXT("Trying to create a Goal that already exists, %s."), *GoalName.ToString());
+		return NAME_None;
 	}
-	
-	if (GetGoalIndex(GoalName) != INDEX_NONE)
+
+	// does this bone exist?
+	const int32 BoneIndex = Asset->Skeleton.GetBoneIndexFromName(BoneName);
+	if (BoneIndex == INDEX_NONE)
 	{
-		return nullptr; // goal already exists!
+		UE_LOG(LogIKRigEditor, Error, TEXT("Trying to create Goal on unknown bone, %s."), *BoneName.ToString());
+		return NAME_None;
 	}
 	
 	FScopedTransaction Transaction(LOCTEXT("AddNewGoal_Label", "Add New Goal"));
@@ -757,16 +953,16 @@ UIKRigEffectorGoal* UIKRigController::AddNewGoal(const FName& GoalName, const FN
 	// connect the new goal to all the solvers
 	for (int32 SolverIndex=0; SolverIndex<GetNumSolvers(); ++SolverIndex)
 	{
-		ConnectGoalToSolver(*NewGoal, SolverIndex);
+		ConnectGoalToSolver(NewGoal->GoalName, SolverIndex);
 	}
  
 	BroadcastNeedsReinitialized();
 	BroadcastGoalsChange();
 	
-	return NewGoal;
+	return NewGoal->GoalName;
 }
 
-bool UIKRigController::RemoveGoal(const FName& GoalName) const
+bool UIKRigController::RemoveGoal(const FName GoalName) const
 {
 	const int32 GoalIndex = GetGoalIndex(GoalName);
 	if (GoalIndex == INDEX_NONE)
@@ -777,7 +973,7 @@ bool UIKRigController::RemoveGoal(const FName& GoalName) const
 	FScopedTransaction Transaction(LOCTEXT("RemoveGoal_Label", "Remove Goal"));
 
 	// remove from all the solvers
-	const FName& GoalToRemove = Asset->Goals[GoalIndex]->GoalName;
+	const FName GoalToRemove = Asset->Goals[GoalIndex]->GoalName;
 	for (UIKRigSolver* Solver : Asset->Solvers)
 	{
 		Solver->Modify();
@@ -803,7 +999,7 @@ bool UIKRigController::RemoveGoal(const FName& GoalName) const
 	return true;
 }
 
-FName UIKRigController::RenameGoal(const FName& OldName, const FName& PotentialNewName) const
+FName UIKRigController::RenameGoal(const FName OldName, const FName PotentialNewName) const
 {
 	if (OldName == PotentialNewName)
 	{
@@ -851,9 +1047,9 @@ FName UIKRigController::RenameGoal(const FName& OldName, const FName& PotentialN
 	return NewName;
 }
 
-FName UIKRigController::GetUniqueGoalName(const FName& NameToMakeUnique) const
+FName UIKRigController::GetUniqueGoalName(const FName NameToMakeUnique) const
 {
-	auto IsNameBeingUsed = [this](const FName& NameToTry) -> bool
+	auto IsNameBeingUsed = [this](const FName NameToTry) -> bool
 	{
 		// check if this goal already exists (case sensitive)
 		int32 ExistingGoalIndex = GetGoalIndex(NameToTry, ENameCase::IgnoreCase);
@@ -876,7 +1072,7 @@ FName UIKRigController::GetUniqueGoalName(const FName& NameToMakeUnique) const
 	return FName(NameToMakeUnique, Number);
 }
 
-bool UIKRigController::ModifyGoal(const FName& GoalName) const
+bool UIKRigController::ModifyGoal(const FName GoalName) const
 {
 	if (UIKRigEffectorGoal* Goal = GetGoal(GoalName))
 	{
@@ -887,7 +1083,7 @@ bool UIKRigController::ModifyGoal(const FName& GoalName) const
 	return false;
 }
 
-bool UIKRigController::SetGoalBone(const FName& GoalName, const FName& NewBoneName) const
+bool UIKRigController::SetGoalBone(const FName GoalName, const FName NewBoneName) const
 {
 	const int32 GoalIndex = GetGoalIndex(GoalName);
 	if (GoalIndex == INDEX_NONE)
@@ -926,7 +1122,7 @@ bool UIKRigController::SetGoalBone(const FName& GoalName, const FName& NewBoneNa
 	return true;
 }
 
-FName UIKRigController::GetBoneForGoal(const FName& GoalName) const
+FName UIKRigController::GetBoneForGoal(const FName GoalName) const
 {
 	for (const UIKRigEffectorGoal* Goal : Asset->Goals)
 	{
@@ -939,48 +1135,64 @@ FName UIKRigController::GetBoneForGoal(const FName& GoalName) const
 	return NAME_None;
 }
 
-UIKRigEffectorGoal* UIKRigController::GetGoalForBone(const FName& BoneName) const
+FName UIKRigController::GetGoalNameForBone(const FName BoneName) const
 {
 	TArray<UIKRigEffectorGoal*>& AllGoals = GetAllGoals();
 	for (UIKRigEffectorGoal* Goal : AllGoals)
 	{
 		if (Goal->BoneName == BoneName)
 		{
-			return Goal;
+			return Goal->GoalName;
 		}
 	}
 
-	return nullptr;
+	return NAME_None;
 }
 
-bool UIKRigController::ConnectGoalToSolver(const UIKRigEffectorGoal& Goal, int32 SolverIndex) const
+bool UIKRigController::ConnectGoalToSolver(const FName GoalName, int32 SolverIndex) const
 {
-	// can't add goal that is not present in the core
-	check(GetGoalIndex(Goal.GoalName) != INDEX_NONE);
+	// get index of the goal with this name
+	const int32 GoalIndex = GetGoalIndex(GoalName);
+	
+	// can't add goal that is not present
+	if (GoalIndex == INDEX_NONE)
+	{
+		UE_LOG(LogIKRigEditor, Warning, TEXT("Trying to connect unknown Goal, {0} to a solve."), *GoalName.ToString());
+		return false;
+	}
+	
 	// can't add goal to a solver with an invalid index
-	check(Asset->Solvers.IsValidIndex(SolverIndex))
+	if (!Asset->Solvers.IsValidIndex(SolverIndex))
+	{
+		UE_LOG(LogIKRigEditor, Warning, TEXT("Trying to connect Goal, %s to a unknown solver with index, %d."), *GoalName.ToString(), SolverIndex);
+		return false;
+	}
 
 	FScopedTransaction Transaction(LOCTEXT("ConnectGoalSolver_Label", "Connect Goal to Solver"));
+
 	UIKRigSolver* Solver = Asset->Solvers[SolverIndex];
+	const UIKRigEffectorGoal& Goal = *GetGoal(GoalName);
 	Solver->Modify();
-	
 	Solver->AddGoal(Goal);
 
 	BroadcastNeedsReinitialized();
+	
 	return true;
 }
 
-bool UIKRigController::DisconnectGoalFromSolver(const FName& GoalToRemove, int32 SolverIndex) const
+bool UIKRigController::DisconnectGoalFromSolver(const FName GoalToRemove, int32 SolverIndex) const
 {
 	// can't remove goal that is not present in the core
 	if (GetGoalIndex(GoalToRemove) == INDEX_NONE)
 	{
+		UE_LOG(LogIKRigEditor, Warning, TEXT("Trying to disconnect unknown Goal, %s."), *GoalToRemove.ToString());
 		return false;
 	}
 	
 	// can't remove goal from a solver with an invalid index
 	if (!Asset->Solvers.IsValidIndex(SolverIndex))
 	{
+		UE_LOG(LogIKRigEditor, Warning, TEXT("Trying to disconnect Goal, %s from an unknown solver with index, %d."), *GoalToRemove.ToString(), SolverIndex);
 		return false;
 	}
 
@@ -994,7 +1206,7 @@ bool UIKRigController::DisconnectGoalFromSolver(const FName& GoalToRemove, int32
 	return true;
 }
 
-bool UIKRigController::IsGoalConnectedToSolver(const FName& GoalName, int32 SolverIndex) const
+bool UIKRigController::IsGoalConnectedToSolver(const FName GoalName, int32 SolverIndex) const
 {
 	if (!Asset->Solvers.IsValidIndex(SolverIndex))
 	{
@@ -1004,7 +1216,7 @@ bool UIKRigController::IsGoalConnectedToSolver(const FName& GoalName, int32 Solv
 	return Asset->Solvers[SolverIndex]->IsGoalConnected(GoalName);
 }
 
-bool UIKRigController::IsGoalConnectedToAnySolver(const FName& GoalName) const
+bool UIKRigController::IsGoalConnectedToAnySolver(const FName GoalName) const
 {
 	for (const TObjectPtr<UIKRigSolver> Solver : Asset->Solvers)
 	{
@@ -1022,17 +1234,7 @@ TArray<UIKRigEffectorGoal*>& UIKRigController::GetAllGoals() const
 	return Asset->Goals;
 }
 
-const UIKRigEffectorGoal* UIKRigController::GetGoal(int32 GoalIndex) const
-{
-	if (!Asset->Goals.IsValidIndex(GoalIndex))
-	{
-		return nullptr;
-	}
-
-	return Asset->Goals[GoalIndex];
-}
-
-UIKRigEffectorGoal* UIKRigController::GetGoal(const FName& GoalName) const
+UIKRigEffectorGoal* UIKRigController::GetGoal(const FName GoalName) const
 {
 	const int32 GoalIndex = GetGoalIndex(GoalName);
 	if (GoalIndex == INDEX_NONE)
@@ -1043,7 +1245,7 @@ UIKRigEffectorGoal* UIKRigController::GetGoal(const FName& GoalName) const
 	return Asset->Goals[GoalIndex];
 }
 
-UObject* UIKRigController::GetGoalSettingsForSolver(const FName& GoalName, int32 SolverIndex) const
+UObject* UIKRigController::GetGoalSettingsForSolver(const FName GoalName, int32 SolverIndex) const
 {
 	const int32 GoalIndex = GetGoalIndex(GoalName);
 	if (GoalIndex == INDEX_NONE)
@@ -1059,7 +1261,7 @@ UObject* UIKRigController::GetGoalSettingsForSolver(const FName& GoalName, int32
 	return Asset->Solvers[SolverIndex]->GetGoalSettings(GoalName);
 }
 
-FTransform UIKRigController::GetGoalCurrentTransform(const FName& GoalName) const
+FTransform UIKRigController::GetGoalCurrentTransform(const FName GoalName) const
 {
 	if(const UIKRigEffectorGoal* Goal = GetGoal(GoalName))
 	{
@@ -1069,7 +1271,7 @@ FTransform UIKRigController::GetGoalCurrentTransform(const FName& GoalName) cons
 	return FTransform::Identity; // no goal with that name
 }
 
-void UIKRigController::SetGoalCurrentTransform(const FName& GoalName, const FTransform& Transform) const
+void UIKRigController::SetGoalCurrentTransform(const FName GoalName, const FTransform& Transform) const
 {
 	UIKRigEffectorGoal* Goal = GetGoal(GoalName);
 	check(Goal);
@@ -1129,7 +1331,7 @@ void UIKRigController::SanitizeGoalName(FString& InOutName)
 	}
 }
 
-int32 UIKRigController::GetGoalIndex(const FName& InGoalName, const ENameCase CompareMethod) const
+int32 UIKRigController::GetGoalIndex(const FName InGoalName, const ENameCase CompareMethod) const
 {
 	return Asset->Goals.IndexOfByPredicate([&](const TObjectPtr<UIKRigEffectorGoal>& Goal)
 	{
@@ -1137,16 +1339,6 @@ int32 UIKRigController::GetGoalIndex(const FName& InGoalName, const ENameCase Co
 	});
 }
 
-FName UIKRigController::GetGoalName(const int32& GoalIndex) const
-{
-	if (!Asset->Goals.IsValidIndex(GoalIndex))
-	{
-		return NAME_None;
-	}
-
-	return Asset->Goals[GoalIndex]->GoalName;
-}
-
 void UIKRigController::BroadcastGoalsChange() const
 {
 	if (Asset)
@@ -1160,5 +1352,22 @@ void UIKRigController::BroadcastGoalsChange() const
 	}
 }
 
+void UIKRigController::BroadcastNeedsReinitialized() const
+{
+	// initialize all solvers
+	const FIKRigSkeleton& IKRigSkeleton = GetIKRigSkeleton();
+	const TArray<UIKRigSolver*>& Solvers = GetSolverArray(); 
+	for (UIKRigSolver* Solver: Solvers)
+	{
+		Solver->Initialize(IKRigSkeleton);
+	}
+
+	// ensure goals are using initial transforms from the current mesh 
+	ResetInitialGoalTransforms();
+
+	// inform outside systems
+	IKRigNeedsInitialized.Broadcast(GetAsset());
+}
+
 #undef LOCTEXT_NAMESPACE
 
diff --git a/Engine/Plugins/Animation/IKRig/Source/IKRigEditor/Private/RigEditor/IKRigEditMode.cpp b/Engine/Plugins/Animation/IKRig/Source/IKRigEditor/Private/RigEditor/IKRigEditMode.cpp
--- a/Engine/Plugins/Animation/IKRig/Source/IKRigEditor/Private/RigEditor/IKRigEditMode.cpp
+++ b/Engine/Plugins/Animation/IKRig/Source/IKRigEditor/Private/RigEditor/IKRigEditMode.cpp
@@ -242,7 +242,7 @@ void FIKRigEditMode::GetBoneColors(
 			// record which bones in the skeleton are affected by this solver
 			const FIKRigSkeleton& Skeleton = Processor->GetSkeleton();
 			const UIKRigController* AssetController = Controller->AssetController;
-			if(const UIKRigSolver* SelectedSolver = AssetController->GetSolver(SelectedSolvers[0].Get()->IndexInStack))
+			if(const UIKRigSolver* SelectedSolver = AssetController->GetSolverAtIndex(SelectedSolvers[0].Get()->IndexInStack))
 			{
 				for (int32 BoneIndex=0; BoneIndex < Skeleton.BoneNames.Num(); ++BoneIndex)
 				{
diff --git a/Engine/Plugins/Animation/IKRig/Source/IKRigEditor/Private/RigEditor/IKRigEditorController.cpp b/Engine/Plugins/Animation/IKRig/Source/IKRigEditor/Private/RigEditor/IKRigEditorController.cpp
--- a/Engine/Plugins/Animation/IKRig/Source/IKRigEditor/Private/RigEditor/IKRigEditorController.cpp
+++ b/Engine/Plugins/Animation/IKRig/Source/IKRigEditor/Private/RigEditor/IKRigEditorController.cpp
@@ -2,14 +2,6 @@
 
 #include "RigEditor/IKRigEditorController.h"
 
-#include "ContentBrowserModule.h"
-#include "IContentBrowserSingleton.h"
-#include "IPersonaToolkit.h"
-#include "SKismetInspector.h"
-#include "Animation/DebugSkelMeshComponent.h"
-#include "Dialog/SCustomDialog.h"
-#include "Dialogs/Dialogs.h"
-
 #include "RigEditor/IKRigController.h"
 #include "RigEditor/SIKRigHierarchy.h"
 #include "RigEditor/SIKRigSolverStack.h"
@@ -19,6 +11,13 @@
 #include "RigEditor/SIKRigOutputLog.h"
 #include "Solvers/IKRig_PBIKSolver.h"
 #include "Widgets/Input/SComboBox.h"
+#include "ContentBrowserModule.h"
+#include "IContentBrowserSingleton.h"
+#include "IPersonaToolkit.h"
+#include "SKismetInspector.h"
+#include "Animation/DebugSkelMeshComponent.h"
+#include "Dialog/SCustomDialog.h"
+#include "Dialogs/Dialogs.h"
 
 #include UE_INLINE_GENERATED_CPP_BY_NAME(IKRigEditorController)
 
@@ -44,7 +43,8 @@ TOptional<FTransform> UIKRigBoneDetails::GetTransform(EIKRigTransformType::Type
 	FTransform GlobalTransform = FTransform::Identity;
 	const bool* IsRelative = nullptr;
 
-	const int32 BoneIndex = AssetPtr->Skeleton.GetBoneIndexFromName(SelectedBone);
+	const FIKRigSkeleton& Skeleton = AssetPtr->GetSkeleton();
+	const int32 BoneIndex = Skeleton.GetBoneIndexFromName(SelectedBone);
 	if(BoneIndex == INDEX_NONE)
 	{
 		return TOptional<FTransform>();
@@ -67,20 +67,20 @@ TOptional<FTransform> UIKRigBoneDetails::GetTransform(EIKRigTransformType::Type
 			}
 			else
 			{
-				GlobalTransform = AssetPtr->Skeleton.CurrentPoseGlobal[BoneIndex];
-				LocalTransform = AssetPtr->Skeleton.CurrentPoseLocal[BoneIndex];
+				GlobalTransform = Skeleton.CurrentPoseGlobal[BoneIndex];
+				LocalTransform = Skeleton.CurrentPoseLocal[BoneIndex];
 			}
 			break;
 		}
 		case EIKRigTransformType::Reference:
 		{
 			IsRelative = ReferenceTransformRelative;
-			GlobalTransform = AssetPtr->Skeleton.RefPoseGlobal[BoneIndex];
+			GlobalTransform = Skeleton.RefPoseGlobal[BoneIndex];
 			LocalTransform = GlobalTransform;
-			const int32 ParentBoneIndex = AssetPtr->Skeleton.ParentIndices[BoneIndex];
+			const int32 ParentBoneIndex = Skeleton.ParentIndices[BoneIndex];
 			if(ParentBoneIndex != INDEX_NONE)
 			{
-				const FTransform ParentTransform = AssetPtr->Skeleton.RefPoseGlobal[ParentBoneIndex];;
+				const FTransform ParentTransform = Skeleton.RefPoseGlobal[ParentBoneIndex];;
 				LocalTransform = GlobalTransform.GetRelativeTransform(ParentTransform);
 			}
 			break;
@@ -295,20 +295,16 @@ EChainSide FRetargetChainAnalyzer::GetSideOfChain(const TArray<int32>& BoneIndic
 	return AverageXPositionOfChain > 0 ? EChainSide::Left :  EChainSide::Right;
 }
 
-void FIKRigEditorController::Initialize(TSharedPtr<FIKRigEditorToolkit> Toolkit, UIKRigDefinition* IKRigAsset)
+void FIKRigEditorController::Initialize(TSharedPtr<FIKRigEditorToolkit> Toolkit, const UIKRigDefinition* IKRigAsset)
 {
 	EditorToolkit = Toolkit;
-	AssetController = UIKRigController::GetIKRigController(IKRigAsset);
+	AssetController = UIKRigController::GetController(IKRigAsset);
 	BoneDetails = NewObject<UIKRigBoneDetails>();
 	
 	// register callback to be informed when rig asset is modified by editor
 	if (!AssetController->OnIKRigNeedsInitialized().IsBoundToObject(this))
 	{
-		ReinitializeDelegateHandle = AssetController->OnIKRigNeedsInitialized().AddSP(
-			this, &FIKRigEditorController::HandleIKRigNeedsInitialized);
-
-		// Initialize editor's instances at first initialization
-		InitializeSolvers();
+		ReinitializeDelegateHandle = AssetController->OnIKRigNeedsInitialized().AddSP(this, &FIKRigEditorController::HandleIKRigNeedsInitialized);
 	}
 }
 
@@ -320,7 +316,7 @@ void FIKRigEditorController::Close() const
 void FIKRigEditorController::PromptUserToAssignMesh()
 {
 	// do we already have a skeletal mesh assigned?
-	if (AssetController->GetAsset()->PreviewSkeletalMesh)
+	if (AssetController->GetSkeletalMesh())
 	{
 		return;
 	}
@@ -373,7 +369,7 @@ void FIKRigEditorController::PromptUserToAssignMesh()
 
 UIKRigProcessor* FIKRigEditorController::GetIKRigProcessor() const
 {
-	if (AnimInstance.IsValid())
+	if (AnimInstance)
 	{
 		return AnimInstance->GetCurrentlyRunningProcessor();
 	}
@@ -402,21 +398,15 @@ void FIKRigEditorController::HandleIKRigNeedsInitialized(UIKRigDefinition* Modif
 	}
 
 	ClearOutputLog();
+
+	// currently running processor needs reinit
+	AnimInstance->SetProcessorNeedsInitialized();
 	
 	// in case the skeletal mesh was swapped out, we need to ensure the preview scene is up-to-date
-	{
-		USkeletalMesh* NewMesh = AssetController->GetAsset()->GetPreviewMesh();
-		const TSharedRef<IPersonaPreviewScene> PreviewScene = EditorToolkit.Pin()->GetPersonaToolkit()->GetPreviewScene();
-		PreviewScene->SetPreviewMesh(NewMesh);
+	USkeletalMesh* NewMesh = AssetController->GetSkeletalMesh();
+	const TSharedRef<IPersonaPreviewScene> PreviewScene = EditorToolkit.Pin()->GetPersonaToolkit()->GetPreviewScene();
+	PreviewScene->SetPreviewMesh(NewMesh); // this will re-initialize the anim instance
 
-		// ensure goals are using initial transforms from the current mesh 
-		AssetController->ResetInitialGoalTransforms();
-	}
-
-	// anim instance needs reinitialized
-	InitializeSolvers();
-	AnimInstance->SetProcessorNeedsInitialized();
-	AnimInstance->InitializeAnimation();
 
 	// update the bone details so it can pull on the current data
 	BoneDetails->AnimInstancePtr = AnimInstance;
@@ -495,14 +485,14 @@ void FIKRigEditorController::AddNewGoals(const TArray<FName>& GoalNames, const T
 		const FName& BoneName = BoneNames[I];
 
 		// create a new goal
-		UIKRigEffectorGoal* NewGoal = AssetController->AddNewGoal(GoalName, BoneName);
-		if (!NewGoal)
+		const FName NewGoalName = AssetController->AddNewGoal(GoalName, BoneName);
+		if (NewGoalName == NAME_None)
 		{
 			continue; // already exists
 		}
 
 		// ask user if they want to assign this goal to a chain (if there is one on this bone)
-		PromptToAssignGoalToChain(NewGoal);
+		PromptToAssignGoalToChain(NewGoalName);
 
 		LastCreatedGoalName = GoalName;
 	}
@@ -678,7 +668,7 @@ void FIKRigEditorController::CreateNewRetargetChains()
 	else
 	{
 		// create an empty chain
-		const FBoneChain Chain(FRetargetChainAnalyzer::GetDefaultChainName(), NAME_None, NAME_None);
+		const FBoneChain Chain(FRetargetChainAnalyzer::GetDefaultChainName(), NAME_None, NAME_None, NAME_None);
 		PromptToAddNewRetargetChain(Chain);
 	}
 	
@@ -795,7 +785,7 @@ void FIKRigEditorController::ShowDetailsForGoal(const FName& GoalName) const
 void FIKRigEditorController::ShowDetailsForGoalSettings(const FName GoalName, const int32 SolverIndex) const
 {
 	// get solver that owns this effector
-	if (const UIKRigSolver* SolverWithEffector = AssetController->GetSolver(SolverIndex))
+	if (const UIKRigSolver* SolverWithEffector = AssetController->GetSolverAtIndex(SolverIndex))
 	{
 		if (UObject* EffectorSettings = SolverWithEffector->GetGoalSettings(GoalName))
 		{
@@ -806,7 +796,7 @@ void FIKRigEditorController::ShowDetailsForGoalSettings(const FName GoalName, co
 
 void FIKRigEditorController::ShowDetailsForSolver(const int32 SolverIndex) const
 {
-	DetailsView->SetObject(AssetController->GetSolver(SolverIndex));
+	DetailsView->SetObject(AssetController->GetSolverAtIndex(SolverIndex));
 }
 
 void FIKRigEditorController::ShowEmptyDetails() const
@@ -876,7 +866,8 @@ void FIKRigEditorController::OnFinishedChangingDetails(const FPropertyChangedEve
 	const bool bPreviewMeshChanged = PropertyChangedEvent.GetPropertyName() == UIKRigDefinition::GetPreviewMeshPropertyName();
 	if (bPreviewMeshChanged)
 	{
-		HandleIKRigNeedsInitialized(AssetController->GetAsset());
+		USkeletalMesh* Mesh = AssetController->GetSkeletalMesh();
+		AssetController->SetSkeletalMesh(Mesh);
 	}
 }
 
@@ -887,8 +878,10 @@ void FIKRigEditorController::SetDetailsView(const TSharedPtr<IDetailsView>& InDe
 	ShowEmptyDetails();
 }
 
-void FIKRigEditorController::PromptToAssignGoalToChain(UIKRigEffectorGoal* NewGoal) const
+void FIKRigEditorController::PromptToAssignGoalToChain(const FName NewGoalName) const
 {
+	const UIKRigEffectorGoal* NewGoal = AssetController->GetGoal(NewGoalName);
+	check(NewGoal);
 	const TArray<FBoneChain>& AllRetargetChains = AssetController->GetRetargetChains();
 	FName ChainToAddGoalTo = NAME_None;
 	for (const FBoneChain& Chain : AllRetargetChains)
@@ -935,15 +928,15 @@ FName FIKRigEditorController::PromptToAddNewRetargetChain(const FBoneChain& Bone
 	DialogArguments.OnOkPressed_Lambda([&BoneChain, &NewChainName, this] ()
 	{
 		// prompt to add a goal to the chain (skipped if already has a goal)
-		UIKRigEffectorGoal* NewGoal = PromptToAddGoalToNewChain(BoneChain);
+		const FName NewGoalName = PromptToAddGoalToNewChain(BoneChain);
 
 		// add the retarget chain
 		NewChainName = AssetController->AddRetargetChain(BoneChain);
 
 		// ask user if they want to assign this goal to a chain (if there is one on this bone)
-		if (NewGoal)
+		if (NewGoalName != NAME_None)
 		{
-			PromptToAssignGoalToChain(NewGoal);
+			PromptToAssignGoalToChain(NewGoalName);
 		}
 		
 		RefreshAllViews();
@@ -958,20 +951,20 @@ FName FIKRigEditorController::PromptToAddNewRetargetChain(const FBoneChain& Bone
 	return NewChainName;
 }
 
-UIKRigEffectorGoal* FIKRigEditorController::PromptToAddGoalToNewChain(const FBoneChain& BoneChain) const
+FName FIKRigEditorController::PromptToAddGoalToNewChain(const FBoneChain& BoneChain) const
 {
 	// check if there is already a goal on the end bone of this chain, if there is, then we're done
-	UIKRigEffectorGoal* ExistingGoal = AssetController->GetGoalForBone(BoneChain.EndBone.BoneName);
-	if (ExistingGoal)
+	const FName ExistingGoalName = AssetController->GetGoalNameForBone(BoneChain.EndBone.BoneName);
+	if (ExistingGoalName != NAME_None)
 	{
-		return ExistingGoal;
+		return ExistingGoalName;
 	}
 
 	// check if there is a corresponding bone, if there is none, then no need to create a goal
 	const FIKRigSkeleton& IKRigSkeleton = AssetController->GetIKRigSkeleton();
 	if (IKRigSkeleton.GetBoneIndexFromName(BoneChain.EndBone.BoneName) == INDEX_NONE)
 	{
-		return nullptr;
+		return NAME_None;
 	}
 
 	// ask user if they want to add a goal to the end of this chain
@@ -992,7 +985,7 @@ UIKRigEffectorGoal* FIKRigEditorController::PromptToAddGoalToNewChain(const FBon
 
 	if (AddGoalToChainDialog->ShowModal() != 0)
 	{
-		return nullptr; // cancel button pressed, or window closed
+		return NAME_None; // cancel button pressed, or window closed
 	}
 
 	// add a default solver if there isn't one already
@@ -1005,7 +998,7 @@ UIKRigEffectorGoal* FIKRigEditorController::PromptToAddGoalToNewChain(const FBon
 
 void FIKRigEditorController::PlayAnimationAsset(UAnimationAsset* AssetToPlay)
 {
-	if (AssetToPlay && AnimInstance.IsValid())
+	if (AssetToPlay && AnimInstance)
 	{
 		AnimInstance->SetAnimationAsset(AssetToPlay);
 	}
@@ -1021,19 +1014,6 @@ void FIKRigEditorController::SetLastSelectedType(EIKRigSelectionType SelectionTy
 	LastSelectedType = SelectionType;
 }
 
-void FIKRigEditorController::InitializeSolvers() const
-{
-	if (AssetController)
-	{
-		const FIKRigSkeleton& IKRigSkeleton = AssetController->GetIKRigSkeleton();
-		const TArray<UIKRigSolver*>& Solvers = AssetController->GetSolverArray(); 
-		for (UIKRigSolver* Solver: Solvers)
-		{
-			Solver->Initialize(IKRigSkeleton);
-		}
-	}
-}
-
 TObjectPtr<UIKRigBoneDetails> FIKRigEditorController::CreateBoneDetails(const TSharedPtr<FIKRigTreeElement const>& InBoneItem) const
 {
 	// ensure that the element is related to a bone
diff --git a/Engine/Plugins/Animation/IKRig/Source/IKRigEditor/Private/RigEditor/IKRigToolkit.cpp b/Engine/Plugins/Animation/IKRig/Source/IKRigEditor/Private/RigEditor/IKRigToolkit.cpp
--- a/Engine/Plugins/Animation/IKRig/Source/IKRigEditor/Private/RigEditor/IKRigToolkit.cpp
+++ b/Engine/Plugins/Animation/IKRig/Source/IKRigEditor/Private/RigEditor/IKRigToolkit.cpp
@@ -83,9 +83,6 @@ void FIKRigEditorToolkit::InitAssetEditor(
 
 	ExtendToolbar();
 	RegenerateMenusAndToolbars();
-
-	// ignored if hierarchy / mesh already imported into IK Rig asset
-	EditorController->PromptUserToAssignMesh();
 }
 
 void FIKRigEditorToolkit::OnClose()
@@ -192,16 +189,12 @@ void FIKRigEditorToolkit::HandleOnPreviewSceneSettingsCustomized(IDetailLayoutBu
 
 void FIKRigEditorToolkit::PostUndo(bool bSuccess)
 {
-	EditorController->ClearOutputLog();
 	EditorController->AssetController->BroadcastNeedsReinitialized();
-	EditorController->RefreshAllViews();
 }
 
 void FIKRigEditorToolkit::PostRedo(bool bSuccess)
 {
-	EditorController->ClearOutputLog();
 	EditorController->AssetController->BroadcastNeedsReinitialized();
-	EditorController->RefreshAllViews();
 }
 
 void FIKRigEditorToolkit::HandlePreviewSceneCreated(const TSharedRef<IPersonaPreviewScene>& InPersonaPreviewScene)
@@ -224,7 +217,7 @@ void FIKRigEditorToolkit::HandlePreviewSceneCreated(const TSharedRef<IPersonaPre
 
 	// set the skeletal mesh on the component
 	// NOTE: this must be done AFTER setting the AnimInstance so that the correct root anim node is loaded
-	USkeletalMesh* Mesh = EditorController->AssetController->GetAsset()->GetPreviewMesh();
+	USkeletalMesh* Mesh = EditorController->AssetController->GetSkeletalMesh();
 	EditorController->SkelMeshComponent->SetSkeletalMesh(Mesh);
 
 	// apply mesh to the preview scene
diff --git a/Engine/Plugins/Animation/IKRig/Source/IKRigEditor/Private/RigEditor/SIKRigAssetBrowser.cpp b/Engine/Plugins/Animation/IKRig/Source/IKRigEditor/Private/RigEditor/SIKRigAssetBrowser.cpp
--- a/Engine/Plugins/Animation/IKRig/Source/IKRigEditor/Private/RigEditor/SIKRigAssetBrowser.cpp
+++ b/Engine/Plugins/Animation/IKRig/Source/IKRigEditor/Private/RigEditor/SIKRigAssetBrowser.cpp
@@ -10,6 +10,7 @@
 #include "Animation/PoseAsset.h"
 #include "RigEditor/IKRigController.h"
 #include "RigEditor/IKRigEditorController.h"
+#include "Engine/SkeletalMesh.h"
 
 #define LOCTEXT_NAMESPACE "IKRigAssetBrowser"
 
@@ -142,8 +143,15 @@ bool SIKRigAssetBrowser::OnShouldFilterAsset(const struct FAssetData& AssetData)
 		return true;
 	}
 
+	// get skeletal mesh
+	USkeletalMesh* SkeletalMesh = Controller->AssetController->GetSkeletalMesh();
+	if (!SkeletalMesh)
+	{
+		return true;
+	}
+	
 	// get skeleton
-	USkeleton* DesiredSkeleton = Controller->AssetController->GetSkeleton();;
+	USkeleton* DesiredSkeleton = SkeletalMesh->GetSkeleton();
 	if (!DesiredSkeleton)
 	{
 		return true;
diff --git a/Engine/Plugins/Animation/IKRig/Source/IKRigEditor/Private/RigEditor/SIKRigHierarchy.cpp b/Engine/Plugins/Animation/IKRig/Source/IKRigEditor/Private/RigEditor/SIKRigHierarchy.cpp
--- a/Engine/Plugins/Animation/IKRig/Source/IKRigEditor/Private/RigEditor/SIKRigHierarchy.cpp
+++ b/Engine/Plugins/Animation/IKRig/Source/IKRigEditor/Private/RigEditor/SIKRigHierarchy.cpp
@@ -64,7 +64,7 @@ TWeakObjectPtr< UObject > FIKRigTreeElement::GetObject() const
 	case IKRigTreeElementType::GOAL:
 		return Controller->AssetController->GetGoal(GoalName);
 	case IKRigTreeElementType::SOLVERGOAL:
-		if (const UIKRigSolver* SolverWithEffector = Controller->AssetController->GetSolver(EffectorIndex))
+		if (const UIKRigSolver* SolverWithEffector = Controller->AssetController->GetSolverAtIndex(EffectorIndex))
 		{
 			return SolverWithEffector->GetGoalSettings(EffectorGoalName);
 		}
@@ -589,7 +589,7 @@ void SIKRigHierarchy::ReplaceItemInSelection(const FText& OldName, const FText&
 	}
 }
 
-void SIKRigHierarchy::GetSelectedBoneChains(TArray<FBoneChain>& OutChains)
+void SIKRigHierarchy::GetSelectedBoneChains(TArray<FBoneChain>& OutChains) const
 {
 	const TSharedPtr<FIKRigEditorController> Controller = EditorController.Pin();
 	if (!Controller.IsValid())
@@ -618,9 +618,10 @@ void SIKRigHierarchy::GetSelectedBoneChains(TArray<FBoneChain>& OutChains)
 	// add goals (if there are any on the end bone of the chain)
 	for (FBoneChain& Chain : OutChains)
 	{
-		if (const UIKRigEffectorGoal* Goal = Controller->AssetController->GetGoalForBone(Chain.EndBone.BoneName))
+		const FName GoalName = Controller->AssetController->GetGoalNameForBone(Chain.EndBone.BoneName);
+		if (GoalName != NAME_None)
 		{
-			Chain.IKGoalName = Goal->GoalName;
+			Chain.IKGoalName = GoalName;
 		}
 	}
 }
@@ -645,7 +646,7 @@ bool SIKRigHierarchy::IsElementConnectedToSolver(TWeakPtr<FIKRigTreeElement> InT
 		return false; // not a valid solver index
 	}
 
-	const UIKRigSolver* Solver = AssetController->GetSolver(SolverIndex);
+	const UIKRigSolver* Solver = AssetController->GetSolverAtIndex(SolverIndex);
 	if (TreeElement->ElementType == IKRigTreeElementType::BONE)
 	{
 		// is this bone affected by this solver?
@@ -959,18 +960,17 @@ void SIKRigHierarchy::ConnectSelectedGoalsToSelectedSolvers(bool bConnect)
 	for (const TSharedPtr<FIKRigTreeElement>& GoalElement : SelectedGoals)
 	{
 		const FName GoalName = GoalElement->GoalName;
-		const int32 GoalIndex = AssetController->GetGoalIndex(GoalName);
-		check(GoalIndex != INDEX_NONE);
-		const UIKRigEffectorGoal& EffectorGoal = *AssetController->GetGoal(GoalIndex);
+		const UIKRigEffectorGoal* Goal = AssetController->GetGoal(GoalName);
+		check(Goal);
 		for (const TSharedPtr<FSolverStackElement>& SolverElement : SelectedSolvers)
 		{
 			if (bConnect)
 			{
-				AssetController->ConnectGoalToSolver(EffectorGoal, SolverElement->IndexInStack);	
+				AssetController->ConnectGoalToSolver(Goal->GoalName, SolverElement->IndexInStack);	
 			}
 			else
 			{
-				AssetController->DisconnectGoalFromSolver(EffectorGoal.GoalName, SolverElement->IndexInStack);	
+				AssetController->DisconnectGoalFromSolver(Goal->GoalName, SolverElement->IndexInStack);	
 			}
 		}
 	}
@@ -1061,7 +1061,7 @@ bool SIKRigHierarchy::CanSetRootBoneOnSolvers()
 	Controller->GetSelectedSolvers(SelectedSolvers);
 	for (const TSharedPtr<FSolverStackElement>& Solver : SelectedSolvers)
 	{
-		if (AssetController->GetSolver(Solver->IndexInStack)->RequiresRootBone())
+		if (AssetController->GetSolverAtIndex(Solver->IndexInStack)->RequiresRootBone())
 		{
 			return true;
 		}
@@ -1123,7 +1123,7 @@ bool SIKRigHierarchy::CanSetEndBoneOnSolvers() const
 	Controller->GetSelectedSolvers(SelectedSolvers);
 	for (const TSharedPtr<FSolverStackElement>& Solver : SelectedSolvers)
 	{
-		if (AssetController->GetSolver(Solver->IndexInStack)->RequiresEndBone())
+		if (AssetController->GetSolverAtIndex(Solver->IndexInStack)->RequiresEndBone())
 		{
 			return true;
 		}
diff --git a/Engine/Plugins/Animation/IKRig/Source/IKRigEditor/Private/RigEditor/SIKRigRetargetChainList.cpp b/Engine/Plugins/Animation/IKRig/Source/IKRigEditor/Private/RigEditor/SIKRigRetargetChainList.cpp
--- a/Engine/Plugins/Animation/IKRig/Source/IKRigEditor/Private/RigEditor/SIKRigRetargetChainList.cpp
+++ b/Engine/Plugins/Animation/IKRig/Source/IKRigEditor/Private/RigEditor/SIKRigRetargetChainList.cpp
@@ -288,7 +288,7 @@ const FReferenceSkeleton& SIKRigRetargetChainRow::GetReferenceSkeleton() const
 		return DummySkeleton; 
 	}
 
-	USkeletalMesh* SkeletalMesh = Controller->AssetController->GetAsset()->GetPreviewMesh();
+	USkeletalMesh* SkeletalMesh = Controller->AssetController->GetSkeletalMesh();
 	if (SkeletalMesh == nullptr)
 	{
 		return DummySkeleton;
@@ -747,15 +747,14 @@ void SIKRigRetargetChainList::MirrorSelectedChains() const
 		MirroredChain.ChainName = UMirrorDataTable::GetSettingsMirrorName(MirroredChain.ChainName);
 		MirroredChain.StartBone = IKRigSkeleton.BoneNames[MirroredIndices[0]];
 		MirroredChain.EndBone = IKRigSkeleton.BoneNames[MirroredIndices.Last()];
-		const UIKRigEffectorGoal* GoalOnMirroredBone = AssetController->GetGoalForBone(MirroredChain.EndBone.BoneName);
-		if (GoalOnMirroredBone)
+		const FName GoalOnMirroredBone = AssetController->GetGoalNameForBone(MirroredChain.EndBone.BoneName);
+		if (GoalOnMirroredBone != NAME_None)
 		{
-			MirroredChain.IKGoalName = GoalOnMirroredBone->GoalName;
+			MirroredChain.IKGoalName = GoalOnMirroredBone;
 		}
 
 		const FName NewChainName = Controller.PromptToAddNewRetargetChain(MirroredChain);
 		const FBoneChain* NewChain = AssetController->GetRetargetChainByName(NewChainName);
-		
 		if (!NewChain)
 		{
 			// user cancelled mirroring the chain
@@ -764,17 +763,17 @@ void SIKRigRetargetChainList::MirrorSelectedChains() const
 		
 		// check old bone has a goal, and the new bone also has a goal
 		// so we can connect the new goal to the same solver(s)
-		const UIKRigEffectorGoal* GoalOnOldChain = AssetController->GetGoalForBone(Chain->EndBone.BoneName);
+		const FName GoalOnOldChainName = AssetController->GetGoalNameForBone(Chain->EndBone.BoneName);
 		const UIKRigEffectorGoal* GoalOnNewChain = AssetController->GetGoal(NewChain->IKGoalName);
-		if (GoalOnOldChain && GoalOnNewChain)
+		if (GoalOnOldChainName != NAME_None && GoalOnNewChain)
 		{
 			// connect to the same solvers
 			const TArray<UIKRigSolver*>& AllSolvers =  AssetController->GetSolverArray();
 			for (int32 SolverIndex=0; SolverIndex<AllSolvers.Num(); ++SolverIndex)
 			{
-				if (AssetController->IsGoalConnectedToSolver(GoalOnOldChain->GoalName, SolverIndex))
+				if (AssetController->IsGoalConnectedToSolver(GoalOnOldChainName, SolverIndex))
 				{
-					AssetController->ConnectGoalToSolver(*GoalOnNewChain, SolverIndex);
+					AssetController->ConnectGoalToSolver(GoalOnNewChain->GoalName, SolverIndex);
 				}
 				else
 				{
diff --git a/Engine/Plugins/Animation/IKRig/Source/IKRigEditor/Private/RigEditor/SIKRigSolverStack.cpp b/Engine/Plugins/Animation/IKRig/Source/IKRigEditor/Private/RigEditor/SIKRigSolverStack.cpp
--- a/Engine/Plugins/Animation/IKRig/Source/IKRigEditor/Private/RigEditor/SIKRigSolverStack.cpp
+++ b/Engine/Plugins/Animation/IKRig/Source/IKRigEditor/Private/RigEditor/SIKRigSolverStack.cpp
@@ -75,12 +75,11 @@ void SIKRigSolverStackItem::Construct(
 					.IsChecked_Lambda([InSolverStack, InStackElement]() -> ECheckBoxState
 					{
 						bool bEnabled = true;
-						if (InSolverStack.IsValid() && InSolverStack->EditorController.IsValid())
+						if (InSolverStack.IsValid() &&
+							InSolverStack->EditorController.IsValid() &&
+							InSolverStack->EditorController.Pin().IsValid())
 						{
-							if (const UIKRigSolver* Solver = InSolverStack->EditorController.Pin()->AssetController->GetSolver(InStackElement->IndexInStack))
-							{
-								bEnabled = Solver->IsEnabled();
-							}
+							bEnabled = InSolverStack->EditorController.Pin()->AssetController->GetSolverEnabled(InStackElement->IndexInStack);
 						}
 						return bEnabled ? ECheckBoxState::Checked : ECheckBoxState::Unchecked;
 					})
@@ -188,7 +187,7 @@ UIKRigSolver* SIKRigSolverStackItem::GetSolver() const
 		return nullptr;
 	}
 	const int32 SolverIndex = StackElement.Pin()->IndexInStack;
-	return SolverStack.Pin()->EditorController.Pin()->AssetController->GetSolver(SolverIndex);
+	return SolverStack.Pin()->EditorController.Pin()->AssetController->GetSolverAtIndex(SolverIndex);
 }
 
 TSharedRef<FIKRigSolverStackDragDropOp> FIKRigSolverStackDragDropOp::New(TWeakPtr<FSolverStackElement> InElement)
@@ -392,7 +391,7 @@ void SIKRigSolverStack::RefreshStackView()
 	const int32 NumSolvers = AssetController->GetNumSolvers();
 	for (int32 i=0; i<NumSolvers; ++i)
 	{
-		const UIKRigSolver* Solver = AssetController->GetSolver(i);
+		const UIKRigSolver* Solver = AssetController->GetSolverAtIndex(i);
 		const FText DisplayName = Solver ? FText::FromString(AssetController->GetSolverUniqueName(i)) : UnknownSolverTxt;
 		TSharedPtr<FSolverStackElement> SolverItem = FSolverStackElement::Make(DisplayName, i);
 		ListViewItems.Add(SolverItem);
diff --git a/Engine/Plugins/Animation/IKRig/Source/IKRigEditor/Public/RetargetEditor/IKRetargetBatchOperation.h b/Engine/Plugins/Animation/IKRig/Source/IKRigEditor/Public/RetargetEditor/IKRetargetBatchOperation.h
--- a/Engine/Plugins/Animation/IKRig/Source/IKRigEditor/Public/RetargetEditor/IKRetargetBatchOperation.h
+++ b/Engine/Plugins/Animation/IKRig/Source/IKRigEditor/Public/RetargetEditor/IKRetargetBatchOperation.h
@@ -5,6 +5,8 @@
 #include "CoreMinimal.h"
 #include "EditorAnimUtils.h"
 
+#include "IKRetargetBatchOperation.generated.h"
+
 class UIKRetargeter;
 struct FScopedSlowTask;
 
@@ -54,16 +56,47 @@ public:
 };
 
 //** Encapsulate ability to batch duplicate and retarget a set of animation assets */
-struct FIKRetargetBatchOperation
+UCLASS(BlueprintType)
+class IKRIGEDITOR_API UIKRetargetBatchOperation : public UObject
 {
+	GENERATED_BODY()
 	
 public:
 
+	/* Convenience function to run a batch animation retarget from Blueprint / Python. This function will duplicate a list of
+	 * assets and use the supplied IK Retargeter to retarget the animation from the source to the target using the
+	 * settings in the provided IK Retargeter asset.
+	 * 
+	 * @param AssetsToRetarget A list of animation assets to retarget (sequences, blendspaces or montages)
+	 * @param SourceMesh The skeletal mesh with desired proportions to playback the assets to retarget
+	 * @param TargetMesh The skeletal mesh to retarget the animation onto.
+	 * @param IKRetargetAsset The IK Retargeter asset with IK Rigs appropriate for the source and target skeletal mesh
+	 * @param Search A string to search for in the file name (replaced with "Replace" string)
+	 * @param Replace A string to replace with in the file name
+	 * @param Suffix A string to add at the end of the new file name
+	 * @param Prefix A string to add to the start of the new file name
+	 * @param bRemapReferencedAssets Whether to remap existing references in the animation assets
+	 * 
+	 * Return: list of new animation files created. */
+	UFUNCTION(BlueprintCallable, Category=IKBatchRetarget)
+	TArray<FAssetData> RunRetarget(
+		const TArray<FAssetData>& AssetsToRetarget,
+		USkeletalMesh* SourceMesh,
+		USkeletalMesh* TargetMesh,
+		UIKRetargeter* IKRetargetAsset,
+		const FString& Search = "",
+		const FString& Replace = "",
+		const FString& Prefix = "",
+		const FString& Suffix = "",
+		const bool bRemapReferencedAssets=true);
+	
 	/* Actually run the process to duplicate and retarget the assets for the given context */
 	void RunRetarget(FIKRetargetBatchOperationContext& Context);
 
 private:
 
+	void Reset();
+
 	/**
 	* Initialize set of referenced assets to retarget.
 	* @return	Number of assets that need retargeting.
@@ -108,7 +141,4 @@ private:
 	TMap<UAnimBlueprint*, UAnimBlueprint*>		DuplicatedBlueprints;
 
 	TMap<UAnimationAsset*, UAnimationAsset*>	RemappedAnimAssets;
-
-	/** If we only chose one object to retarget store it here */
-	UObject* SingleTargetObject = nullptr;
 };
diff --git a/Engine/Plugins/Animation/IKRig/Source/IKRigEditor/Public/RetargetEditor/IKRetargetDetails.h b/Engine/Plugins/Animation/IKRig/Source/IKRigEditor/Public/RetargetEditor/IKRetargetDetails.h
--- a/Engine/Plugins/Animation/IKRig/Source/IKRigEditor/Public/RetargetEditor/IKRetargetDetails.h
+++ b/Engine/Plugins/Animation/IKRig/Source/IKRigEditor/Public/RetargetEditor/IKRetargetDetails.h
@@ -4,7 +4,7 @@
 
 #include "CoreMinimal.h"
 #include "IDetailCustomization.h"
-#include "IKRetargeterController.h"
+#include "RetargetEditor/IKRetargeterController.h"
 #include "Retargeter/IKRetargetProcessor.h"
 
 #include "IKRetargetDetails.generated.h"
diff --git a/Engine/Plugins/Animation/IKRig/Source/IKRigEditor/Public/RetargetEditor/IKRetargetEditorController.h b/Engine/Plugins/Animation/IKRig/Source/IKRigEditor/Public/RetargetEditor/IKRetargetEditorController.h
--- a/Engine/Plugins/Animation/IKRig/Source/IKRigEditor/Public/RetargetEditor/IKRetargetEditorController.h
+++ b/Engine/Plugins/Animation/IKRig/Source/IKRigEditor/Public/RetargetEditor/IKRetargetEditorController.h
@@ -101,7 +101,7 @@ public:
 	/** callback when IK Rig asset's retarget chain has been renamed */
 	void HandleRetargetChainRenamed(UIKRigDefinition* ModifiedIKRig, FName OldName, FName NewName) const;
 	/** callback when IK Rig asset's retarget chain has been removed */
-	void HandleRetargetChainRemoved(UIKRigDefinition* ModifiedIKRig, const FName& InChainRemoved) const;
+	void HandleRetargetChainRemoved(UIKRigDefinition* ModifiedIKRig, const FName InChainRemoved) const;
 	/** callback when IK Retargeter asset requires reinitialization */
 	void HandleRetargeterNeedsInitialized() const;
 	FDelegateHandle RetargeterReInitDelegateHandle;
@@ -132,9 +132,9 @@ public:
 	/** viewport anim instance */
 	UIKRetargetAnimInstance* GetAnimInstance(const ERetargetSourceOrTarget SourceOrTarget) const;
 	UPROPERTY(transient, NonTransactional)
-	TWeakObjectPtr<class UIKRetargetAnimInstance> SourceAnimInstance;
+	TObjectPtr<UIKRetargetAnimInstance> SourceAnimInstance;
 	UPROPERTY(transient, NonTransactional)
-	TWeakObjectPtr<class UIKRetargetAnimInstance> TargetAnimInstance;
+	TObjectPtr<UIKRetargetAnimInstance> TargetAnimInstance;
 	
 	/** store pointers to various tabs of UI,
 	 * have to manage access to these because they can be null if the tabs are closed */
@@ -322,19 +322,19 @@ private:
 	TSharedPtr<SWindow> IKRigPickerWindow;
 
 	/** the current output mode of the retargeter */
-	ERetargeterOutputMode OutputMode;
+	ERetargeterOutputMode OutputMode = ERetargeterOutputMode::RunRetarget;
 	/** slider value to blend between reference pose and retarget pose */
 	float RetargetPosePreviewBlend = 1.0f;
 	
 	/** which skeleton are we editing / viewing? */
-	ERetargetSourceOrTarget CurrentlyEditingSourceOrTarget;
+	ERetargetSourceOrTarget CurrentlyEditingSourceOrTarget = ERetargetSourceOrTarget::Target;
 
 	/** current selection set */
-	bool bIsRootSelected;
-	UPrimitiveComponent* SelectedMesh;
+	bool bIsRootSelected = false;
+	UPrimitiveComponent* SelectedMesh = nullptr;
 	TArray<FName> SelectedChains;
 	TMap<ERetargetSourceOrTarget, TArray<FName>> SelectedBoneNames;
-	ERetargetSelectionType LastSelectedItem;
+	ERetargetSelectionType LastSelectedItem = ERetargetSelectionType::NONE;
 	UPROPERTY()
 	TMap<FName,TObjectPtr<UIKRetargetBoneDetails>> AllBoneDetails;
 
diff --git a/Engine/Plugins/Animation/IKRig/Source/IKRigEditor/Public/RetargetEditor/IKRetargeterController.h b/Engine/Plugins/Animation/IKRig/Source/IKRigEditor/Public/RetargetEditor/IKRetargeterController.h
--- a/Engine/Plugins/Animation/IKRig/Source/IKRigEditor/Public/RetargetEditor/IKRetargeterController.h
+++ b/Engine/Plugins/Animation/IKRig/Source/IKRigEditor/Public/RetargetEditor/IKRetargeterController.h
@@ -2,7 +2,6 @@
 #pragma once
 
 #include "CoreMinimal.h"
-#include "IKRetargetAnimInstanceProxy.h"
 #include "UObject/Object.h"
 
 #include "IKRetargeterController.generated.h"
@@ -15,6 +14,7 @@ class UIKRigDefinition;
 class UIKRetargeter;
 class USkeletalMesh;
 
+UENUM()
 enum class EAutoMapChainType : uint8
 {
 	Exact, // map chains that have exactly the same name (case insensitive)
@@ -22,164 +22,246 @@ enum class EAutoMapChainType : uint8
 	Clear, // clear all mappings, set them to None
 };
 
-/** A singleton (per-asset) class used to make modifications to a UIKRetargeter asset.
-* Call the static GetController(UIKRetargeter* Asset) function to get the controller for the asset you want to modify. */ 
+// A stateless singleton (1-per-asset) class used to make modifications to a UIKRetargeter asset.
+// Use UIKRetargeter.GetController() to get the controller for the asset you want to modify.  
 UCLASS(config = Engine, hidecategories = UObject)
 class IKRIGEDITOR_API UIKRetargeterController : public UObject
 {
 	GENERATED_BODY()
 
 public:
+	
+	// Get access to the retargeter asset.
+	// Warning: Do not make modifications to the asset directly. Use this API instead. 
+	UIKRetargeter* GetAsset() const;
+	
+private:
+	
+	// The actual asset that this Controller modifies. This is the only field this class should have.
+	UIKRetargeter* Asset = nullptr;
+
+	//
+	// GENERAL PUBLIC/SCRIPTING API
+	//
 
+public:
+	
 	/** Use this to get the controller for the given retargeter asset */
+	UFUNCTION(BlueprintCallable, Category=IKRetargeter)
 	static UIKRetargeterController* GetController(UIKRetargeter* InRetargeterAsset);
-	/** Get access to the retargeter asset.
-	 *@warning Do not make modifications to the asset directly. Use this API instead. */
-	UIKRetargeter* GetAsset() const;
 	
-
-	/** SOURCE / TARGET
-	* 
-	*/
-	/** Set the IK Rig to use as the source or target (to copy animation FROM/TO) */
+	// Set the IK Rig to use as the source or target (to copy animation FROM/TO) 
+	UFUNCTION(BlueprintCallable, Category=IKRetargeter)
 	void SetIKRig(const ERetargetSourceOrTarget SourceOrTarget, UIKRigDefinition* IKRig) const;
-	/** Get either source or target IK Rig */
+	
+	// Get either source or target IK Rig 
+	UFUNCTION(BlueprintCallable, Category=IKRetargeter)
 	const UIKRigDefinition* GetIKRig(const ERetargetSourceOrTarget SourceOrTarget) const;
-	/** Get either source or target IK Rig */
-	UIKRigDefinition* GetIKRigWriteable(const ERetargetSourceOrTarget SourceOrTarget) const;
-	/** Set the preview skeletal mesh for either source or target */
-	void SetPreviewMesh(const ERetargetSourceOrTarget SourceOrTarget, USkeletalMesh* PreviewMesh);
-	/** Get the preview skeletal mesh */
+
+	// Set the preview skeletal mesh for either source or target
+	UFUNCTION(BlueprintCallable, Category=IKRetargeter)
+	void SetPreviewMesh(const ERetargetSourceOrTarget SourceOrTarget, USkeletalMesh* InPreviewMesh) const;
+	
+	// Get the preview skeletal mesh
+	UFUNCTION(BlueprintCallable, Category=IKRetargeter)
 	USkeletalMesh* GetPreviewMesh(const ERetargetSourceOrTarget SourceOrTarget) const;
-	/** Get if we've already asked to fix the root height for the given skeletal mesh */
+
+	//
+	// GENERAL C++ ONLY API
+	//
+
+	// Ensures all internal data is compatible with assigned meshes and ready to edit.
+	// - Removes bones from retarget poses that are no longer in skeleton
+	// - Removes chain settings for chains that are no longer in target IK Rig
+	void CleanAsset() const;
+	
+	// Get either source or target IK Rig 
+	UIKRigDefinition* GetIKRigWriteable(const ERetargetSourceOrTarget SourceOrTarget) const;
+	
+	// Get if we've already asked to fix the root height for the given skeletal mesh 
 	bool GetAskedToFixRootHeightForMesh(USkeletalMesh* Mesh) const;
-	/** Set if we've asked to fix the root height for the given skeletal mesh */
+	
+	// Set if we've asked to fix the root height for the given skeletal mesh 
 	void SetAskedToFixRootHeightForMesh(USkeletalMesh* Mesh, bool InAsked) const;
-	/** Get name of the Root bone used for retargeting. */
+	
+	// Get name of the Root bone used for retargeting. 
 	FName GetRetargetRootBone(const ERetargetSourceOrTarget SourceOrTarget) const;
 
-	/** RETARGET CHAIN MAPPING
-	* 
-	*/
-	/** Get names of all the bone chains. */
-	void GetChainNames(const ERetargetSourceOrTarget SourceOrTarget, TArray<FName>& OutNames) const;
-	/** Remove invalid chain mappings (no longer existing in currently referenced source/target IK Rig assets)
-	 * NOTE: be sure to reinitialize any running processors after chain mappings are cleaned! */
-	void CleanChainMapping() const;
-	/** Use fuzzy string search to find "best" Source chain to map to each Target chain */
+	//
+	// RETARGET CHAIN MAPPING PUBLIC/SCRIPTING API
+	//
+	
+	// Use fuzzy string search to find "best" Source chain to map to each Target chain
+	UFUNCTION(BlueprintCallable, Category=IKRetargeter)
 	void AutoMapChains(const EAutoMapChainType AutoMapType, const bool bForceRemap) const;
-	/** Callback when IK Rig chain is added or removed. */
-	void OnRetargetChainAdded(UIKRigDefinition* IKRig) const;
-	/** Callback when IK Rig chain is renamed. Retains existing mappings using the new name */
-	void OnRetargetChainRenamed(UIKRigDefinition* IKRig, FName OldChainName, FName NewChainName) const;
-	/** Callback when IK Rig chain is removed. */
-	void OnRetargetChainRemoved(UIKRigDefinition* IKRig, const FName& InChainRemoved) const;
-	/** Set the source chain to map to a given target chain */
-	void SetSourceChainForTargetChain(URetargetChainSettings* ChainMap, FName SourceChainToMapTo) const;
-	/** Get read-only access to the list of chain mappings */
-	const TArray<TObjectPtr<URetargetChainSettings>>& GetChainMappings() const;
-	/** Get read-only acces to the chain settings associated with a given target chain */
-	const URetargetChainSettings* GetChainMappingByTargetChainName(const FName& TargetChainName) const;
-	/** Get the name of the goal assigned to a chain */
+
+	// Assign a source chain to the given target chain. Animation will be copied from the source to the target.
+	UFUNCTION(BlueprintCallable, Category=IKRetargeter)
+	bool SetSourceChain(FName SourceChainName, FName TargetChainName) const;
+
+	// Get the name of the source chain mapped to a given target chain (the chain animation is copied FROM).
+	UFUNCTION(BlueprintCallable, Category=IKRetargeter)
+	FName GetSourceChain(const FName& TargetChainName) const;
+	
+	// Get read-only access to the list of settings for each target chain
+	UFUNCTION(BlueprintCallable, Category=IKRetargeter)
+	const TArray<URetargetChainSettings*>& GetAllChainSettings() const;
+
+	//
+	// RETARGET CHAIN MAPPING C++ ONLY API
+	//
+	
+	// Get the name of the goal assigned to a chain 
 	FName GetChainGoal(const TObjectPtr<URetargetChainSettings> ChainSettings) const;
-	/** Get whether the given chain's IK goal is connected to a solver */
+	
+	// Get whether the given chain's IK goal is connected to a solver 
 	bool IsChainGoalConnectedToASolver(const FName& GoalName) const;
-	/** END RETARGET CHAIN MAPPING */
-
-	/** RETARGET POSE EDITING
-	 * 
-	 */
-	/** Remove bones from retarget poses that are no longer in skeleton */
-	void CleanPoseLists() const;
-	/** Remove bones from retarget poses that are no longer in skeleton */
-	void CleanPoseList(const ERetargetSourceOrTarget SourceOrTarget) const;
-	/** Add new retarget pose. */
-	void AddRetargetPose(
-		const FName& NewPoseName,
-		const FIKRetargetPose* ToDuplicate,
-		const ERetargetSourceOrTarget SourceOrTarget) const;
-	/** Rename current retarget pose. */
-	void RenameCurrentRetargetPose(
-		const FName& NewPoseName,
-		const ERetargetSourceOrTarget SourceOrTarget) const;
-	/** Remove a retarget pose. */
-	void RemoveRetargetPose(
-	const FName& PoseToRemove, 
-	const ERetargetSourceOrTarget SourceOrTarget) const;
-	/** Reset a retarget pose for the specified bones.
-	 *If BonesToReset is Empty, will removes all stored deltas, returning pose to reference pose */
-	void ResetRetargetPose(
-		const FName& PoseToReset,
-		const TArray<FName>& BonesToReset,
-		const ERetargetSourceOrTarget SourceOrTarget) const;
-	/** Get the current retarget pose */
+	
+	// Call this when IK Rig chain is added or removed. 
+	void HandleRetargetChainAdded(UIKRigDefinition* IKRig) const;
+	
+	// Call this when IK Rig chain is renamed. Retains existing mappings using the new name 
+	void HandleRetargetChainRenamed(UIKRigDefinition* IKRig, FName OldChainName, FName NewChainName) const;
+	
+	// Call this when IK Rig chain is removed. 
+	void HandleRetargetChainRemoved(UIKRigDefinition* IKRig, const FName& InChainRemoved) const;
+	
+private:
+
+	// Remove invalid chain mappings (no longer existing in currently referenced source/target IK Rig assets)
+	// NOTE: be sure to reinitialize any running processors after chain mappings are cleaned!
+	void CleanChainMapping() const;
+	
+	// Get names of all the bone chains.
+	void GetChainNames(const ERetargetSourceOrTarget SourceOrTarget, TArray<FName>& OutNames) const;
+	
+	// convenience to get chain settings by name
+	URetargetChainSettings* GetChainSettings(const FName& TargetChainName) const;
+
+	// Sort the Asset ChainMapping based on the StartBone of the target chains. 
+	void SortChainMapping() const;
+
+	//
+	// RETARGET POSE PUBLIC/SCRIPTING API
+	//
+
+public:
+	
+	// Add new retarget pose. Returns the name of the new retarget pose.
+	UFUNCTION(BlueprintCallable, Category=IKRetargeter)
+	FName CreateRetargetPose(const FName& NewPoseName, const ERetargetSourceOrTarget SourceOrTarget) const;
+
+	// Remove a retarget pose. Returns true if the pose was found and removed.
+	UFUNCTION(BlueprintCallable, Category=IKRetargeter)
+	bool RemoveRetargetPose(const FName& PoseToRemove, const ERetargetSourceOrTarget SourceOrTarget) const;
+
+	// Duplicate a retarget pose. Returns the name of the new, duplicate pose (or NAME_None if PoseToDuplicate is not found).
+	UFUNCTION(BlueprintCallable, Category=IKRetargeter)
+	FName DuplicateRetargetPose(const FName PoseToDuplicate, const FName NewName, const ERetargetSourceOrTarget SourceOrTarget) const;
+	
+	// Rename current retarget pose. Returns true if a pose was renamed.
+	UFUNCTION(BlueprintCallable, Category=IKRetargeter)
+	bool RenameRetargetPose(const FName OldPoseName, const FName NewPoseName, const ERetargetSourceOrTarget SourceOrTarget) const;
+
+	// Change which retarget pose is used by the retargeter at runtime
+	UFUNCTION(BlueprintCallable, Category=IKRetargeter)
+	bool SetCurrentRetargetPose(FName CurrentPose, const ERetargetSourceOrTarget SourceOrTarget) const;
+	
+	// Get the current retarget pose
+	UFUNCTION(BlueprintCallable, Category=IKRetargeter)
     FName GetCurrentRetargetPoseName(const ERetargetSourceOrTarget SourceOrTarget) const;
-	/** Change which retarget pose is used by the retargeter at runtime */
-	void SetCurrentRetargetPose(FName CurrentPose, const ERetargetSourceOrTarget SourceOrTarget) const;
-	/** Get access to array of retarget poses */
+	
+	// Get access to array of retarget poses
+	UFUNCTION(BlueprintCallable, Category=IKRetargeter)
 	TMap<FName, FIKRetargetPose>& GetRetargetPoses(const ERetargetSourceOrTarget SourceOrTarget) const;
-	/** Get the current retarget pose */
+	
+	// Get the current retarget pose
+	UFUNCTION(BlueprintCallable, Category=IKRetargeter)
 	FIKRetargetPose& GetCurrentRetargetPose(const ERetargetSourceOrTarget SourceOrTarget) const;
-	/** Set a delta rotation for a given bone for the current retarget pose (used in Edit Mode in the retarget editor) */
+
+	// Reset a retarget pose for the specified bones.
+	// If BonesToReset is Empty, will removes all stored deltas, returning pose to reference pose
+	UFUNCTION(BlueprintCallable, Category=IKRetargeter)
+	void ResetRetargetPose(const FName& PoseToReset, const TArray<FName>& BonesToReset, const ERetargetSourceOrTarget SourceOrTarget) const;
+	
+	// Set a delta rotation for a given bone for the current retarget pose
+	UFUNCTION(BlueprintCallable, Category=IKRetargeter)
 	void SetRotationOffsetForRetargetPoseBone(
 		const FName& BoneName,
 		const FQuat& RotationOffset,
 		const ERetargetSourceOrTarget SkeletonMode) const;
-	/** Get a delta rotation for a given bone for the current retarget pose (used in Edit Mode in the retarget editor) */
+	
+	// Get a delta rotation for a given bone for the current retarget pose
+	UFUNCTION(BlueprintCallable, Category=IKRetargeter)
 	FQuat GetRotationOffsetForRetargetPoseBone(
 		const FName& BoneName,
 		const ERetargetSourceOrTarget SourceOrTarget) const;
-	/** Add a delta translation to the root bone (used in Edit Mode in the retarget editor) */
-	void AddTranslationOffsetToRetargetRootBone(
+	
+	// Set the translation offset on the retarget root bone for the current retarget pose
+	UFUNCTION(BlueprintCallable, Category=IKRetargeter)
+	void SetRootOffsetInRetargetPose(
 		const FVector& TranslationOffset,
 		const ERetargetSourceOrTarget SourceOrTarget) const;
-	/** Add a numbered suffix to the given pose name to make it unique. */
+
+	// Get the translation offset on the retarget root bone for the current retarget pose
+	UFUNCTION(BlueprintCallable, Category=IKRetargeter)
+	FVector GetRootOffsetInRetargetPose(const ERetargetSourceOrTarget SourceOrTarget) const;
+
+	//
+	// RETARGET POSE C++ ONLY API
+	//
+	
+	// Add a numbered suffix to the given pose name to make it unique. 
 	FName MakePoseNameUnique(
 		const FString& PoseName,
 		const ERetargetSourceOrTarget SourceOrTarget) const;
-	/** END RETARGET POSE EDITING */
+	
+private:
+	
+	// Remove bones from retarget poses that are no longer in skeleton 
+	void CleanPoseList(const ERetargetSourceOrTarget SourceOrTarget) const;
 
+
+	//
+	// DELEGATE CALLBACKS FOR OUTSIDE SYSTEMS
+	//
+	
 private:
 	
-	/** Called whenever the retargeter is modified in such a way that would require re-initialization by dependent systems.*/
 	DECLARE_MULTICAST_DELEGATE(FOnRetargeterNeedsInitialized);
 	FOnRetargeterNeedsInitialized RetargeterNeedsInitialized;
-
-	/** Called whenever either the source or target IK Rig asset's are swapped out.*/
+	
 	DECLARE_MULTICAST_DELEGATE_OneParam(FOnIKRigReplaced, ERetargetSourceOrTarget);
 	FOnIKRigReplaced IKRigReplaced;
 	
-	/** Called whenever either the source or target Preview Mesh asset's are swapped out.*/
 	DECLARE_MULTICAST_DELEGATE_OneParam(FOnPreviewMeshReplaced, ERetargetSourceOrTarget);
 	FOnPreviewMeshReplaced PreviewMeshReplaced;
-	
-public:
-	
-	FOnPreviewMeshReplaced& OnPreviewMeshReplaced(){ return PreviewMeshReplaced; };
+
 	void BroadcastPreviewMeshReplaced(ERetargetSourceOrTarget SourceOrTarget) const
 	{
 		PreviewMeshReplaced.Broadcast(SourceOrTarget);
 	}
 
-	FOnIKRigReplaced& OnIKRigReplaced(){ return IKRigReplaced; };
 	void BroadcastIKRigReplaced(ERetargetSourceOrTarget SourceOrTarget) const
 	{
 		IKRigReplaced.Broadcast(SourceOrTarget);
 	}
 	
-	FOnRetargeterNeedsInitialized& OnRetargeterNeedsInitialized(){ return RetargeterNeedsInitialized; };
 	void BroadcastNeedsReinitialized() const
 	{
 		RetargeterNeedsInitialized.Broadcast();
 	}
 	
-private:
-
-	URetargetChainSettings* GetChainMap(const FName& TargetChainName) const;
-
-	/** Sort the Asset ChainMapping based on the StartBone of the target chains. */
-	void SortChainMapping() const;
+public:
 
-	/** The actual asset that this Controller modifies. */
-	TObjectPtr<UIKRetargeter> Asset = nullptr;
+	// Attach a delegate to be notified whenever either the source or target Preview Mesh asset's are swapped out.
+	FOnPreviewMeshReplaced& OnPreviewMeshReplaced(){ return PreviewMeshReplaced; };
+	
+	// Attach a delegate to be notified whenever either the source or target IK Rig asset's are swapped out.
+	FOnIKRigReplaced& OnIKRigReplaced(){ return IKRigReplaced; };
+	
+	// Attach a delegate to be notified whenever the retargeter is modified in such a way that would require re-initialization of the processor.
+	FOnRetargeterNeedsInitialized& OnRetargeterNeedsInitialized(){ return RetargeterNeedsInitialized; };
+	
+	friend class UIKRetargeter;
 };
diff --git a/Engine/Plugins/Animation/IKRig/Source/IKRigEditor/Public/RigEditor/IKRigController.h b/Engine/Plugins/Animation/IKRig/Source/IKRigEditor/Public/RigEditor/IKRigController.h
--- a/Engine/Plugins/Animation/IKRig/Source/IKRigEditor/Public/RigEditor/IKRigController.h
+++ b/Engine/Plugins/Animation/IKRig/Source/IKRigEditor/Public/RigEditor/IKRigController.h
@@ -17,8 +17,14 @@ class USkeletalMesh;
 class USkeleton;
 struct FBoneChain;
 
-/** A singleton (per-asset) class used to make modifications to a UIKRigDefinition asset.
- * Call the static GetIKRigController() function to get the controller for the asset you want to modify. */ 
+// A singleton (per-asset) class used to make modifications to a UIKRigDefinition asset
+//
+// All modifications to an IKRigDefinition must go through this controller.
+//
+// Editors can subscribe to the callbacks on this controller to be notified of changes that require reinitialization
+// of a running IK Rig processor instance.
+// The API here is split into public/scripting sections which are accessible from Blueprint/Python and sections that
+// are only relevant to editors in C++.
 UCLASS(config = Engine, hidecategories = UObject)
 class IKRIGEDITOR_API UIKRigController : public UObject
 {
@@ -26,193 +32,330 @@ class IKRIGEDITOR_API UIKRigController : public UObject
 
 public:
 
-	/** Use this to get the controller for the given IKRig */
-	static UIKRigController* GetIKRigController(UIKRigDefinition* InIKRigDefinition);
+	//
+	// GENERAL PUBLIC/SCRIPTING API
+	//
 	
-	/** Get the asset this controller controls. */
-	UIKRigDefinition* GetAsset() const;
-	/** Get unique asset integer ID as a name */
-	FName GetAssetIDAsName() const;
-
-	/** SKELETON
-	 * 
-	 */
-	/** Sets the preview mesh to use. Loads the hierarchy into the asset's IKRigSkeleton.
-	 * Returns true if the mesh was able to be set. False if it was incompatible for any reason. */
+	// Use this to get the controller for the given IKRig
+	UFUNCTION(BlueprintCallable, Category=IKRig)
+	static UIKRigController* GetController(const UIKRigDefinition* InIKRigDefinition);
+	
+	// Sets the preview mesh to use. Loads the hierarchy into the asset's IKRigSkeleton.
+	// Returns true if the mesh was able to be set. False if it was incompatible for any reason. 
+	UFUNCTION(BlueprintCallable, Category=IKRig)
 	bool SetSkeletalMesh(USkeletalMesh* SkeletalMesh, bool bTransact=false) const;
-	/** Get read-access to the IKRig skeleton representation */
+
+	// Get the skeletal mesh this asset is initialized with 
+	UFUNCTION(BlueprintCallable, Category=IKRig)
+	USkeletalMesh* GetSkeletalMesh() const;
+
+	// Returns true if the provided skeletal mesh could be used with this IK Rig.
+	UFUNCTION(BlueprintCallable, Category=IKRig)
+	bool IsSkeletalMeshCompatible(USkeletalMesh* SkeletalMeshToCheck) const;
+
+	//
+	// GENERAL C++ ONLY API
+	//
+	
+	// Get the asset this controller controls.
+	// Warning: Do not use for general editing of the data model.
+	// If you need to make modifications to the asset, do so through the API provided by this controller. 
+	UIKRigDefinition* GetAsset() const;
+	
+	// force all currently connected processors to reinitialize using latest asset state 
+	void BroadcastNeedsReinitialized() const;
+	
+	//
+	// SKELETON C++ ONLY API
+	//
+	
+	// Get read-access to the IKRig skeleton representation 
 	const FIKRigSkeleton& GetIKRigSkeleton() const;
-	/** Get the USkeleton asset this rig was initialized with */
-	USkeleton* GetSkeleton() const;
-	/** Include/exclude a bone from all the solvers. All bones are included by default. */
-	void SetBoneExcluded(const FName& BoneName, const bool bExclude) const;
-	/** Returns true if the given Bone is excluded, false otherwise. */
-	bool GetBoneExcluded(const FName& BoneName) const;
-	/** Get the global-space retarget pose transform of the given Bone. Can only be called AFTER skeleton is initialized. */
-	FTransform GetRefPoseTransformOfBone(const FName& BoneName) const;
-	/** END SKELETON */
-
-	/** SOLVERS
-	 * 
-	 */
-	/** Add a new solver of the given type to the bottom of the stack. Returns the stack index. */
+
+	//
+	// SOLVERS PUBLIC/SCRIPTING API
+	//
+
+	// Add a new solver of the given type to the bottom of the stack. Returns the stack index. 
 	int32 AddSolver(TSubclassOf<UIKRigSolver> InSolverClass) const;
-	/** Remove the solver at the given stack index. */
-	void RemoveSolver(const int32 SolverIndex) const;
-	/** Move the solver at the given index to the target index. */
+
+	// Adda new solver, specifying the type as a fully qualified package path, ie ("/Script/IKRig/. 
+	UFUNCTION(BlueprintCallable, Category=IKRig)
+	int32 AddSolver(const FString& InSolverClassPath) const;
+
+	// Remove the solver at the given stack index. 
+	UFUNCTION(BlueprintCallable, Category=IKRig)
+	bool RemoveSolver(const int32 SolverIndex) const;
+
+	// Get access to the given solver. 
+	UFUNCTION(BlueprintCallable, Category=IKRig)
+	UIKRigSolver* GetSolverAtIndex(int32 Index) const;
+
+	// Get access to the given solver. 
+	UFUNCTION(BlueprintCallable, Category=IKRig)
+	int32 GetIndexOfSolver(UIKRigSolver* Solver) const;
+
+	// Get the number of solvers in the stack. 
+	UFUNCTION(BlueprintCallable, Category=IKRig)
+	int32 GetNumSolvers() const;
+
+	// Move the solver at the given index to the target index. 
+	UFUNCTION(BlueprintCallable, Category=IKRig)
 	bool MoveSolverInStack(int32 SolverToMoveIndex, int32 TargetSolverIndex) const;
-	/** Enabled/disable the given solver. */
+
+	// Set enabled/disabled status of the given solver. 
+	UFUNCTION(BlueprintCallable, Category=IKRig)
 	bool SetSolverEnabled(int32 SolverIndex, bool bIsEnabled) const;
-	/** Get access to the given solver. */
-	UIKRigSolver* GetSolver(int32 Index) const;
-	/** Get the number of solvers in the stack. */
-	int32 GetNumSolvers() const;
-	/** Set the root bone on a given solver. (not all solvers support root bones, checks CanSetRootBone() first) */
-	void SetRootBone(const FName& RootBoneName, int32 SolverIndex) const;
-	/** Set the end bone on a given solver. (not all solvers require extra end bones, checks CanSetEndBone() first) */
-	void SetEndBone(const FName& EndBoneName, int32 SolverIndex) const;
-	/** Get read-only access to the array of solvers. */
+
+	// Get enabled status of the given solver. 
+	UFUNCTION(BlueprintCallable, Category=IKRig)
+	bool GetSolverEnabled(int32 SolverIndex) const;
+
+	// Set the root bone on a given solver. (not all solvers support root bones, checks CanSetRootBone() first) 
+	UFUNCTION(BlueprintCallable, Category=IKRig)
+	bool SetRootBone(const FName RootBoneName, int32 SolverIndex) const;
+
+	// Get the name of the root bone on a given solver. (not all solvers support root bones, checks CanSetRootBone() first) 
+	UFUNCTION(BlueprintCallable, Category=IKRig)
+	FName GetRootBone(int32 SolverIndex) const;
+
+	// Set the end bone on a given solver. (not all solvers require extra end bones, checks CanSetEndBone() first) 
+	UFUNCTION(BlueprintCallable, Category=IKRig)
+	bool SetEndBone(const FName EndBoneName, int32 SolverIndex) const;
+
+	// Get the name of the end bone on a given solver. (not all solvers require extra end bones, checks CanSetEndBone() first) 
+	UFUNCTION(BlueprintCallable, Category=IKRig)
+	FName GetEndBone(int32 SolverIndex) const;
+
+	//
+	// SOLVERS C++ ONLY API
+	//
+	
+	// Get read-only access to the array of solvers. 
 	const TArray<UIKRigSolver*>& GetSolverArray() const;
-	/** Get unique label for a given solver. Returns dash separated index and name like so, "0 - SolverName". */
-	FString GetSolverUniqueName(int32 SolverIndex);
-	/** END SOLVERS */
-
-	/** GOALS
-	 * 
-	 */
-	/** Add a new Goal associated with the given Bone. GoalName must be unique. Bones can have multiple Goals (rare). */
-	UIKRigEffectorGoal* AddNewGoal(const FName& GoalName, const FName& BoneName) const;
-	/** Remove the Goal by name. */
-	bool RemoveGoal(const FName& GoalName) const;
-	/** Rename a Goal. Returns new name, which may be different after being sanitized. Returns NAME_None if this fails.*/
-	FName RenameGoal(const FName& OldName, const FName& PotentialNewName) const;
-	/** Modify a Goal for a transaction. Returns true if Goal found.*/
-	bool ModifyGoal(const FName& GoalName) const;
-	/** The the Bone that the given Goal should be parented to / associated with. */
-	bool SetGoalBone(const FName& GoalName, const FName& NewBoneName) const;
-	/** The the Bone associated with the given Goal. */
-	FName GetBoneForGoal(const FName& GoalName) const;
-	/** The the Goal associated with the given Bone (may be null. */
-	UIKRigEffectorGoal* GetGoalForBone(const FName& BoneName) const;
-	/** Connect the given Goal to the given Solver. This creates an "Effector" with settings specific to this Solver.*/
-	bool ConnectGoalToSolver(const UIKRigEffectorGoal& Goal, int32 SolverIndex) const;
-	/** Disconnect the given Goal from the given Solver. This removes the Effector that associates the Goal with the Solver.*/
-	bool DisconnectGoalFromSolver(const FName& GoalToRemove, int32 SolverIndex) const;
-	/** Returns true if the given Goal is connected to the given Solver. False otherwise. */
-	bool IsGoalConnectedToSolver(const FName& GoalName, int32 SolverIndex) const;
-	/** Returns true if the given Goal is connected to ANY solver. False otherwise. */
-	bool IsGoalConnectedToAnySolver(const FName& GoalName) const;
-	/** Get the index of the given Goal in the list of Goals. */
-	int32 GetGoalIndex(const FName& InGoalName, const ENameCase CompareMethod = ENameCase::IgnoreCase) const;
-	/** Get the name of Goal at the given index. */
-	FName GetGoalName(const int32& GoalIndex) const;
-	/** Get access to the list of Goals. */
+
+	// Get unique label for a given solver. Returns dash separated index and name like so, "0 - SolverName". 
+	FString GetSolverUniqueName(int32 SolverIndex) const;
+
+	//
+	// GOALS PUBLIC/SCRIPTING API
+	//
+	
+	// Add a new Goal associated with the given Bone. GoalName must be unique. Bones can have multiple Goals (rare). 
+	UFUNCTION(BlueprintCallable, Category=IKRig)
+	FName AddNewGoal(const FName GoalName, const FName BoneName) const;
+
+	// Remove the Goal by name. 
+	UFUNCTION(BlueprintCallable, Category=IKRig)
+	bool RemoveGoal(const FName GoalName) const;
+
+	// Rename a Goal. Returns new name, which may be different after being sanitized. Returns NAME_None if this fails.
+	UFUNCTION(BlueprintCallable, Category=IKRig)
+	FName RenameGoal(const FName OldName, const FName PotentialNewName) const;
+
+	// The the Bone that the given Goal should be parented to / associated with. 
+	UFUNCTION(BlueprintCallable, Category=IKRig)
+	bool SetGoalBone(const FName GoalName, const FName NewBoneName) const;
+
+	// The the Bone associated with the given Goal. 
+	UFUNCTION(BlueprintCallable, Category=IKRig)
+	FName GetBoneForGoal(const FName GoalName) const;
+
+	// Get the Goal associated with the given Bone (may be NAME_None) 
+	UFUNCTION(BlueprintCallable, Category=IKRig)
+	FName GetGoalNameForBone(const FName BoneName) const;
+
+	// Connect the given Goal to the given Solver. This creates an "Effector" with settings specific to this Solver.
+	UFUNCTION(BlueprintCallable, Category=IKRig)
+	bool ConnectGoalToSolver(const FName GoalName, int32 SolverIndex) const;
+
+	// Disconnect the given Goal from the given Solver. This removes the Effector that associates the Goal with the Solver.
+	UFUNCTION(BlueprintCallable, Category=IKRig)
+	bool DisconnectGoalFromSolver(const FName GoalToRemove, int32 SolverIndex) const;
+
+	// Returns true if the given Goal is connected to the given Solver. False otherwise. 
+	UFUNCTION(BlueprintCallable, Category=IKRig)
+	bool IsGoalConnectedToSolver(const FName GoalName, int32 SolverIndex) const;
+
+	// Returns true if the given Goal is connected to ANY solver. False otherwise. 
+	UFUNCTION(BlueprintCallable, Category=IKRig)
+	bool IsGoalConnectedToAnySolver(const FName GoalName) const;
+
+	// Get read-write access to the Goal with the given name.
+	UFUNCTION(BlueprintCallable, Category=IKRig)
+	UIKRigEffectorGoal* GetGoal(const FName GoalName) const;
+
+	// Get access to the list of Goals. 
+	UFUNCTION(BlueprintCallable, Category=IKRig)
 	TArray<UIKRigEffectorGoal*>& GetAllGoals() const;
-	/** Get read-only access to the Goal at the given index. */
-	const UIKRigEffectorGoal* GetGoal(int32 GoalIndex) const;
-	/** Get read-write access to the Goal with the given name. */
-	UIKRigEffectorGoal* GetGoal(const FName& GoalName) const;
-	/** Get the UObject for the settings associated with the given Goal in the given Solver.
-	 ** Solvers can define their own per-Goal settings depending on their needs. These are termed "Effectors". */
-	UObject* GetGoalSettingsForSolver(const FName& GoalName, int32 SolverIndex) const;
-	/** Get the global-space transform of the given Goal. This may be set by the user in the editor, or at runtime. */
-	FTransform GetGoalCurrentTransform(const FName& GoalName) const;
-	/** Set the Goal to the given transform. */
-	void SetGoalCurrentTransform(const FName& GoalName, const FTransform& Transform) const;
-	/** Reset all Goals back to their initial transforms. */
+	
+	//
+	// GOALS C++ ONLY API
+	//
+	
+	// Get the index of the given Goal in the list of Goals. 
+	int32 GetGoalIndex(const FName InGoalName, const ENameCase CompareMethod = ENameCase::IgnoreCase) const;
+
+	// Get the UObject for the settings associated with the given Goal in the given Solver.
+	// Solvers can define their own per-Goal settings depending on their needs. These are termed "Effectors". 
+	UObject* GetGoalSettingsForSolver(const FName GoalName, int32 SolverIndex) const;
+
+	// Get the global-space transform of the given Goal. This may be set by the user in the editor, or at runtime. 
+	FTransform GetGoalCurrentTransform(const FName GoalName) const;
+
+	// Set the Goal to the given transform. 
+	void SetGoalCurrentTransform(const FName GoalName, const FTransform& Transform) const;
+
+	// Reset all Goals back to their initial transforms. 
 	void ResetGoalTransforms() const;
-	/** Reset initial goal transform */
+
+	// Reset initial goal transform 
 	void ResetInitialGoalTransforms() const;
-	/** Ensure that the given name adheres to required standards for Goal names (no special characters etc..)*/
+
+	// Ensure that the given name adheres to required standards for Goal names (no special characters etc..)
 	static void SanitizeGoalName(FString& InOutName);
-	/** Add a suffix as needed to ensure the Goal name is unique */
-	FName GetUniqueGoalName(const FName& NameToMakeUnique) const;
-	
-	/** END Goals */
-
-	/** Bone Settings
-	 * 
-	 */
-	/** Add settings to the given Bone/Solver. Does nothing if Bone already has settings in this Solver.*/
-	void AddBoneSetting(const FName& BoneName, int32 SolverIndex) const;
-	/** Remove settings for the given Bone/Solver. Does nothing if Bone doesn't have setting in this Solver.*/
-	void RemoveBoneSetting(const FName& BoneName, int32 SolverIndex) const;
-	/** Returns true if this Bone can have settings in the given Solver. */
-	bool CanAddBoneSetting(const FName& BoneName, int32 SolverIndex) const;
-	/** Returns true if settings for this Bone can be removed from the given Solver. */
-	bool CanRemoveBoneSetting(const FName& BoneName, int32 SolverIndex) const;
-	/** Get the generic (Solver-specific) Bone settings UObject for this Bone in the given Solver.*/
-	UObject* GetSettingsForBone(const FName& BoneName, int32 SolverIndex) const;
-	/** Returns true if the given Bone has any settings in any Solver. */
-	bool DoesBoneHaveSettings(const FName& BoneName) const;
-	/** END Bone Settings */
-
-	/** Retargeting Options and Retarget Bone Chains
-	 * 
-	 */
-	/** Add a Chain with the given Name and Start/End bones. Returns newly created chain name. */
-	FName AddRetargetChain(const FBoneChain& BoneChain) const;
-	/** Remove a Chain with the given name. Returns true if a Chain was removed. */
-	bool RemoveRetargetChain(const FName& ChainName) const;
-	/** Renamed the given Chain. Returns the new name (same as old if unsuccessful). */
-	FName RenameRetargetChain(const FName& ChainName, const FName& NewChainName) const;
-	/** Set the Start Bone for the given Chain. Returns true if operation was successful. */
-	bool SetRetargetChainStartBone(const FName& ChainName, const FName& StartBoneName) const;
-	/** Set the End Bone for the given Chain. Returns true if operation was successful. */
-	bool SetRetargetChainEndBone(const FName& ChainName, const FName& EndBoneName) const;
-	/** Set the Goal for the given Chain. Returns true if operation was successful. */
-	bool SetRetargetChainGoal(const FName& ChainName, const FName& GoalName) const;
-	/** Get the Goal name for the given Chain. */
-	FName GetRetargetChainGoal(const FName& ChainName) const;
-	/** Get the End Bone name for the given Chain. */
-    FName GetRetargetChainStartBone(const FName& ChainName) const;
-    /** Get the Start Bone name for the given Chain. */
-    FName GetRetargetChainEndBone(const FName& ChainName) const;
-	/** Get read-only access to a single retarget chain with the given name */
-	const FBoneChain* GetRetargetChainByName(const FName& ChainName) const;
-	/** Get read-only access to the list of Chains. */
+
+	// Add a suffix as needed to ensure the Goal name is unique 
+	FName GetUniqueGoalName(const FName NameToMakeUnique) const;
+
+	// Modify a Goal for a transaction. Returns true if Goal found.
+	bool ModifyGoal(const FName GoalName) const;
+
+	//
+	// BONES PUBLIC/SCRIPTING API
+	//
+
+	// Include/exclude a bone from all the solvers. All bones are included by default. 
+	UFUNCTION(BlueprintCallable, Category=IKRig)
+	bool SetBoneExcluded(const FName BoneName, const bool bExclude) const;
+
+	// Returns true if the given Bone is excluded, false otherwise. 
+	UFUNCTION(BlueprintCallable, Category=IKRig)
+	bool GetBoneExcluded(const FName BoneName) const;
+	
+	// Add settings to the given Bone/Solver. Does nothing if Bone already has settings in this Solver.
+	UFUNCTION(BlueprintCallable, Category=IKRig)
+	bool AddBoneSetting(const FName BoneName, int32 SolverIndex) const;
+
+	// Remove settings for the given Bone/Solver. Does nothing if Bone doesn't have setting in this Solver.
+	UFUNCTION(BlueprintCallable, Category=IKRig)
+	bool RemoveBoneSetting(const FName BoneName, int32 SolverIndex) const;
+
+	// Get the generic (Solver-specific) Bone settings UObject for this Bone in the given Solver.
+	UFUNCTION(BlueprintCallable, Category=IKRig)
+	UObject* GetSettingsForBone(const FName BoneName, int32 SolverIndex) const;
+
+	// Get the global-space retarget pose transform of the given Bone.
+	UFUNCTION(BlueprintCallable, Category=IKRig)
+	FTransform GetRefPoseTransformOfBone(const FName BoneName) const;
+
+	//
+	// BONES C++ ONLY API
+	//
+
+	// Returns true if this Bone can have settings in the given Solver. 
+	bool CanAddBoneSetting(const FName BoneName, int32 SolverIndex) const;
+
+	// Returns true if settings for this Bone can be removed from the given Solver. 
+	bool CanRemoveBoneSetting(const FName BoneName, int32 SolverIndex) const;
+
+	// Returns true if the given Bone has any settings in any Solver. 
+	bool DoesBoneHaveSettings(const FName BoneName) const;
+
+	//
+	// RETARGETING PUBLIC/SCRIPTING API
+	//
+
+	// Add a new chain with the given Chain and Bone names. Returns newly created chain name (uniquified).
+	// Note: only the ChainName is required here, all else can be set later.
+	UFUNCTION(BlueprintCallable, Category=IKRig)
+	FName AddRetargetChain(const FName ChainName, const FName StartBoneName, const FName EndBoneName, const FName GoalName) const;
+
+	// Remove a Chain with the given name. Returns true if a Chain was removed. 
+	UFUNCTION(BlueprintCallable, Category=IKRig)
+	bool RemoveRetargetChain(const FName ChainName) const;
+
+	// Renamed the given Chain. Returns the new name (same as old if unsuccessful). 
+	UFUNCTION(BlueprintCallable, Category=IKRig)
+	FName RenameRetargetChain(const FName ChainName, const FName NewChainName) const;
+
+	// Set the Start Bone for the given Chain. Returns true if operation was successful. 
+	UFUNCTION(BlueprintCallable, Category=IKRig)
+	bool SetRetargetChainStartBone(const FName ChainName, const FName StartBoneName) const;
+
+	// Set the End Bone for the given Chain. Returns true if operation was successful. 
+	UFUNCTION(BlueprintCallable, Category=IKRig)
+	bool SetRetargetChainEndBone(const FName ChainName, const FName EndBoneName) const;
+
+	// Set the Goal for the given Chain. Returns true if operation was successful. 
+	UFUNCTION(BlueprintCallable, Category=IKRig)
+	bool SetRetargetChainGoal(const FName ChainName, const FName GoalName) const;
+
+	// Get the Goal name for the given Chain. 
+	UFUNCTION(BlueprintCallable, Category=IKRig)
+	FName GetRetargetChainGoal(const FName ChainName) const;
+
+	// Get the End Bone name for the given Chain. 
+	UFUNCTION(BlueprintCallable, Category=IKRig)
+    FName GetRetargetChainStartBone(const FName ChainName) const;
+
+	// Get the Start Bone name for the given Chain. 
+	UFUNCTION(BlueprintCallable, Category=IKRig)
+    FName GetRetargetChainEndBone(const FName ChainName) const;
+
+	// Get read-only access to the list of Chains. 
+	UFUNCTION(BlueprintCallable, Category=IKRig)
 	const TArray<FBoneChain>& GetRetargetChains() const;
-	/** Set the Root Bone of the retargeting (can only be one). */
-	void SetRetargetRoot(const FName& RootBoneName) const;
-	/** Get the name of the Root Bone of the retargeting (can only be one). */
+
+	// Set the Root Bone of the retargeting (can only be one). 
+	UFUNCTION(BlueprintCallable, Category=IKRig)
+	bool SetRetargetRoot(const FName RootBoneName) const;
+
+	// Get the name of the Root Bone of the retargeting (can only be one). 
+	UFUNCTION(BlueprintCallable, Category=IKRig)
 	FName GetRetargetRoot() const;
-	/** Sorts the Chains from Root to tip based on the Start Bone of each Chain. */
+
+	//
+	// RETARGETING C++ ONLY API
+	//
+	
+	// Add a Chain with the given BoneChain settings. Returns newly created chain name. 
+	FName AddRetargetChain(const FBoneChain& BoneChain) const;
+	
+	// Get read-only access to a single retarget chain with the given name 
+	const FBoneChain* GetRetargetChainByName(const FName ChainName) const;
+	
+	// Get the name of the retarget chain that contains the given Bone. Returns NAME_None if Bone not in a Chain. 
+	FName GetRetargetChainFromBone(const FName BoneName, const FIKRigSkeleton* OptionalSkeleton) const;
+
+	// Get the name of the retarget chain that contains the given Goal. Returns NAME_None if Goal not in a Chain. 
+	FName GetRetargetChainFromGoal(const FName GoalName) const;
+
+	// Sorts the Chains from Root to tip based on the Start Bone of each Chain. 
 	void SortRetargetChains() const;
-	/** Make unique name for a retargeting bone chain. Adds a numbered suffix to make it unique.*/
-	FName GetUniqueRetargetChainName(const FName& NameToMakeUnique) const;
-	/** Returns true if this is a valid chain. Produces array of bone indices between start and end (inclusive).
-	 * Optionally provide a runtime skeleton from an IKRigProcessor to get indices for a running instance (otherwise uses stored hierarchy in asset)*/
-	bool ValidateChain(
-		const FName& ChainName,
-		const FIKRigSkeleton* OptionalSkeleton,
-		TSet<int32>& OutChainIndices) const;
-	/** Get the name of the retarget chain that contains the given Bone. Returns NAME_None if Bone not in a Chain. */
-	FName GetRetargetChainFromBone(const FName& BoneName, const FIKRigSkeleton* OptionalSkeleton) const;
-	/** Get the name of the retarget chain that contains the given Goal. Returns NAME_None if Goal not in a Chain. */
-	FName GetRetargetChainFromGoal(const FName& GoalName) const;
-	/** END retarget chains */
-
-	// force all currently connected processors to reinitialize using latest asset state
-	void BroadcastNeedsReinitialized() const
-	{
-		IKRigNeedsInitialized.Broadcast(GetAsset());
-	}
+
+	// Make unique name for a retargeting bone chain. Adds a numbered suffix to make it unique.
+	FName GetUniqueRetargetChainName(const FName NameToMakeUnique) const;
+
+	// Returns true if this is a valid chain. Produces array of bone indices between start and end (inclusive).
+	// Optionally provide a runtime skeleton from an IKRigProcessor to get indices for a running instance (otherwise uses stored hierarchy in asset)
+	bool ValidateChain( const FName ChainName, const FIKRigSkeleton* OptionalSkeleton, TSet<int32>& OutChainIndices) const;
 
 private:
 
-	/** Called whenever the rig is modified in such a way that would require re-initialization by dependent systems.*/
+	// Called whenever the rig is modified in such a way that would require re-initialization by dependent systems.
 	DECLARE_MULTICAST_DELEGATE_OneParam(FOnIKRigNeedsInitialized, UIKRigDefinition*);
 	FOnIKRigNeedsInitialized IKRigNeedsInitialized;
 
-	/** Called whenever a retarget chain is added.*/
+	// Called whenever a retarget chain is added.
 	DECLARE_MULTICAST_DELEGATE_OneParam(FOnRetargetChainAdded, UIKRigDefinition*);
 	FOnRetargetChainAdded RetargetChainAdded;
 
-	/** Called whenever a retarget chain is removed.*/
-	DECLARE_MULTICAST_DELEGATE_TwoParams(FOnRetargetChainRemoved, UIKRigDefinition*, const FName& /*chain name*/);
+	// Called whenever a retarget chain is removed.
+	DECLARE_MULTICAST_DELEGATE_TwoParams(FOnRetargetChainRemoved, UIKRigDefinition*, const FName /*chain name*/);
 	FOnRetargetChainRemoved RetargetChainRemoved;
 
-	/** Called whenever a retarget chain is renamed.*/
+	// Called whenever a retarget chain is renamed.
 	DECLARE_MULTICAST_DELEGATE_ThreeParams(FOnRetargetChainRenamed, UIKRigDefinition*, FName /*old name*/, FName /*new name*/);
 	FOnRetargetChainRenamed RetargetChainRenamed;
 
@@ -225,13 +368,10 @@ public:
 
 private:
 
-	/** The actual IKRigDefinition asset that this Controller modifies. */
+	// The actual IKRigDefinition asset that this Controller modifies. 
 	UPROPERTY(transient)
 	TObjectPtr<UIKRigDefinition> Asset = nullptr;
 
-	/** Lazy-generated map of Controllers to IK Rig Assets. Avoids duplicate controllers. */
-	static TMap<UIKRigDefinition*, UIKRigController*> AssetToControllerMap;
-
 	// broadcast changes within the asset goals array
 	void BroadcastGoalsChange() const;
 	
diff --git a/Engine/Plugins/Animation/IKRig/Source/IKRigEditor/Public/RigEditor/IKRigEditorController.h b/Engine/Plugins/Animation/IKRig/Source/IKRigEditor/Public/RigEditor/IKRigEditorController.h
--- a/Engine/Plugins/Animation/IKRig/Source/IKRigEditor/Public/RigEditor/IKRigEditorController.h
+++ b/Engine/Plugins/Animation/IKRig/Source/IKRigEditor/Public/RigEditor/IKRigEditorController.h
@@ -86,7 +86,7 @@ enum class EChainSide
 	Center
 };
 
-/** use string matching and skeletal analysis to suggest a reasonable default name for retarget chains */
+// use string matching and skeletal analysis to suggest a reasonable default name for retarget chains 
 struct FRetargetChainAnalyzer
 {
 	TSharedPtr<FTextFilterExpressionEvaluator> TextFilter;
@@ -103,167 +103,165 @@ struct FRetargetChainAnalyzer
 	static EChainSide GetSideOfChain(const TArray<int32>& BoneIndices, const FIKRigSkeleton& IKRigSkeleton);
 };
 
-/** a home for cross-widget communication to synchronize state across all tabs and viewport */
+// a home for cross-widget communication to synchronize state across all tabs and viewport 
 class FIKRigEditorController : public TSharedFromThis<FIKRigEditorController>, FGCObject
 {
 public:
 
-	/** initialize the editor controller to an instance of the IK Rig editor */
-	void Initialize(TSharedPtr<FIKRigEditorToolkit> Toolkit, UIKRigDefinition* Asset);
-	/** cleanup when editor closed */
+	// initialize the editor controller to an instance of the IK Rig editor 
+	void Initialize(TSharedPtr<FIKRigEditorToolkit> Toolkit, const UIKRigDefinition* Asset);
+	// cleanup when editor closed 
 	void Close() const;
 
 	// initial setup, used when asset is empty (no hierarchy, no mesh assigned)
 	void PromptUserToAssignMesh();
 
-	/** get the currently active processor running the IK Rig in the editor */
+	// get the currently active processor running the IK Rig in the editor 
 	UIKRigProcessor* GetIKRigProcessor() const;
-	/** get the currently running IKRig skeleton (if there is a running processor) */
+	// get the currently running IKRig skeleton (if there is a running processor) 
 	const FIKRigSkeleton* GetCurrentIKRigSkeleton() const;
 
-	/** callback when IK Rig requires re-initialization */
+	// callback when IK Rig requires re-initialization 
 	void HandleIKRigNeedsInitialized(UIKRigDefinition* ModifiedIKRig) const;
 
-	/** create goals */
+	// create goals 
 	void AddNewGoals(const TArray<FName>& GoalNames, const TArray<FName>& BoneNames);
-	/** clear all selected objects */
+	// clear all selected objects 
 	void ClearSelection();
-	/** callback when goal is selected in the viewport */
+	// callback when goal is selected in the viewport 
 	void HandleGoalSelectedInViewport(const FName& GoalName, bool bReplace) const;
-	/** callback when bone is selected in the viewport */
+	// callback when bone is selected in the viewport 
 	void HandleBoneSelectedInViewport(const FName& BoneName, bool bReplace) const;
-	/** reset all goals to initial transforms */
+	// reset all goals to initial transforms 
 	void Reset() const;
-	/** refresh all views */
+	// refresh all views 
 	void RefreshAllViews() const;
-	/** refresh just the skeleton tree view */
+	// refresh just the skeleton tree view 
 	void RefreshTreeView() const;
-	/** clear the output log */
+	// clear the output log 
 	void ClearOutputLog() const;
 
-	/** return list of those solvers in the stack that are selected by user */
+	// return list of those solvers in the stack that are selected by user 
 	void GetSelectedSolvers(TArray<TSharedPtr<FSolverStackElement> >& OutSelectedSolvers);
-	/** get index of the first selected solver, return INDEX_None if nothing selected */
+	// get index of the first selected solver, return INDEX_None if nothing selected 
 	int32 GetSelectedSolverIndex();
-	/** get names of all goals that are selected */
+	// get names of all goals that are selected 
 	void GetSelectedGoalNames(TArray<FName>& OutGoalNames) const;
-	/** return the number of selected goals */
+	// return the number of selected goals 
 	int32 GetNumSelectedGoals() const;
-	/** get names of all bones that are selected */
+	// get names of all bones that are selected 
 	void GetSelectedBoneNames(TArray<FName>& OutBoneNames) const;
-	/** get all bones that are selected */
+	// get all bones that are selected 
 	void GetSelectedBones(TArray<TSharedPtr<FIKRigTreeElement>>& OutBoneItems) const;
-	/** returns true if Goal is currently selected */
+	// returns true if Goal is currently selected 
 	bool IsGoalSelected(const FName& GoalName) const;
-	/** is anything selected in the skeleton view? */
+	// is anything selected in the skeleton view? 
 	bool DoesSkeletonHaveSelectedItems() const;
-	/** get name of the selected retargeting chain */
+	// get name of the selected retargeting chain 
 	TArray<FName> GetSelectedChains() const;
-	/** get chains selected in skeleton view */
+	// get chains selected in skeleton view 
 	void GetChainsSelectedInSkeletonView(TArray<FBoneChain>& InOutChains);
 	
-	/** create new retarget chains from selected bones (or single empty chain if no selection) */
+	// create new retarget chains from selected bones (or single empty chain if no selection) 
 	void CreateNewRetargetChains();
 	
-	/** show single transform of bone in details view */
+	// show single transform of bone in details view 
 	void ShowDetailsForBone(const FName BoneName) const;
-	/** show single BONE settings in details view */
+	// show single BONE settings in details view 
 	void ShowDetailsForBoneSettings(const FName& BoneName, int32 SolverIndex) const;
-	/** show single GOAL settings in details view */
+	// show single GOAL settings in details view 
 	void ShowDetailsForGoal(const FName& GoalName) const;
-	/** show single EFFECTOR settings in details view */
+	// show single EFFECTOR settings in details view 
 	void ShowDetailsForGoalSettings(const FName GoalName, const int32 SolverIndex) const;
-	/** show single SOLVER settings in details view */
+	// show single SOLVER settings in details view 
 	void ShowDetailsForSolver(const int32 SolverIndex) const;
-	/** show nothing in details view */
+	// show nothing in details view 
 	void ShowEmptyDetails() const;
-	/** show selected items in details view */
+	// show selected items in details view 
 	void ShowDetailsForElements(const TArray<TSharedPtr<FIKRigTreeElement>>& InItems) const;
-	/** callback when detail is edited */
+	// callback when detail is edited 
 	void OnFinishedChangingDetails(const FPropertyChangedEvent& PropertyChangedEvent) const;
 	
-	/** set details tab view */
+	// set details tab view 
 	void SetDetailsView(const TSharedPtr<class IDetailsView>& InDetailsView);
-	/** set skeleton tab view */
+	// set skeleton tab view 
 	void SetHierarchyView(const TSharedPtr<SIKRigHierarchy>& InSkeletonView){ SkeletonView = InSkeletonView; };
-	/** set solver stack tab view */
+	// set solver stack tab view 
 	void SetSolverStackView(const TSharedPtr<SIKRigSolverStack>& InSolverStackView){ SolverStackView = InSolverStackView; };
-	/** set retargeting tab view */
+	// set retargeting tab view 
 	void SetRetargetingView(const TSharedPtr<SIKRigRetargetChainList>& InRetargetingView){ RetargetingView = InRetargetingView; };
-	/** set output log view */
+	// set output log view 
 	void SetOutputLogView(const TSharedPtr<SIKRigOutputLog>& InOutputLogView){ OutputLogView = InOutputLogView; };
 
-	/** right after importing a skeleton, we ask user what solver they want to use */
+	// right after importing a skeleton, we ask user what solver they want to use 
 	bool PromptToAddDefaultSolver() const;
-	/** show user the new retarget chain they are about to create (provides option to edit name) */
+	// show user the new retarget chain they are about to create (provides option to edit name) 
 	FName PromptToAddNewRetargetChain(const FBoneChain& BoneChain) const;
-	/** prompt user if they want to add a goal to a newly created chain */
-	UIKRigEffectorGoal* PromptToAddGoalToNewChain(const FBoneChain& BoneChain) const;
-	/** right after creating a goal, we ask user if they want it assigned to a retarget chain */
-	void PromptToAssignGoalToChain(UIKRigEffectorGoal* NewGoal) const;
+	// prompt user if they want to add a goal to a newly created chain
+	FName PromptToAddGoalToNewChain(const FBoneChain& BoneChain) const;
+	// right after creating a goal, we ask user if they want it assigned to a retarget chain 
+	void PromptToAssignGoalToChain(FName NewGoalName) const;
 
-	/** play preview animation on running anim instance in editor (before IK) */
+	// play preview animation on running anim instance in editor (before IK) 
 	void PlayAnimationAsset(UAnimationAsset* AssetToPlay);
 	
-	/** all modifications to the data model should go through this controller */
+	// all modifications to the data model should go through this controller 
 	UIKRigController* AssetController;
 
-	/** viewport skeletal mesh */
+	// viewport skeletal mesh 
 	UDebugSkelMeshComponent* SkelMeshComponent;
 
-	/** viewport anim instance */
+	// viewport anim instance 
 	UPROPERTY(transient, NonTransactional)
-	TWeakObjectPtr<UIKRigAnimInstance> AnimInstance;
+	TObjectPtr<UIKRigAnimInstance> AnimInstance;
 
-	/** the persona toolkit */
+	// the persona toolkit 
 	TWeakPtr<FIKRigEditorToolkit> EditorToolkit;
 
-	/** used to analyze a retarget chain */
+	// used to analyze a retarget chain 
 	FRetargetChainAnalyzer ChainAnalyzer;
 
-	/** UI and viewport selection state */
+	// UI and viewport selection state 
 	bool bManipulatingGoals = false;
 
-	/** record which part of the UI was last selected */
+	// record which part of the UI was last selected 
 	EIKRigSelectionType GetLastSelectedType() const;
 	void SetLastSelectedType(EIKRigSelectionType SelectionType);
-	/** END selection type */
+	// END selection type 
 
-	/** FGCObject interface */
+	// FGCObject interface 
 	virtual void AddReferencedObjects(FReferenceCollector& Collector) override
 	{
 		Collector.AddReferencedObject(BoneDetails);
+		Collector.AddReferencedObject(AnimInstance);
 	};
 	virtual FString GetReferencerName() const override { return "IKRigEditorController"; };
-	/** END FGCObject interface */
+	// END FGCObject interface 
 
-	/** UIKRigBoneDetails factory **/
+	// UIKRigBoneDetails factory *
 	TObjectPtr<UIKRigBoneDetails> CreateBoneDetails(const TSharedPtr<FIKRigTreeElement const>& InItem) const;
 	
 private:
-	
-	/** Initializes editor's solvers instances */
-	void InitializeSolvers() const;
 
-	/** asset properties tab */
+	// asset properties tab 
 	TSharedPtr<IDetailsView> DetailsView;
 
-	/** the skeleton tree view */
+	// the skeleton tree view 
 	TSharedPtr<SIKRigHierarchy> SkeletonView;
 	
-	/** the solver stack view */
+	// the solver stack view 
 	TSharedPtr<SIKRigSolverStack> SolverStackView;
 
-	/** the solver stack view */
+	// the solver stack view 
 	TSharedPtr<SIKRigRetargetChainList> RetargetingView;
 
-	/** asset browser view */
+	// asset browser view 
 	TSharedPtr<SIKRigAssetBrowser> AssetBrowserView;
 
-	/** output log view */
+	// output log view 
 	TSharedPtr<SIKRigOutputLog> OutputLogView;
 
-	/** when prompting user to assign a skeletal mesh */
+	// when prompting user to assign a skeletal mesh 
 	TSharedPtr<SWindow> MeshPickerWindow;
 
 	EIKRigSelectionType LastSelectedType;
@@ -271,7 +269,7 @@ private:
 	UPROPERTY()
 	TObjectPtr<UIKRigBoneDetails> BoneDetails;
 
-	/** remove delegate when closing editor */
+	// remove delegate when closing editor 
 	FDelegateHandle ReinitializeDelegateHandle;
 
 	friend struct FIKRigOutputLogTabSummoner;
diff --git a/Engine/Plugins/Animation/IKRig/Source/IKRigEditor/Public/RigEditor/SIKRigHierarchy.h b/Engine/Plugins/Animation/IKRig/Source/IKRigEditor/Public/RigEditor/SIKRigHierarchy.h
--- a/Engine/Plugins/Animation/IKRig/Source/IKRigEditor/Public/RigEditor/SIKRigHierarchy.h
+++ b/Engine/Plugins/Animation/IKRig/Source/IKRigEditor/Public/RigEditor/SIKRigHierarchy.h
@@ -188,7 +188,7 @@ public:
 		const FName& ItemName,
 		IKRigTreeElementType ItemType,
 		const bool bReplace);
-	void GetSelectedBoneChains(TArray<FBoneChain>& OutChains);
+	void GetSelectedBoneChains(TArray<FBoneChain>& OutChains) const;
 	TArray<TSharedPtr<FIKRigTreeElement>> GetSelectedItems() const;
 	bool HasSelectedItems() const;
 	/** END selection state queries */
