Modified IK Rig solver API to have execution pins in BP and use classes instead of class name strings.

diff --git a/Engine/Plugins/Animation/IKRig/Source/IKRig/Public/Solvers/IKRig_BodyMover.h b/Engine/Plugins/Animation/IKRig/Source/IKRig/Public/Solvers/IKRig_BodyMover.h
--- a/Engine/Plugins/Animation/IKRig/Source/IKRig/Public/Solvers/IKRig_BodyMover.h
+++ b/Engine/Plugins/Animation/IKRig/Source/IKRig/Public/Solvers/IKRig_BodyMover.h
@@ -26,7 +26,7 @@ public:
 	float InfluenceMultiplier = 1.0f;
 };
 
-UCLASS(EditInlineNew)
+UCLASS(BlueprintType, EditInlineNew)
 class IKRIG_API UIKRig_BodyMover : public UIKRigSolver
 {
 	GENERATED_BODY()
diff --git a/Engine/Plugins/Animation/IKRig/Source/IKRig/Public/Solvers/IKRig_LimbSolver.h b/Engine/Plugins/Animation/IKRig/Source/IKRig/Public/Solvers/IKRig_LimbSolver.h
--- a/Engine/Plugins/Animation/IKRig/Source/IKRig/Public/Solvers/IKRig_LimbSolver.h
+++ b/Engine/Plugins/Animation/IKRig/Source/IKRig/Public/Solvers/IKRig_LimbSolver.h
@@ -24,7 +24,7 @@ public:
 	FName BoneName = NAME_None;
 };
 
-UCLASS(EditInlineNew)
+UCLASS(BlueprintType, EditInlineNew)
 class IKRIG_API UIKRig_LimbSolver : public UIKRigSolver
 {
 	GENERATED_BODY()
diff --git a/Engine/Plugins/Animation/IKRig/Source/IKRig/Public/Solvers/IKRig_PBIKSolver.h b/Engine/Plugins/Animation/IKRig/Source/IKRig/Public/Solvers/IKRig_PBIKSolver.h
--- a/Engine/Plugins/Animation/IKRig/Source/IKRig/Public/Solvers/IKRig_PBIKSolver.h
+++ b/Engine/Plugins/Animation/IKRig/Source/IKRig/Public/Solvers/IKRig_PBIKSolver.h
@@ -172,7 +172,7 @@ public:
 	}
 };
 
-UCLASS(EditInlineNew, config = Engine, hidecategories = UObject)
+UCLASS(BlueprintType, EditInlineNew, config = Engine, hidecategories = UObject)
 class IKRIG_API UIKRigPBIKSolver : public UIKRigSolver
 {
 	GENERATED_BODY()
diff --git a/Engine/Plugins/Animation/IKRig/Source/IKRig/Public/Solvers/IKRig_PoleSolver.h b/Engine/Plugins/Animation/IKRig/Source/IKRig/Public/Solvers/IKRig_PoleSolver.h
--- a/Engine/Plugins/Animation/IKRig/Source/IKRig/Public/Solvers/IKRig_PoleSolver.h
+++ b/Engine/Plugins/Animation/IKRig/Source/IKRig/Public/Solvers/IKRig_PoleSolver.h
@@ -27,7 +27,7 @@ public:
 	float Alpha = 1.0f;
 };
 
-UCLASS(EditInlineNew)
+UCLASS(BlueprintType, EditInlineNew)
 class IKRIG_API UIKRig_PoleSolver : public UIKRigSolver
 {
 	GENERATED_BODY()
diff --git a/Engine/Plugins/Animation/IKRig/Source/IKRig/Public/Solvers/IKRig_SetTransform.h b/Engine/Plugins/Animation/IKRig/Source/IKRig/Public/Solvers/IKRig_SetTransform.h
--- a/Engine/Plugins/Animation/IKRig/Source/IKRig/Public/Solvers/IKRig_SetTransform.h
+++ b/Engine/Plugins/Animation/IKRig/Source/IKRig/Public/Solvers/IKRig_SetTransform.h
@@ -27,7 +27,7 @@ public:
 	float Alpha = 1.0f;
 };
 
-UCLASS(EditInlineNew)
+UCLASS(BlueprintType, EditInlineNew)
 class IKRIG_API UIKRig_SetTransform : public UIKRigSolver
 {
 	GENERATED_BODY()
diff --git a/Engine/Plugins/Animation/IKRig/Source/IKRigEditor/Private/RigEditor/IKRigController.cpp b/Engine/Plugins/Animation/IKRig/Source/IKRigEditor/Private/RigEditor/IKRigController.cpp
--- a/Engine/Plugins/Animation/IKRig/Source/IKRigEditor/Private/RigEditor/IKRigController.cpp
+++ b/Engine/Plugins/Animation/IKRig/Source/IKRigEditor/Private/RigEditor/IKRigController.cpp
@@ -683,38 +683,6 @@ int32 UIKRigController::AddSolver(TSubclassOf<UIKRigSolver> InIKRigSolverClass)
 	return SolverIndex;
 }
 
-int32 UIKRigController::AddSolver(const FString& InSolverClassPath) const
-{
-	if (InSolverClassPath.IsEmpty() || InSolverClassPath == FName(NAME_None).ToString())
-	{
-		UE_LOG(LogIKRigEditor, Warning, TEXT("Could not add solver to IK Rig. No solver class specified."));
-		return INDEX_NONE;
-	}
-
-	// account for class paths coming in fully qualified or just with class name
-	UPackage* Package = nullptr;
-	FString PackageName;
-	FString CPPTypeObjectName = InSolverClassPath;
-	if (InSolverClassPath.Split(TEXT("."), &PackageName, &CPPTypeObjectName))
-	{
-		Package = FindPackage(nullptr, *PackageName);
-	}
-
-	// first attempt to find object in specified package
-	UObject* Object = FindObject<UObject>(Package, *CPPTypeObjectName);
-	if (!Object)
-	{
-		// attempt to find object in ANY package
-		Object = FindFirstObject<UObject>(*InSolverClassPath, EFindFirstObjectOptions::NativeFirst | EFindFirstObjectOptions::EnsureIfAmbiguous);
-	}
-
-	// assign to subclass of base solver (null if cast fails)
-	//UClass* ObjClass = Object->GetClass();
-	const TSubclassOf<UIKRigSolver> SolverClass = CastChecked<UClass>(Object);
-	
-	return AddSolver(SolverClass);
-}
-
 bool UIKRigController::RemoveSolver(const int32 SolverIndex) const
 {
 	check(Asset)
diff --git a/Engine/Plugins/Animation/IKRig/Source/IKRigEditor/Public/RetargetEditor/IKRetargeterController.h b/Engine/Plugins/Animation/IKRig/Source/IKRigEditor/Public/RetargetEditor/IKRetargeterController.h
--- a/Engine/Plugins/Animation/IKRig/Source/IKRigEditor/Public/RetargetEditor/IKRetargeterController.h
+++ b/Engine/Plugins/Animation/IKRig/Source/IKRigEditor/Public/RetargetEditor/IKRetargeterController.h
@@ -51,7 +51,7 @@ public:
 	static UIKRetargeterController* GetController(UIKRetargeter* InRetargeterAsset);
 	
 	// Set the IK Rig to use as the source or target (to copy animation FROM/TO) 
-	UFUNCTION(BlueprintCallable, Category=IKRetargeter)
+	UFUNCTION(BlueprintCallable, Category=IKRetargeter, meta=(BlueprintPure = false))
 	void SetIKRig(const ERetargetSourceOrTarget SourceOrTarget, UIKRigDefinition* IKRig) const;
 	
 	// Get either source or target IK Rig 
@@ -59,7 +59,7 @@ public:
 	const UIKRigDefinition* GetIKRig(const ERetargetSourceOrTarget SourceOrTarget) const;
 
 	// Set the preview skeletal mesh for either source or target
-	UFUNCTION(BlueprintCallable, Category=IKRetargeter)
+	UFUNCTION(BlueprintCallable, Category=IKRetargeter, meta=(BlueprintPure = false))
 	void SetPreviewMesh(const ERetargetSourceOrTarget SourceOrTarget, USkeletalMesh* InPreviewMesh) const;
 	
 	// Get the preview skeletal mesh
@@ -92,11 +92,11 @@ public:
 	//
 	
 	// Use fuzzy string search to find "best" Source chain to map to each Target chain
-	UFUNCTION(BlueprintCallable, Category=IKRetargeter)
+	UFUNCTION(BlueprintCallable, Category=IKRetargeter, meta=(BlueprintPure = false))
 	void AutoMapChains(const EAutoMapChainType AutoMapType, const bool bForceRemap) const;
 
 	// Assign a source chain to the given target chain. Animation will be copied from the source to the target.
-	UFUNCTION(BlueprintCallable, Category=IKRetargeter)
+	UFUNCTION(BlueprintCallable, Category=IKRetargeter, meta=(BlueprintPure = false))
 	bool SetSourceChain(FName SourceChainName, FName TargetChainName) const;
 
 	// Get the name of the source chain mapped to a given target chain (the chain animation is copied FROM).
@@ -148,23 +148,23 @@ private:
 public:
 	
 	// Add new retarget pose. Returns the name of the new retarget pose.
-	UFUNCTION(BlueprintCallable, Category=IKRetargeter)
+	UFUNCTION(BlueprintCallable, Category=IKRetargeter, meta=(BlueprintPure = false))
 	FName CreateRetargetPose(const FName& NewPoseName, const ERetargetSourceOrTarget SourceOrTarget) const;
 
 	// Remove a retarget pose. Returns true if the pose was found and removed.
-	UFUNCTION(BlueprintCallable, Category=IKRetargeter)
+	UFUNCTION(BlueprintCallable, Category=IKRetargeter, meta=(BlueprintPure = false))
 	bool RemoveRetargetPose(const FName& PoseToRemove, const ERetargetSourceOrTarget SourceOrTarget) const;
 
 	// Duplicate a retarget pose. Returns the name of the new, duplicate pose (or NAME_None if PoseToDuplicate is not found).
-	UFUNCTION(BlueprintCallable, Category=IKRetargeter)
+	UFUNCTION(BlueprintCallable, Category=IKRetargeter, meta=(BlueprintPure = false))
 	FName DuplicateRetargetPose(const FName PoseToDuplicate, const FName NewName, const ERetargetSourceOrTarget SourceOrTarget) const;
 	
 	// Rename current retarget pose. Returns true if a pose was renamed.
-	UFUNCTION(BlueprintCallable, Category=IKRetargeter)
+	UFUNCTION(BlueprintCallable, Category=IKRetargeter, meta=(BlueprintPure = false))
 	bool RenameRetargetPose(const FName OldPoseName, const FName NewPoseName, const ERetargetSourceOrTarget SourceOrTarget) const;
 
 	// Change which retarget pose is used by the retargeter at runtime
-	UFUNCTION(BlueprintCallable, Category=IKRetargeter)
+	UFUNCTION(BlueprintCallable, Category=IKRetargeter, meta=(BlueprintPure = false))
 	bool SetCurrentRetargetPose(FName CurrentPose, const ERetargetSourceOrTarget SourceOrTarget) const;
 	
 	// Get the current retarget pose
@@ -181,24 +181,24 @@ public:
 
 	// Reset a retarget pose for the specified bones.
 	// If BonesToReset is Empty, will removes all stored deltas, returning pose to reference pose
-	UFUNCTION(BlueprintCallable, Category=IKRetargeter)
+	UFUNCTION(BlueprintCallable, Category=IKRetargeter, meta=(BlueprintPure = false))
 	void ResetRetargetPose(const FName& PoseToReset, const TArray<FName>& BonesToReset, const ERetargetSourceOrTarget SourceOrTarget) const;
 	
 	// Set a delta rotation for a given bone for the current retarget pose
-	UFUNCTION(BlueprintCallable, Category=IKRetargeter)
+	UFUNCTION(BlueprintCallable, Category=IKRetargeter, meta=(BlueprintPure = false))
 	void SetRotationOffsetForRetargetPoseBone(
 		const FName& BoneName,
 		const FQuat& RotationOffset,
 		const ERetargetSourceOrTarget SkeletonMode) const;
 	
 	// Get a delta rotation for a given bone for the current retarget pose
-	UFUNCTION(BlueprintCallable, Category=IKRetargeter)
+	UFUNCTION(BlueprintCallable, Category=IKRetargeter, meta=(BlueprintPure = false))
 	FQuat GetRotationOffsetForRetargetPoseBone(
 		const FName& BoneName,
 		const ERetargetSourceOrTarget SourceOrTarget) const;
 	
 	// Set the translation offset on the retarget root bone for the current retarget pose
-	UFUNCTION(BlueprintCallable, Category=IKRetargeter)
+	UFUNCTION(BlueprintCallable, Category=IKRetargeter, meta=(BlueprintPure = false))
 	void SetRootOffsetInRetargetPose(
 		const FVector& TranslationOffset,
 		const ERetargetSourceOrTarget SourceOrTarget) const;
diff --git a/Engine/Plugins/Animation/IKRig/Source/IKRigEditor/Public/RigEditor/IKRigController.h b/Engine/Plugins/Animation/IKRig/Source/IKRigEditor/Public/RigEditor/IKRigController.h
--- a/Engine/Plugins/Animation/IKRig/Source/IKRigEditor/Public/RigEditor/IKRigController.h
+++ b/Engine/Plugins/Animation/IKRig/Source/IKRigEditor/Public/RigEditor/IKRigController.h
@@ -76,15 +76,12 @@ public:
 	// SOLVERS PUBLIC/SCRIPTING API
 	//
 
-	// Add a new solver of the given type to the bottom of the stack. Returns the stack index. 
+	// Add a new solver of the given type to the bottom of the stack. Returns the stack index.
+	UFUNCTION(BlueprintCallable, Category=IKRig, meta=(BlueprintPure = false))
 	int32 AddSolver(TSubclassOf<UIKRigSolver> InSolverClass) const;
 
-	// Adda new solver, specifying the type as a fully qualified package path, ie ("/Script/IKRig/. 
-	UFUNCTION(BlueprintCallable, Category=IKRig)
-	int32 AddSolver(const FString& InSolverClassPath) const;
-
 	// Remove the solver at the given stack index. 
-	UFUNCTION(BlueprintCallable, Category=IKRig)
+	UFUNCTION(BlueprintCallable, Category=IKRig, meta=(BlueprintPure = false))
 	bool RemoveSolver(const int32 SolverIndex) const;
 
 	// Get access to the given solver. 
@@ -100,11 +97,11 @@ public:
 	int32 GetNumSolvers() const;
 
 	// Move the solver at the given index to the target index. 
-	UFUNCTION(BlueprintCallable, Category=IKRig)
+	UFUNCTION(BlueprintCallable, Category=IKRig, meta=(BlueprintPure = false))
 	bool MoveSolverInStack(int32 SolverToMoveIndex, int32 TargetSolverIndex) const;
 
 	// Set enabled/disabled status of the given solver. 
-	UFUNCTION(BlueprintCallable, Category=IKRig)
+	UFUNCTION(BlueprintCallable, Category=IKRig, meta=(BlueprintPure = false))
 	bool SetSolverEnabled(int32 SolverIndex, bool bIsEnabled) const;
 
 	// Get enabled status of the given solver. 
@@ -112,7 +109,7 @@ public:
 	bool GetSolverEnabled(int32 SolverIndex) const;
 
 	// Set the root bone on a given solver. (not all solvers support root bones, checks CanSetRootBone() first) 
-	UFUNCTION(BlueprintCallable, Category=IKRig)
+	UFUNCTION(BlueprintCallable, Category=IKRig, meta=(BlueprintPure = false))
 	bool SetRootBone(const FName RootBoneName, int32 SolverIndex) const;
 
 	// Get the name of the root bone on a given solver. (not all solvers support root bones, checks CanSetRootBone() first) 
@@ -120,7 +117,7 @@ public:
 	FName GetRootBone(int32 SolverIndex) const;
 
 	// Set the end bone on a given solver. (not all solvers require extra end bones, checks CanSetEndBone() first) 
-	UFUNCTION(BlueprintCallable, Category=IKRig)
+	UFUNCTION(BlueprintCallable, Category=IKRig, meta=(BlueprintPure = false))
 	bool SetEndBone(const FName EndBoneName, int32 SolverIndex) const;
 
 	// Get the name of the end bone on a given solver. (not all solvers require extra end bones, checks CanSetEndBone() first) 
@@ -142,19 +139,19 @@ public:
 	//
 	
 	// Add a new Goal associated with the given Bone. GoalName must be unique. Bones can have multiple Goals (rare). 
-	UFUNCTION(BlueprintCallable, Category=IKRig)
+	UFUNCTION(BlueprintCallable, Category=IKRig, meta=(BlueprintPure = false))
 	FName AddNewGoal(const FName GoalName, const FName BoneName) const;
 
 	// Remove the Goal by name. 
-	UFUNCTION(BlueprintCallable, Category=IKRig)
+	UFUNCTION(BlueprintCallable, Category=IKRig, meta=(BlueprintPure = false))
 	bool RemoveGoal(const FName GoalName) const;
 
 	// Rename a Goal. Returns new name, which may be different after being sanitized. Returns NAME_None if this fails.
-	UFUNCTION(BlueprintCallable, Category=IKRig)
+	UFUNCTION(BlueprintCallable, Category=IKRig, meta=(BlueprintPure = false))
 	FName RenameGoal(const FName OldName, const FName PotentialNewName) const;
 
 	// The the Bone that the given Goal should be parented to / associated with. 
-	UFUNCTION(BlueprintCallable, Category=IKRig)
+	UFUNCTION(BlueprintCallable, Category=IKRig, meta=(BlueprintPure = false))
 	bool SetGoalBone(const FName GoalName, const FName NewBoneName) const;
 
 	// The the Bone associated with the given Goal. 
@@ -166,11 +163,11 @@ public:
 	FName GetGoalNameForBone(const FName BoneName) const;
 
 	// Connect the given Goal to the given Solver. This creates an "Effector" with settings specific to this Solver.
-	UFUNCTION(BlueprintCallable, Category=IKRig)
+	UFUNCTION(BlueprintCallable, Category=IKRig, meta=(BlueprintPure = false))
 	bool ConnectGoalToSolver(const FName GoalName, int32 SolverIndex) const;
 
 	// Disconnect the given Goal from the given Solver. This removes the Effector that associates the Goal with the Solver.
-	UFUNCTION(BlueprintCallable, Category=IKRig)
+	UFUNCTION(BlueprintCallable, Category=IKRig, meta=(BlueprintPure = false))
 	bool DisconnectGoalFromSolver(const FName GoalToRemove, int32 SolverIndex) const;
 
 	// Returns true if the given Goal is connected to the given Solver. False otherwise. 
@@ -226,7 +223,7 @@ public:
 	//
 
 	// Include/exclude a bone from all the solvers. All bones are included by default. 
-	UFUNCTION(BlueprintCallable, Category=IKRig)
+	UFUNCTION(BlueprintCallable, Category=IKRig, meta=(BlueprintPure = false))
 	bool SetBoneExcluded(const FName BoneName, const bool bExclude) const;
 
 	// Returns true if the given Bone is excluded, false otherwise. 
@@ -234,11 +231,11 @@ public:
 	bool GetBoneExcluded(const FName BoneName) const;
 	
 	// Add settings to the given Bone/Solver. Does nothing if Bone already has settings in this Solver.
-	UFUNCTION(BlueprintCallable, Category=IKRig)
+	UFUNCTION(BlueprintCallable, Category=IKRig, meta=(BlueprintPure = false))
 	bool AddBoneSetting(const FName BoneName, int32 SolverIndex) const;
 
 	// Remove settings for the given Bone/Solver. Does nothing if Bone doesn't have setting in this Solver.
-	UFUNCTION(BlueprintCallable, Category=IKRig)
+	UFUNCTION(BlueprintCallable, Category=IKRig, meta=(BlueprintPure = false))
 	bool RemoveBoneSetting(const FName BoneName, int32 SolverIndex) const;
 
 	// Get the generic (Solver-specific) Bone settings UObject for this Bone in the given Solver.
@@ -268,27 +265,27 @@ public:
 
 	// Add a new chain with the given Chain and Bone names. Returns newly created chain name (uniquified).
 	// Note: only the ChainName is required here, all else can be set later.
-	UFUNCTION(BlueprintCallable, Category=IKRig)
+	UFUNCTION(BlueprintCallable, Category=IKRig, meta=(BlueprintPure = false))
 	FName AddRetargetChain(const FName ChainName, const FName StartBoneName, const FName EndBoneName, const FName GoalName) const;
 
 	// Remove a Chain with the given name. Returns true if a Chain was removed. 
-	UFUNCTION(BlueprintCallable, Category=IKRig)
+	UFUNCTION(BlueprintCallable, Category=IKRig, meta=(BlueprintPure = false))
 	bool RemoveRetargetChain(const FName ChainName) const;
 
 	// Renamed the given Chain. Returns the new name (same as old if unsuccessful). 
-	UFUNCTION(BlueprintCallable, Category=IKRig)
+	UFUNCTION(BlueprintCallable, Category=IKRig, meta=(BlueprintPure = false))
 	FName RenameRetargetChain(const FName ChainName, const FName NewChainName) const;
 
 	// Set the Start Bone for the given Chain. Returns true if operation was successful. 
-	UFUNCTION(BlueprintCallable, Category=IKRig)
+	UFUNCTION(BlueprintCallable, Category=IKRig, meta=(BlueprintPure = false))
 	bool SetRetargetChainStartBone(const FName ChainName, const FName StartBoneName) const;
 
 	// Set the End Bone for the given Chain. Returns true if operation was successful. 
-	UFUNCTION(BlueprintCallable, Category=IKRig)
+	UFUNCTION(BlueprintCallable, Category=IKRig, meta=(BlueprintPure = false))
 	bool SetRetargetChainEndBone(const FName ChainName, const FName EndBoneName) const;
 
 	// Set the Goal for the given Chain. Returns true if operation was successful. 
-	UFUNCTION(BlueprintCallable, Category=IKRig)
+	UFUNCTION(BlueprintCallable, Category=IKRig, meta=(BlueprintPure = false))
 	bool SetRetargetChainGoal(const FName ChainName, const FName GoalName) const;
 
 	// Get the Goal name for the given Chain. 
@@ -308,7 +305,7 @@ public:
 	const TArray<FBoneChain>& GetRetargetChains() const;
 
 	// Set the Root Bone of the retargeting (can only be one). 
-	UFUNCTION(BlueprintCallable, Category=IKRig)
+	UFUNCTION(BlueprintCallable, Category=IKRig, meta=(BlueprintPure = false))
 	bool SetRetargetRoot(const FName RootBoneName) const;
 
 	// Get the name of the Root Bone of the retargeting (can only be one). 
