Removing pointer back to controller class from IK Rig and Retarget assets.

diff --git a/Engine/Plugins/Animation/IKRig/Source/IKRig/Public/IKRigDefinition.h b/Engine/Plugins/Animation/IKRig/Source/IKRig/Public/IKRigDefinition.h
--- a/Engine/Plugins/Animation/IKRig/Source/IKRig/Public/IKRigDefinition.h
+++ b/Engine/Plugins/Animation/IKRig/Source/IKRig/Public/IKRigDefinition.h
@@ -241,10 +241,6 @@ public:
 	/** The thickness of the Goals in the editor viewport.*/
 	UPROPERTY(EditAnywhere, Category = "Viewport Goal Settings",  meta = (ClampMin = "0.01", UIMin = "0.0", UIMax = "10.0"))
 	float GoalThickness = 0.7f;
-
-	/** The controller responsible for managing this asset's data (all mutation goes through this) */
-	UPROPERTY(Transient)
-	mutable TObjectPtr<UObject> Controller;
 #endif
 
 private:
diff --git a/Engine/Plugins/Animation/IKRig/Source/IKRig/Public/Retargeter/IKRetargeter.h b/Engine/Plugins/Animation/IKRig/Source/IKRig/Public/Retargeter/IKRetargeter.h
--- a/Engine/Plugins/Animation/IKRig/Source/IKRig/Public/Retargeter/IKRetargeter.h
+++ b/Engine/Plugins/Animation/IKRig/Source/IKRig/Public/Retargeter/IKRetargeter.h
@@ -31,7 +31,7 @@ struct IKRIG_API FRetargetChainMap
 	FName TargetChain = NAME_None;
 };
 
-UCLASS()
+UCLASS(BlueprintType)
 class IKRIG_API URetargetChainSettings : public UObject
 {
 	GENERATED_BODY()
@@ -429,10 +429,6 @@ public:
 	
 private:
 
-	/** The controller responsible for managing this asset's data (all mutation goes through this) */
-	UPROPERTY(Transient)
-	mutable TObjectPtr<UObject> Controller;
-
 	// only ask to fix the root height once, then warn thereafter (don't nag) 
 	UPROPERTY()
 	TSet<TObjectPtr<USkeletalMesh>> MeshesAskedToFixRootHeightFor;
diff --git a/Engine/Plugins/Animation/IKRig/Source/IKRigEditor/Private/RetargetEditor/IKRetargeterController.cpp b/Engine/Plugins/Animation/IKRig/Source/IKRigEditor/Private/RetargetEditor/IKRetargeterController.cpp
--- a/Engine/Plugins/Animation/IKRig/Source/IKRigEditor/Private/RetargetEditor/IKRetargeterController.cpp
+++ b/Engine/Plugins/Animation/IKRig/Source/IKRigEditor/Private/RetargetEditor/IKRetargeterController.cpp
@@ -14,6 +14,8 @@
 
 #define LOCTEXT_NAMESPACE "IKRetargeterController"
 
+TMap<UIKRetargeter*, UIKRetargeterController*> UIKRetargeterController::Controllers;
+
 UIKRetargeterController* UIKRetargeterController::GetController(UIKRetargeter* InRetargeterAsset)
 {
 	if (!InRetargeterAsset)
@@ -21,15 +23,14 @@ UIKRetargeterController* UIKRetargeterController::GetController(UIKRetargeter* I
 		return nullptr;
 	}
 
-	if (!InRetargeterAsset->Controller)
+	if (!Controllers.Contains(InRetargeterAsset))
 	{
 		UIKRetargeterController* Controller = NewObject<UIKRetargeterController>();
-		Controller->Asset = InRetargeterAsset;
-		InRetargeterAsset->Controller = Controller;
+		Controller->Asset = const_cast<UIKRetargeter*>(InRetargeterAsset);
+		Controllers.Add(Controller->Asset, Controller);
 	}
 
-	UIKRetargeterController* Controller = Cast<UIKRetargeterController>(InRetargeterAsset->Controller);
-	return Controller;
+	return Controllers[InRetargeterAsset];
 }
 
 UIKRetargeter* UIKRetargeterController::GetAsset() const
diff --git a/Engine/Plugins/Animation/IKRig/Source/IKRigEditor/Private/RigEditor/IKRigController.cpp b/Engine/Plugins/Animation/IKRig/Source/IKRigEditor/Private/RigEditor/IKRigController.cpp
--- a/Engine/Plugins/Animation/IKRig/Source/IKRigEditor/Private/RigEditor/IKRigController.cpp
+++ b/Engine/Plugins/Animation/IKRig/Source/IKRigEditor/Private/RigEditor/IKRigController.cpp
@@ -15,6 +15,8 @@
 
 #define LOCTEXT_NAMESPACE "IKRigController"
 
+TMap<UIKRigDefinition*, UIKRigController*> UIKRigController::Controllers;
+
 UIKRigController* UIKRigController::GetController(const UIKRigDefinition* InIKRigDefinition)
 {
 	if (!InIKRigDefinition)
@@ -22,14 +24,14 @@ UIKRigController* UIKRigController::GetController(const UIKRigDefinition* InIKRi
 		return nullptr;
 	}
 
-	if (!InIKRigDefinition->Controller)
+	if (!Controllers.Contains(InIKRigDefinition))
 	{
 		UIKRigController* Controller = NewObject<UIKRigController>();
 		Controller->Asset = const_cast<UIKRigDefinition*>(InIKRigDefinition);
-		InIKRigDefinition->Controller = Controller;
+		Controllers.Add(Controller->Asset, Controller);
 	}
 
-	return Cast<UIKRigController>(InIKRigDefinition->Controller);
+	return Controllers[InIKRigDefinition];
 }
 
 UIKRigDefinition* UIKRigController::GetAsset() const
@@ -894,7 +896,7 @@ FName UIKRigController::AddNewGoal(const FName GoalName, const FName BoneName) c
 	// does goal already exist?
 	if (GetGoalIndex(GoalName) != INDEX_NONE)
 	{
-		UE_LOG(LogIKRigEditor, Error, TEXT("Trying to create a Goal that already exists, %s."), *GoalName.ToString());
+		UE_LOG(LogIKRigEditor, Warning, TEXT("Trying to create a Goal that already exists, %s."), *GoalName.ToString());
 		return NAME_None;
 	}
 
@@ -902,7 +904,7 @@ FName UIKRigController::AddNewGoal(const FName GoalName, const FName BoneName) c
 	const int32 BoneIndex = Asset->Skeleton.GetBoneIndexFromName(BoneName);
 	if (BoneIndex == INDEX_NONE)
 	{
-		UE_LOG(LogIKRigEditor, Error, TEXT("Trying to create Goal on unknown bone, %s."), *BoneName.ToString());
+		UE_LOG(LogIKRigEditor, Warning, TEXT("Trying to create Goal on unknown bone, %s."), *BoneName.ToString());
 		return NAME_None;
 	}
 	
diff --git a/Engine/Plugins/Animation/IKRig/Source/IKRigEditor/Public/RetargetEditor/IKRetargeterController.h b/Engine/Plugins/Animation/IKRig/Source/IKRigEditor/Public/RetargetEditor/IKRetargeterController.h
--- a/Engine/Plugins/Animation/IKRig/Source/IKRigEditor/Public/RetargetEditor/IKRetargeterController.h
+++ b/Engine/Plugins/Animation/IKRig/Source/IKRigEditor/Public/RetargetEditor/IKRetargeterController.h
@@ -40,6 +40,9 @@ private:
 	// The actual asset that this Controller modifies. This is the only field this class should have.
 	UIKRetargeter* Asset = nullptr;
 
+	// Static global map of assets to their associated controller
+	static TMap<UIKRetargeter*, UIKRetargeterController*> Controllers;
+
 	//
 	// GENERAL PUBLIC/SCRIPTING API
 	//
diff --git a/Engine/Plugins/Animation/IKRig/Source/IKRigEditor/Public/RigEditor/IKRigController.h b/Engine/Plugins/Animation/IKRig/Source/IKRigEditor/Public/RigEditor/IKRigController.h
--- a/Engine/Plugins/Animation/IKRig/Source/IKRigEditor/Public/RigEditor/IKRigController.h
+++ b/Engine/Plugins/Animation/IKRig/Source/IKRigEditor/Public/RigEditor/IKRigController.h
@@ -369,6 +369,9 @@ private:
 	UPROPERTY(transient)
 	TObjectPtr<UIKRigDefinition> Asset = nullptr;
 
+	// Static global map of assets to their associated controller
+	static TMap<UIKRigDefinition*, UIKRigController*> Controllers;
+
 	// broadcast changes within the asset goals array
 	void BroadcastGoalsChange() const;
 	
