fix a crash due to wrong variable being used

diff --git a/Engine/Plugins/Animation/IKRig/Source/IKRigEditor/Private/RigEditor/IKRigController.cpp b/Engine/Plugins/Animation/IKRig/Source/IKRigEditor/Private/RigEditor/IKRigController.cpp
--- a/Engine/Plugins/Animation/IKRig/Source/IKRigEditor/Private/RigEditor/IKRigController.cpp
+++ b/Engine/Plugins/Animation/IKRig/Source/IKRigEditor/Private/RigEditor/IKRigController.cpp
@@ -469,7 +469,7 @@ bool UIKRigController::ValidateChain(
 	const int32 EndBoneIndex = Skeleton.GetBoneIndexFromName(Chain->EndBone.BoneName);
 
 	const bool bHasStartBone = StartBoneIndex != INDEX_NONE;
-	const bool bHasEndBone = StartBoneIndex != INDEX_NONE;
+	const bool bHasEndBone = EndBoneIndex != INDEX_NONE;
 
 	// chain has neither start nor end bone
 	if (!bHasStartBone && !bHasEndBone)
