Fixed bug, ik rig editor crashes when open a long time.

diff --git a/Engine/Plugins/Animation/IKRig/Source/IKRigEditor/Public/RigEditor/IKRigEditorController.h b/Engine/Plugins/Animation/IKRig/Source/IKRigEditor/Public/RigEditor/IKRigEditorController.h
--- a/Engine/Plugins/Animation/IKRig/Source/IKRigEditor/Public/RigEditor/IKRigEditorController.h
+++ b/Engine/Plugins/Animation/IKRig/Source/IKRigEditor/Public/RigEditor/IKRigEditorController.h
@@ -205,7 +205,8 @@ public:
 	// play preview animation on running anim instance in editor (before IK) 
 	void PlayAnimationAsset(UAnimationAsset* AssetToPlay);
 	
-	// all modifications to the data model should go through this controller 
+	// all modifications to the data model should go through this controller
+	UPROPERTY(transient)
 	UIKRigController* AssetController;
 
 	// viewport skeletal mesh 
@@ -234,6 +235,7 @@ public:
 	{
 		Collector.AddReferencedObject(BoneDetails);
 		Collector.AddReferencedObject(AnimInstance);
+		Collector.AddReferencedObject(AssetController);
 	};
 	virtual FString GetReferencerName() const override { return "IKRigEditorController"; };
 	// END FGCObject interface 
