Fixing issue, IK Rig logs displayed at numbers in global message log window.

diff --git a/Engine/Plugins/Animation/IKRig/Source/IKRig/IKRig.Build.cs b/Engine/Plugins/Animation/IKRig/Source/IKRig/IKRig.Build.cs
--- a/Engine/Plugins/Animation/IKRig/Source/IKRig/IKRig.Build.cs
+++ b/Engine/Plugins/Animation/IKRig/Source/IKRig/IKRig.Build.cs
@@ -14,7 +14,8 @@ public class IKRig : ModuleRules
 				"CoreUObject",
 				"Engine",
 				"AnimationCore",
-				"PBIK",
+				"PBIK", 
+				"MessageLog",
 			});
 
 		PublicDependencyModuleNames.AddRange(
diff --git a/Engine/Plugins/Animation/IKRig/Source/IKRig/Private/IKRigLogger.cpp b/Engine/Plugins/Animation/IKRig/Source/IKRig/Private/IKRigLogger.cpp
--- a/Engine/Plugins/Animation/IKRig/Source/IKRig/Private/IKRigLogger.cpp
+++ b/Engine/Plugins/Animation/IKRig/Source/IKRig/Private/IKRigLogger.cpp
@@ -2,13 +2,28 @@
 
 #include "IKRigLogger.h"
 
+#include "Modules/ModuleManager.h"
+#include "MessageLogModule.h"
 #include "Logging/MessageLog.h"
 
 DEFINE_LOG_CATEGORY(LogIKRig);
 
-void FIKRigLogger::SetLogTarget(const FName InLogName)
+void FIKRigLogger::SetLogTarget(const FName InLogName, const FText& LogLabel)
 {
 	LogName = InLogName;
+
+	FMessageLogModule& MessageLogModule = FModuleManager::LoadModuleChecked<FMessageLogModule>("MessageLog");
+
+	if (!MessageLogModule.IsRegisteredLogListing(InLogName))
+	{
+		FMessageLogInitializationOptions InitOptions;
+		InitOptions.bShowFilters = false;
+		InitOptions.bShowPages = false;
+		InitOptions.bAllowClear = false;
+		InitOptions.bShowInLogWindow = false;
+		InitOptions.bDiscardDuplicates = true;
+		MessageLogModule.RegisterLogListing(InLogName, LogLabel, InitOptions);
+	}
 }
 
 FName FIKRigLogger::GetLogTarget() const
diff --git a/Engine/Plugins/Animation/IKRig/Source/IKRig/Private/IKRigProcessor.cpp b/Engine/Plugins/Animation/IKRig/Source/IKRig/Private/IKRigProcessor.cpp
--- a/Engine/Plugins/Animation/IKRig/Source/IKRig/Private/IKRigProcessor.cpp
+++ b/Engine/Plugins/Animation/IKRig/Source/IKRig/Private/IKRigProcessor.cpp
@@ -13,7 +13,7 @@
 UIKRigProcessor::UIKRigProcessor()
 {
 	const FName LogName = FName("IKRig_",GetUniqueID());
-	Log.SetLogTarget(LogName);
+	Log.SetLogTarget(LogName, LOCTEXT("IKRigLogName", "IK Rig Log"));
 }
 
 void UIKRigProcessor::Initialize(
diff --git a/Engine/Plugins/Animation/IKRig/Source/IKRig/Private/Retargeter/IKRetargetProcessor.cpp b/Engine/Plugins/Animation/IKRig/Source/IKRig/Private/Retargeter/IKRetargetProcessor.cpp
--- a/Engine/Plugins/Animation/IKRig/Source/IKRig/Private/Retargeter/IKRetargetProcessor.cpp
+++ b/Engine/Plugins/Animation/IKRig/Source/IKRig/Private/Retargeter/IKRetargetProcessor.cpp
@@ -1469,7 +1469,7 @@ FVector FPoleVectorMatcher::GetChainNormal(
 UIKRetargetProcessor::UIKRetargetProcessor()
 {
 	const FName LogName = FName("IKRetarget_",GetUniqueID());
-	Log.SetLogTarget(LogName);
+	Log.SetLogTarget(LogName, LOCTEXT("IKRetargetLogName", "IK Retarget Log"));
 }
 
 void UIKRetargetProcessor::Initialize(
diff --git a/Engine/Plugins/Animation/IKRig/Source/IKRig/Public/IKRigLogger.h b/Engine/Plugins/Animation/IKRig/Source/IKRig/Public/IKRigLogger.h
--- a/Engine/Plugins/Animation/IKRig/Source/IKRig/Public/IKRigLogger.h
+++ b/Engine/Plugins/Animation/IKRig/Source/IKRig/Public/IKRigLogger.h
@@ -10,7 +10,7 @@ DECLARE_LOG_CATEGORY_EXTERN(LogIKRig, Warning, All);
 struct IKRIG_API FIKRigLogger
 {
 	/** Set the name of the log to output messages to */
-	void SetLogTarget(const FName InLogName);
+	void SetLogTarget(const FName InLogName, const FText& LogLabel);
 	/** Get the name this log is currently outputting to */
 	FName GetLogTarget() const;
 	/** Log a warning message to display to user. */
diff --git a/Engine/Plugins/Animation/IKRig/Source/IKRigEditor/Private/RigEditor/SIKRigOutputLog.cpp b/Engine/Plugins/Animation/IKRig/Source/IKRigEditor/Private/RigEditor/SIKRigOutputLog.cpp
--- a/Engine/Plugins/Animation/IKRig/Source/IKRigEditor/Private/RigEditor/SIKRigOutputLog.cpp
+++ b/Engine/Plugins/Animation/IKRig/Source/IKRigEditor/Private/RigEditor/SIKRigOutputLog.cpp
@@ -12,22 +12,11 @@ void SIKRigOutputLog::Construct(
 {
 	FMessageLogModule& MessageLogModule = FModuleManager::LoadModuleChecked<FMessageLogModule>("MessageLog");
 
-	if (MessageLogModule.IsRegisteredLogListing(InLogName))
-	{
-		MessageLogListing = MessageLogModule.GetLogListing(InLogName);
-	}
-	else
-	{
-		FMessageLogInitializationOptions InitOptions;
-		InitOptions.bShowFilters = false;
-		InitOptions.bShowPages = false;
-		InitOptions.bAllowClear = false;
-		InitOptions.bShowInLogWindow = false;
-		InitOptions.bDiscardDuplicates = true;
-		MessageLogModule.RegisterLogListing(InLogName, FText::FromString("IK Rig Log"), InitOptions);
-		MessageLogListing = MessageLogModule.GetLogListing(InLogName);
-	}
-
+	// we have to ensure that the log is already registered at this point because otherwise GetLogListing will register
+	// it without the proper initialization settings required to filter these logs out of the main Message Log UI
+	ensure(MessageLogModule.IsRegisteredLogListing(InLogName));
+	
+	MessageLogListing = MessageLogModule.GetLogListing(InLogName);
 	OutputLogWidget = MessageLogModule.CreateLogListingWidget( MessageLogListing.ToSharedRef() );
 	
 	ChildSlot
