Fixing shipping build breakage from CL 23782468

diff --git a/Engine/Plugins/Animation/IKRig/Source/IKRig/IKRig.Build.cs b/Engine/Plugins/Animation/IKRig/Source/IKRig/IKRig.Build.cs
--- a/Engine/Plugins/Animation/IKRig/Source/IKRig/IKRig.Build.cs
+++ b/Engine/Plugins/Animation/IKRig/Source/IKRig/IKRig.Build.cs
@@ -14,8 +14,7 @@ public class IKRig : ModuleRules
 				"CoreUObject",
 				"Engine",
 				"AnimationCore",
-				"PBIK", 
-				"MessageLog",
+				"PBIK"
 			});
 
 		PublicDependencyModuleNames.AddRange(
@@ -39,6 +38,7 @@ public class IKRig : ModuleRules
 				new string[]
 				{
 					"ApplicationCore",
+					"MessageLog",
 				});
 		}
 	}
diff --git a/Engine/Plugins/Animation/IKRig/Source/IKRig/Private/IKRigLogger.cpp b/Engine/Plugins/Animation/IKRig/Source/IKRig/Private/IKRigLogger.cpp
--- a/Engine/Plugins/Animation/IKRig/Source/IKRig/Private/IKRigLogger.cpp
+++ b/Engine/Plugins/Animation/IKRig/Source/IKRig/Private/IKRigLogger.cpp
@@ -1,10 +1,12 @@
 ﻿// Copyright Epic Games, Inc. All Rights Reserved.
 
 #include "IKRigLogger.h"
+#include "Logging/MessageLog.h"
 
+#if WITH_EDITOR
 #include "Modules/ModuleManager.h"
 #include "MessageLogModule.h"
-#include "Logging/MessageLog.h"
+#endif
 
 DEFINE_LOG_CATEGORY(LogIKRig);
 
@@ -12,8 +14,8 @@ void FIKRigLogger::SetLogTarget(const FName InLogName, const FText& LogLabel)
 {
 	LogName = InLogName;
 
+#if WITH_EDITOR
 	FMessageLogModule& MessageLogModule = FModuleManager::LoadModuleChecked<FMessageLogModule>("MessageLog");
-
 	if (!MessageLogModule.IsRegisteredLogListing(InLogName))
 	{
 		FMessageLogInitializationOptions InitOptions;
@@ -24,6 +26,8 @@ void FIKRigLogger::SetLogTarget(const FName InLogName, const FText& LogLabel)
 		InitOptions.bDiscardDuplicates = true;
 		MessageLogModule.RegisterLogListing(InLogName, LogLabel, InitOptions);
 	}
+#endif
+	
 }
 
 FName FIKRigLogger::GetLogTarget() const
