Fixes issue, IK Rig not reinitializing when swapping to an incompatible mesh resulting in suppression of warning messages.

diff --git a/Engine/Plugins/Animation/IKRig/Source/IKRigEditor/Private/RigEditor/IKRigController.cpp b/Engine/Plugins/Animation/IKRig/Source/IKRigEditor/Private/RigEditor/IKRigController.cpp
--- a/Engine/Plugins/Animation/IKRig/Source/IKRigEditor/Private/RigEditor/IKRigController.cpp
+++ b/Engine/Plugins/Animation/IKRig/Source/IKRigEditor/Private/RigEditor/IKRigController.cpp
@@ -562,13 +562,6 @@ FName UIKRigController::GetRetargetChainFromGoal(const FName GoalName) const
 
 bool UIKRigController::SetSkeletalMesh(USkeletalMesh* SkeletalMesh, bool bTransact) const
 {
-	// first determine runtime compatibility between the IK Rig asset and the skeleton we're trying to run it on
-	if (!IsSkeletalMeshCompatible(SkeletalMesh))
-	{
-		UE_LOG(LogIKRigEditor, Warning, TEXT("Trying to initialize IKRig with a Skeleton that is missing required bones. See output log. {0}"), *Asset->GetName());
-		return false;
-	}
-
 	FScopedTransaction Transaction(LOCTEXT("SetSkeletalMesh_Label", "Set Skeletal Mesh"));
 	if (bTransact)
 	{	
