Retargeter UI/UX polish.
- Run button is green
- Run button on it's own row in layout
- Defaults to "Run" mode when opening editor
- Stop button icon is red
- Stop button disabled if already stopped

diff --git a/Engine/Plugins/Animation/IKRig/Source/IKRigEditor/Private/RetargetEditor/IKRetargetEditor.cpp b/Engine/Plugins/Animation/IKRig/Source/IKRigEditor/Private/RetargetEditor/IKRetargetEditor.cpp
--- a/Engine/Plugins/Animation/IKRig/Source/IKRigEditor/Private/RetargetEditor/IKRetargetEditor.cpp
+++ b/Engine/Plugins/Animation/IKRig/Source/IKRigEditor/Private/RetargetEditor/IKRetargetEditor.cpp
@@ -92,6 +92,9 @@ void FIKRetargetEditor::InitAssetEditor(
 	// initial setup, ignored if IK Rig is already assigned
 	EditorController->PromptUserToAssignIKRig(ERetargetSourceOrTarget::Source);
 	EditorController->PromptUserToAssignIKRig(ERetargetSourceOrTarget::Target);
+
+	// run retarget by default
+	EditorController->SetRetargeterMode(ERetargeterOutputMode::RunRetarget);
 }
 
 void FIKRetargetEditor::OnClose()
diff --git a/Engine/Plugins/Animation/IKRig/Source/IKRigEditor/Private/RetargetEditor/IKRetargetEditorController.cpp b/Engine/Plugins/Animation/IKRig/Source/IKRigEditor/Private/RetargetEditor/IKRetargetEditorController.cpp
--- a/Engine/Plugins/Animation/IKRig/Source/IKRigEditor/Private/RetargetEditor/IKRetargetEditorController.cpp
+++ b/Engine/Plugins/Animation/IKRig/Source/IKRigEditor/Private/RetargetEditor/IKRetargetEditorController.cpp
@@ -123,6 +123,17 @@ void FRetargetPlaybackManager::ResumePlayback() const
 	}
 }
 
+bool FRetargetPlaybackManager::IsStopped() const
+{
+	UIKRetargetAnimInstance* AnimInstance = EditorController.Pin()->SourceAnimInstance.Get();
+	if (!AnimInstance)
+	{
+		return true;
+	}
+	
+	return !AnimInstance->GetAnimationAsset();
+}
+
 void FIKRetargetEditorController::Initialize(TSharedPtr<FIKRetargetEditor> InEditor, UIKRetargeter* InAsset)
 {
 	Editor = InEditor;
diff --git a/Engine/Plugins/Animation/IKRig/Source/IKRigEditor/Private/RetargetEditor/SIKRetargetPoseEditor.cpp b/Engine/Plugins/Animation/IKRig/Source/IKRigEditor/Private/RetargetEditor/SIKRetargetPoseEditor.cpp
--- a/Engine/Plugins/Animation/IKRig/Source/IKRigEditor/Private/RetargetEditor/SIKRetargetPoseEditor.cpp
+++ b/Engine/Plugins/Animation/IKRig/Source/IKRigEditor/Private/RetargetEditor/SIKRetargetPoseEditor.cpp
@@ -26,139 +26,144 @@ void SIKRetargetPoseEditor::Construct(
 	[
 		SNew(SVerticalBox)
 
-		// poses
+		// run/pause retargeter
 		+SVerticalBox::Slot()
 		.Padding(2.0f)
 		.AutoHeight()
+		.HAlign(HAlign_Center)
 		[
-			SNew(SBox)
-			.HAlign(HAlign_Center)
+			SNew(SHorizontalBox)
+			+SHorizontalBox::Slot()
+			.AutoWidth()
+			.Padding(4,0)
 			[
-				SNew(SHorizontalBox)
-				+SHorizontalBox::Slot()
-				.Padding(4,0)
+				SNew(SButton)
+				.OnClicked(Controller, &FIKRetargetEditorController::HandleShowRetargetPose)
+				.IsEnabled_Lambda([Controller]
+				{
+					return Controller->GetRetargeterMode() != ERetargeterOutputMode::EditRetargetPose;
+				})
+				.ButtonStyle(FAppStyle::Get(), "FlatButton.Success")
+				.ForegroundColor(FLinearColor::Black)
+				.Content()
 				[
-					SNew(SVerticalBox)
-
-					// add pose selection label and combobox
-					+SVerticalBox::Slot()
-					.VAlign(VAlign_Center)
-					.Padding(4,0)
+					SNew(SHorizontalBox)
+					+SHorizontalBox::Slot()
+					.AutoWidth()
+					.HAlign(HAlign_Left)
 					[
-						SNew(STextBlock).MinDesiredWidth(180).Text(LOCTEXT("CurrentPose", "Current Retarget Pose:"))
+						SNew(SImage)
+						.ColorAndOpacity(FLinearColor::Black)
+						.Image_Lambda([Controller]() -> const FSlateBrush*
+						{
+							if (Controller->GetRetargeterMode() == ERetargeterOutputMode::RunRetarget)
+							{
+								return FAppStyle::GetBrush("GenericPause");
+							}
+							return FAppStyle::GetBrush("GenericPlay");
+						})
 					]
-						
-					// pose selection combobox
-					+SVerticalBox::Slot()
+			
+					+SHorizontalBox::Slot()
+					.HAlign(HAlign_Right)
+					.Padding(4,0)
 					[
-						SAssignNew(PoseListComboBox, SComboBox<TSharedPtr<FName>>)
-						.OptionsSource(&PoseNames)
-						.OnComboBoxOpening(this, &SIKRetargetPoseEditor::Refresh)
-						.OnGenerateWidget_Lambda([](TSharedPtr<FName> InItem)
+						SNew(STextBlock)
+						.ToolTipText(LOCTEXT("BlendPoseToolTip", "Pause playback and show the retarget pose. Use slider below to blend between reference pose (0) and retarget pose (1)."))
+						.Text_Lambda([Controller]() -> const FText
 						{
-							return SNew(STextBlock).Text(FText::FromName(*InItem.Get()));
+							if (Controller->GetRetargeterMode() == ERetargeterOutputMode::RunRetarget)
+							{
+								return LOCTEXT("ShowPoseText", "Show Retarget Pose");
+							}
+							return LOCTEXT("RunRetargetText", "Run Retargeter");
 						})
-						.OnSelectionChanged(Controller, &FIKRetargetEditorController::OnPoseSelected)
-						.Content()
-						[
-							SNew(STextBlock).Text(Controller, &FIKRetargetEditorController::GetCurrentPoseName)
-						]
 					]
 				]
+			]
 
-				+SHorizontalBox::Slot()
-				.Padding(4,0)
+			// stop playback button (contextually disabled)
+			+SHorizontalBox::Slot()
+			.Padding(4,0)
+			.AutoWidth()
+			[
+				SNew(SButton)
+				.OnClicked_Lambda([Controller]
+				{
+					Controller->PlaybackManager->StopPlayback();
+					return FReply::Handled();
+				})
+				.IsEnabled_Lambda([Controller]
+				{
+					return !Controller->PlaybackManager->IsStopped();
+				})
+				.Content()
 				[
-					SNew(SVerticalBox)
+					SNew(SImage)
+					.ToolTipText(LOCTEXT("StopPlaybackToolTip", "Stop playback of current animation and show the reference pose."))
+					.Image(FAppStyle::GetBrush("GenericStop"))
+					.ColorAndOpacity_Lambda([Controller]
+					{
+						const bool bRunning = !Controller->PlaybackManager->IsStopped();
+						return bRunning ? FLinearColor::Red : FSlateColor::UseForeground();
+					})
+				]
+			]
+		]
+		
+		+SVerticalBox::Slot()
+		.Padding(2.0f)
+		.AutoHeight()
+		[
+			SNew(SSeparator)
+		]
 
-					// show retarget pose / run retarget
-					+SVerticalBox::Slot()
-					.Padding(0,4)
-					[
-						SNew(SHorizontalBox)
-						+SHorizontalBox::Slot()
-						.AutoWidth()
-						.Padding(4,0)
-						[
-							SNew(SButton)
-							.OnClicked(Controller, &FIKRetargetEditorController::HandleShowRetargetPose)
-							.IsEnabled_Lambda([Controller]
-							{
-								return Controller->GetRetargeterMode() != ERetargeterOutputMode::EditRetargetPose;
-							})
-							.Content()
-							[
-								SNew(SHorizontalBox)
-								+SHorizontalBox::Slot()
-								.AutoWidth()
-								.HAlign(HAlign_Left)
-								[
-									SNew(SImage)
-									.Image_Lambda([Controller]() -> const FSlateBrush*
-									{
-										if (Controller->GetRetargeterMode() == ERetargeterOutputMode::RunRetarget)
-										{
-											return FAppStyle::GetBrush("GenericPause");
-										}
-										return FAppStyle::GetBrush("GenericPlay");
-									})
-									.ColorAndOpacity(FSlateColor::UseForeground())
-								]
-						
-								+SHorizontalBox::Slot()
-								.HAlign(HAlign_Right)
-								.Padding(4,0)
-								[
-									SNew(STextBlock)
-									.ToolTipText(LOCTEXT("BlendPoseToolTip", "Pause playback and show the retarget pose. Use slider below to blend between reference pose (0) and retarget pose (1)."))
-									.Text_Lambda([Controller]() -> const FText
-									{
-										if (Controller->GetRetargeterMode() == ERetargeterOutputMode::RunRetarget)
-										{
-											return LOCTEXT("ShowPoseText", "Show Retarget Pose");
-										}
-										return LOCTEXT("RunRetargetText", "Run Retargeter");
-									})
-								]
-							]
-						]
+		+SVerticalBox::Slot()
+		.Padding(2.0f)
+		.AutoHeight()
+		.HAlign(HAlign_Center)
+		[
+			SNew(SHorizontalBox)
+			
+			// pose selection label
+			+SHorizontalBox::Slot()
+			.AutoWidth()
+			.VAlign(VAlign_Center)
+			.Padding(4,0)
+			[
+				SNew(STextBlock).Text(LOCTEXT("CurrentPose", "Current Retarget Pose:"))
+			]
 						
-						+SHorizontalBox::Slot()
-						.Padding(4,0)
-						.AutoWidth()
-						[
-							SNew(SButton)
-							.OnClicked_Lambda([Controller]
-							{
-								Controller->PlaybackManager->StopPlayback();
-								return FReply::Handled();
-							})
-							.IsEnabled_Lambda([Controller]
-							{
-								return Controller->GetRetargeterMode() == ERetargeterOutputMode::RunRetarget;
-							})
-							.Content()
-							[
-								SNew(SImage)
-								.ToolTipText(LOCTEXT("StopPlaybackToolTip", "Stop playback of current animation and show the reference pose."))
-								.Image(FAppStyle::GetBrush("GenericStop"))
-								.ColorAndOpacity(FSlateColor::UseForeground())
-							]
-						]
-					]
-
-					// blend retarget pose slider
-					+SVerticalBox::Slot()
-					[
-						SNew(SSpinBox<float>)
-						.Font(FAppStyle::Get().GetFontStyle("PropertyWindow.NormalFont"))
-						.MinValue(0.0f)
-						.MaxValue(1.0f)
-						.Value(Controller, &FIKRetargetEditorController::GetRetargetPoseAmount)
-						.OnValueChanged(Controller, &FIKRetargetEditorController::SetRetargetPoseAmount)
-					]
+			// pose selection combobox
+			+SHorizontalBox::Slot()
+			.AutoWidth()
+			[
+				SAssignNew(PoseListComboBox, SComboBox<TSharedPtr<FName>>)
+				.OptionsSource(&PoseNames)
+				.OnComboBoxOpening(this, &SIKRetargetPoseEditor::Refresh)
+				.OnGenerateWidget_Lambda([](TSharedPtr<FName> InItem)
+				{
+					return SNew(STextBlock).Text(FText::FromName(*InItem.Get()));
+				})
+				.OnSelectionChanged(Controller, &FIKRetargetEditorController::OnPoseSelected)
+				.Content()
+				[
+					SNew(STextBlock).Text(Controller, &FIKRetargetEditorController::GetCurrentPoseName)
 				]
-			]	
+			]
+
+			// pose blending slider
+			+SHorizontalBox::Slot()
+			.AutoWidth()
+			[
+				SNew(SSpinBox<float>)
+				.Font(FAppStyle::Get().GetFontStyle("PropertyWindow.NormalFont"))
+				.MinDesiredWidth(100)
+				.MinValue(0.0f)
+				.MaxValue(1.0f)
+				.Value(Controller, &FIKRetargetEditorController::GetRetargetPoseAmount)
+				.OnValueChanged(Controller, &FIKRetargetEditorController::SetRetargetPoseAmount)
+			]
 		]
 
 		// pose editing toolbar
diff --git a/Engine/Plugins/Animation/IKRig/Source/IKRigEditor/Public/RetargetEditor/IKRetargetEditorController.h b/Engine/Plugins/Animation/IKRig/Source/IKRigEditor/Public/RetargetEditor/IKRetargetEditorController.h
--- a/Engine/Plugins/Animation/IKRig/Source/IKRigEditor/Public/RetargetEditor/IKRetargetEditorController.h
+++ b/Engine/Plugins/Animation/IKRig/Source/IKRigEditor/Public/RetargetEditor/IKRetargetEditorController.h
@@ -68,6 +68,7 @@ struct FRetargetPlaybackManager : public TSharedFromThis<FRetargetPlaybackManage
 	void StopPlayback() const;
 	void PausePlayback();
 	void ResumePlayback() const;
+	bool IsStopped() const;
 	
 private:
 
