Improving "New Retarget Chain" flow in IK Rig editor.

diff --git a/Engine/Plugins/Animation/IKRig/Source/IKRigEditor/Private/RigEditor/IKRigEditorController.cpp b/Engine/Plugins/Animation/IKRig/Source/IKRigEditor/Private/RigEditor/IKRigEditorController.cpp
--- a/Engine/Plugins/Animation/IKRig/Source/IKRigEditor/Private/RigEditor/IKRigEditorController.cpp
+++ b/Engine/Plugins/Animation/IKRig/Source/IKRigEditor/Private/RigEditor/IKRigEditorController.cpp
@@ -17,7 +17,6 @@
 #include "SKismetInspector.h"
 #include "Animation/DebugSkelMeshComponent.h"
 #include "Dialog/SCustomDialog.h"
-#include "Dialogs/Dialogs.h"
 
 #include UE_INLINE_GENERATED_CPP_BY_NAME(IKRigEditorController)
 
@@ -668,7 +667,7 @@ void FIKRigEditorController::CreateNewRetargetChains()
 	else
 	{
 		// create an empty chain
-		const FBoneChain Chain(FRetargetChainAnalyzer::GetDefaultChainName(), NAME_None, NAME_None, NAME_None);
+		FBoneChain Chain(FRetargetChainAnalyzer::GetDefaultChainName(), NAME_None, NAME_None, NAME_None);
 		PromptToAddNewRetargetChain(Chain);
 	}
 	
@@ -917,83 +916,154 @@ void FIKRigEditorController::PromptToAssignGoalToChain(const FName NewGoalName)
 	AssetController->SetRetargetChainGoal(ChainToAddGoalTo, NewGoal->GoalName);
 }
 
-FName FIKRigEditorController::PromptToAddNewRetargetChain(const FBoneChain& BoneChain) const
+FName FIKRigEditorController::PromptToAddNewRetargetChain(FBoneChain& BoneChain) const
 {
-	const TSharedPtr<FStructOnScope> StructToDisplay = MakeShareable(new FStructOnScope(FBoneChain::StaticStruct(), (uint8*)&BoneChain));
-	TSharedRef<SKismetInspector> KismetInspector = SNew(SKismetInspector);
-	KismetInspector->ShowSingleStruct(StructToDisplay);
+	TArray<SCustomDialog::FButton> Buttons;
+	Buttons.Add(SCustomDialog::FButton(LOCTEXT("AddChain", "Add Chain")));
 
-	FName NewChainName = NAME_None;
-	SGenericDialogWidget::FArguments DialogArguments;
-	DialogArguments.OnOkPressed_Lambda([&BoneChain, &NewChainName, this] ()
+	const bool bHasExistingGoal = BoneChain.IKGoalName != NAME_None;
+	if (bHasExistingGoal)
 	{
-		// prompt to add a goal to the chain (skipped if already has a goal)
-		const FName NewGoalName = PromptToAddGoalToNewChain(BoneChain);
-
-		// add the retarget chain
-		NewChainName = AssetController->AddRetargetChain(BoneChain);
-
-		// ask user if they want to assign this goal to a chain (if there is one on this bone)
-		if (NewGoalName != NAME_None)
-		{
-			PromptToAssignGoalToChain(NewGoalName);
-		}
-		
-		RefreshAllViews();
-	});
-
-	SGenericDialogWidget::OpenDialog(
-		LOCTEXT("SIKRigRetargetChains", "Add New Retarget Chain"),
-		KismetInspector,
-		DialogArguments,
-		true);
-
-	return NewChainName;
-}
-
-FName FIKRigEditorController::PromptToAddGoalToNewChain(const FBoneChain& BoneChain) const
-{
-	// check if there is already a goal on the end bone of this chain, if there is, then we're done
-	const FName ExistingGoalName = AssetController->GetGoalNameForBone(BoneChain.EndBone.BoneName);
-	if (ExistingGoalName != NAME_None)
-	{
-		return ExistingGoalName;
+		Buttons.Add(SCustomDialog::FButton(LOCTEXT("AddChain", "Add Chain with Existing Goal")));
 	}
-
-	// check if there is a corresponding bone, if there is none, then no need to create a goal
-	const FIKRigSkeleton& IKRigSkeleton = AssetController->GetIKRigSkeleton();
-	if (IKRigSkeleton.GetBoneIndexFromName(BoneChain.EndBone.BoneName) == INDEX_NONE)
+	else
 	{
-		return NAME_None;
+		Buttons.Add(SCustomDialog::FButton(LOCTEXT("AddChain", "Add Chain with New Goal")));
 	}
+	
+	Buttons.Add(SCustomDialog::FButton(LOCTEXT("Cancel", "Cancel")));
 
 	// ask user if they want to add a goal to the end of this chain
-	const TSharedRef<SCustomDialog> AddGoalToChainDialog = SNew(SCustomDialog)
-		.Title(FText(LOCTEXT("AddGoalToNewChainTitleLabel", "Add Goal to New Chain")))
+	const TSharedRef<SCustomDialog> AddNewRetargetChainDialog =
+		SNew(SCustomDialog)
+		.Title(FText(LOCTEXT("AddNewChainTitleLabel", "Add New Retarget Chain")))
 		.Content()
 		[
-			SNew(STextBlock)
-			.Text(FText::Format(
-				LOCTEXT("AddGoalToNewChainLabel", "Add goal to end bone, {0} of new chain, {1}?"),
-				FText::FromName(BoneChain.EndBone.BoneName), 
-				FText::FromName(BoneChain.ChainName)))
+			SNew(SHorizontalBox)
+			+SHorizontalBox::Slot()
+			.AutoWidth()
+			[
+				SNew(SVerticalBox)
+
+				+SVerticalBox::Slot()
+				.AutoHeight()
+				.Padding(2, 1)
+				.VAlign(VAlign_Center)
+				[
+					SNew(SHorizontalBox)
+					+SHorizontalBox::Slot()
+					[
+						SNew(STextBlock).Text(LOCTEXT("ChainNameLabel", "Chain Name:"))
+					]
+
+					+SHorizontalBox::Slot()
+					[
+						SNew(SEditableTextBox)
+						.MinDesiredWidth(200.0f)
+						.OnTextChanged_Lambda([&BoneChain](const FText& InText)
+						{
+							BoneChain.ChainName = FName(InText.ToString());
+						})
+						.Text(FText::FromName(BoneChain.ChainName))
+					]
+				]
+
+				+SVerticalBox::Slot()
+				.AutoHeight()
+				.Padding(2, 1)
+				.VAlign(VAlign_Center)
+				[
+					SNew(SHorizontalBox)
+					+SHorizontalBox::Slot()
+					[
+						SNew(STextBlock).Text(LOCTEXT("StartBoneLabel", "Start Bone:"))
+					]
+
+					+SHorizontalBox::Slot()
+					[
+						SNew(SEditableTextBox)
+						.Text(FText::FromName(BoneChain.StartBone.BoneName))
+						.IsReadOnly(true)
+					]
+				]
+
+				+SVerticalBox::Slot()
+				.AutoHeight()
+				.Padding(2, 1)
+				.VAlign(VAlign_Center)
+				[
+					SNew(SHorizontalBox)
+					+SHorizontalBox::Slot()
+					[
+						SNew(STextBlock).Text(LOCTEXT("EndBoneLabel", "End Bone:"))
+					]
+
+					+SHorizontalBox::Slot()
+					[
+						SNew(SEditableTextBox)
+						.Text(FText::FromName(BoneChain.EndBone.BoneName))
+						.IsReadOnly(true)
+					]
+				]
+
+				+SVerticalBox::Slot()
+				.AutoHeight()
+				.Padding(2, 1)
+				.VAlign(VAlign_Center)
+				[
+					SNew(SHorizontalBox)
+					+SHorizontalBox::Slot()
+					[
+						SNew(STextBlock).Text(LOCTEXT("GoalLabel", "Goal:"))
+					]
+
+					+SHorizontalBox::Slot()
+					[
+						SNew(SEditableTextBox)
+						.Text(FText::FromName(BoneChain.IKGoalName))
+						.IsReadOnly(true)
+					]
+				]
+			]
+			
 		]
-		.Buttons({
-			SCustomDialog::FButton(LOCTEXT("AddGoal", "Add Goal")),
-			SCustomDialog::FButton(LOCTEXT("NoGoal", "No Goal"))
-	});
+		.Buttons(Buttons);
 
-	if (AddGoalToChainDialog->ShowModal() != 0)
+	// show the dialog and handle user choice
+	const int32 UserChoice = AddNewRetargetChainDialog->ShowModal();
+	if (UserChoice == 2)
 	{
-		return NAME_None; // cancel button pressed, or window closed
+		return NAME_None;  // cancel button pressed, or window closed
 	}
 
-	// add a default solver if there isn't one already
-	PromptToAddDefaultSolver();
+	// add the retarget chain
+	const FName NewChainName = AssetController->AddRetargetChain(BoneChain);
+
+	// did user choose to assign a goal
+	if (UserChoice == 1)
+	{
+		FName GoalName = NAME_None;
+		if (bHasExistingGoal)
+		{
+			// use the existing goal
+			GoalName = BoneChain.IKGoalName;
+		}
+		else
+		{
+			// add a default solver if there isn't one already
+			PromptToAddDefaultSolver();
+			// create new goal
+			const FName NewGoalName = FName(FText::Format(LOCTEXT("GoalOnNewChainName", "{0}_Goal"), FText::FromName(BoneChain.ChainName)).ToString());
+			GoalName = AssetController->AddNewGoal(NewGoalName, BoneChain.EndBone.BoneName);
+		}
+		
+		// assign the existing goal to it
+		AssetController->SetRetargetChainGoal(NewChainName, GoalName);
+	}
 
-	// add the goal
-	const FName NewGoalName = FName(FText::Format(LOCTEXT("GoalOnNewChainName", "{0}_Goal"), FText::FromName(BoneChain.ChainName)).ToString());
-	return AssetController->AddNewGoal(NewGoalName, BoneChain.EndBone.BoneName);
+	RefreshAllViews();
+	
+	return NewChainName;
 }
 
 void FIKRigEditorController::PlayAnimationAsset(UAnimationAsset* AssetToPlay)
diff --git a/Engine/Plugins/Animation/IKRig/Source/IKRigEditor/Public/RigEditor/IKRigEditorController.h b/Engine/Plugins/Animation/IKRig/Source/IKRigEditor/Public/RigEditor/IKRigEditorController.h
--- a/Engine/Plugins/Animation/IKRig/Source/IKRigEditor/Public/RigEditor/IKRigEditorController.h
+++ b/Engine/Plugins/Animation/IKRig/Source/IKRigEditor/Public/RigEditor/IKRigEditorController.h
@@ -196,9 +196,7 @@ public:
 	// right after importing a skeleton, we ask user what solver they want to use 
 	bool PromptToAddDefaultSolver() const;
 	// show user the new retarget chain they are about to create (provides option to edit name) 
-	FName PromptToAddNewRetargetChain(const FBoneChain& BoneChain) const;
-	// prompt user if they want to add a goal to a newly created chain
-	FName PromptToAddGoalToNewChain(const FBoneChain& BoneChain) const;
+	FName PromptToAddNewRetargetChain(FBoneChain& BoneChain) const;
 	// right after creating a goal, we ask user if they want it assigned to a retarget chain 
 	void PromptToAssignGoalToChain(FName NewGoalName) const;
 
