IK Rig Add Retarget Chain UI polish.

diff --git a/Engine/Plugins/Animation/IKRig/Source/IKRigEditor/Private/RigEditor/IKRigEditorController.cpp b/Engine/Plugins/Animation/IKRig/Source/IKRigEditor/Private/RigEditor/IKRigEditorController.cpp
--- a/Engine/Plugins/Animation/IKRig/Source/IKRigEditor/Private/RigEditor/IKRigEditorController.cpp
+++ b/Engine/Plugins/Animation/IKRig/Source/IKRigEditor/Private/RigEditor/IKRigEditorController.cpp
@@ -924,11 +924,11 @@ FName FIKRigEditorController::PromptToAddNewRetargetChain(FBoneChain& BoneChain)
 	const bool bHasExistingGoal = BoneChain.IKGoalName != NAME_None;
 	if (bHasExistingGoal)
 	{
-		Buttons.Add(SCustomDialog::FButton(LOCTEXT("AddChain", "Add Chain with Existing Goal")));
+		Buttons.Add(SCustomDialog::FButton(LOCTEXT("AddChain", "Add Chain using Goal")));
 	}
 	else
 	{
-		Buttons.Add(SCustomDialog::FButton(LOCTEXT("AddChain", "Add Chain with New Goal")));
+		Buttons.Add(SCustomDialog::FButton(LOCTEXT("AddChain", "Add Chain and Goal")));
 	}
 	
 	Buttons.Add(SCustomDialog::FButton(LOCTEXT("Cancel", "Cancel")));
@@ -937,11 +937,11 @@ FName FIKRigEditorController::PromptToAddNewRetargetChain(FBoneChain& BoneChain)
 	const TSharedRef<SCustomDialog> AddNewRetargetChainDialog =
 		SNew(SCustomDialog)
 		.Title(FText(LOCTEXT("AddNewChainTitleLabel", "Add New Retarget Chain")))
+		.HAlignContent(HAlign_Center)
 		.Content()
 		[
-			SNew(SHorizontalBox)
-			+SHorizontalBox::Slot()
-			.AutoWidth()
+			SNew(SBox)
+			.HAlign(HAlign_Center)
 			[
 				SNew(SVerticalBox)
 
@@ -953,7 +953,7 @@ FName FIKRigEditorController::PromptToAddNewRetargetChain(FBoneChain& BoneChain)
 					SNew(SHorizontalBox)
 					+SHorizontalBox::Slot()
 					[
-						SNew(STextBlock).Text(LOCTEXT("ChainNameLabel", "Chain Name:"))
+						SNew(STextBlock).Text(LOCTEXT("ChainNameLabel", "Chain Name"))
 					]
 
 					+SHorizontalBox::Slot()
@@ -976,7 +976,7 @@ FName FIKRigEditorController::PromptToAddNewRetargetChain(FBoneChain& BoneChain)
 					SNew(SHorizontalBox)
 					+SHorizontalBox::Slot()
 					[
-						SNew(STextBlock).Text(LOCTEXT("StartBoneLabel", "Start Bone:"))
+						SNew(STextBlock).Text(LOCTEXT("StartBoneLabel", "Start Bone"))
 					]
 
 					+SHorizontalBox::Slot()
@@ -995,7 +995,7 @@ FName FIKRigEditorController::PromptToAddNewRetargetChain(FBoneChain& BoneChain)
 					SNew(SHorizontalBox)
 					+SHorizontalBox::Slot()
 					[
-						SNew(STextBlock).Text(LOCTEXT("EndBoneLabel", "End Bone:"))
+						SNew(STextBlock).Text(LOCTEXT("EndBoneLabel", "End Bone"))
 					]
 
 					+SHorizontalBox::Slot()
@@ -1014,7 +1014,7 @@ FName FIKRigEditorController::PromptToAddNewRetargetChain(FBoneChain& BoneChain)
 					SNew(SHorizontalBox)
 					+SHorizontalBox::Slot()
 					[
-						SNew(STextBlock).Text(LOCTEXT("GoalLabel", "Goal:"))
+						SNew(STextBlock).Text(LOCTEXT("GoalLabel", "Goal"))
 					]
 
 					+SHorizontalBox::Slot()
@@ -1031,7 +1031,7 @@ FName FIKRigEditorController::PromptToAddNewRetargetChain(FBoneChain& BoneChain)
 
 	// show the dialog and handle user choice
 	const int32 UserChoice = AddNewRetargetChainDialog->ShowModal();
-	if (UserChoice == 2)
+	if (UserChoice == 2 || UserChoice < 0)
 	{
 		return NAME_None;  // cancel button pressed, or window closed
 	}
