Upgrading the IK Rig and Retarget BP/Python APIs.

diff --git a/Engine/Plugins/Animation/IKRig/Source/IKRig/Public/IKRigDefinition.h b/Engine/Plugins/Animation/IKRig/Source/IKRig/Public/IKRigDefinition.h
--- a/Engine/Plugins/Animation/IKRig/Source/IKRig/Public/IKRigDefinition.h
+++ b/Engine/Plugins/Animation/IKRig/Source/IKRig/Public/IKRigDefinition.h
@@ -147,7 +147,7 @@ struct IKRIG_API FBoneChain
 		, IKGoalName(InGoalName)
 	{}
 
-	UPROPERTY(EditAnywhere, Category = BoneChain)
+	UPROPERTY(BlueprintReadWrite, EditAnywhere, Category = BoneChain)
 	FName ChainName;
 
 	UPROPERTY(VisibleAnywhere, Category = BoneChain)
@@ -156,7 +156,7 @@ struct IKRIG_API FBoneChain
 	UPROPERTY(VisibleAnywhere, Category = BoneChain)
 	FBoneReference EndBone;
 	
-	UPROPERTY(VisibleAnywhere, Category = BoneChain)
+	UPROPERTY(BlueprintReadOnly, VisibleAnywhere, Category = BoneChain)
 	FName IKGoalName;
 };
 
diff --git a/Engine/Plugins/Animation/IKRig/Source/IKRig/Public/Retargeter/IKRetargeter.h b/Engine/Plugins/Animation/IKRig/Source/IKRig/Public/Retargeter/IKRetargeter.h
--- a/Engine/Plugins/Animation/IKRig/Source/IKRig/Public/Retargeter/IKRetargeter.h
+++ b/Engine/Plugins/Animation/IKRig/Source/IKRig/Public/Retargeter/IKRetargeter.h
@@ -47,15 +47,15 @@ public:
 	// END UObject 
 	
 	// The chain on the Source IK Rig asset to copy animation FROM. 
-	UPROPERTY(VisibleAnywhere, Category = "Chain Mapping")
+	UPROPERTY(BlueprintReadOnly, VisibleAnywhere, Category = "Chain Mapping")
 	FName SourceChain;
 
 	// The chain on the Target IK Rig asset to copy animation TO. 
-	UPROPERTY(VisibleAnywhere, Category = "Chain Mapping")
+	UPROPERTY(BlueprintReadOnly, VisibleAnywhere, Category = "Chain Mapping")
 	FName TargetChain;
 
 	// The settings used to control the motion on this target chain. 
-	UPROPERTY(EditAnywhere, Category = "Settings")
+	UPROPERTY(BlueprintReadWrite, EditAnywhere, Category = "Settings")
 	FTargetChainSettings Settings;
 
 #if WITH_EDITORONLY_DATA
diff --git a/Engine/Plugins/Animation/IKRig/Source/IKRig/Public/Solvers/IKRig_BodyMover.h b/Engine/Plugins/Animation/IKRig/Source/IKRig/Public/Solvers/IKRig_BodyMover.h
--- a/Engine/Plugins/Animation/IKRig/Source/IKRig/Public/Solvers/IKRig_BodyMover.h
+++ b/Engine/Plugins/Animation/IKRig/Source/IKRig/Public/Solvers/IKRig_BodyMover.h
@@ -7,7 +7,7 @@
 
 #include "IKRig_BodyMover.generated.h"
 
-UCLASS()
+UCLASS(BlueprintType)
 class IKRIG_API UIKRig_BodyMoverEffector : public UObject
 {
 	GENERATED_BODY()
@@ -15,14 +15,14 @@ class IKRIG_API UIKRig_BodyMoverEffector : public UObject
 public:
 	UIKRig_BodyMoverEffector() { SetFlags(RF_Transactional); }
 	
-	UPROPERTY(VisibleAnywhere, Category = "Body Mover Effector")
+	UPROPERTY(BlueprintReadOnly, VisibleAnywhere, Category = "Body Mover Effector")
 	FName GoalName;
 
-	UPROPERTY(VisibleAnywhere, Category = "Body Mover Effector")
+	UPROPERTY(BlueprintReadOnly, VisibleAnywhere, Category = "Body Mover Effector")
 	FName BoneName;
 
 	/** Scale the influence this effector has on the body. Range is 0-10. Default is 1.0. */
-	UPROPERTY(EditAnywhere, Category = "Body Mover Effector", meta = (ClampMin = "0", ClampMax = "10", UIMin = "0.0", UIMax = "10.0"))
+	UPROPERTY(BlueprintReadWrite, EditAnywhere, Category = "Body Mover Effector", meta = (ClampMin = "0", ClampMax = "10", UIMin = "0.0", UIMax = "10.0"))
 	float InfluenceMultiplier = 1.0f;
 };
 
diff --git a/Engine/Plugins/Animation/IKRig/Source/IKRig/Public/Solvers/IKRig_LimbSolver.h b/Engine/Plugins/Animation/IKRig/Source/IKRig/Public/Solvers/IKRig_LimbSolver.h
--- a/Engine/Plugins/Animation/IKRig/Source/IKRig/Public/Solvers/IKRig_LimbSolver.h
+++ b/Engine/Plugins/Animation/IKRig/Source/IKRig/Public/Solvers/IKRig_LimbSolver.h
@@ -9,7 +9,7 @@
 
 #include "IKRig_LimbSolver.generated.h"
 
-UCLASS()
+UCLASS(BlueprintType)
 class IKRIG_API UIKRig_LimbEffector : public UObject
 {
 	GENERATED_BODY()
@@ -17,10 +17,10 @@ class IKRIG_API UIKRig_LimbEffector : public UObject
 public:
 	UIKRig_LimbEffector() { SetFlags(RF_Transactional); }
 
-	UPROPERTY(VisibleAnywhere, Category = "Limb IK Effector")
+	UPROPERTY(BlueprintReadOnly, VisibleAnywhere, Category = "Limb IK Effector")
 	FName GoalName = NAME_None;
 
-	UPROPERTY(VisibleAnywhere, Category = "Limb IK Effector")
+	UPROPERTY(BlueprintReadOnly, VisibleAnywhere, Category = "Limb IK Effector")
 	FName BoneName = NAME_None;
 };
 
@@ -33,54 +33,54 @@ public:
 
 	UIKRig_LimbSolver();
 	
-	UPROPERTY(VisibleAnywhere, DisplayName = "Root Bone", Category = "Limb IK Settings")
+	UPROPERTY(BlueprintReadOnly, VisibleAnywhere, DisplayName = "Root Bone", Category = "Limb IK Settings")
 	FName RootName = NAME_None;
 
 	/** Precision (distance to the target) */
-	UPROPERTY(EditAnywhere, Category = "Limb IK Settings", meta = (ClampMin = "0.0"))
+	UPROPERTY(BlueprintReadWrite, EditAnywhere, Category = "Limb IK Settings", meta = (ClampMin = "0.0"))
 	float ReachPrecision = 0.01f;
 	
 	// TWO BONES SETTINGS
 	
 	/** Hinge Bones Rotation Axis. This is essentially the plane normal for (hip - knee - foot). */
-	UPROPERTY(EditAnywhere, Category = "Limb IK Settings|Two Bones")
+	UPROPERTY(BlueprintReadWrite, EditAnywhere, Category = "Limb IK Settings|Two Bones")
 	TEnumAsByte<EAxis::Type> HingeRotationAxis = EAxis::None;
 
 	// FABRIK SETTINGS
 
 	/** Number of Max Iterations to reach the target */
-	UPROPERTY(EditAnywhere, Category = "Limb IK Settings|FABRIK", meta = (ClampMin = "0", UIMin = "0", UIMax = "100"))
+	UPROPERTY(BlueprintReadWrite, EditAnywhere, Category = "Limb IK Settings|FABRIK", meta = (ClampMin = "0", UIMin = "0", UIMax = "100"))
 	int32 MaxIterations = 12;
 
 	/** Enable/Disable rotational limits */
-	UPROPERTY(EditAnywhere, Category = "Limb IK Settings|FABRIK|Limits")
+	UPROPERTY(BlueprintReadWrite, EditAnywhere, Category = "Limb IK Settings|FABRIK|Limits")
 	bool bEnableLimit = false;
 
 	/** Only used if bEnableRotationLimit is enabled. Prevents the leg from folding onto itself,
 	* and forces at least this angle between Parent and Child bone. */
-	UPROPERTY(EditAnywhere, Category = "Limb IK Settings|FABRIK|Limits", meta = (EditCondition="bEnableLimit", ClampMin = "0.0", ClampMax = "90.0", UIMin = "0.0", UIMax = "90.0"))
+	UPROPERTY(BlueprintReadWrite, EditAnywhere, Category = "Limb IK Settings|FABRIK|Limits", meta = (EditCondition="bEnableLimit", ClampMin = "0.0", ClampMax = "90.0", UIMin = "0.0", UIMax = "90.0"))
 	float MinRotationAngle = 15.f;
 
 	/** Pull averaging only has a visual impact when we have more than 2 bones (3 links). */
-	UPROPERTY(EditAnywhere, Category = "Limb IK Settings|FABRIK|Pull Averaging")
+	UPROPERTY(BlueprintReadWrite, EditAnywhere, Category = "Limb IK Settings|FABRIK|Pull Averaging")
 	bool bAveragePull = true;
 
 	/** Re-position limb to distribute pull: 0 = foot, 0.5 = balanced, 1.f = hip */
-	UPROPERTY(EditAnywhere, Category = "Limb IK Settings|FABRIK|Pull Averaging", meta = (ClampMin = "0.0", ClampMax = "1.0",UIMin = "0.0", UIMax = "1.0"))
+	UPROPERTY(BlueprintReadWrite, EditAnywhere, Category = "Limb IK Settings|FABRIK|Pull Averaging", meta = (ClampMin = "0.0", ClampMax = "1.0",UIMin = "0.0", UIMax = "1.0"))
 	float PullDistribution = 0.5f;
 	
 	/** Move end effector towards target. If we are compressing the chain, limit displacement.*/
-	UPROPERTY(EditAnywhere, Category = "Limb IK Settings|FABRIK|Pull Averaging", meta = (ClampMin = "0.0", ClampMax = "1.0",UIMin = "0.0", UIMax = "1.0"))
+	UPROPERTY(BlueprintReadWrite, EditAnywhere, Category = "Limb IK Settings|FABRIK|Pull Averaging", meta = (ClampMin = "0.0", ClampMax = "1.0",UIMin = "0.0", UIMax = "1.0"))
 	float ReachStepAlpha = 0.7f;
 
 	// TWIST SETTINGS
 	
 	/** Enable Knee Twist correction, by comparing Foot FK with Foot IK orientation. */
-	UPROPERTY(EditAnywhere, Category = "Limb IK Settings|Twist")
+	UPROPERTY(BlueprintReadWrite, EditAnywhere, Category = "Limb IK Settings|Twist")
 	bool bEnableTwistCorrection = false;
 	
 	/** Forward Axis for Foot bone. */
-	UPROPERTY(EditAnywhere, Category = "Limb IK Settings|Twist", meta = (EditCondition="bEnableTwistCorrection"))
+	UPROPERTY(BlueprintReadWrite, EditAnywhere, Category = "Limb IK Settings|Twist", meta = (EditCondition="bEnableTwistCorrection"))
 	TEnumAsByte<EAxis::Type> EndBoneForwardAxis = EAxis::Y;
 	
 	/** UIKRigSolver interface */
diff --git a/Engine/Plugins/Animation/IKRig/Source/IKRig/Public/Solvers/IKRig_PBIKSolver.h b/Engine/Plugins/Animation/IKRig/Source/IKRig/Public/Solvers/IKRig_PBIKSolver.h
--- a/Engine/Plugins/Animation/IKRig/Source/IKRig/Public/Solvers/IKRig_PBIKSolver.h
+++ b/Engine/Plugins/Animation/IKRig/Source/IKRig/Public/Solvers/IKRig_PBIKSolver.h
@@ -9,7 +9,7 @@
 
 #include "IKRig_PBIKSolver.generated.h"
 
-UCLASS()
+UCLASS(BlueprintType)
 class IKRIG_API UIKRig_FBIKEffector : public UObject
 {
 	GENERATED_BODY()
@@ -18,18 +18,18 @@ public:
 	UIKRig_FBIKEffector() { SetFlags(RF_Transactional); }
 	
 	/** The Goal that is driving this effector's transform. */
-	UPROPERTY(VisibleAnywhere, Category = "Full Body IK Effector")
+	UPROPERTY(BlueprintReadOnly, VisibleAnywhere, Category = "Full Body IK Effector")
 	FName GoalName;
 	
 	/** The bone that this effector will pull on. */
-	UPROPERTY(VisibleAnywhere, Category = "Full Body IK Effector")
+	UPROPERTY(BlueprintReadOnly, VisibleAnywhere, Category = "Full Body IK Effector")
 	FName BoneName;
 
 	/** Range 0-1 (default is 1.0). The strength of the effector when pulling the bone towards it's target location.
 	* At 0.0, the effector does not pull at all, but the bones between the effector and the root will still slightly resist motion from other effectors.
 	* This can thus act as a "stabilizer" for parts of the body that you do not want to behave in a pure FK fashion.
 	*/
-	UPROPERTY(EditAnywhere, Category = "Full Body IK Effector", meta = (ClampMin = "0", ClampMax = "1", UIMin = "0.0", UIMax = "1.0"))
+	UPROPERTY(BlueprintReadWrite, EditAnywhere, Category = "Full Body IK Effector", meta = (ClampMin = "0", ClampMax = "1", UIMin = "0.0", UIMax = "1.0"))
 	float StrengthAlpha = 1.0f;
 
 	/** Range 0-1 (default is 1.0). When enabled (greater than 0.0), the solver internally partitions the skeleton into 'chains' which extend from the effector to the nearest fork in the skeleton.
@@ -37,12 +37,12 @@ public:
 	*This can improve the results for sparse bone chains, and significantly improve convergence on dense bone chains.
 	*But it may cause undesirable results in highly constrained bone chains (like robot arms).
 	*/
-	UPROPERTY(EditAnywhere, Category = "Full Body IK Effector", meta = (ClampMin = "0", ClampMax = "1", UIMin = "0.0", UIMax = "1.0"))
+	UPROPERTY(BlueprintReadWrite, EditAnywhere, Category = "Full Body IK Effector", meta = (ClampMin = "0", ClampMax = "1", UIMin = "0.0", UIMax = "1.0"))
 	float PullChainAlpha = 1.0f;
 
 	/** Range 0-1 (default is 1.0).
 	*Blends the effector bone rotation between the rotation of the effector transform (1.0) and the rotation of the input bone (0.0).*/
-	UPROPERTY(EditAnywhere, Category = "Full Body IK Effector", meta = (ClampMin = "0", ClampMax = "1", UIMin = "0.0", UIMax = "1.0"))
+	UPROPERTY(BlueprintReadWrite, EditAnywhere, Category = "Full Body IK Effector", meta = (ClampMin = "0", ClampMax = "1", UIMin = "0.0", UIMax = "1.0"))
 	float PinRotation = 1.0f;
 
 	UPROPERTY(Transient)
@@ -56,7 +56,7 @@ public:
 	}
 };
 
-UCLASS()
+UCLASS(BlueprintType)
 class IKRIG_API UIKRig_PBIKBoneSettings : public UObject
 {
 	GENERATED_BODY()
@@ -74,65 +74,65 @@ public:
 	}
 
 	/** The bone these settings are applied to. */
-	UPROPERTY(VisibleAnywhere, Category = Bone, meta = (Constant, CustomWidget = "BoneName"))
+	UPROPERTY(BlueprintReadOnly, VisibleAnywhere, Category = Bone, meta = (Constant, CustomWidget = "BoneName"))
 	FName Bone;
 	
 	/** Range is 0 to 1 (Default is 0). At higher values, the bone will resist rotating (forcing other bones to compensate). */
-	UPROPERTY(EditAnywhere, Category = Stiffness, meta = (ClampMin = "0", ClampMax = "1", UIMin = "0.0", UIMax = "1.0"))
+	UPROPERTY(BlueprintReadWrite, EditAnywhere, Category = Stiffness, meta = (ClampMin = "0", ClampMax = "1", UIMin = "0.0", UIMax = "1.0"))
 	float RotationStiffness = 0.0f;
 
 	/** Range is 0 to 1 (Default is 0). At higher values, the bone will resist translational motion (forcing other bones to compensate). */
-	UPROPERTY(EditAnywhere, Category = Stiffness, meta = (ClampMin = "0", ClampMax = "1", UIMin = "0.0", UIMax = "1.0"))
+	UPROPERTY(BlueprintReadWrite, EditAnywhere, Category = Stiffness, meta = (ClampMin = "0", ClampMax = "1", UIMin = "0.0", UIMax = "1.0"))
 	float PositionStiffness = 0.0f;
 
 	/** Limit the rotation angle of the bone on the X axis.
 	 *Free: can rotate freely in this axis.
 	 *Limited: rotation is clamped between the min/max angles relative to the Skeletal Mesh reference pose.
 	 *Locked: no rotation is allowed in this axis (will remain at reference pose angle). */
-	UPROPERTY(EditAnywhere, Category = Limits)
+	UPROPERTY(BlueprintReadWrite, EditAnywhere, Category = Limits)
 	EPBIKLimitType X;
 	/**Range is -180 to 0 (Default is 0). Degrees of rotation in the negative X direction to allow when joint is in "Limited" mode. */
-	UPROPERTY(EditAnywhere, Category = Limits, meta = (ClampMin = "-180", ClampMax = "0", UIMin = "-180.0", UIMax = "0.0"))
+	UPROPERTY(BlueprintReadWrite, EditAnywhere, Category = Limits, meta = (ClampMin = "-180", ClampMax = "0", UIMin = "-180.0", UIMax = "0.0"))
 	float MinX = 0.0f;
 	/**Range is 0 to 180 (Default is 0). Degrees of rotation in the positive X direction to allow when joint is in "Limited" mode. */
-	UPROPERTY(EditAnywhere, Category = Limits, meta = (ClampMin = "0", ClampMax = "180", UIMin = "0.0", UIMax = "180.0"))
+	UPROPERTY(BlueprintReadWrite, EditAnywhere, Category = Limits, meta = (ClampMin = "0", ClampMax = "180", UIMin = "0.0", UIMax = "180.0"))
 	float MaxX = 0.0f;
 
 	/** Limit the rotation angle of the bone on the Y axis.
 	*Free: can rotate freely in this axis.
 	*Limited: rotation is clamped between the min/max angles relative to the Skeletal Mesh reference pose.
 	*Locked: no rotation is allowed in this axis (will remain at input pose angle). */
-	UPROPERTY(EditAnywhere, Category = Limits)
+	UPROPERTY(BlueprintReadWrite, EditAnywhere, Category = Limits)
 	EPBIKLimitType Y;
 	/**Range is -180 to 0 (Default is 0). Degrees of rotation in the negative Y direction to allow when joint is in "Limited" mode. */
-	UPROPERTY(EditAnywhere, Category = Limits, meta = (ClampMin = "-180", ClampMax = "0", UIMin = "-180.0", UIMax = "0.0"))
+	UPROPERTY(BlueprintReadWrite, EditAnywhere, Category = Limits, meta = (ClampMin = "-180", ClampMax = "0", UIMin = "-180.0", UIMax = "0.0"))
 	float MinY = 0.0f;
 	/**Range is 0 to 180 (Default is 0). Degrees of rotation in the positive Y direction to allow when joint is in "Limited" mode. */
-	UPROPERTY(EditAnywhere, Category = Limits, meta = (ClampMin = "0", ClampMax = "180", UIMin = "0.0", UIMax = "180.0"))
+	UPROPERTY(BlueprintReadWrite, EditAnywhere, Category = Limits, meta = (ClampMin = "0", ClampMax = "180", UIMin = "0.0", UIMax = "180.0"))
 	float MaxY = 0.0f;
 
 	/** Limit the rotation angle of the bone on the Z axis.
 	*Free: can rotate freely in this axis.
 	*Limited: rotation is clamped between the min/max angles relative to the Skeletal Mesh reference pose.
 	*Locked: no rotation is allowed in this axis (will remain at input pose angle). */
-	UPROPERTY(EditAnywhere, Category = Limits)
+	UPROPERTY(BlueprintReadWrite, EditAnywhere, Category = Limits)
 	EPBIKLimitType Z;
 	/**Range is -180 to 0 (Default is 0). Degrees of rotation in the negative Z direction to allow when joint is in "Limited" mode. */
-	UPROPERTY(EditAnywhere, Category = Limits, meta = (ClampMin = "-180", ClampMax = "0", UIMin = "-180.0", UIMax = "0.0"))
+	UPROPERTY(BlueprintReadWrite, EditAnywhere, Category = Limits, meta = (ClampMin = "-180", ClampMax = "0", UIMin = "-180.0", UIMax = "0.0"))
 	float MinZ = 0.0f;
 	/**Range is 0 to 180 (Default is 0). Degrees of rotation in the positive Z direction to allow when joint is in "Limited" mode. */
-	UPROPERTY(EditAnywhere, Category = Limits, meta = (ClampMin = "0", ClampMax = "180", UIMin = "0.0", UIMax = "180.0"))
+	UPROPERTY(BlueprintReadWrite, EditAnywhere, Category = Limits, meta = (ClampMin = "0", ClampMax = "180", UIMin = "0.0", UIMax = "180.0"))
 	float MaxZ = 0.0f;
 
 	/**When true, this bone will "prefer" to rotate in the direction specified by the Preferred Angles when the chain it belongs to is compressed.
 	 * Preferred Angles should be the first method used to fix bones that bend in the wrong direction (rather than limits).
 	 * The resulting angles can be visualized on their own by temporarily setting the Solver iterations to 0 and moving the effectors.*/
-	UPROPERTY(EditAnywhere, Category = PreferredAngles)
+	UPROPERTY(BlueprintReadWrite, EditAnywhere, Category = PreferredAngles)
 	bool bUsePreferredAngles = false;
 	/**The local Euler angles (in degrees) used to rotate this bone when the chain it belongs to is squashed.
 	 * This happens by moving the effector at the tip of the chain towards the root of the chain.
 	 * This can be used to coerce knees and elbows to bend in the anatomically "correct" direction without resorting to limits (which may require more iterations to converge).*/
-	UPROPERTY(EditAnywhere, Category = PreferredAngles)
+	UPROPERTY(BlueprintReadWrite, EditAnywhere, Category = PreferredAngles)
 	FVector PreferredAngles;
 
 	void CopyToCoreStruct(PBIK::FBoneSettings& Settings) const
@@ -184,48 +184,48 @@ public:
 	 *
 	 * If you want to use the solver on a single chain of bones, and NOT translate the chain, ensure that "PinRoot" is
 	 * checked on to disable the root from translating to reach the effector goals.*/
-	UPROPERTY(VisibleAnywhere, Category = "Full Body IK Settings")
+	UPROPERTY(BlueprintReadWrite, VisibleAnywhere, Category = "Full Body IK Settings")
 	FName RootBone;
 
 	/** High iteration counts can help solve complex joint configurations with competing constraints, but will increase runtime cost. Default is 20. */
-	UPROPERTY(EditAnywhere, Category = SolverSettings, meta = (ClampMin = "0", ClampMax = "1000", UIMin = "0.0", UIMax = "200.0"))
+	UPROPERTY(BlueprintReadWrite, EditAnywhere, Category = SolverSettings, meta = (ClampMin = "0", ClampMax = "1000", UIMin = "0.0", UIMax = "200.0"))
 	int32 Iterations = 20;
 
 	/** A global mass multiplier; higher values will make the joints more stiff, but require more iterations. Typical range is 0.0 to 10.0. */
-	UPROPERTY(EditAnywhere, Category = SolverSettings, meta = (ClampMin = "0", UIMin = "0.0", UIMax = "10.0"))
+	UPROPERTY(BlueprintReadWrite, EditAnywhere, Category = SolverSettings, meta = (ClampMin = "0", UIMin = "0.0", UIMax = "10.0"))
 	float MassMultiplier = 1.0f;
 
 	/** If true, joints will translate to reach the effectors; causing bones to lengthen if necessary. Good for cartoon effects. Default is false. */
-	UPROPERTY(EditAnywhere, Category = SolverSettings)
+	UPROPERTY(BlueprintReadWrite, EditAnywhere, Category = SolverSettings)
 	bool bAllowStretch = false;
 
 	/** (Default is PrePull) Set the behavior of the solver root.
 	*Pre Pull: translates and rotates the root (and all children) by the average motion of the stretched effectors to help achieve faster convergence when reaching far.
 	*Pin to Input: locks the translation and rotation of the root bone to the input pose. Overrides any bone settings applied to the root. Good for partial-body solves.
 	*Free: treats the root bone like any other and allows it to move freely or according to any bone settings applied to it. */
-	UPROPERTY(EditAnywhere, Category = RootBehavior)
+	UPROPERTY(BlueprintReadWrite, EditAnywhere, Category = RootBehavior)
 	EPBIKRootBehavior RootBehavior = EPBIKRootBehavior::PrePull;
 
 	/** Settings only applicable when Root Behavior is set to "PrePull". Use these values to adjust the gross movement and orientation of the entire skeleton. */
-	UPROPERTY(EditAnywhere, Category = RootBehavior)
+	UPROPERTY(BlueprintReadWrite, EditAnywhere, Category = RootBehavior)
 	FRootPrePullSettings PrePullRootSettings;
 
 	/** A global multiplier for all Pull Chain Alpha values on all effectors. Range is 0.0 to 1.0. Default is 1.0. */
-	UPROPERTY(EditAnywhere, Category = AdvancedSettings, meta = (ClampMin = "0", ClampMax = "1", UIMin = "0.0", UIMax = "1.0"))
+	UPROPERTY(BlueprintReadWrite, EditAnywhere, Category = AdvancedSettings, meta = (ClampMin = "0", ClampMax = "1", UIMin = "0.0", UIMax = "1.0"))
 	float PullChainAlpha = 1.0f;
 
 	/** Maximum angle that a joint can be rotated per constraint iteration. Lower this value if the solve is diverging. Range is 0.0 to 180.0. Default is 30. */
-	UPROPERTY(EditAnywhere, Category = AdvancedSettings, meta = (ClampMin = "0", ClampMax = "45", UIMin = "0.0", UIMax = "180.0"))
+	UPROPERTY(BlueprintReadWrite, EditAnywhere, Category = AdvancedSettings, meta = (ClampMin = "0", ClampMax = "45", UIMin = "0.0", UIMax = "180.0"))
 	float MaxAngle = 30.f;
 
 	/** Pushes constraints beyond their normal amount to speed up convergence. Increasing this may speed up convergence, but at the cost of stability. Range is 1.0 - 2.0. Default is 1.3. */
-	UPROPERTY(EditAnywhere, Category = AdvancedSettings, meta = (ClampMin = "1", UIMin = "1.0", UIMax = "10.0"))
+	UPROPERTY(BlueprintReadWrite, EditAnywhere, Category = AdvancedSettings, meta = (ClampMin = "1", UIMin = "1.0", UIMax = "10.0"))
 	float OverRelaxation = 1.3f;
 	
 	/** When true, the solver is reset each tick to start from the current input pose. Default is true.
 	 * If false, incoming animated poses are ignored and the solver starts from the results of the previous solve.
 	 * In very limited circumstances, it can be beneficial to use the pose from the previous frame.*/
-	UPROPERTY(EditAnywhere, Category = AdvancedSettings)
+	UPROPERTY(BlueprintReadWrite, EditAnywhere, Category = AdvancedSettings)
 	bool bStartSolveFromInputPose = true;
 	
 	UPROPERTY()
diff --git a/Engine/Plugins/Animation/IKRig/Source/IKRig/Public/Solvers/IKRig_PoleSolver.h b/Engine/Plugins/Animation/IKRig/Source/IKRig/Public/Solvers/IKRig_PoleSolver.h
--- a/Engine/Plugins/Animation/IKRig/Source/IKRig/Public/Solvers/IKRig_PoleSolver.h
+++ b/Engine/Plugins/Animation/IKRig/Source/IKRig/Public/Solvers/IKRig_PoleSolver.h
@@ -7,7 +7,7 @@
 
 #include "IKRig_PoleSolver.generated.h"
 
-UCLASS()
+UCLASS(BlueprintType)
 class IKRIG_API UIKRig_PoleSolverEffector : public UObject
 {
 	GENERATED_BODY()
@@ -16,14 +16,14 @@ public:
 
 	UIKRig_PoleSolverEffector() { SetFlags(RF_Transactional); }
 	
-	UPROPERTY(VisibleAnywhere, Category = "Pole Solver Effector")
+	UPROPERTY(BlueprintReadOnly, VisibleAnywhere, Category = "Pole Solver Effector")
 	FName GoalName = NAME_None;
 
-	UPROPERTY(VisibleAnywhere, Category = "Pole Solver Effector")
+	UPROPERTY(BlueprintReadOnly, VisibleAnywhere, Category = "Pole Solver Effector")
 	FName BoneName = NAME_None;
 
 	/** Blend the effector on/off. Range is 0-1. Default is 1.0. */
-	UPROPERTY(EditAnywhere, Category = "Pole Solver Effector", meta = (ClampMin = "0.0", ClampMax = "1.0", UIMin = "0.0", UIMax = "1.0"))
+	UPROPERTY(BlueprintReadWrite, EditAnywhere, Category = "Pole Solver Effector", meta = (ClampMin = "0.0", ClampMax = "1.0", UIMin = "0.0", UIMax = "1.0"))
 	float Alpha = 1.0f;
 };
 
diff --git a/Engine/Plugins/Animation/IKRig/Source/IKRigEditor/Private/RetargetEditor/IKRetargetBatchOperation.cpp b/Engine/Plugins/Animation/IKRig/Source/IKRigEditor/Private/RetargetEditor/IKRetargetBatchOperation.cpp
--- a/Engine/Plugins/Animation/IKRig/Source/IKRigEditor/Private/RetargetEditor/IKRetargetBatchOperation.cpp
+++ b/Engine/Plugins/Animation/IKRig/Source/IKRigEditor/Private/RetargetEditor/IKRetargetBatchOperation.cpp
@@ -459,7 +459,7 @@ void UIKRetargetBatchOperation::GetNewAssets(TArray<UObject*>& NewAssets) const
 }
 
 
-TArray<FAssetData> UIKRetargetBatchOperation::RunRetarget(
+TArray<FAssetData> UIKRetargetBatchOperation::DuplicateAndRetarget(
 	const TArray<FAssetData>& AssetsToRetarget,
 	USkeletalMesh* SourceMesh,
 	USkeletalMesh* TargetMesh,
@@ -486,14 +486,18 @@ TArray<FAssetData> UIKRetargetBatchOperation::RunRetarget(
 	Context.bRemapReferencedAssets = bRemapReferencedAssets;
 
 	// actually run the batch operation
-	RunRetarget(Context);
+	UIKRetargetBatchOperation* BatchOperation = NewObject<UIKRetargetBatchOperation>();
+	BatchOperation->AddToRoot();
+	BatchOperation->RunRetarget(Context);
 
 	// create array of FAssetData to return
 	TArray<FAssetData> Results;
-	for (const UAnimationAsset* RetargetedAsset : AnimationAssetsToRetarget)
+	for (const UAnimationAsset* RetargetedAsset : BatchOperation->AnimationAssetsToRetarget)
 	{
 		Results.Add(FAssetData(RetargetedAsset));
 	}
+	
+	BatchOperation->RemoveFromRoot();
 	return Results;
 }
 
diff --git a/Engine/Plugins/Animation/IKRig/Source/IKRigEditor/Private/RetargetEditor/IKRetargeterController.cpp b/Engine/Plugins/Animation/IKRig/Source/IKRigEditor/Private/RetargetEditor/IKRetargeterController.cpp
--- a/Engine/Plugins/Animation/IKRig/Source/IKRigEditor/Private/RetargetEditor/IKRetargeterController.cpp
+++ b/Engine/Plugins/Animation/IKRig/Source/IKRigEditor/Private/RetargetEditor/IKRetargeterController.cpp
@@ -124,6 +124,49 @@ USkeletalMesh* UIKRetargeterController::GetPreviewMesh(const ERetargetSourceOrTa
 	return IKRig->GetPreviewMesh();
 }
 
+FTargetRootSettings UIKRetargeterController::GetRootSettings() const
+{
+	return GetAsset()->GetRootSettingsUObject()->Settings;
+}
+
+void UIKRetargeterController::SetRootSettings(const FTargetRootSettings& RootSettings) const
+{
+	GetAsset()->GetRootSettingsUObject()->Settings = RootSettings;
+}
+
+FRetargetGlobalSettings UIKRetargeterController::GetGlobalSettings() const
+{
+	return GetAsset()->GetGlobalSettingsUObject()->Settings;
+}
+
+void UIKRetargeterController::SetGlobalSettings(const FRetargetGlobalSettings& GlobalSettings) const
+{
+	GetAsset()->GetGlobalSettingsUObject()->Settings = GlobalSettings;
+}
+
+FTargetChainSettings UIKRetargeterController::GetRetargetChainSettings(const FName& TargetChainName) const
+{
+	URetargetChainSettings* ChainSettings = GetChainSettings(TargetChainName);
+	if (!ChainSettings)
+	{
+		return FTargetChainSettings();
+	}
+
+	return ChainSettings->Settings;
+}
+
+bool UIKRetargeterController::SetRetargetChainSettings(const FName& TargetChainName, const FTargetChainSettings& Settings) const
+{
+	URetargetChainSettings* ChainSettings = GetChainSettings(TargetChainName);
+	if (!ChainSettings)
+	{
+		ChainSettings->Settings = Settings;
+		return true;
+	}
+
+	return false;
+}
+
 bool UIKRetargeterController::GetAskedToFixRootHeightForMesh(USkeletalMesh* Mesh) const
 {
 	return GetAsset()->MeshesAskedToFixRootHeightFor.Contains(Mesh);
diff --git a/Engine/Plugins/Animation/IKRig/Source/IKRigEditor/Private/RigEditor/IKRigController.cpp b/Engine/Plugins/Animation/IKRig/Source/IKRigEditor/Private/RigEditor/IKRigController.cpp
--- a/Engine/Plugins/Animation/IKRig/Source/IKRigEditor/Private/RigEditor/IKRigController.cpp
+++ b/Engine/Plugins/Animation/IKRig/Source/IKRigEditor/Private/RigEditor/IKRigController.cpp
@@ -138,7 +138,7 @@ bool UIKRigController::CanRemoveBoneSetting(const FName BoneName, int32 SolverIn
 	return true;
 }
 
-UObject* UIKRigController::GetSettingsForBone(const FName BoneName, int32 SolverIndex) const
+UObject* UIKRigController::GetBoneSettings(const FName BoneName, int32 SolverIndex) const
 {
 	const UIKRigSolver* Solver = GetSolverAtIndex(SolverIndex);
 	if (!Solver)
@@ -1117,7 +1117,7 @@ bool UIKRigController::ConnectGoalToSolver(const FName GoalName, int32 SolverInd
 	// can't add goal that is not present
 	if (GoalIndex == INDEX_NONE)
 	{
-		UE_LOG(LogIKRigEditor, Warning, TEXT("Trying to connect unknown Goal, {0} to a solve."), *GoalName.ToString());
+		UE_LOG(LogIKRigEditor, Warning, TEXT("Trying to connect unknown Goal, {0} to a solver."), *GoalName.ToString());
 		return false;
 	}
 	
diff --git a/Engine/Plugins/Animation/IKRig/Source/IKRigEditor/Private/RigEditor/IKRigEditorController.cpp b/Engine/Plugins/Animation/IKRig/Source/IKRigEditor/Private/RigEditor/IKRigEditorController.cpp
--- a/Engine/Plugins/Animation/IKRig/Source/IKRigEditor/Private/RigEditor/IKRigEditorController.cpp
+++ b/Engine/Plugins/Animation/IKRig/Source/IKRigEditor/Private/RigEditor/IKRigEditorController.cpp
@@ -770,7 +770,7 @@ void FIKRigEditorController::ShowDetailsForBone(const FName BoneName) const
 
 void FIKRigEditorController::ShowDetailsForBoneSettings(const FName& BoneName, int32 SolverIndex) const
 {
-	if (UObject* BoneSettings = AssetController->GetSettingsForBone(BoneName, SolverIndex))
+	if (UObject* BoneSettings = AssetController->GetBoneSettings(BoneName, SolverIndex))
 	{
 		DetailsView->SetObject(BoneSettings);
 	}
diff --git a/Engine/Plugins/Animation/IKRig/Source/IKRigEditor/Private/RigEditor/SIKRigHierarchy.cpp b/Engine/Plugins/Animation/IKRig/Source/IKRigEditor/Private/RigEditor/SIKRigHierarchy.cpp
--- a/Engine/Plugins/Animation/IKRig/Source/IKRigEditor/Private/RigEditor/SIKRigHierarchy.cpp
+++ b/Engine/Plugins/Animation/IKRig/Source/IKRigEditor/Private/RigEditor/SIKRigHierarchy.cpp
@@ -60,7 +60,7 @@ TWeakObjectPtr< UObject > FIKRigTreeElement::GetObject() const
 		}
 		return OptionalBoneDetails;
 	case IKRigTreeElementType::BONE_SETTINGS:
-		return Controller->AssetController->GetSettingsForBone(BoneSettingBoneName, BoneSettingsSolverIndex);
+		return Controller->AssetController->GetBoneSettings(BoneSettingBoneName, BoneSettingsSolverIndex);
 	case IKRigTreeElementType::GOAL:
 		return Controller->AssetController->GetGoal(GoalName);
 	case IKRigTreeElementType::SOLVERGOAL:
diff --git a/Engine/Plugins/Animation/IKRig/Source/IKRigEditor/Public/RetargetEditor/IKRetargetBatchOperation.h b/Engine/Plugins/Animation/IKRig/Source/IKRigEditor/Public/RetargetEditor/IKRetargetBatchOperation.h
--- a/Engine/Plugins/Animation/IKRig/Source/IKRigEditor/Public/RetargetEditor/IKRetargetBatchOperation.h
+++ b/Engine/Plugins/Animation/IKRig/Source/IKRigEditor/Public/RetargetEditor/IKRetargetBatchOperation.h
@@ -79,7 +79,7 @@ public:
 	 * 
 	 * Return: list of new animation files created. */
 	UFUNCTION(BlueprintCallable, Category=IKBatchRetarget)
-	TArray<FAssetData> RunRetarget(
+	static TArray<FAssetData> DuplicateAndRetarget(
 		const TArray<FAssetData>& AssetsToRetarget,
 		USkeletalMesh* SourceMesh,
 		USkeletalMesh* TargetMesh,
diff --git a/Engine/Plugins/Animation/IKRig/Source/IKRigEditor/Public/RetargetEditor/IKRetargeterController.h b/Engine/Plugins/Animation/IKRig/Source/IKRigEditor/Public/RetargetEditor/IKRetargeterController.h
--- a/Engine/Plugins/Animation/IKRig/Source/IKRigEditor/Public/RetargetEditor/IKRetargeterController.h
+++ b/Engine/Plugins/Animation/IKRig/Source/IKRigEditor/Public/RetargetEditor/IKRetargeterController.h
@@ -6,6 +6,9 @@
 
 #include "IKRetargeterController.generated.h"
 
+struct FRetargetGlobalSettings;
+struct FTargetRootSettings;
+struct FTargetChainSettings;
 struct FIKRetargetPose;
 enum class ERetargetSourceOrTarget : uint8;
 class FIKRetargetEditorController;
@@ -66,6 +69,34 @@ public:
 	UFUNCTION(BlueprintCallable, Category=IKRetargeter)
 	USkeletalMesh* GetPreviewMesh(const ERetargetSourceOrTarget SourceOrTarget) const;
 
+	//
+	// GET/SET SETTINGS PUBLIC/SCRIPTING API
+	//
+	
+	// Get a copy of the retarget root settings for this asset.
+	UFUNCTION(BlueprintCallable, BlueprintPure=false, Category=IKRetargeter)
+	FTargetRootSettings GetRootSettings() const;
+	
+	// Set the retarget root settings for this asset.
+	UFUNCTION(BlueprintCallable, BlueprintPure=false, Category=IKRetargeter)
+	void SetRootSettings(const FTargetRootSettings& RootSettings) const;
+
+	// Get a copy of the global settings for this asset.
+	UFUNCTION(BlueprintCallable, BlueprintPure=false, Category=IKRetargeter)
+	FRetargetGlobalSettings GetGlobalSettings() const;
+
+	// Get a copy of the global settings for this asset.
+	UFUNCTION(BlueprintCallable, BlueprintPure=false, Category=IKRetargeter)
+	void SetGlobalSettings(const FRetargetGlobalSettings& GlobalSettings) const;
+
+	// Get a copy of the settings for the target chain by name.
+	UFUNCTION(BlueprintCallable, BlueprintPure=false, Category=IKRetargeter)
+	FTargetChainSettings GetRetargetChainSettings(const FName& TargetChainName) const;
+
+	// Set the settings for the target chain by name. Returns true if the chain settings were applied, false otherwise.
+	UFUNCTION(BlueprintCallable, BlueprintPure=false, Category=IKRetargeter)
+	bool SetRetargetChainSettings(const FName& TargetChainName, const FTargetChainSettings& Settings) const;
+
 	//
 	// GENERAL C++ ONLY API
 	//
@@ -134,13 +165,13 @@ private:
 	
 	// Get names of all the bone chains.
 	void GetChainNames(const ERetargetSourceOrTarget SourceOrTarget, TArray<FName>& OutNames) const;
-	
-	// convenience to get chain settings by name
-	URetargetChainSettings* GetChainSettings(const FName& TargetChainName) const;
 
 	// Sort the Asset ChainMapping based on the StartBone of the target chains. 
 	void SortChainMapping() const;
 
+	// convenience to get chain settings UObject by name
+	URetargetChainSettings* GetChainSettings(const FName& TargetChainName) const;
+
 	//
 	// RETARGET POSE PUBLIC/SCRIPTING API
 	//
diff --git a/Engine/Plugins/Animation/IKRig/Source/IKRigEditor/Public/RigEditor/IKRigController.h b/Engine/Plugins/Animation/IKRig/Source/IKRigEditor/Public/RigEditor/IKRigController.h
--- a/Engine/Plugins/Animation/IKRig/Source/IKRigEditor/Public/RigEditor/IKRigController.h
+++ b/Engine/Plugins/Animation/IKRig/Source/IKRigEditor/Public/RigEditor/IKRigController.h
@@ -42,7 +42,7 @@ public:
 	
 	// Sets the preview mesh to use. Loads the hierarchy into the asset's IKRigSkeleton.
 	// Returns true if the mesh was able to be set. False if it was incompatible for any reason. 
-	UFUNCTION(BlueprintCallable, Category=IKRig)
+	UFUNCTION(BlueprintCallable, BlueprintPure=false, Category=IKRig)
 	bool SetSkeletalMesh(USkeletalMesh* SkeletalMesh, bool bTransact=false) const;
 
 	// Get the skeletal mesh this asset is initialized with 
@@ -185,6 +185,11 @@ public:
 	// Get access to the list of Goals. 
 	UFUNCTION(BlueprintCallable, Category=IKRig)
 	TArray<UIKRigEffectorGoal*>& GetAllGoals() const;
+
+	// Get the UObject for the settings associated with the given Goal in the given Solver.
+	// Solvers can define their own per-Goal settings depending on their needs. These are termed "Effectors".
+	UFUNCTION(BlueprintCallable, Category=IKRig)
+	UObject* GetGoalSettingsForSolver(const FName GoalName, int32 SolverIndex) const;
 	
 	//
 	// GOALS C++ ONLY API
@@ -193,10 +198,6 @@ public:
 	// Get the index of the given Goal in the list of Goals. 
 	int32 GetGoalIndex(const FName InGoalName, const ENameCase CompareMethod = ENameCase::IgnoreCase) const;
 
-	// Get the UObject for the settings associated with the given Goal in the given Solver.
-	// Solvers can define their own per-Goal settings depending on their needs. These are termed "Effectors". 
-	UObject* GetGoalSettingsForSolver(const FName GoalName, int32 SolverIndex) const;
-
 	// Get the global-space transform of the given Goal. This may be set by the user in the editor, or at runtime. 
 	FTransform GetGoalCurrentTransform(const FName GoalName) const;
 
@@ -239,8 +240,8 @@ public:
 	bool RemoveBoneSetting(const FName BoneName, int32 SolverIndex) const;
 
 	// Get the generic (Solver-specific) Bone settings UObject for this Bone in the given Solver.
-	UFUNCTION(BlueprintCallable, Category=IKRig)
-	UObject* GetSettingsForBone(const FName BoneName, int32 SolverIndex) const;
+	UFUNCTION(BlueprintCallable, BlueprintPure=false, Category=IKRig)
+	UObject* GetBoneSettings(const FName BoneName, int32 SolverIndex) const;
 
 	// Get the global-space retarget pose transform of the given Bone.
 	UFUNCTION(BlueprintCallable, Category=IKRig)
diff --git a/Engine/Plugins/Experimental/FullBodyIK/Source/PBIK/Public/Core/PBIKSolver.h b/Engine/Plugins/Experimental/FullBodyIK/Source/PBIK/Public/Core/PBIKSolver.h
--- a/Engine/Plugins/Experimental/FullBodyIK/Source/PBIK/Public/Core/PBIKSolver.h
+++ b/Engine/Plugins/Experimental/FullBodyIK/Source/PBIK/Public/Core/PBIKSolver.h
@@ -95,7 +95,7 @@ struct FEffector
 	
 } // namespace
 
-UENUM()
+UENUM(BlueprintType)
 enum class EPBIKRootBehavior : uint8
 {
 	PrePull,
@@ -109,35 +109,35 @@ struct FRootPrePullSettings
 	GENERATED_BODY()
 	
 	/** Range 0-1, default is 0. Apply a large scale rotation offset to the entire body prior to constraint solving.*/
-	UPROPERTY(EditAnywhere, Category = RootRotation, meta = (ClampMin = "0.0", ClampMax = "1.0", UIMin = "0.0", UIMax = "1.0"))
+	UPROPERTY(BlueprintReadWrite, EditAnywhere, Category = RootRotation, meta = (ClampMin = "0.0", ClampMax = "1.0", UIMin = "0.0", UIMax = "1.0"))
 	float RotationAlpha = 0.0f;
 
 	/** Range 0-1, default is 1. Blend contribution to rotation offset in the X axis in component space.*/
-	UPROPERTY(EditAnywhere, Category = RootRotation, meta = (ClampMin = "0.0", ClampMax = "1.0", UIMin = "0.0", UIMax = "1.0"))
+	UPROPERTY(BlueprintReadWrite, EditAnywhere, Category = RootRotation, meta = (ClampMin = "0.0", ClampMax = "1.0", UIMin = "0.0", UIMax = "1.0"))
 	float RotationAlphaX = 1.0f;
 
 	/** Range 0-1, default is 1. Blend contribution to rotation offset in the Y axis in component space.*/
-	UPROPERTY(EditAnywhere, Category = RootRotation, meta = (ClampMin = "0.0", ClampMax = "1.0", UIMin = "0.0", UIMax = "1.0"))
+	UPROPERTY(BlueprintReadWrite, EditAnywhere, Category = RootRotation, meta = (ClampMin = "0.0", ClampMax = "1.0", UIMin = "0.0", UIMax = "1.0"))
 	float RotationAlphaY = 1.0f;
 
 	/** Range 0-1, default is 1. Blend contribution to rotation offset in the Z axis in component space.*/
-	UPROPERTY(EditAnywhere, Category = RootRotation, meta = (ClampMin = "0.0", ClampMax = "1.0", UIMin = "0.0", UIMax = "1.0"))
+	UPROPERTY(BlueprintReadWrite, EditAnywhere, Category = RootRotation, meta = (ClampMin = "0.0", ClampMax = "1.0", UIMin = "0.0", UIMax = "1.0"))
 	float RotationAlphaZ = 1.0f;
 
 	/** Range 0-1, default is 1. Apply a large scale position offset to the entire body prior to constraint solving.*/
-	UPROPERTY(EditAnywhere, Category = RootPosition, meta = (ClampMin = "0.0", ClampMax = "1.0", UIMin = "0.0", UIMax = "1.0"))
+	UPROPERTY(BlueprintReadWrite, EditAnywhere, Category = RootPosition, meta = (ClampMin = "0.0", ClampMax = "1.0", UIMin = "0.0", UIMax = "1.0"))
 	float PositionAlpha = 1.0f;
 
 	/** Range 0-1, default is 1. Blend contribution to position offset in the X axis in component space.*/
-	UPROPERTY(EditAnywhere, Category = RootPosition, meta = (ClampMin = "0.0", ClampMax = "1.0", UIMin = "0.0", UIMax = "1.0"))
+	UPROPERTY(BlueprintReadWrite, EditAnywhere, Category = RootPosition, meta = (ClampMin = "0.0", ClampMax = "1.0", UIMin = "0.0", UIMax = "1.0"))
 	float PositionAlphaX = 1.0f;
 
 	/** Range 0-1, default is 1. Blend contribution to position offset in the Y axis in component space.*/
-	UPROPERTY(EditAnywhere, Category = RootPosition, meta = (ClampMin = "0.0", ClampMax = "1.0", UIMin = "0.0", UIMax = "1.0"))
+	UPROPERTY(BlueprintReadWrite, EditAnywhere, Category = RootPosition, meta = (ClampMin = "0.0", ClampMax = "1.0", UIMin = "0.0", UIMax = "1.0"))
 	float PositionAlphaY = 1.0f;
 
 	/** Range 0-1, default is 1. Blend contribution to position offset in the Z axis in component space.*/
-	UPROPERTY(EditAnywhere, Category = RootPosition, meta = (ClampMin = "0.0", ClampMax = "1.0", UIMin = "0.0", UIMax = "1.0"))
+	UPROPERTY(BlueprintReadWrite, EditAnywhere, Category = RootPosition, meta = (ClampMin = "0.0", ClampMax = "1.0", UIMin = "0.0", UIMax = "1.0"))
 	float PositionAlphaZ = 1.0f;
 };
 
@@ -147,44 +147,44 @@ struct PBIK_API FPBIKSolverSettings
 	GENERATED_BODY()
 
 	/** High iteration counts can help solve complex joint configurations with competing constraints, but will increase runtime cost. Default is 20. */
-	UPROPERTY(EditAnywhere, Category = SolverSettings, meta = (ClampMin = "0", ClampMax = "1000", UIMin = "0.0", UIMax = "200.0"))
+	UPROPERTY(BlueprintReadWrite, EditAnywhere, Category = SolverSettings, meta = (ClampMin = "0", ClampMax = "1000", UIMin = "0.0", UIMax = "200.0"))
 	int32 Iterations = 20;
 
 	/** A global mass multiplier; higher values will make the joints more stiff, but require more iterations. Typical range is 0.0 to 10.0. */
-	UPROPERTY(EditAnywhere, Category = SolverSettings, meta = (ClampMin = "0", UIMin = "0.0", UIMax = "10.0"))
+	UPROPERTY(BlueprintReadWrite, EditAnywhere, Category = SolverSettings, meta = (ClampMin = "0", UIMin = "0.0", UIMax = "10.0"))
 	float MassMultiplier = 1.0f;
 
 	/** If true, joints will translate to reach the effectors; causing bones to lengthen if necessary. Good for cartoon effects. Default is false. */
-	UPROPERTY(EditAnywhere, Category = SolverSettings)
+	UPROPERTY(BlueprintReadWrite, EditAnywhere, Category = SolverSettings)
 	bool bAllowStretch = false;
 
 	/** (Default is PrePull) Set the behavior of the solver root.
 	*Pre Pull: translates and rotates the root (and all children) by the average motion of the stretched effectors to help achieve faster convergence when reaching far.
 	*Pin to Input: locks the translation and rotation of the root bone to the input pose. Overrides any bone settings applied to the root. Good for partial-body solves.
 	*Free: treats the root bone like any other and allows it to move freely or according to any bone settings applied to it. */
-	UPROPERTY(EditAnywhere, Category = RootBehavior)
+	UPROPERTY(BlueprintReadWrite, EditAnywhere, Category = RootBehavior)
 	EPBIKRootBehavior RootBehavior = EPBIKRootBehavior::PrePull;
 
 	/** Settings only applicable when Root Behavior is set to "PrePull". Use these values to adjust the gross movement and orientation of the entire skeleton. */
-	UPROPERTY(EditAnywhere, Category = RootBehavior)
+	UPROPERTY(BlueprintReadWrite, EditAnywhere, Category = RootBehavior)
 	FRootPrePullSettings PrePullRootSettings;
 
 	/** A global multiplier for all Pull Chain Alpha values on all effectors. Range is 0.0 to 1.0. Default is 1.0. */
-	UPROPERTY(EditAnywhere, Category = AdvancedSettings, meta = (ClampMin = "0", ClampMax = "1", UIMin = "0.0", UIMax = "1.0"))
+	UPROPERTY(BlueprintReadWrite, EditAnywhere, Category = AdvancedSettings, meta = (ClampMin = "0", ClampMax = "1", UIMin = "0.0", UIMax = "1.0"))
 	float GlobalPullChainAlpha = 1.0f;
 
 	/** Maximum angle that a joint can be rotated per constraint iteration. Lower this value if the solve is diverging. Range is 0.0 to 180.0. Default is 30. */
-	UPROPERTY(EditAnywhere, Category = AdvancedSettings, meta = (ClampMin = "0", ClampMax = "45", UIMin = "0.0", UIMax = "180.0"))
+	UPROPERTY(BlueprintReadWrite, EditAnywhere, Category = AdvancedSettings, meta = (ClampMin = "0", ClampMax = "45", UIMin = "0.0", UIMax = "180.0"))
 	float MaxAngle = 30.f;
 
 	/** Pushes constraints beyond their normal amount to speed up convergence. Increasing this may speed up convergence, but at the cost of stability. Range is 1.0 - 2.0. Default is 1.3. */
-	UPROPERTY(EditAnywhere, Category = AdvancedSettings, meta = (ClampMin = "1", UIMin = "1.0", UIMax = "2.0"))
+	UPROPERTY(BlueprintReadWrite, EditAnywhere, Category = AdvancedSettings, meta = (ClampMin = "1", UIMin = "1.0", UIMax = "2.0"))
 	float OverRelaxation = 1.3f;
 	
 	/** When true, the solver is reset each tick to start from the current input pose. Default is true.
 	 * If false, incoming animated poses are ignored and the solver starts from the results of the previous solve.
 	 * In very limited circumstances, it can be beneficial to use the pose from the previous frame.*/
-	UPROPERTY(EditAnywhere, Category = AdvancedSettings)
+	UPROPERTY(BlueprintReadWrite, EditAnywhere, Category = AdvancedSettings)
 	bool bStartSolveFromInputPose = true;
 };
 
