IK Rig: fixes clashing LOCTEXT labels in IKRigEditorControlller.cpp

diff --git a/Engine/Plugins/Animation/IKRig/Source/IKRigEditor/Private/RigEditor/IKRigEditorController.cpp b/Engine/Plugins/Animation/IKRig/Source/IKRigEditor/Private/RigEditor/IKRigEditorController.cpp
--- a/Engine/Plugins/Animation/IKRig/Source/IKRigEditor/Private/RigEditor/IKRigEditorController.cpp
+++ b/Engine/Plugins/Animation/IKRig/Source/IKRigEditor/Private/RigEditor/IKRigEditorController.cpp
@@ -924,11 +924,11 @@ FName FIKRigEditorController::PromptToAddNewRetargetChain(FBoneChain& BoneChain)
 	const bool bHasExistingGoal = BoneChain.IKGoalName != NAME_None;
 	if (bHasExistingGoal)
 	{
-		Buttons.Add(SCustomDialog::FButton(LOCTEXT("AddChain", "Add Chain using Goal")));
+		Buttons.Add(SCustomDialog::FButton(LOCTEXT("AddChainUsingGoal", "Add Chain using Goal")));
 	}
 	else
 	{
-		Buttons.Add(SCustomDialog::FButton(LOCTEXT("AddChain", "Add Chain and Goal")));
+		Buttons.Add(SCustomDialog::FButton(LOCTEXT("AddChainAndGoal", "Add Chain and Goal")));
 	}
 	
 	Buttons.Add(SCustomDialog::FButton(LOCTEXT("Cancel", "Cancel")));
