IK Rig: fixing static analysis warning issue 254362

diff --git a/Engine/Plugins/Animation/IKRig/Source/IKRigEditor/Private/RetargetEditor/IKRetargeterController.cpp b/Engine/Plugins/Animation/IKRig/Source/IKRigEditor/Private/RetargetEditor/IKRetargeterController.cpp
--- a/Engine/Plugins/Animation/IKRig/Source/IKRigEditor/Private/RetargetEditor/IKRetargeterController.cpp
+++ b/Engine/Plugins/Animation/IKRig/Source/IKRigEditor/Private/RetargetEditor/IKRetargeterController.cpp
@@ -157,8 +157,7 @@ FTargetChainSettings UIKRetargeterController::GetRetargetChainSettings(const FNa
 
 bool UIKRetargeterController::SetRetargetChainSettings(const FName& TargetChainName, const FTargetChainSettings& Settings) const
 {
-	URetargetChainSettings* ChainSettings = GetChainSettings(TargetChainName);
-	if (!ChainSettings)
+	if (URetargetChainSettings* ChainSettings = GetChainSettings(TargetChainName))
 	{
 		ChainSettings->Settings = Settings;
 		return true;
