Renamed PBIK to FBIK.
Exposed IK Goal properties to BP.
Exposed Body Mover properties to BP.

#Swarm https://p4-swarm.epicgames.net/reviews/23877321

diff --git a/Engine/Plugins/Animation/IKRig/Config/BaseIKRig.ini b/Engine/Plugins/Animation/IKRig/Config/BaseIKRig.ini
--- a/Engine/Plugins/Animation/IKRig/Config/BaseIKRig.ini
+++ b/Engine/Plugins/Animation/IKRig/Config/BaseIKRig.ini
@@ -1,3 +1,6 @@
-+ClassRedirects=(OldName="AnimNode_IKRetargeter",NewName="AnimNode_RetargetPoseFromMesh")
-
+[CoreRedirects]
++ClassRedirects=(OldName="/Script/IKRig.IKRigPBIKSolver", NewName="/Script/IKRig.IKRigFBIKSolver")
++ClassRedirects=(OldName="/Script/IKRig.IKRig_PBIKBoneSettings", NewName="/Script/IKRig.IKRig_FBIKBoneSettings")
++ClassRedirects=(OldName="/Script/IKRig.AnimNode_IKRetargeter", NewName="/Script/IKRig.AnimNode_RetargetPoseFromMesh")
 +PropertyRedirects=(OldName="/Script/IKRig.TargetChainIKSettings.StaticOffset",NewName="TranslationGlobalOffset")
+
diff --git a/Engine/Plugins/Animation/IKRig/Source/IKRig/Private/Solvers/IKRig_PBIKSolver.cpp b/Engine/Plugins/Animation/IKRig/Source/IKRig/Private/Solvers/IKRig_PBIKSolver.cpp
--- a/Engine/Plugins/Animation/IKRig/Source/IKRig/Private/Solvers/IKRig_PBIKSolver.cpp
+++ b/Engine/Plugins/Animation/IKRig/Source/IKRig/Private/Solvers/IKRig_PBIKSolver.cpp
@@ -1,15 +1,15 @@
 // Copyright Epic Games, Inc. All Rights Reserved.
 
-#include "Solvers/IKRig_PBIKSolver.h"
+#include "Solvers/IKRig_FBIKSolver.h"
 
 #include "IKRigDataTypes.h"
 #include "IKRigSkeleton.h"
 
-#include UE_INLINE_GENERATED_CPP_BY_NAME(IKRig_PBIKSolver)
+#include UE_INLINE_GENERATED_CPP_BY_NAME(IKRig_FBIKSolver)
 
-#define LOCTEXT_NAMESPACE "UIKRigPBIKSolver"
+#define LOCTEXT_NAMESPACE "UIKRigFBIKSolver"
 
-void UIKRigPBIKSolver::Initialize(const FIKRigSkeleton& InSkeleton)
+void UIKRigFBIKSolver::Initialize(const FIKRigSkeleton& InSkeleton)
 {	
 	// check how many effectors are assigned to a bone
 	int NumEffectors = 0;
@@ -56,7 +56,7 @@ void UIKRigPBIKSolver::Initialize(const FIKRigSkeleton& InSkeleton)
 	Solver.Initialize();
 }
 
-void UIKRigPBIKSolver::Solve(FIKRigSkeleton& IKRigSkeleton, const FIKRigGoalContainer& Goals)
+void UIKRigFBIKSolver::Solve(FIKRigSkeleton& IKRigSkeleton, const FIKRigGoalContainer& Goals)
 {
 	if (!Solver.IsReadyToSimulate())
 	{
@@ -77,7 +77,7 @@ void UIKRigPBIKSolver::Solve(FIKRigSkeleton& IKRigSkeleton, const FIKRigGoalCont
 	}
 
 	// update bone settings
-	for (const UIKRig_PBIKBoneSettings* BoneSetting : BoneSettings)
+	for (const UIKRig_FBIKBoneSettings* BoneSetting : BoneSettings)
 	{
 		const int32 BoneIndex = Solver.GetBoneIndex(BoneSetting->Bone);
 		if (PBIK::FBoneSettings* InternalSettings = Solver.GetBoneSettings(BoneIndex))
@@ -136,9 +136,9 @@ void UIKRigPBIKSolver::Solve(FIKRigSkeleton& IKRigSkeleton, const FIKRigGoalCont
 	}
 }
 
-void UIKRigPBIKSolver::GetBonesWithSettings(TSet<FName>& OutBonesWithSettings) const
+void UIKRigFBIKSolver::GetBonesWithSettings(TSet<FName>& OutBonesWithSettings) const
 {
-	for (UIKRig_PBIKBoneSettings* BoneSetting : BoneSettings)
+	for (UIKRig_FBIKBoneSettings* BoneSetting : BoneSettings)
 	{
 		if (BoneSetting)
 		{
@@ -147,9 +147,9 @@ void UIKRigPBIKSolver::GetBonesWithSettings(TSet<FName>& OutBonesWithSettings) c
 	}
 }
 
-void UIKRigPBIKSolver::UpdateSolverSettings(UIKRigSolver* InSettings)
+void UIKRigFBIKSolver::UpdateSolverSettings(UIKRigSolver* InSettings)
 {
-	if(UIKRigPBIKSolver* Settings = Cast<UIKRigPBIKSolver>(InSettings))
+	if(UIKRigFBIKSolver* Settings = Cast<UIKRigFBIKSolver>(InSettings))
 	{
 		Iterations = Settings->Iterations;
 		MassMultiplier = Settings->MassMultiplier;
@@ -175,9 +175,9 @@ void UIKRigPBIKSolver::UpdateSolverSettings(UIKRigSolver* InSettings)
 		}
 
 		// copy bone settings
-		for (const UIKRig_PBIKBoneSettings* InBoneSetting : Settings->BoneSettings)
+		for (const UIKRig_FBIKBoneSettings* InBoneSetting : Settings->BoneSettings)
 		{
-			for (UIKRig_PBIKBoneSettings* BoneSetting : BoneSettings)
+			for (UIKRig_FBIKBoneSettings* BoneSetting : BoneSettings)
 			{
 				if (BoneSetting && BoneSetting->Bone == InBoneSetting->Bone)
 				{
@@ -189,7 +189,7 @@ void UIKRigPBIKSolver::UpdateSolverSettings(UIKRigSolver* InSettings)
 	}
 }
 
-void UIKRigPBIKSolver::RemoveGoal(const FName& GoalName)
+void UIKRigFBIKSolver::RemoveGoal(const FName& GoalName)
 {
 	// ensure goal even exists in this solver
 	const int32 GoalIndex = GetIndexOfGoal(GoalName);
@@ -204,12 +204,12 @@ void UIKRigPBIKSolver::RemoveGoal(const FName& GoalName)
 
 #if WITH_EDITOR
 
-FText UIKRigPBIKSolver::GetNiceName() const
+FText UIKRigFBIKSolver::GetNiceName() const
 {
 	return FText(LOCTEXT("SolverName", "Full Body IK"));
 }
 
-bool UIKRigPBIKSolver::GetWarningMessage(FText& OutWarningMessage) const
+bool UIKRigFBIKSolver::GetWarningMessage(FText& OutWarningMessage) const
 {
 	if (RootBone == NAME_None)
 	{
@@ -226,7 +226,7 @@ bool UIKRigPBIKSolver::GetWarningMessage(FText& OutWarningMessage) const
 	return false;
 }
 
-void UIKRigPBIKSolver::AddGoal(const UIKRigEffectorGoal& NewGoal)
+void UIKRigFBIKSolver::AddGoal(const UIKRigEffectorGoal& NewGoal)
 {
 	UIKRig_FBIKEffector* NewEffector = NewObject<UIKRig_FBIKEffector>(this, UIKRig_FBIKEffector::StaticClass());
 	NewEffector->GoalName = NewGoal.GoalName;
@@ -234,7 +234,7 @@ void UIKRigPBIKSolver::AddGoal(const UIKRigEffectorGoal& NewGoal)
 	Effectors.Add(NewEffector);
 }
 
-void UIKRigPBIKSolver::RenameGoal(const FName& OldName, const FName& NewName)
+void UIKRigFBIKSolver::RenameGoal(const FName& OldName, const FName& NewName)
 {
 	// ensure goal even exists in this solver
 	const int32 GoalIndex = GetIndexOfGoal(OldName);
@@ -248,7 +248,7 @@ void UIKRigPBIKSolver::RenameGoal(const FName& OldName, const FName& NewName)
 	Effectors[GoalIndex]->GoalName = NewName;
 }
 
-void UIKRigPBIKSolver::SetGoalBone(const FName& GoalName, const FName& NewBoneName)
+void UIKRigFBIKSolver::SetGoalBone(const FName& GoalName, const FName& NewBoneName)
 {
 	// ensure goal even exists in this solver
 	const int32 GoalIndex = GetIndexOfGoal(GoalName);
@@ -262,17 +262,17 @@ void UIKRigPBIKSolver::SetGoalBone(const FName& GoalName, const FName& NewBoneNa
 	Effectors[GoalIndex]->BoneName = NewBoneName;
 }
 
-bool UIKRigPBIKSolver::IsGoalConnected(const FName& GoalName) const
+bool UIKRigFBIKSolver::IsGoalConnected(const FName& GoalName) const
 {
 	return GetIndexOfGoal(GoalName) != INDEX_NONE;
 }
 
-void UIKRigPBIKSolver::SetRootBone(const FName& RootBoneName)
+void UIKRigFBIKSolver::SetRootBone(const FName& RootBoneName)
 {
 	RootBone = RootBoneName;
 }
 
-UObject* UIKRigPBIKSolver::GetGoalSettings(const FName& GoalName) const
+UObject* UIKRigFBIKSolver::GetGoalSettings(const FName& GoalName) const
 {
 	const int32 GoalIndex = GetIndexOfGoal(GoalName);
 	if (GoalIndex == INDEX_NONE)
@@ -283,22 +283,22 @@ UObject* UIKRigPBIKSolver::GetGoalSettings(const FName& GoalName) const
 	return Effectors[GoalIndex];
 }
 
-void UIKRigPBIKSolver::AddBoneSetting(const FName& BoneName)
+void UIKRigFBIKSolver::AddBoneSetting(const FName& BoneName)
 {
 	if (GetBoneSetting(BoneName))
 	{
 		return; // already have settings on this bone
 	}
 
-	UIKRig_PBIKBoneSettings* NewBoneSettings = NewObject<UIKRig_PBIKBoneSettings>(this, UIKRig_PBIKBoneSettings::StaticClass());
+	UIKRig_FBIKBoneSettings* NewBoneSettings = NewObject<UIKRig_FBIKBoneSettings>(this, UIKRig_FBIKBoneSettings::StaticClass());
 	NewBoneSettings->Bone = BoneName;
 	BoneSettings.Add(NewBoneSettings);
 }
 
-void UIKRigPBIKSolver::RemoveBoneSetting(const FName& BoneName)
+void UIKRigFBIKSolver::RemoveBoneSetting(const FName& BoneName)
 {
-	UIKRig_PBIKBoneSettings* BoneSettingToRemove = nullptr; 
-	for (UIKRig_PBIKBoneSettings* BoneSetting : BoneSettings)
+	UIKRig_FBIKBoneSettings* BoneSettingToRemove = nullptr; 
+	for (UIKRig_FBIKBoneSettings* BoneSetting : BoneSettings)
 	{
 		if (BoneSetting->Bone == BoneName)
 		{
@@ -313,9 +313,9 @@ void UIKRigPBIKSolver::RemoveBoneSetting(const FName& BoneName)
 	}
 }
 
-UObject* UIKRigPBIKSolver::GetBoneSetting(const FName& BoneName) const
+UObject* UIKRigFBIKSolver::GetBoneSetting(const FName& BoneName) const
 {
-	for (UIKRig_PBIKBoneSettings* BoneSetting : BoneSettings)
+	for (UIKRig_FBIKBoneSettings* BoneSetting : BoneSettings)
 	{
 		if (BoneSetting && BoneSetting->Bone == BoneName)
 		{
@@ -326,7 +326,7 @@ UObject* UIKRigPBIKSolver::GetBoneSetting(const FName& BoneName) const
 	return nullptr;
 }
 
-void UIKRigPBIKSolver::DrawBoneSettings(
+void UIKRigFBIKSolver::DrawBoneSettings(
 	const FName& BoneName,
 	const FIKRigSkeleton& IKRigSkeleton,
 	FPrimitiveDrawInterface* PDI) const
@@ -334,7 +334,7 @@ void UIKRigPBIKSolver::DrawBoneSettings(
 	
 }
 
-bool UIKRigPBIKSolver::IsBoneAffectedBySolver(const FName& BoneName, const FIKRigSkeleton& IKRigSkeleton) const
+bool UIKRigFBIKSolver::IsBoneAffectedBySolver(const FName& BoneName, const FIKRigSkeleton& IKRigSkeleton) const
 {
 	// nothing is affected by solver without a root bone assigned or at least 1 effector
 	if (RootBone == NAME_None || Effectors.IsEmpty())
@@ -360,7 +360,7 @@ bool UIKRigPBIKSolver::IsBoneAffectedBySolver(const FName& BoneName, const FIKRi
 	return false;
 }
 
-void UIKRigPBIKSolver::PostLoad()
+void UIKRigFBIKSolver::PostLoad()
 {
 	Super::PostLoad();
 	
@@ -382,7 +382,7 @@ void UIKRigPBIKSolver::PostLoad()
 
 #endif
 
-int32 UIKRigPBIKSolver::GetIndexOfGoal(const FName& OldName) const
+int32 UIKRigFBIKSolver::GetIndexOfGoal(const FName& OldName) const
 {
 	for (int32 i=0; i<Effectors.Num(); ++i)
 	{
diff --git a/Engine/Plugins/Animation/IKRig/Source/IKRig/Public/IKRigDefinition.h b/Engine/Plugins/Animation/IKRig/Source/IKRig/Public/IKRigDefinition.h
--- a/Engine/Plugins/Animation/IKRig/Source/IKRig/Public/IKRigDefinition.h
+++ b/Engine/Plugins/Animation/IKRig/Source/IKRig/Public/IKRigDefinition.h
@@ -17,7 +17,7 @@ class IPropertyHandle;
 class UIKRigSolver;
 
 #if WITH_EDITORONLY_DATA
-UENUM()
+UENUM(BlueprintType)
 enum class EIKRigGoalPreviewMode : uint8
 {
 	Additive		UMETA(DisplayName = "Additive"),
@@ -34,7 +34,7 @@ namespace EIKRigTransformType
 }
 #endif
 
-UCLASS()
+UCLASS(BlueprintType)
 class IKRIG_API UIKRigEffectorGoal : public UObject
 {
 	GENERATED_BODY()
@@ -43,27 +43,27 @@ public:
 
 	/** The name used to refer to this goal from outside systems.
 	 * This is the name to use when referring to this Goal from Blueprint, Anim Graph, Control Rig or IK Retargeter.*/
-	UPROPERTY(VisibleAnywhere, Category = "Goal Settings")
+	UPROPERTY(BlueprintReadOnly, VisibleAnywhere, Category = "Goal Settings")
 	FName GoalName;
 
 	/** The name of the bone that this Goal is located at.*/
-	UPROPERTY(VisibleAnywhere, Category = "Goal Settings")
+	UPROPERTY(BlueprintReadOnly, VisibleAnywhere, Category = "Goal Settings")
 	FName BoneName;
 
 	/** Range 0-1, default is 1. Blend between the input bone position (0.0) and the current goal position (1.0).*/
-	UPROPERTY(EditAnywhere, Category = "Goal Settings", meta = (ClampMin = "0.0", ClampMax = "1.0", UIMin = "0.0", UIMax = "1.0"))
+	UPROPERTY(BlueprintReadWrite, EditAnywhere, Category = "Goal Settings", meta = (ClampMin = "0.0", ClampMax = "1.0", UIMin = "0.0", UIMax = "1.0"))
 	float PositionAlpha = 1.0f;
 	
 	/** Range 0-1, default is 1. Blend between the input bone rotation (0.0) and the current goal rotation (1.0).*/
-	UPROPERTY(EditAnywhere, Category = "Goal Settings", meta = (ClampMin = "0.0", ClampMax = "1.0", UIMin = "0.0", UIMax = "1.0"))
+	UPROPERTY(BlueprintReadWrite, EditAnywhere, Category = "Goal Settings", meta = (ClampMin = "0.0", ClampMax = "1.0", UIMin = "0.0", UIMax = "1.0"))
 	float RotationAlpha = 1.0f;
 
 	/** The current transform of this Goal, in the Global Space of the character.*/
-	UPROPERTY(EditAnywhere, Category = "Goal Settings")
+	UPROPERTY(BlueprintReadWrite, EditAnywhere, Category = "Goal Settings")
 	FTransform CurrentTransform;
 
 	/** The initial transform of this Goal, as defined by the initial transform of the Goal's bone in the retarget pose.*/
-	UPROPERTY(VisibleAnywhere, Category = "Goal Settings")
+	UPROPERTY(BlueprintReadOnly, VisibleAnywhere, Category = "Goal Settings")
 	FTransform InitialTransform;
 
 	bool operator==(const UIKRigEffectorGoal& Other) const { return GoalName == Other.GoalName; }
@@ -74,23 +74,23 @@ public:
 	* "Additive" interprets the Goal transform as being relative to the input pose. Useful for previewing animations.
 	* "Absolute" pins the Goal transform to the Gizmo in the viewport.
 	*/
-	UPROPERTY(EditAnywhere, Category = "Goal Settings")
+	UPROPERTY(BlueprintReadWrite, EditAnywhere, Category = "Goal Settings")
 	EIKRigGoalPreviewMode PreviewMode = EIKRigGoalPreviewMode::Additive;
 	
 	/**The size of the Goal gizmo drawing in the editor viewport.*/
-	UPROPERTY(EditAnywhere, Category = "Viewport Goal Settings", meta = (ClampMin = "0.1", ClampMax = "1000.0", UIMin = "0.1", UIMax = "100.0"))
+	UPROPERTY(BlueprintReadWrite, EditAnywhere, Category = "Viewport Goal Settings", meta = (ClampMin = "0.1", ClampMax = "1000.0", UIMin = "0.1", UIMax = "100.0"))
 	float SizeMultiplier = 1.0f;
 
 	/**The thickness of the Goal gizmo drawing in the editor viewport.*/
-	UPROPERTY(EditAnywhere, Category = "Viewport Goal Settings",  meta = (ClampMin = "0.0", ClampMax = "10.0", UIMin = "0.0", UIMax = "5.0"))
+	UPROPERTY(BlueprintReadWrite, EditAnywhere, Category = "Viewport Goal Settings",  meta = (ClampMin = "0.0", ClampMax = "10.0", UIMin = "0.0", UIMax = "5.0"))
 	float ThicknessMultiplier = 0.7f;
 
 	/** Should position data be exposed in Blueprint */
-	UPROPERTY(EditAnywhere, Category = "Exposure")
+	UPROPERTY(BlueprintReadWrite, EditAnywhere, Category = "Exposure")
 	bool bExposePosition = false;
 
 	/** Should rotation data be exposed in Blueprint */
-	UPROPERTY(EditAnywhere, Category = "Exposure")
+	UPROPERTY(BlueprintReadWrite, EditAnywhere, Category = "Exposure")
 	bool bExposeRotation = false;
 	
 	virtual void PostLoad() override
diff --git a/Engine/Plugins/Animation/IKRig/Source/IKRig/Public/Solvers/IKRig_BodyMover.h b/Engine/Plugins/Animation/IKRig/Source/IKRig/Public/Solvers/IKRig_BodyMover.h
--- a/Engine/Plugins/Animation/IKRig/Source/IKRig/Public/Solvers/IKRig_BodyMover.h
+++ b/Engine/Plugins/Animation/IKRig/Source/IKRig/Public/Solvers/IKRig_BodyMover.h
@@ -38,47 +38,47 @@ public:
 	FName RootBone;
 
 	/** Blend the translational effect of this solver on/off. Range is 0-1. Default is 1.0. */
-	UPROPERTY(EditAnywhere, Category = "Body Mover Settings", meta = (UIMin = "0.0", UIMax = "1.0"))
+	UPROPERTY(BlueprintReadWrite, EditAnywhere, Category = "Body Mover Settings", meta = (UIMin = "0.0", UIMax = "1.0"))
 	float PositionAlpha = 1.0f;
 
 	/** Multiply the POSITIVE X translation. Range is 0-1. Default is 1.0. */
-	UPROPERTY(EditAnywhere, Category = "Body Mover Settings", meta = (UIMin = "0.0", UIMax = "1.0"))
+	UPROPERTY(BlueprintReadWrite, EditAnywhere, Category = "Body Mover Settings", meta = (UIMin = "0.0", UIMax = "1.0"))
 	float PositionPositiveX = 1.0f;
 
 	/** Multiply the NEGATIVE X translation. Range is 0-1. Default is 1.0. */
-	UPROPERTY(EditAnywhere, Category = "Body Mover Settings", meta = (UIMin = "0.0", UIMax = "1.0"))
+	UPROPERTY(BlueprintReadWrite, EditAnywhere, Category = "Body Mover Settings", meta = (UIMin = "0.0", UIMax = "1.0"))
 	float PositionNegativeX = 1.0f;
 
 	/** Multiply the POSITIVE Y translation. Range is 0-1. Default is 1.0. */
-	UPROPERTY(EditAnywhere, Category = "Body Mover Settings", meta = (UIMin = "0.0", UIMax = "1.0"))
+	UPROPERTY(BlueprintReadWrite, EditAnywhere, Category = "Body Mover Settings", meta = (UIMin = "0.0", UIMax = "1.0"))
 	float PositionPositiveY = 1.0f;
 
 	/** Multiply the NEGATIVE Y translation. Range is 0-1. Default is 1.0. */
-	UPROPERTY(EditAnywhere, Category = "Body Mover Settings", meta = (UIMin = "0.0", UIMax = "1.0"))
+	UPROPERTY(BlueprintReadWrite, EditAnywhere, Category = "Body Mover Settings", meta = (UIMin = "0.0", UIMax = "1.0"))
 	float PositionNegativeY = 1.0f;
 
 	/** Multiply the POSITIVE Z translation. Range is 0-1. Default is 1.0. */
-	UPROPERTY(EditAnywhere, Category = "Body Mover Settings", meta = (UIMin = "0.0", UIMax = "1.0"))
+	UPROPERTY(BlueprintReadWrite, EditAnywhere, Category = "Body Mover Settings", meta = (UIMin = "0.0", UIMax = "1.0"))
 	float PositionPositiveZ = 1.0f;
 
 	/** Multiply the NEGATIVE Z translation. Range is 0-1. Default is 1.0. */
-	UPROPERTY(EditAnywhere, Category = "Body Mover Settings", meta = (UIMin = "0.0", UIMax = "1.0"))
+	UPROPERTY(BlueprintReadWrite, EditAnywhere, Category = "Body Mover Settings", meta = (UIMin = "0.0", UIMax = "1.0"))
 	float PositionNegativeZ = 1.0f;
 
 	/** Blend the total rotational effect on/off. Range is 0-1. Default is 1.0. */
-	UPROPERTY(EditAnywhere, Category = "Body Mover Settings", meta = (UIMin = "0.0", UIMax = "1.0"))
+	UPROPERTY(BlueprintReadWrite, EditAnywhere, Category = "Body Mover Settings", meta = (UIMin = "0.0", UIMax = "1.0"))
 	float RotationAlpha = 1.0f;
 
 	/** Blend the X-axis rotational effect on/off. Range is 0-1. Default is 1.0. */
-	UPROPERTY(EditAnywhere, Category = "Body Mover Settings", meta = (UIMin = "0.0", UIMax = "1.0"))
+	UPROPERTY(BlueprintReadWrite, EditAnywhere, Category = "Body Mover Settings", meta = (UIMin = "0.0", UIMax = "1.0"))
 	float RotateXAlpha = 1.0f;
 
 	/** Blend the Y-axis rotational effect on/off. Range is 0-1. Default is 1.0. */
-	UPROPERTY(EditAnywhere, Category = "Body Mover Settings", meta = (UIMin = "0.0", UIMax = "1.0"))
+	UPROPERTY(BlueprintReadWrite, EditAnywhere, Category = "Body Mover Settings", meta = (UIMin = "0.0", UIMax = "1.0"))
 	float RotateYAlpha = 1.0f;
 
 	/** Blend the Z-axis rotational effect on/off. Range is 0-1. Default is 1.0. */
-	UPROPERTY(EditAnywhere, Category = "Body Mover Settings", meta = (UIMin = "0.0", UIMax = "1.0"))
+	UPROPERTY(BlueprintReadWrite, EditAnywhere, Category = "Body Mover Settings", meta = (UIMin = "0.0", UIMax = "1.0"))
 	float RotateZAlpha = 1.0f;
 	
 	UPROPERTY()
diff --git a/Engine/Plugins/Animation/IKRig/Source/IKRig/Public/Solvers/IKRig_FBIKSolver.h b/Engine/Plugins/Animation/IKRig/Source/IKRig/Public/Solvers/IKRig_FBIKSolver.h
new file mode 100644
--- /dev/null
+++ b/Engine/Plugins/Animation/IKRig/Source/IKRig/Public/Solvers/IKRig_FBIKSolver.h
@@ -0,0 +1,281 @@
+// Copyright Epic Games, Inc. All Rights Reserved.
+
+#pragma once
+
+#include "IKRigDefinition.h"
+#include "IKRigSolver.h"
+#include "Core/PBIKSolver.h"
+#include "PBIK_Shared.h"
+
+#include "IKRig_FBIKSolver.generated.h"
+
+UCLASS(BlueprintType)
+class IKRIG_API UIKRig_FBIKEffector : public UObject
+{
+	GENERATED_BODY()
+
+public:
+	UIKRig_FBIKEffector() { SetFlags(RF_Transactional); }
+	
+	/** The Goal that is driving this effector's transform. */
+	UPROPERTY(BlueprintReadOnly, VisibleAnywhere, Category = "Full Body IK Effector")
+	FName GoalName;
+	
+	/** The bone that this effector will pull on. */
+	UPROPERTY(BlueprintReadOnly, VisibleAnywhere, Category = "Full Body IK Effector")
+	FName BoneName;
+
+	/** Range 0-1 (default is 1.0). The strength of the effector when pulling the bone towards it's target location.
+	* At 0.0, the effector does not pull at all, but the bones between the effector and the root will still slightly resist motion from other effectors.
+	* This can thus act as a "stabilizer" for parts of the body that you do not want to behave in a pure FK fashion.
+	*/
+	UPROPERTY(BlueprintReadWrite, EditAnywhere, Category = "Full Body IK Effector", meta = (ClampMin = "0", ClampMax = "1", UIMin = "0.0", UIMax = "1.0"))
+	float StrengthAlpha = 1.0f;
+
+	/** Range 0-1 (default is 1.0). When enabled (greater than 0.0), the solver internally partitions the skeleton into 'chains' which extend from the effector to the nearest fork in the skeleton.
+	*These chains are pre-rotated and translated, as a whole, towards the effector targets.
+	*This can improve the results for sparse bone chains, and significantly improve convergence on dense bone chains.
+	*But it may cause undesirable results in highly constrained bone chains (like robot arms).
+	*/
+	UPROPERTY(BlueprintReadWrite, EditAnywhere, Category = "Full Body IK Effector", meta = (ClampMin = "0", ClampMax = "1", UIMin = "0.0", UIMax = "1.0"))
+	float PullChainAlpha = 1.0f;
+
+	/** Range 0-1 (default is 1.0).
+	*Blends the effector bone rotation between the rotation of the effector transform (1.0) and the rotation of the input bone (0.0).*/
+	UPROPERTY(BlueprintReadWrite, EditAnywhere, Category = "Full Body IK Effector", meta = (ClampMin = "0", ClampMax = "1", UIMin = "0.0", UIMax = "1.0"))
+	float PinRotation = 1.0f;
+
+	UPROPERTY(Transient)
+	int32 IndexInSolver = -1;
+
+	void CopySettings(const UIKRig_FBIKEffector* Other)
+	{
+		StrengthAlpha = Other->StrengthAlpha;
+		PullChainAlpha = Other->PullChainAlpha;
+		PinRotation = Other->PinRotation;
+	}
+};
+
+UCLASS(BlueprintType)
+class IKRIG_API UIKRig_FBIKBoneSettings : public UObject
+{
+	GENERATED_BODY()
+
+public:
+	
+	UIKRig_FBIKBoneSettings()
+		: Bone(NAME_None), 
+		X(EPBIKLimitType::Free),
+		Y(EPBIKLimitType::Free),
+		Z(EPBIKLimitType::Free),
+		PreferredAngles(FVector::ZeroVector)
+	{
+		SetFlags(RF_Transactional);
+	}
+
+	/** The bone these settings are applied to. */
+	UPROPERTY(BlueprintReadOnly, VisibleAnywhere, Category = Bone, meta = (Constant, CustomWidget = "BoneName"))
+	FName Bone;
+	
+	/** Range is 0 to 1 (Default is 0). At higher values, the bone will resist rotating (forcing other bones to compensate). */
+	UPROPERTY(BlueprintReadWrite, EditAnywhere, Category = Stiffness, meta = (ClampMin = "0", ClampMax = "1", UIMin = "0.0", UIMax = "1.0"))
+	float RotationStiffness = 0.0f;
+
+	/** Range is 0 to 1 (Default is 0). At higher values, the bone will resist translational motion (forcing other bones to compensate). */
+	UPROPERTY(BlueprintReadWrite, EditAnywhere, Category = Stiffness, meta = (ClampMin = "0", ClampMax = "1", UIMin = "0.0", UIMax = "1.0"))
+	float PositionStiffness = 0.0f;
+
+	/** Limit the rotation angle of the bone on the X axis.
+	 *Free: can rotate freely in this axis.
+	 *Limited: rotation is clamped between the min/max angles relative to the Skeletal Mesh reference pose.
+	 *Locked: no rotation is allowed in this axis (will remain at reference pose angle). */
+	UPROPERTY(BlueprintReadWrite, EditAnywhere, Category = Limits)
+	EPBIKLimitType X;
+	/**Range is -180 to 0 (Default is 0). Degrees of rotation in the negative X direction to allow when joint is in "Limited" mode. */
+	UPROPERTY(BlueprintReadWrite, EditAnywhere, Category = Limits, meta = (ClampMin = "-180", ClampMax = "0", UIMin = "-180.0", UIMax = "0.0"))
+	float MinX = 0.0f;
+	/**Range is 0 to 180 (Default is 0). Degrees of rotation in the positive X direction to allow when joint is in "Limited" mode. */
+	UPROPERTY(BlueprintReadWrite, EditAnywhere, Category = Limits, meta = (ClampMin = "0", ClampMax = "180", UIMin = "0.0", UIMax = "180.0"))
+	float MaxX = 0.0f;
+
+	/** Limit the rotation angle of the bone on the Y axis.
+	*Free: can rotate freely in this axis.
+	*Limited: rotation is clamped between the min/max angles relative to the Skeletal Mesh reference pose.
+	*Locked: no rotation is allowed in this axis (will remain at input pose angle). */
+	UPROPERTY(BlueprintReadWrite, EditAnywhere, Category = Limits)
+	EPBIKLimitType Y;
+	/**Range is -180 to 0 (Default is 0). Degrees of rotation in the negative Y direction to allow when joint is in "Limited" mode. */
+	UPROPERTY(BlueprintReadWrite, EditAnywhere, Category = Limits, meta = (ClampMin = "-180", ClampMax = "0", UIMin = "-180.0", UIMax = "0.0"))
+	float MinY = 0.0f;
+	/**Range is 0 to 180 (Default is 0). Degrees of rotation in the positive Y direction to allow when joint is in "Limited" mode. */
+	UPROPERTY(BlueprintReadWrite, EditAnywhere, Category = Limits, meta = (ClampMin = "0", ClampMax = "180", UIMin = "0.0", UIMax = "180.0"))
+	float MaxY = 0.0f;
+
+	/** Limit the rotation angle of the bone on the Z axis.
+	*Free: can rotate freely in this axis.
+	*Limited: rotation is clamped between the min/max angles relative to the Skeletal Mesh reference pose.
+	*Locked: no rotation is allowed in this axis (will remain at input pose angle). */
+	UPROPERTY(BlueprintReadWrite, EditAnywhere, Category = Limits)
+	EPBIKLimitType Z;
+	/**Range is -180 to 0 (Default is 0). Degrees of rotation in the negative Z direction to allow when joint is in "Limited" mode. */
+	UPROPERTY(BlueprintReadWrite, EditAnywhere, Category = Limits, meta = (ClampMin = "-180", ClampMax = "0", UIMin = "-180.0", UIMax = "0.0"))
+	float MinZ = 0.0f;
+	/**Range is 0 to 180 (Default is 0). Degrees of rotation in the positive Z direction to allow when joint is in "Limited" mode. */
+	UPROPERTY(BlueprintReadWrite, EditAnywhere, Category = Limits, meta = (ClampMin = "0", ClampMax = "180", UIMin = "0.0", UIMax = "180.0"))
+	float MaxZ = 0.0f;
+
+	/**When true, this bone will "prefer" to rotate in the direction specified by the Preferred Angles when the chain it belongs to is compressed.
+	 * Preferred Angles should be the first method used to fix bones that bend in the wrong direction (rather than limits).
+	 * The resulting angles can be visualized on their own by temporarily setting the Solver iterations to 0 and moving the effectors.*/
+	UPROPERTY(BlueprintReadWrite, EditAnywhere, Category = PreferredAngles)
+	bool bUsePreferredAngles = false;
+	/**The local Euler angles (in degrees) used to rotate this bone when the chain it belongs to is squashed.
+	 * This happens by moving the effector at the tip of the chain towards the root of the chain.
+	 * This can be used to coerce knees and elbows to bend in the anatomically "correct" direction without resorting to limits (which may require more iterations to converge).*/
+	UPROPERTY(BlueprintReadWrite, EditAnywhere, Category = PreferredAngles)
+	FVector PreferredAngles;
+
+	void CopyToCoreStruct(PBIK::FBoneSettings& Settings) const
+	{
+		Settings.RotationStiffness = RotationStiffness;
+		Settings.PositionStiffness = PositionStiffness;
+		Settings.X = static_cast<PBIK::ELimitType>(X);
+		Settings.MinX = MinX;
+		Settings.MaxX = MaxX;
+		Settings.Y = static_cast<PBIK::ELimitType>(Y);
+		Settings.MinY = MinY;
+		Settings.MaxY = MaxY;
+		Settings.Z = static_cast<PBIK::ELimitType>(Z);
+		Settings.MinZ = MinZ;
+		Settings.MaxZ = MaxZ;
+		Settings.bUsePreferredAngles = bUsePreferredAngles;
+		Settings.PreferredAngles.Pitch = PreferredAngles.Y;
+		Settings.PreferredAngles.Yaw = PreferredAngles.Z;
+		Settings.PreferredAngles.Roll = PreferredAngles.X;
+	}
+
+	void CopySettings(const UIKRig_FBIKBoneSettings* Other)
+	{	
+		RotationStiffness = Other->RotationStiffness;
+		PositionStiffness = Other->PositionStiffness;
+		X = Other->X;
+		MinX = Other->MinX;
+		MaxX = Other->MaxX;
+		Y = Other->Y;
+		MinY = Other->MinY;
+		MaxY = Other->MaxY;
+		Z = Other->Z;
+		MinZ = Other->MinZ;
+		MaxZ = Other->MaxZ;
+		bUsePreferredAngles = Other->bUsePreferredAngles;
+		PreferredAngles = Other->PreferredAngles;
+	}
+};
+
+UCLASS(BlueprintType, EditInlineNew, config = Engine, hidecategories = UObject)
+class IKRIG_API UIKRigFBIKSolver : public UIKRigSolver
+{
+	GENERATED_BODY()
+
+public:
+
+	/** All bones above this bone in the hierarchy will be completely ignored by the solver. Typically this is set to
+	 * the top-most skinned bone in the Skeletal Mesh (ie Pelvis, Hips etc), NOT the actual root of the entire skeleton.
+	 *
+	 * If you want to use the solver on a single chain of bones, and NOT translate the chain, ensure that "PinRoot" is
+	 * checked on to disable the root from translating to reach the effector goals.*/
+	UPROPERTY(BlueprintReadWrite, VisibleAnywhere, Category = "Full Body IK Settings")
+	FName RootBone;
+
+	/** High iteration counts can help solve complex joint configurations with competing constraints, but will increase runtime cost. Default is 20. */
+	UPROPERTY(BlueprintReadWrite, EditAnywhere, Category = SolverSettings, meta = (ClampMin = "0", ClampMax = "1000", UIMin = "0.0", UIMax = "200.0"))
+	int32 Iterations = 20;
+
+	/** A global mass multiplier; higher values will make the joints more stiff, but require more iterations. Typical range is 0.0 to 10.0. */
+	UPROPERTY(BlueprintReadWrite, EditAnywhere, Category = SolverSettings, meta = (ClampMin = "0", UIMin = "0.0", UIMax = "10.0"))
+	float MassMultiplier = 1.0f;
+
+	/** If true, joints will translate to reach the effectors; causing bones to lengthen if necessary. Good for cartoon effects. Default is false. */
+	UPROPERTY(BlueprintReadWrite, EditAnywhere, Category = SolverSettings)
+	bool bAllowStretch = false;
+
+	/** (Default is PrePull) Set the behavior of the solver root.
+	*Pre Pull: translates and rotates the root (and all children) by the average motion of the stretched effectors to help achieve faster convergence when reaching far.
+	*Pin to Input: locks the translation and rotation of the root bone to the input pose. Overrides any bone settings applied to the root. Good for partial-body solves.
+	*Free: treats the root bone like any other and allows it to move freely or according to any bone settings applied to it. */
+	UPROPERTY(BlueprintReadWrite, EditAnywhere, Category = RootBehavior)
+	EPBIKRootBehavior RootBehavior = EPBIKRootBehavior::PrePull;
+
+	/** Settings only applicable when Root Behavior is set to "PrePull". Use these values to adjust the gross movement and orientation of the entire skeleton. */
+	UPROPERTY(BlueprintReadWrite, EditAnywhere, Category = RootBehavior)
+	FRootPrePullSettings PrePullRootSettings;
+
+	/** A global multiplier for all Pull Chain Alpha values on all effectors. Range is 0.0 to 1.0. Default is 1.0. */
+	UPROPERTY(BlueprintReadWrite, EditAnywhere, Category = AdvancedSettings, meta = (ClampMin = "0", ClampMax = "1", UIMin = "0.0", UIMax = "1.0"))
+	float PullChainAlpha = 1.0f;
+
+	/** Maximum angle that a joint can be rotated per constraint iteration. Lower this value if the solve is diverging. Range is 0.0 to 180.0. Default is 30. */
+	UPROPERTY(BlueprintReadWrite, EditAnywhere, Category = AdvancedSettings, meta = (ClampMin = "0", ClampMax = "45", UIMin = "0.0", UIMax = "180.0"))
+	float MaxAngle = 30.f;
+
+	/** Pushes constraints beyond their normal amount to speed up convergence. Increasing this may speed up convergence, but at the cost of stability. Range is 1.0 - 2.0. Default is 1.3. */
+	UPROPERTY(BlueprintReadWrite, EditAnywhere, Category = AdvancedSettings, meta = (ClampMin = "1", UIMin = "1.0", UIMax = "10.0"))
+	float OverRelaxation = 1.3f;
+	
+	/** When true, the solver is reset each tick to start from the current input pose. Default is true.
+	 * If false, incoming animated poses are ignored and the solver starts from the results of the previous solve.
+	 * In very limited circumstances, it can be beneficial to use the pose from the previous frame.*/
+	UPROPERTY(BlueprintReadWrite, EditAnywhere, Category = AdvancedSettings)
+	bool bStartSolveFromInputPose = true;
+	
+	UPROPERTY()
+	TArray<TObjectPtr<UIKRig_FBIKEffector>> Effectors;
+
+	UPROPERTY()
+	TArray<TObjectPtr<UIKRig_FBIKBoneSettings>> BoneSettings;
+
+	/** UIKRigSolver interface */
+	// runtime
+	virtual void Initialize(const FIKRigSkeleton& IKRigSkeleton) override;
+	virtual void Solve(FIKRigSkeleton& IKRigSkeleton, const FIKRigGoalContainer& Goals) override;
+
+	virtual FName GetRootBone() const override { return RootBone; };
+	virtual void GetBonesWithSettings(TSet<FName>& OutBonesWithSettings) const override;
+
+	virtual void UpdateSolverSettings(UIKRigSolver* InSettings) override;
+	virtual void RemoveGoal(const FName& GoalName) override;
+	
+#if WITH_EDITOR
+	virtual FText GetNiceName() const override;
+	virtual bool GetWarningMessage(FText& OutWarningMessage) const override;
+	// goals
+	virtual void AddGoal(const UIKRigEffectorGoal& NewGoal) override;
+	virtual void RenameGoal(const FName& OldName, const FName& NewName) override;
+	virtual void SetGoalBone(const FName& GoalName, const FName& NewBoneName) override;
+	virtual bool IsGoalConnected(const FName& GoalName) const override;
+	virtual UObject* GetGoalSettings(const FName& GoalName) const override;
+	// bone settings
+	virtual void AddBoneSetting(const FName& BoneName) override;
+	virtual void RemoveBoneSetting(const FName& BoneName) override;
+	virtual UObject* GetBoneSetting(const FName& BoneName) const override;
+	virtual bool UsesBoneSettings() const override{ return true;};
+	virtual void DrawBoneSettings(const FName& BoneName, const FIKRigSkeleton& IKRigSkeleton, FPrimitiveDrawInterface* PDI) const override;
+	// root bone can be set on this solver
+	virtual bool RequiresRootBone() const override { return true; };
+	virtual void SetRootBone(const FName& RootBoneName) override;
+	virtual bool IsBoneAffectedBySolver(const FName& BoneName, const FIKRigSkeleton& IKRigSkeleton) const override;
+	/** END UIKRigSolver interface */
+
+	//** UObject */
+	virtual void PostLoad() override;
+	//** END UObject */
+#endif
+
+private:
+
+	FPBIKSolver Solver;
+
+	int32 GetIndexOfGoal(const FName& GoalName) const;
+};
+
+
diff --git a/Engine/Plugins/Animation/IKRig/Source/IKRig/Public/Solvers/IKRig_PBIKSolver.h b/Engine/Plugins/Animation/IKRig/Source/IKRig/Public/Solvers/IKRig_PBIKSolver.h
deleted file mode 100644
diff --git a/Engine/Plugins/Animation/IKRig/Source/IKRigEditor/Private/RigEditor/IKRigEditorController.cpp b/Engine/Plugins/Animation/IKRig/Source/IKRigEditor/Private/RigEditor/IKRigEditorController.cpp
--- a/Engine/Plugins/Animation/IKRig/Source/IKRigEditor/Private/RigEditor/IKRigEditorController.cpp
+++ b/Engine/Plugins/Animation/IKRig/Source/IKRigEditor/Private/RigEditor/IKRigEditorController.cpp
@@ -9,7 +9,7 @@
 #include "RigEditor/IKRigToolkit.h"
 #include "RigEditor/SIKRigAssetBrowser.h"
 #include "RigEditor/SIKRigOutputLog.h"
-#include "Solvers/IKRig_PBIKSolver.h"
+#include "Solvers/IKRig_FBIKSolver.h"
 #include "Widgets/Input/SComboBox.h"
 #include "ContentBrowserModule.h"
 #include "IContentBrowserSingleton.h"
@@ -711,7 +711,7 @@ bool FIKRigEditorController::PromptToAddDefaultSolver() const
 	TSharedPtr<FIKRigSolverTypeAndName> SelectedSolver = SolverTypes[0];
 	for (TSharedPtr<FIKRigSolverTypeAndName>& SolverType :SolverTypes)
 	{
-		if (SolverType->SolverType == UIKRigPBIKSolver::StaticClass())
+		if (SolverType->SolverType == UIKRigFBIKSolver::StaticClass())
 		{
 			SelectedSolver = SolverType;
 			break;
