IK Retargeter: Fixing bug, could not clear retarget root bone.

diff --git a/Engine/Plugins/Animation/IKRig/Source/IKRigEditor/Private/RigEditor/IKRigController.cpp b/Engine/Plugins/Animation/IKRig/Source/IKRigEditor/Private/RigEditor/IKRigController.cpp
--- a/Engine/Plugins/Animation/IKRig/Source/IKRigEditor/Private/RigEditor/IKRigController.cpp
+++ b/Engine/Plugins/Animation/IKRig/Source/IKRigEditor/Private/RigEditor/IKRigController.cpp
@@ -382,16 +382,16 @@ bool UIKRigController::SetRetargetRoot(const FName RootBoneName) const
 {
 	check(Asset)
 
-	if (Asset->Skeleton.GetBoneIndexFromName(RootBoneName) == INDEX_NONE)
+	FName NewRootBone = RootBoneName;
+	if (RootBoneName != NAME_None && Asset->Skeleton.GetBoneIndexFromName(RootBoneName) == INDEX_NONE)
 	{
-		UE_LOG(LogIKRigEditor, Warning, TEXT("Bone does not exist, %s."), *RootBoneName.ToString());
-		return false; // bone doesn't exist
+		NewRootBone = NAME_None;
 	}
 
 	FScopedTransaction Transaction(LOCTEXT("SetRetargetRootBone_Label", "Set Retarget Root Bone"));
 	Asset->Modify();
 	
-	Asset->RetargetDefinition.RootBone = RootBoneName;
+	Asset->RetargetDefinition.RootBone = NewRootBone;
 
 	BroadcastNeedsReinitialized();
 
