IK Retargeter: not reinitializing after editing pose.

diff --git a/Engine/Plugins/Animation/IKRig/Source/IKRigEditor/Private/RetargetEditor/IKRetargetAnimInstanceProxy.cpp b/Engine/Plugins/Animation/IKRig/Source/IKRigEditor/Private/RetargetEditor/IKRetargetAnimInstanceProxy.cpp
--- a/Engine/Plugins/Animation/IKRig/Source/IKRigEditor/Private/RetargetEditor/IKRetargetAnimInstanceProxy.cpp
+++ b/Engine/Plugins/Animation/IKRig/Source/IKRigEditor/Private/RetargetEditor/IKRetargetAnimInstanceProxy.cpp
@@ -110,6 +110,11 @@ void FIKRetargetAnimInstanceProxy::ConfigureAnimInstance(
 void FIKRetargetAnimInstanceProxy::SetRetargetMode(const ERetargeterOutputMode& InOutputMode)
 {
 	OutputMode = InOutputMode;
+	
+	if (UIKRetargetProcessor* Processor = RetargetNode->GetRetargetProcessor())
+	{
+		Processor->SetNeedsInitialized();
+	}
 }
 
 void FIKRetargetAnimInstanceProxy::SetRetargetPoseBlend(const float& InRetargetPoseBlend) const
