Fix retarget chain mapping "exact all" mode not resetting non-exact matches to none.

#swarm https://p4-swarm.epicgames.net/reviews/24083883

diff --git a/Engine/Plugins/Animation/IKRig/Source/IKRigEditor/Private/RetargetEditor/IKRetargeterController.cpp b/Engine/Plugins/Animation/IKRig/Source/IKRigEditor/Private/RetargetEditor/IKRetargeterController.cpp
--- a/Engine/Plugins/Animation/IKRig/Source/IKRigEditor/Private/RetargetEditor/IKRetargeterController.cpp
+++ b/Engine/Plugins/Animation/IKRig/Source/IKRigEditor/Private/RetargetEditor/IKRetargeterController.cpp
@@ -343,7 +343,10 @@ void UIKRetargeterController::AutoMapChains(const EAutoMapChainType AutoMapType,
 			}
 		case EAutoMapChainType::Exact:
 			{
-				// auto-map chains with same name
+				// if no exact match is found, then set to None
+				ChainMap->SourceChain = NAME_None;
+				
+				// auto-map chains with EXACT same name
 				for (int32 ChainIndex=0; ChainIndex<SourceChainNames.Num(); ++ChainIndex)
 				{
 					if (SourceChainNames[ChainIndex] == ChainMap->TargetChain)
