Disallow loading incompatible skeletons in IK Rig editor. Forces re-initialization when swapping skeletal meshes.

diff --git a/Engine/Plugins/Animation/IKRig/Source/IKRigEditor/Private/RigEditor/IKRigController.cpp b/Engine/Plugins/Animation/IKRig/Source/IKRigEditor/Private/RigEditor/IKRigController.cpp
--- a/Engine/Plugins/Animation/IKRig/Source/IKRigEditor/Private/RigEditor/IKRigController.cpp
+++ b/Engine/Plugins/Animation/IKRig/Source/IKRigEditor/Private/RigEditor/IKRigController.cpp
@@ -562,6 +562,14 @@ FName UIKRigController::GetRetargetChainFromGoal(const FName GoalName) const
 
 bool UIKRigController::SetSkeletalMesh(USkeletalMesh* SkeletalMesh, bool bTransact) const
 {
+	// first determine runtime compatibility between the IK Rig asset and the skeleton we're trying to run it on
+	if (!IsSkeletalMeshCompatible(SkeletalMesh))
+	{
+		UE_LOG(LogIKRigEditor, Warning, TEXT("Trying to initialize IKRig with a Skeleton that is missing required bones. See output log. {0}"), *Asset->GetName());
+		BroadcastNeedsReinitialized();
+		return false;
+	}
+	
 	FScopedTransaction Transaction(LOCTEXT("SetSkeletalMesh_Label", "Set Skeletal Mesh"));
 	if (bTransact)
 	{	
diff --git a/Engine/Plugins/Animation/IKRig/Source/IKRigEditor/Private/RigEditor/IKRigEditorController.cpp b/Engine/Plugins/Animation/IKRig/Source/IKRigEditor/Private/RigEditor/IKRigEditorController.cpp
--- a/Engine/Plugins/Animation/IKRig/Source/IKRigEditor/Private/RigEditor/IKRigEditorController.cpp
+++ b/Engine/Plugins/Animation/IKRig/Source/IKRigEditor/Private/RigEditor/IKRigEditorController.cpp
@@ -403,8 +403,21 @@ void FIKRigEditorController::HandleIKRigNeedsInitialized(UIKRigDefinition* Modif
 	
 	// in case the skeletal mesh was swapped out, we need to ensure the preview scene is up-to-date
 	USkeletalMesh* NewMesh = AssetController->GetSkeletalMesh();
+	SkelMeshComponent->SetSkeletalMesh(NewMesh);
 	const TSharedRef<IPersonaPreviewScene> PreviewScene = EditorToolkit.Pin()->GetPersonaToolkit()->GetPreviewScene();
-	PreviewScene->SetPreviewMesh(NewMesh); // this will re-initialize the anim instance
+	if (PreviewScene->GetPreviewMesh() != NewMesh)
+	{
+		PreviewScene->SetPreviewMeshComponent(SkelMeshComponent);
+		PreviewScene->SetPreviewMesh(NewMesh);
+	}
+
+	// re-initializes the anim instances running in the viewport
+	if (AnimInstance)
+	{
+		SkelMeshComponent->PreviewInstance = AnimInstance;
+		AnimInstance->InitializeAnimation();
+		SkelMeshComponent->EnablePreview(true, nullptr);
+	}
 
 
 	// update the bone details so it can pull on the current data
