Fixes, too much validation happening in batch retargeter; prevents partially compatible mesh / IK Rigs from being retargeted.

diff --git a/Engine/Plugins/Animation/IKRig/Source/IKRigEditor/Private/RetargetEditor/IKRetargetBatchOperation.cpp b/Engine/Plugins/Animation/IKRig/Source/IKRigEditor/Private/RetargetEditor/IKRetargetBatchOperation.cpp
--- a/Engine/Plugins/Animation/IKRig/Source/IKRigEditor/Private/RetargetEditor/IKRetargetBatchOperation.cpp
+++ b/Engine/Plugins/Animation/IKRig/Source/IKRigEditor/Private/RetargetEditor/IKRetargetBatchOperation.cpp
@@ -559,22 +559,6 @@ void UIKRetargetBatchOperation::RunRetarget(FIKRetargetBatchOperationContext& Co
 			return;
 		}
 	}
-
-	// validate source mesh is compatible with source IK Rig
-	const UIKRigController* SrcRigController = UIKRigController::GetController(SrcIKRig);
-	if (!SrcRigController->IsSkeletalMeshCompatible(Context.SourceMesh))
-	{
-		UE_LOG(LogTemp, Warning, TEXT("Batch retarget aborted. The provided source mesh was not compatible with the source IK Rig."));
-		return;
-	}
-	
-	// validate target mesh is compatible with target IK Rig
-	const UIKRigController* TgtRigController = UIKRigController::GetController(TgtIKRig);
-	if (!TgtRigController->IsSkeletalMeshCompatible(Context.TargetMesh))
-	{
-		UE_LOG(LogTemp, Warning, TEXT("Batch retarget aborted. The provided target mesh was not compatible with the target IK Rig."));
-		return;
-	}
 	
 	// show progress bar
 	FScopedSlowTask Progress(NumAssets + 2, LOCTEXT("GatheringBatchRetarget", "Gathering animation assets..."));
