Fixes IK Rig python API GetRootBone() on 3 solvers.

diff --git a/Engine/Plugins/Animation/IKRig/Source/IKRig/Public/Solvers/IKRig_LimbSolver.h b/Engine/Plugins/Animation/IKRig/Source/IKRig/Public/Solvers/IKRig_LimbSolver.h
--- a/Engine/Plugins/Animation/IKRig/Source/IKRig/Public/Solvers/IKRig_LimbSolver.h
+++ b/Engine/Plugins/Animation/IKRig/Source/IKRig/Public/Solvers/IKRig_LimbSolver.h
@@ -89,6 +89,8 @@ public:
 
 	virtual void UpdateSolverSettings(UIKRigSolver* InSettings) override;
 	virtual void RemoveGoal(const FName& GoalName) override;
+
+	virtual FName GetRootBone() const override { return RootName; };
 	
 #if WITH_EDITOR
 	virtual FText GetNiceName() const override;
diff --git a/Engine/Plugins/Animation/IKRig/Source/IKRig/Public/Solvers/IKRig_PoleSolver.h b/Engine/Plugins/Animation/IKRig/Source/IKRig/Public/Solvers/IKRig_PoleSolver.h
--- a/Engine/Plugins/Animation/IKRig/Source/IKRig/Public/Solvers/IKRig_PoleSolver.h
+++ b/Engine/Plugins/Animation/IKRig/Source/IKRig/Public/Solvers/IKRig_PoleSolver.h
@@ -49,6 +49,8 @@ public:
 	virtual void UpdateSolverSettings(UIKRigSolver* InSettings) override;
 	virtual void RemoveGoal(const FName& GoalName) override;
 
+	virtual FName GetRootBone() const override { return RootName; };
+
 #if WITH_EDITOR
 	virtual FText GetNiceName() const override;
 	virtual bool GetWarningMessage(FText& OutWarningMessage) const override;
diff --git a/Engine/Plugins/Animation/IKRig/Source/IKRig/Public/Solvers/IKRig_SetTransform.h b/Engine/Plugins/Animation/IKRig/Source/IKRig/Public/Solvers/IKRig_SetTransform.h
--- a/Engine/Plugins/Animation/IKRig/Source/IKRig/Public/Solvers/IKRig_SetTransform.h
+++ b/Engine/Plugins/Animation/IKRig/Source/IKRig/Public/Solvers/IKRig_SetTransform.h
@@ -50,6 +50,7 @@ public:
 	virtual void Solve(FIKRigSkeleton& IKRigSkeleton, const FIKRigGoalContainer& Goals) override;
 	virtual void UpdateSolverSettings(UIKRigSolver* InSettings) override;
 	virtual void RemoveGoal(const FName& GoalName) override;
+	virtual FName GetRootBone() const override { return RootBone; };
 #if WITH_EDITOR
 	virtual FText GetNiceName() const override;
 	virtual bool GetWarningMessage(FText& OutWarningMessage) const override;
