Fixing issue, renaming default retarget pose from blueprint causes crash.

diff --git a/Engine/Plugins/Animation/IKRig/Source/IKRigEditor/Private/RetargetEditor/IKRetargeterController.cpp b/Engine/Plugins/Animation/IKRig/Source/IKRigEditor/Private/RetargetEditor/IKRetargeterController.cpp
--- a/Engine/Plugins/Animation/IKRig/Source/IKRigEditor/Private/RetargetEditor/IKRetargeterController.cpp
+++ b/Engine/Plugins/Animation/IKRig/Source/IKRigEditor/Private/RetargetEditor/IKRetargeterController.cpp
@@ -588,6 +588,13 @@ bool UIKRetargeterController::RenameRetargetPose(const FName OldPoseName, const
 		return false;
 	}
 
+	// do not allow renaming the default pose (this is disallowed from the UI, but must be done here as well for API usage)
+	if (OldPoseName == UIKRetargeter::GetDefaultPoseName())
+	{
+		UE_LOG(LogIKRigEditor, Warning, TEXT("Trying to rename the default pose. This is not allowed."));
+    	return false;
+	}
+
 	// check if we're renaming the current pose
 	const bool bWasCurrentPose = GetCurrentRetargetPoseName(SourceOrTarget) == OldPoseName;
 	
@@ -599,11 +606,10 @@ bool UIKRetargeterController::RenameRetargetPose(const FName OldPoseName, const
 
 	// replace key in the map
 	TMap<FName, FIKRetargetPose>& Poses = GetRetargetPoses(SourceOrTarget);
-	const FName CurrentPoseName = GetCurrentRetargetPoseName(SourceOrTarget);
-	const FIKRetargetPose CurrentPoseData = GetCurrentRetargetPose(SourceOrTarget);
-	Poses.Remove(CurrentPoseName);
+	const FIKRetargetPose OldPoseData = Poses[OldPoseName];
+	Poses.Remove(OldPoseName);
 	Poses.Shrink();
-	Poses.Add(UniqueNewPoseName, CurrentPoseData);
+	Poses.Add(UniqueNewPoseName, OldPoseData);
 
 	// make this the current retarget pose, iff the old one was
 	if (bWasCurrentPose)
