Change "IK Rig" category in content menu to "Retargeting"

diff --git a/Engine/Plugins/Animation/IKRig/Source/IKRigEditor/Private/RetargetEditor/AssetTypeActions_IKRetargeter.cpp b/Engine/Plugins/Animation/IKRig/Source/IKRigEditor/Private/RetargetEditor/AssetTypeActions_IKRetargeter.cpp
--- a/Engine/Plugins/Animation/IKRig/Source/IKRigEditor/Private/RetargetEditor/AssetTypeActions_IKRetargeter.cpp
+++ b/Engine/Plugins/Animation/IKRig/Source/IKRigEditor/Private/RetargetEditor/AssetTypeActions_IKRetargeter.cpp
@@ -98,7 +98,7 @@ const TArray<FText>& FAssetTypeActions_IKRetargeter::GetSubMenus() const
 {
 	static const TArray<FText> SubMenus
 	{
-		LOCTEXT("AnimIKRigSubMenu", "IK Rig")
+		LOCTEXT("AnimRetargetingSubMenu", "Retargeting")
 	};
 	return SubMenus;
 }
diff --git a/Engine/Plugins/Animation/IKRig/Source/IKRigEditor/Private/RigEditor/AssetTypeActions_IKRigDefinition.cpp b/Engine/Plugins/Animation/IKRig/Source/IKRigEditor/Private/RigEditor/AssetTypeActions_IKRigDefinition.cpp
--- a/Engine/Plugins/Animation/IKRig/Source/IKRigEditor/Private/RigEditor/AssetTypeActions_IKRigDefinition.cpp
+++ b/Engine/Plugins/Animation/IKRig/Source/IKRigEditor/Private/RigEditor/AssetTypeActions_IKRigDefinition.cpp
@@ -39,7 +39,7 @@ const TArray<FText>& FAssetTypeActions_IKRigDefinition::GetSubMenus() const
 {
 	static const TArray<FText> SubMenus
 	{
-		LOCTEXT("AnimIKRigSubMenu", "IK Rig")
+		LOCTEXT("AnimRetargetingSubMenu", "Retargeting")
 	};
 	return SubMenus;
 }
