Exposed retarget phases to python API.

diff --git a/Engine/Plugins/Animation/IKRig/Source/IKRig/Public/Retargeter/IKRetargetSettings.h b/Engine/Plugins/Animation/IKRig/Source/IKRig/Public/Retargeter/IKRetargetSettings.h
--- a/Engine/Plugins/Animation/IKRig/Source/IKRig/Public/Retargeter/IKRetargetSettings.h
+++ b/Engine/Plugins/Animation/IKRig/Source/IKRig/Public/Retargeter/IKRetargetSettings.h
@@ -276,19 +276,19 @@ struct IKRIG_API FRetargetGlobalSettings
 	/** When false, the motion of the Retarget Root bone is not copied from the source. Useful for debugging issues with the root settings.
 	 * Note: the retargeting order is: Root > FK > IK
 	 */
-	UPROPERTY(EditAnywhere, Category = RetargetPhases)
+	UPROPERTY(EditAnywhere, BlueprintReadWrite, Category = RetargetPhases)
 	bool bEnableRoot = true;
 	
 	/** When false, limbs are not copied via FK. Useful for debugging limb issues suspected to be caused by FK chain settings.
 	 * Note: the retargeting order is: Root > FK > IK
 	 */
-	UPROPERTY(EditAnywhere, Category = RetargetPhases)
+	UPROPERTY(EditAnywhere, BlueprintReadWrite, Category = RetargetPhases)
 	bool bEnableFK = true;
 	
 	/** When false, IK is not applied as part of retargeter. Useful for debugging limb issues suspected to be caused by IK.
 	 * Note: the retargeting order is: Root > FK > IK
 	 */
-	UPROPERTY(EditAnywhere, Category = RetargetPhases)
+	UPROPERTY(EditAnywhere, BlueprintReadWrite, Category = RetargetPhases)
 	bool bEnableIK = true;
 
 	/** Enable IK Warping.
