Creating /Rig folder in IKRig plugin. Organizing IK Rig code better.

diff --git a/Engine/Plugins/Animation/IKRig/Source/IKRig/Private/AnimNodes/AnimNode_IKRig.cpp b/Engine/Plugins/Animation/IKRig/Source/IKRig/Private/AnimNodes/AnimNode_IKRig.cpp
--- a/Engine/Plugins/Animation/IKRig/Source/IKRig/Private/AnimNodes/AnimNode_IKRig.cpp
+++ b/Engine/Plugins/Animation/IKRig/Source/IKRig/Private/AnimNodes/AnimNode_IKRig.cpp
@@ -1,10 +1,12 @@
 // Copyright Epic Games, Inc. All Rights Reserved.
 
 #include "AnimNodes/AnimNode_IKRig.h"
+
+#include "Rig/IKRigProcessor.h"
+#include "Rig/IKRigDefinition.h"
+#include "Rig/Solvers/IKRigSolver.h"
+
 #include "Components/SkeletalMeshComponent.h"
-#include "IKRigProcessor.h"
-#include "IKRigDefinition.h"
-#include "IKRigSolver.h"
 #include "ActorComponents/IKRigInterface.h"
 #include "Animation/AnimInstanceProxy.h"
 #include "Algo/ForEach.h"
diff --git a/Engine/Plugins/Animation/IKRig/Source/IKRig/Private/IKRigDataTypes.cpp b/Engine/Plugins/Animation/IKRig/Source/IKRig/Private/IKRigDataTypes.cpp
deleted file mode 100644
diff --git a/Engine/Plugins/Animation/IKRig/Source/IKRig/Private/IKRigDefinition.cpp b/Engine/Plugins/Animation/IKRig/Source/IKRig/Private/IKRigDefinition.cpp
deleted file mode 100644
diff --git a/Engine/Plugins/Animation/IKRig/Source/IKRig/Private/IKRigProcessor.cpp b/Engine/Plugins/Animation/IKRig/Source/IKRig/Private/IKRigProcessor.cpp
deleted file mode 100644
diff --git a/Engine/Plugins/Animation/IKRig/Source/IKRig/Private/IKRigSkeleton.cpp b/Engine/Plugins/Animation/IKRig/Source/IKRig/Private/IKRigSkeleton.cpp
deleted file mode 100644
diff --git a/Engine/Plugins/Animation/IKRig/Source/IKRig/Private/IKRigSolver.cpp b/Engine/Plugins/Animation/IKRig/Source/IKRig/Private/IKRigSolver.cpp
deleted file mode 100644
diff --git a/Engine/Plugins/Animation/IKRig/Source/IKRig/Private/Retargeter/IKRetargetProcessor.cpp b/Engine/Plugins/Animation/IKRig/Source/IKRig/Private/Retargeter/IKRetargetProcessor.cpp
--- a/Engine/Plugins/Animation/IKRig/Source/IKRig/Private/Retargeter/IKRetargetProcessor.cpp
+++ b/Engine/Plugins/Animation/IKRig/Source/IKRig/Private/Retargeter/IKRetargetProcessor.cpp
@@ -2,14 +2,16 @@
 
 #include "Retargeter/IKRetargetProcessor.h"
 
-#include "IKRigDefinition.h"
 #include "IKRigLogger.h"
-#include "IKRigProcessor.h"
-#include "Engine/SkeletalMesh.h"
-#include "Kismet/KismetMathLibrary.h"
+#include "Rig/IKRigDefinition.h"
+#include "Rig/Solvers/PointsToRotation.h"
+#include "Rig/IKRigProcessor.h"
+
 #include "Retargeter/IKRetargeter.h"
 #include "Retargeter/IKRetargetProfile.h"
-#include "Solvers/PointsToRotation.h"
+
+#include "Engine/SkeletalMesh.h"
+#include "Kismet/KismetMathLibrary.h"
 #include "UObject/UnrealTypePrivate.h"
 
 #include UE_INLINE_GENERATED_CPP_BY_NAME(IKRetargetProcessor)
diff --git a/Engine/Plugins/Animation/IKRig/Source/IKRig/Private/Rig/IKRigDataTypes.cpp b/Engine/Plugins/Animation/IKRig/Source/IKRig/Private/Rig/IKRigDataTypes.cpp
new file mode 100644
--- /dev/null
+++ b/Engine/Plugins/Animation/IKRig/Source/IKRig/Private/Rig/IKRigDataTypes.cpp
@@ -0,0 +1,66 @@
+// Copyright Epic Games, Inc. All Rights Reserved.
+
+#include "Rig/IKRigDataTypes.h"
+#include "Rig/IKRigDefinition.h"
+
+#include UE_INLINE_GENERATED_CPP_BY_NAME(IKRigDataTypes)
+
+void FIKRigGoalContainer::SetIKGoal(const FIKRigGoal& InGoal)
+{
+	FIKRigGoal* Goal = FindGoalWriteable(InGoal.Name);
+	if (!Goal)
+	{
+		// container hasn't seen this goal before, create new one, copying the input goal
+		Goals.Emplace(InGoal);
+		return;
+	}
+
+	// copy settings to existing goal
+	*Goal = InGoal;
+}
+
+void FIKRigGoalContainer::SetIKGoal(const UIKRigEffectorGoal* InEffectorGoal)
+{
+	FIKRigGoal* Goal = FindGoalWriteable(InEffectorGoal->GoalName);
+	if (!Goal)
+	{
+		// container hasn't seen this goal before, create new one, copying the Effector goal
+		Goals.Emplace(InEffectorGoal);
+		return;
+	}
+
+	// goals in editor have "preview mode" which allows them to be specified relative to their
+	// initial pose to simulate an additive runtime behavior
+#if WITH_EDITOR
+	if (InEffectorGoal->PreviewMode == EIKRigGoalPreviewMode::Additive)
+	{
+		// transform to be interpreted as an offset, relative to input pose
+		Goal->Position = InEffectorGoal->CurrentTransform.GetTranslation() - InEffectorGoal->InitialTransform.GetTranslation();
+		const FQuat RelativeRotation = InEffectorGoal->CurrentTransform.GetRotation() * InEffectorGoal->InitialTransform.GetRotation().Inverse();
+		Goal->Rotation = RelativeRotation.Rotator();
+		Goal->PositionSpace = EIKRigGoalSpace::Additive;
+		Goal->RotationSpace = EIKRigGoalSpace::Additive;
+	}else
+#endif
+	{
+		// transform to be interpreted as absolute and in component space
+		Goal->Position = InEffectorGoal->CurrentTransform.GetTranslation();
+		Goal->Rotation = InEffectorGoal->CurrentTransform.Rotator();
+		Goal->PositionSpace = EIKRigGoalSpace::Component;
+		Goal->RotationSpace = EIKRigGoalSpace::Component;
+	}
+	
+    Goal->PositionAlpha = InEffectorGoal->PositionAlpha;
+    Goal->RotationAlpha = InEffectorGoal->RotationAlpha;
+}
+
+const FIKRigGoal* FIKRigGoalContainer::FindGoalByName(const FName& GoalName) const
+{
+	return Goals.FindByPredicate([GoalName](const FIKRigGoal& Other)	{return Other.Name == GoalName;});
+}
+
+FIKRigGoal* FIKRigGoalContainer::FindGoalWriteable(const FName& GoalName) const
+{
+	return const_cast<FIKRigGoal*>(FindGoalByName(GoalName));
+}
+
diff --git a/Engine/Plugins/Animation/IKRig/Source/IKRig/Private/Rig/IKRigDefinition.cpp b/Engine/Plugins/Animation/IKRig/Source/IKRig/Private/Rig/IKRigDefinition.cpp
new file mode 100644
--- /dev/null
+++ b/Engine/Plugins/Animation/IKRig/Source/IKRig/Private/Rig/IKRigDefinition.cpp
@@ -0,0 +1,303 @@
+// Copyright Epic Games, Inc. All Rights Reserved.
+
+#include "Rig/IKRigDefinition.h"
+
+#include "IKRigObjectVersion.h"
+#include "Engine/SkeletalMesh.h"
+
+#include UE_INLINE_GENERATED_CPP_BY_NAME(IKRigDefinition)
+
+#if WITH_EDITOR
+#include "HAL/PlatformApplicationMisc.h"
+#endif
+
+#define LOCTEXT_NAMESPACE "IKRigDefinition"
+
+#if WITH_EDITOR
+const FName UIKRigDefinition::GetPreviewMeshPropertyName() { return GET_MEMBER_NAME_STRING_CHECKED(UIKRigDefinition, PreviewSkeletalMesh); };
+
+TOptional<FTransform::FReal> UIKRigEffectorGoal::GetNumericValue(
+	ESlateTransformComponent::Type Component,
+	ESlateRotationRepresentation::Type Representation,
+	ESlateTransformSubComponent::Type SubComponent,
+	EIKRigTransformType::Type TransformType) const
+{
+	switch(TransformType)
+	{
+		case EIKRigTransformType::Current:
+		{
+			return SAdvancedTransformInputBox<FTransform>::GetNumericValueFromTransform(
+				CurrentTransform,
+				Component,
+				Representation,
+				SubComponent);
+		}
+		case EIKRigTransformType::Reference:
+		{
+			return SAdvancedTransformInputBox<FTransform>::GetNumericValueFromTransform(
+				InitialTransform,
+				Component,
+				Representation,
+				SubComponent);
+		}
+		default:
+		{
+			break;
+		}
+	}
+	return TOptional<FTransform::FReal>();
+}
+
+TTuple<FTransform, FTransform> UIKRigEffectorGoal::PrepareNumericValueChanged( ESlateTransformComponent::Type Component,
+																			ESlateRotationRepresentation::Type Representation,
+																			ESlateTransformSubComponent::Type SubComponent,
+																			FTransform::FReal Value,
+																			EIKRigTransformType::Type TransformType) const
+{
+	const FTransform& InTransform = TransformType == EIKRigTransformType::Current ? CurrentTransform : InitialTransform;
+	FTransform OutTransform = InTransform;
+	SAdvancedTransformInputBox<FTransform>::ApplyNumericValueChange(OutTransform, Value, Component, Representation, SubComponent);
+	return MakeTuple(InTransform, OutTransform);
+}
+
+void UIKRigEffectorGoal::SetTransform(const FTransform& InTransform, EIKRigTransformType::Type InTransformType)
+{
+	// we assume that InTransform is not equal to the one it's being assigned to
+	Modify();
+
+	FTransform& TransformChanged = InTransformType == EIKRigTransformType::Current ? CurrentTransform : InitialTransform;
+	TransformChanged = InTransform;
+}
+
+namespace
+{
+	
+template<typename DataType>
+void GetContentFromData(const DataType& InData, FString& Content)
+{
+	TBaseStructure<DataType>::Get()->ExportText(Content, &InData, &InData, nullptr, PPF_None, nullptr);
+}
+
+class FIKRigEffectorGoalErrorPipe : public FOutputDevice
+{
+public:
+	int32 NumErrors;
+
+	FIKRigEffectorGoalErrorPipe()
+		: FOutputDevice()
+		, NumErrors(0)
+	{}
+
+	virtual void Serialize(const TCHAR* V, ELogVerbosity::Type Verbosity, const class FName& Category) override
+	{
+		NumErrors++;
+	}
+};
+
+template<typename DataType>
+bool GetDataFromContent(const FString& Content, DataType& OutData)
+{
+	FIKRigEffectorGoalErrorPipe ErrorPipe;
+	static UScriptStruct* DataStruct = TBaseStructure<DataType>::Get();
+	DataStruct->ImportText(*Content, &OutData, nullptr, PPF_None, &ErrorPipe, DataStruct->GetName(), true);
+	return (ErrorPipe.NumErrors == 0);
+}
+	
+}
+
+void UIKRigEffectorGoal::OnCopyToClipboard(ESlateTransformComponent::Type Component, EIKRigTransformType::Type TransformType) const
+{
+	const FTransform& Xfo = TransformType == EIKRigTransformType::Current ? CurrentTransform : InitialTransform;
+
+	FString Content;
+	switch(Component)
+	{
+	case ESlateTransformComponent::Location:
+		{
+			GetContentFromData(Xfo.GetLocation(), Content);
+			break;
+		}
+	case ESlateTransformComponent::Rotation:
+		{
+			GetContentFromData(Xfo.Rotator(), Content);
+			break;
+		}
+	case ESlateTransformComponent::Scale:
+		{
+			GetContentFromData(Xfo.GetScale3D(), Content);
+			break;
+		}
+	case ESlateTransformComponent::Max:
+	default:
+		{
+			GetContentFromData(Xfo, Content);
+			break;
+		}
+	}
+
+	if (!Content.IsEmpty())
+	{
+		FPlatformApplicationMisc::ClipboardCopy(*Content);
+	}
+}
+
+void UIKRigEffectorGoal::OnPasteFromClipboard(ESlateTransformComponent::Type Component, EIKRigTransformType::Type TransformType)
+{
+	FString Content;
+	FPlatformApplicationMisc::ClipboardPaste(Content);
+
+	if (Content.IsEmpty())
+	{
+		return;
+	}
+
+	FTransform& Xfo = TransformType == EIKRigTransformType::Current ? CurrentTransform : InitialTransform;
+	Modify();
+	
+	switch(Component)
+	{
+		case ESlateTransformComponent::Location:
+		{
+			FVector Data = Xfo.GetLocation();
+			if (GetDataFromContent(Content, Data))
+			{
+				Xfo.SetLocation(Data);
+			}
+			break;
+		}
+		case ESlateTransformComponent::Rotation:
+		{
+			FRotator Data = Xfo.Rotator();
+			if (GetDataFromContent(Content, Data))
+			{
+				Xfo.SetRotation(FQuat(Data));
+			}
+			break;
+		}
+		case ESlateTransformComponent::Scale:
+		{
+			FVector Data = Xfo.GetScale3D();
+			if (GetDataFromContent(Content, Data))
+			{
+				Xfo.SetScale3D(Data);
+			}
+			break;
+		}
+		case ESlateTransformComponent::Max:
+		default:
+		{
+			FTransform Data = Xfo;
+			if (GetDataFromContent(Content, Data))
+			{
+				Xfo = Data;
+			}
+			break;
+		}
+	}
+}
+
+bool UIKRigEffectorGoal::TransformDiffersFromDefault(
+	ESlateTransformComponent::Type Component,
+	TSharedPtr<IPropertyHandle> PropertyHandle) const
+{
+	if(PropertyHandle->GetProperty()->GetFName() == GET_MEMBER_NAME_CHECKED(UIKRigEffectorGoal, CurrentTransform))
+	{
+		switch(Component)
+		{
+			case ESlateTransformComponent::Location:
+			{
+				return !(CurrentTransform.GetLocation() - InitialTransform.GetLocation()).IsNearlyZero();
+			}
+			case ESlateTransformComponent::Rotation:
+			{
+				return !(CurrentTransform.Rotator() - InitialTransform.Rotator()).IsNearlyZero();
+			}
+			case ESlateTransformComponent::Scale:
+			{
+				return !(CurrentTransform.GetScale3D() - InitialTransform.GetScale3D()).IsNearlyZero();
+			}
+			default:
+			{
+				return false;
+			}
+		}
+	}
+	return false;
+}
+
+void UIKRigEffectorGoal::ResetTransformToDefault(
+	ESlateTransformComponent::Type Component,
+	TSharedPtr<IPropertyHandle> PropertyHandle)
+{
+	switch(Component)
+	{
+		case ESlateTransformComponent::Location:
+		{
+			CurrentTransform.SetLocation(InitialTransform.GetLocation());
+			break;
+		}
+		case ESlateTransformComponent::Rotation:
+		{
+			CurrentTransform.SetRotation(InitialTransform.GetRotation());
+			break;
+		}
+		case ESlateTransformComponent::Scale:
+		{
+			CurrentTransform.SetScale3D(InitialTransform.GetScale3D());
+			break;
+		}
+		default:
+		{
+			break;
+		}
+	}
+}
+
+#endif
+
+FBoneChain* FRetargetDefinition::GetEditableBoneChainByName(FName ChainName)
+{
+	for (FBoneChain& Chain : BoneChains)
+	{
+		if (Chain.ChainName == ChainName)
+		{
+			return &Chain;
+		}
+	}
+	
+	return nullptr;
+}
+
+void UIKRigDefinition::Serialize(FArchive& Ar)
+{
+	Super::Serialize(Ar);
+	Ar.UsingCustomVersion(FIKRigObjectVersion::GUID);
+}
+
+const FBoneChain* UIKRigDefinition::GetRetargetChainByName(FName ChainName) const
+{
+	for (const FBoneChain& Chain : RetargetDefinition.BoneChains)
+	{
+		if (Chain.ChainName == ChainName)
+		{
+			return &Chain;
+		}
+	}
+	
+	return nullptr;
+}
+
+/** IInterface_PreviewMeshProvider interface */
+void UIKRigDefinition::SetPreviewMesh(USkeletalMesh* PreviewMesh, bool bMarkAsDirty)
+{
+	PreviewSkeletalMesh = PreviewMesh;
+}
+
+USkeletalMesh* UIKRigDefinition::GetPreviewMesh() const
+{
+	return PreviewSkeletalMesh.LoadSynchronous();
+}
+/** END IInterface_PreviewMeshProvider interface */
+
+#undef LOCTEXT_NAMESPACE
+
diff --git a/Engine/Plugins/Animation/IKRig/Source/IKRig/Private/Rig/IKRigProcessor.cpp b/Engine/Plugins/Animation/IKRig/Source/IKRig/Private/Rig/IKRigProcessor.cpp
new file mode 100644
--- /dev/null
+++ b/Engine/Plugins/Animation/IKRig/Source/IKRig/Private/Rig/IKRigProcessor.cpp
@@ -0,0 +1,478 @@
+// Copyright Epic Games, Inc. All Rights Reserved.
+
+#include "Rig/IKRigProcessor.h"
+#include "Rig/IKRigDefinition.h"
+#include "Rig/Solvers/IKRigSolver.h"
+
+#include "Engine/SkeletalMesh.h"
+
+#include UE_INLINE_GENERATED_CPP_BY_NAME(IKRigProcessor)
+
+#define LOCTEXT_NAMESPACE "IKRigProcessor"
+
+UIKRigProcessor::UIKRigProcessor()
+{
+	const FName LogName = FName("IKRig_",GetUniqueID());
+	Log.SetLogTarget(LogName, LOCTEXT("IKRigLogName", "IK Rig Log"));
+}
+
+void UIKRigProcessor::Initialize(
+	const UIKRigDefinition* InRigAsset,
+	const USkeletalMesh* SkeletalMesh,
+	const TArray<FName>& ExcludedGoals)
+{
+	// we instantiate UObjects here which MUST be done on game thread...
+	check(IsInGameThread());
+	check(InRigAsset);
+	
+	bInitialized = false;
+
+	// bail out if we've already tried initializing with this exact version of the rig asset
+	if (bTriedToInitialize)
+	{
+		return; // don't keep spamming
+	}
+
+	// ok, lets try to initialize
+	bTriedToInitialize = true;
+
+	// copy skeleton data from the actual skeleton we want to run on
+	Skeleton.SetInputSkeleton(SkeletalMesh, InRigAsset->GetSkeleton().ExcludedBones);
+	
+	if (InRigAsset->GetSkeleton().BoneNames.IsEmpty())
+	{
+		Log.LogError(FText::Format(
+			LOCTEXT("NoSkeleton", "Trying to initialize IK Rig, '{0}' that has no skeleton."),
+			FText::FromString(InRigAsset->GetName())));
+		return;
+	}
+
+	if (!UIKRigProcessor::IsIKRigCompatibleWithSkeleton(InRigAsset, SkeletalMesh, &Log))
+	{
+		Log.LogError(FText::Format(
+			LOCTEXT("SkeletonMissingRequiredBones", "Trying to initialize IKRig, '{0}' with a Skeleton that is missing required bones. See prior warnings."),
+			FText::FromString(InRigAsset->GetName())));
+		return;
+	}
+	
+	// initialize goals based on source asset
+	GoalContainer.Empty();
+	const TArray<UIKRigEffectorGoal*>& EffectorGoals = InRigAsset->GetGoalArray();
+	for (const UIKRigEffectorGoal* EffectorGoal : EffectorGoals)
+	{
+		// skip excluded goals
+		if (ExcludedGoals.Contains(EffectorGoal->GoalName))
+		{
+			continue;
+		}
+		
+		// add a copy of the Goal to the container
+		GoalContainer.SetIKGoal(EffectorGoal);
+	}
+	
+	// initialize goal bones from asset
+	GoalBones.Reset();
+	for (const UIKRigEffectorGoal* EffectorGoal : EffectorGoals)
+	{
+		// skip excluded goals
+		if (ExcludedGoals.Contains(EffectorGoal->GoalName))
+		{
+			continue;
+		}
+		
+		FGoalBone NewGoalBone;
+		NewGoalBone.BoneName = EffectorGoal->BoneName;
+		NewGoalBone.BoneIndex = Skeleton.GetBoneIndexFromName(EffectorGoal->BoneName);
+
+		// validate that the skeleton we are trying to solve this goal on contains the bone the goal expects
+		if (NewGoalBone.BoneIndex == INDEX_NONE)
+		{
+			Log.LogError(FText::Format(
+				LOCTEXT("MissingGoalBone", "IK Rig, '{0}' has a Goal, '{1}' that references an unknown bone, '{2}'. Cannot evaluate."),
+				FText::FromString(InRigAsset->GetName()), FText::FromName(EffectorGoal->GoalName), FText::FromName(EffectorGoal->BoneName) ));
+			return;
+		}
+
+		// validate that there is not already a different goal, with the same name, that is using a different bone
+		// (all goals with the same name must reference the same bone within a single IK Rig)
+		if (const FGoalBone* Bone = GoalBones.Find(EffectorGoal->GoalName))
+		{
+			if (Bone->BoneName != NewGoalBone.BoneName)
+			{
+				Log.LogError(FText::Format(
+				LOCTEXT("DuplicateGoal", "IK Rig, '{0}' has a Goal, '{1}' that references different bones in different solvers, '{2}' and '{3}'. Cannot evaluate."),
+                FText::FromString(InRigAsset->GetName()),
+                FText::FromName(EffectorGoal->GoalName),
+                FText::FromName(Bone->BoneName), 
+                FText::FromName(NewGoalBone.BoneName)));
+				return;
+			}
+		}
+		
+		GoalBones.Add(EffectorGoal->GoalName, NewGoalBone);
+	}
+
+	// create copies of all the solvers in the IK rig
+	const TArray<UIKRigSolver*>& AssetSolvers = InRigAsset->GetSolverArray();
+	Solvers.Reset(AssetSolvers.Num());
+	int32 SolverIndex = 0;
+	for (const UIKRigSolver* IKRigSolver : AssetSolvers)
+	{
+		if (!IKRigSolver)
+		{
+			// this can happen if asset references deleted IK Solver type which should only happen during development (if at all)
+			Log.LogWarning(FText::Format(
+				LOCTEXT("UnknownSolver", "IK Rig, '{0}' has null/unknown solver in it. Please remove it."),
+				FText::FromString(InRigAsset->GetName())));
+			continue;
+		}
+
+		// create duplicate solver instance with unique name
+		FString Name = IKRigSolver->GetName() + "_SolverInstance_";
+		Name.AppendInt(SolverIndex++);
+		UIKRigSolver* Solver = DuplicateObject(IKRigSolver, this, FName(*Name));
+
+		// remove excluded goals from the solver
+		for (const FName& ExcludedGoal : ExcludedGoals)
+		{
+			Solver->RemoveGoal(ExcludedGoal);
+		}
+
+		// initialize it and store in the processor
+		Solver->Initialize(Skeleton);
+		Solvers.Add(Solver);
+	}
+
+	// validate retarget chains
+	const TArray<FBoneChain>& Chains = InRigAsset->GetRetargetChains();
+	TArray<int32> OutBoneIndices;
+	for (const FBoneChain& Chain : Chains)
+	{
+		if (!Skeleton.ValidateChainAndGetBones(Chain, OutBoneIndices))
+		{
+			Log.LogWarning(FText::Format(
+				LOCTEXT("InvalidRetargetChain", "Invalid Retarget Chain: '{0}'. End bone is not a child of the start bone in Skeletal Mesh, '{1}'."),
+				FText::FromString(Chain.ChainName.ToString()), FText::FromString(SkeletalMesh->GetName())));
+		}
+	}
+
+	Log.LogInfo(FText::Format(
+				LOCTEXT("SuccessfulInit", "IK Rig, '{0}' ready to run on {1}."),
+				FText::FromString(InRigAsset->GetName()), FText::FromString(SkeletalMesh->GetName())));
+	
+	bInitialized = true;
+}
+
+bool UIKRigProcessor::IsIKRigCompatibleWithSkeleton(
+	const UIKRigDefinition* InRigAsset,
+	const FIKRigInputSkeleton& InputSkeleton,
+	const FIKRigLogger* Log)
+{
+	// first we validate that all the required bones are in the input skeleton...
+	
+	TSet<FName> RequiredBones;
+	const TArray<UIKRigSolver*>& AssetSolvers = InRigAsset->GetSolverArray();
+	for (const UIKRigSolver* Solver : AssetSolvers)
+	{
+		const FName RootBone = Solver->GetRootBone();
+		if (RootBone != NAME_None)
+		{
+			RequiredBones.Add(RootBone);
+		}
+
+		Solver->GetBonesWithSettings(RequiredBones);
+	}
+
+	const TArray<UIKRigEffectorGoal*>& Goals = InRigAsset->GetGoalArray();
+	for (const UIKRigEffectorGoal* Goal : Goals)
+	{
+		RequiredBones.Add(Goal->BoneName);
+	}
+
+	bool bAllRequiredBonesFound = true;
+	for (const FName& RequiredBone : RequiredBones)
+	{
+		if (!InputSkeleton.BoneNames.Contains(RequiredBone))
+		{
+			if (Log)
+			{
+				Log->LogError(FText::Format(
+			LOCTEXT("MissingBone", "IK Rig, '{0}' is missing a required bone, '{1}' in the Skeletal Mesh."),
+				FText::FromString(InRigAsset->GetName()),
+				FText::FromName(RequiredBone)));
+			}
+			
+			bAllRequiredBonesFound = false;
+		}
+	}
+
+	if (!bAllRequiredBonesFound)
+	{
+		return false;
+	}
+
+	// now we validate that hierarchy matches for all required bones...
+	bool bAllParentsValid = true;
+	const FIKRigSkeleton& AssetSkeleton = InRigAsset->GetSkeleton();
+	for (const FName& RequiredBone : RequiredBones)
+	{
+		const int32 InputBoneIndex = InputSkeleton.BoneNames.Find(RequiredBone);
+		const int32 AssetBoneIndex = AssetSkeleton.BoneNames.Find(RequiredBone);
+
+		// we shouldn't get this far otherwise due to early return above...
+		check(InputBoneIndex != INDEX_NONE && AssetBoneIndex != INDEX_NONE)
+
+		// validate that input skeleton hierarchy is as expected
+		const int32 AssetParentIndex = AssetSkeleton.ParentIndices[AssetBoneIndex];
+		if (AssetSkeleton.BoneNames.IsValidIndex(AssetParentIndex)) // root bone has no parent
+		{
+			const FName& AssetParentName = AssetSkeleton.BoneNames[AssetParentIndex];
+			const int32 InputParentIndex = InputSkeleton.ParentIndices[InputBoneIndex];
+			if (!InputSkeleton.BoneNames.IsValidIndex(InputParentIndex))
+			{
+				bAllParentsValid = false;
+
+				if (Log)
+				{
+					Log->LogError(FText::Format(
+					LOCTEXT("InvalidParent", "IK Rig is running on a skeleton with a required bone, '{0}', that expected to have a valid parent. The expected parent was, '{1}'."),
+					FText::FromName(RequiredBone),
+					FText::FromName(AssetParentName)));
+				}
+				
+				continue;
+			}
+			
+			const FName& InputParentName = InputSkeleton.BoneNames[InputParentIndex];
+			if (AssetParentName != InputParentName)
+			{
+				if (Log)
+				{
+					// we only warn about this, because it may be nice not to have the exact same hierarchy
+					Log->LogWarning(FText::Format(
+					LOCTEXT("UnexpectedParent", "IK Rig is running on a skeleton with a required bone, '{0}', that has a different parent '{1}'. The expected parent was, '{2}'."),
+					FText::FromName(RequiredBone),
+					FText::FromName(InputParentName),
+					FText::FromName(AssetParentName)));
+				}
+				
+				continue;
+			}
+		}
+	}
+
+	return bAllParentsValid;
+}
+
+void UIKRigProcessor::SetInputPoseGlobal(const TArray<FTransform>& InGlobalBoneTransforms) 
+{
+	check(bInitialized);
+	check(InGlobalBoneTransforms.Num() == Skeleton.CurrentPoseGlobal.Num());
+	Skeleton.CurrentPoseGlobal = InGlobalBoneTransforms;
+	Skeleton.UpdateAllLocalTransformFromGlobal();
+}
+
+void UIKRigProcessor::SetInputPoseToRefPose()
+{
+	check(bInitialized);
+	Skeleton.CurrentPoseGlobal = Skeleton.RefPoseGlobal;
+	Skeleton.UpdateAllLocalTransformFromGlobal();
+}
+
+void UIKRigProcessor::SetIKGoal(const FIKRigGoal& InGoal)
+{
+	check(bInitialized);
+	GoalContainer.SetIKGoal(InGoal);
+}
+
+void UIKRigProcessor::SetIKGoal(const UIKRigEffectorGoal* InGoal)
+{
+	check(bInitialized);
+	GoalContainer.SetIKGoal(InGoal);
+}
+
+void UIKRigProcessor::Solve(const FTransform& ComponentToWorld)
+{
+	check(bInitialized);
+	
+	// convert goals into component space and blend towards input pose by alpha
+	ResolveFinalGoalTransforms(ComponentToWorld);
+
+	// run all the solvers
+	for (UIKRigSolver* Solver : Solvers)
+	{
+		#if WITH_EDITOR
+		if (Solver->IsEnabled())
+		{
+			Solver->Solve(Skeleton, GoalContainer);
+		}
+		#else
+		Solver->Solve(Skeleton, GoalContainer);
+		#endif
+	}
+
+	// make sure rotations are normalized coming out
+	Skeleton.NormalizeRotations(Skeleton.CurrentPoseGlobal);
+}
+
+void UIKRigProcessor::CopyOutputGlobalPoseToArray(TArray<FTransform>& OutputPoseGlobal) const
+{
+	OutputPoseGlobal = Skeleton.CurrentPoseGlobal;
+}
+
+void UIKRigProcessor::Reset()
+{
+	Solvers.Reset();
+	GoalContainer.Empty();
+	GoalBones.Reset();
+	Skeleton.Reset();
+	SetNeedsInitialized();
+}
+
+void UIKRigProcessor::SetNeedsInitialized()
+{
+	bInitialized = false;
+	bTriedToInitialize = false;
+};
+
+void UIKRigProcessor::CopyAllInputsFromSourceAssetAtRuntime(const UIKRigDefinition* SourceAsset)
+{
+	check(SourceAsset)
+	
+	if (!bInitialized)
+	{
+		return;
+	}
+	
+	// copy goal settings
+	const TArray<UIKRigEffectorGoal*>& AssetGoals =  SourceAsset->GetGoalArray();
+	for (const UIKRigEffectorGoal* AssetGoal : AssetGoals)
+	{
+		SetIKGoal(AssetGoal);
+	}
+
+	// copy solver settings
+	const TArray<UIKRigSolver*>& AssetSolvers = SourceAsset->GetSolverArray();
+	check(Solvers.Num() == AssetSolvers.Num()); // if number of solvers has been changed, processor should have been reinitialized
+	for (int32 SolverIndex=0; SolverIndex<Solvers.Num(); ++SolverIndex)
+	{
+		Solvers[SolverIndex]->SetEnabled(AssetSolvers[SolverIndex]->IsEnabled());
+		Solvers[SolverIndex]->UpdateSolverSettings(AssetSolvers[SolverIndex]);
+	}
+}
+
+const FIKRigGoalContainer& UIKRigProcessor::GetGoalContainer() const
+{
+	check(bInitialized);
+	return GoalContainer;
+}
+
+const FGoalBone* UIKRigProcessor::GetGoalBone(const FName& GoalName) const
+{
+	return GoalBones.Find(GoalName);
+}
+
+FIKRigSkeleton& UIKRigProcessor::GetSkeletonWriteable()
+{
+	return Skeleton;
+}
+
+const FIKRigSkeleton& UIKRigProcessor::GetSkeleton() const
+{
+	return Skeleton;
+}
+
+void UIKRigProcessor::ResolveFinalGoalTransforms(const FTransform& WorldToComponent)
+{
+	for (FIKRigGoal& Goal : GoalContainer.Goals)
+	{
+		if (!GoalBones.Contains(Goal.Name))
+		{
+			// user is changing goals after initialization
+			// not necessarily a bad thing, but new goal names won't work until re-init
+			continue;
+		}
+
+		const FGoalBone& GoalBone = GoalBones[Goal.Name];
+		const FTransform& InputPoseBoneTransform = Skeleton.CurrentPoseGlobal[GoalBone.BoneIndex];
+
+		FVector ComponentSpaceGoalPosition = Goal.Position;
+		FQuat ComponentSpaceGoalRotation = Goal.Rotation.Quaternion();
+
+		// FIXME find a way to cache SourceBoneIndex to avoid calling Skeleton.GetBoneIndexFromName here
+		// we may use FGoalBone::OptSourceIndex for this
+		int SourceBoneIndex = INDEX_NONE;
+		if (Goal.TransformSource == EIKRigGoalTransformSource::Bone && Goal.SourceBone.BoneName != NAME_None)
+		{
+			SourceBoneIndex = Skeleton.GetBoneIndexFromName(Goal.SourceBone.BoneName);
+		}
+
+		if (SourceBoneIndex != INDEX_NONE)
+		{
+			ComponentSpaceGoalPosition = Skeleton.CurrentPoseGlobal[SourceBoneIndex].GetLocation();
+		}
+		else
+		{
+			// put goal POSITION in Component Space
+			switch (Goal.PositionSpace)
+			{
+			case EIKRigGoalSpace::Additive:
+				// add position offset to bone position
+				ComponentSpaceGoalPosition = Skeleton.CurrentPoseGlobal[GoalBone.BoneIndex].GetLocation() + Goal.Position;
+				break;
+			case EIKRigGoalSpace::Component:
+				// was already supplied in Component Space
+				break;
+			case EIKRigGoalSpace::World:
+				// convert from World Space to Component Space
+				ComponentSpaceGoalPosition = WorldToComponent.TransformPosition(Goal.Position);
+				break;
+			default:
+				checkNoEntry();
+				break;
+			}
+		}
+		
+		// put goal ROTATION in Component Space
+		if (SourceBoneIndex != INDEX_NONE)
+		{
+			ComponentSpaceGoalRotation = Skeleton.CurrentPoseGlobal[SourceBoneIndex].GetRotation();
+		}
+		else
+		{
+			switch (Goal.RotationSpace)
+			{
+			case EIKRigGoalSpace::Additive:
+				// add rotation offset to bone rotation
+				ComponentSpaceGoalRotation = Goal.Rotation.Quaternion() * Skeleton.CurrentPoseGlobal[GoalBone.BoneIndex].GetRotation();
+				break;
+			case EIKRigGoalSpace::Component:
+				// was already supplied in Component Space
+				break;
+			case EIKRigGoalSpace::World:
+				// convert from World Space to Component Space
+				ComponentSpaceGoalRotation = WorldToComponent.TransformRotation(Goal.Rotation.Quaternion());
+				break;
+			default:
+				checkNoEntry();
+				break;
+			}
+		}
+
+		// blend by alpha from the input pose, to the supplied goal transform
+		// when Alpha is 0, the goal transform matches the bone transform at the input pose.
+		// when Alpha is 1, the goal transform is left fully intact
+		Goal.FinalBlendedPosition = FMath::Lerp(
+            InputPoseBoneTransform.GetTranslation(),
+            ComponentSpaceGoalPosition,
+            Goal.PositionAlpha);
+		
+		Goal.FinalBlendedRotation = FQuat::FastLerp(
+            InputPoseBoneTransform.GetRotation(),
+            ComponentSpaceGoalRotation,
+            Goal.RotationAlpha);
+	}
+}
+
+#undef LOCTEXT_NAMESPACE
+
diff --git a/Engine/Plugins/Animation/IKRig/Source/IKRig/Private/Rig/IKRigSkeleton.cpp b/Engine/Plugins/Animation/IKRig/Source/IKRig/Private/Rig/IKRigSkeleton.cpp
new file mode 100644
--- /dev/null
+++ b/Engine/Plugins/Animation/IKRig/Source/IKRig/Private/Rig/IKRigSkeleton.cpp
@@ -0,0 +1,549 @@
+// Copyright Epic Games, Inc. All Rights Reserved.
+
+#include "Rig/IKRigSkeleton.h"
+
+#include "Rig/IKRigDefinition.h"
+#include "Algo/LevenshteinDistance.h"
+#include "Engine/SkeletalMesh.h"
+#include "Math/Bounds.h"
+
+#include UE_INLINE_GENERATED_CPP_BY_NAME(IKRigSkeleton)
+
+// This is the default end of branch index value, meaning we haven't cached it yet
+#define FIKRIGSKELETON_INVALID_EO_BRANCH_INDEX -2
+
+FIKRigInputSkeleton::FIKRigInputSkeleton(const USkeletalMesh* SkeletalMesh)
+{
+	Initialize(SkeletalMesh);
+}
+
+void FIKRigInputSkeleton::Initialize(const USkeletalMesh* SkeletalMesh)
+{
+	Reset();
+
+	if (!SkeletalMesh)
+	{
+		return;
+	}
+	
+	SkeletalMeshName = FName(SkeletalMesh->GetName());
+	const FReferenceSkeleton& RefSkeleton = SkeletalMesh->GetRefSkeleton();
+	const TArray<FMeshBoneInfo>& BoneInfo = RefSkeleton.GetRefBoneInfo();
+	for (int32 BoneIndex=0; BoneIndex<BoneInfo.Num(); ++BoneIndex)
+	{
+		BoneNames.Add(BoneInfo[BoneIndex].Name);
+		ParentIndices.Add(BoneInfo[BoneIndex].ParentIndex);
+		LocalRefPose.Add(RefSkeleton.GetRefBonePose()[BoneIndex]);
+	}
+}
+
+void FIKRigInputSkeleton::Reset()
+{
+	*this = FIKRigInputSkeleton();
+}
+
+void FIKRigSkeleton::SetInputSkeleton(const USkeletalMesh* SkeletalMesh, const TArray<FName>& InExcludedBones)
+{
+	const FIKRigInputSkeleton InputSkeleton = FIKRigInputSkeleton(SkeletalMesh);
+	SetInputSkeleton(InputSkeleton, InExcludedBones);
+}
+
+void FIKRigSkeleton::SetInputSkeleton(const FIKRigInputSkeleton& InputSkeleton, const TArray<FName>& InExcludedBones)
+{
+	check(InputSkeleton.BoneNames.Num() == InputSkeleton.ParentIndices.Num() &&
+		InputSkeleton.BoneNames.Num() == InputSkeleton.LocalRefPose.Num())
+	
+	// reset all containers
+	Reset();
+	
+	// we use the bone names and parent indices from the input skeleton
+	BoneNames = InputSkeleton.BoneNames;
+	ParentIndices = InputSkeleton.ParentIndices;
+	
+	// bones are excluded at the skeleton level, instead of per-solver
+	ExcludedBones = InExcludedBones;
+	
+	// get a compacted local ref pose based on the stored names
+	TArray<FTransform> CompactRefPoseLocal;
+	// copy a compacted reference pose using only the bones in the input skeleton
+	for (int32 BoneIndex=0; BoneIndex<InputSkeleton.BoneNames.Num(); ++BoneIndex)
+	{
+		// store ref pose of this bone
+		CompactRefPoseLocal.Add(InputSkeleton.LocalRefPose[BoneIndex]);
+	}
+	
+	// copy local ref pose to global
+	ConvertLocalPoseToGlobal(ParentIndices, CompactRefPoseLocal, RefPoseGlobal);
+
+	// initialize CURRENT GLOBAL pose to reference pose
+	CurrentPoseGlobal = RefPoseGlobal;
+
+	// initialize CURRENT LOCAL pose from global pose
+	UpdateAllLocalTransformFromGlobal();
+
+	// initialize branch caching
+	CachedEndOfBranchIndices.Init(FIKRIGSKELETON_INVALID_EO_BRANCH_INDEX, ParentIndices.Num());
+}
+
+void FIKRigSkeleton::Reset()
+{
+	*this = FIKRigSkeleton();
+}
+
+int32 FIKRigSkeleton::GetBoneIndexFromName(const FName InName) const
+{
+	for (int32 Index=0; Index<BoneNames.Num(); ++Index)
+	{
+		if (InName == BoneNames[Index])
+		{
+			return Index;
+		}
+	}
+
+	return INDEX_NONE;
+}
+
+const FName& FIKRigSkeleton::GetBoneNameFromIndex(const int32 BoneIndex) const
+{
+	if (BoneNames.IsValidIndex(BoneIndex))
+	{
+		return BoneNames[BoneIndex];
+	}
+
+	static const FName InvalidName(NAME_None);  
+	return InvalidName;
+}
+
+int32 FIKRigSkeleton::GetParentIndex(const int32 BoneIndex) const
+{
+	if (!BoneNames.IsValidIndex(BoneIndex))
+	{
+		return INDEX_NONE;
+	}
+
+	return ParentIndices[BoneIndex];
+}
+
+int32 FIKRigSkeleton::GetParentIndexThatIsNotExcluded(const int32 BoneIndex) const
+{
+	// find first parent that is NOT excluded
+	int32 ParentIndex = GetParentIndex(BoneIndex);
+	while(ParentIndex != INDEX_NONE)
+	{
+		if (!IsBoneExcluded(ParentIndex))
+		{
+			break;
+		}
+		ParentIndex = GetParentIndex(ParentIndex);
+	}
+	
+	return ParentIndex;
+}
+
+int32 FIKRigSkeleton::GetChildIndices(const int32 ParentBoneIndex, TArray<int32>& Children) const
+{
+	Children.Reset();
+
+	const int32 LastBranchIndex = GetCachedEndOfBranchIndex(ParentBoneIndex);
+	if (LastBranchIndex != INDEX_NONE)
+	{
+		for (int32 BoneIndex = ParentBoneIndex + 1; BoneIndex <= LastBranchIndex; BoneIndex++)
+		{
+			if (GetParentIndex(BoneIndex) == ParentBoneIndex)
+			{
+				Children.Add(BoneIndex);
+			}
+		}
+	}
+
+	return Children.Num();
+}
+
+int32 FIKRigSkeleton::GetCachedEndOfBranchIndex(const int32 InBoneIndex) const
+{
+	if (!CachedEndOfBranchIndices.IsValidIndex(InBoneIndex))
+	{
+		return INDEX_NONE;
+	}
+
+	// already cached
+	if (CachedEndOfBranchIndices[InBoneIndex] != FIKRIGSKELETON_INVALID_EO_BRANCH_INDEX)
+	{
+		return CachedEndOfBranchIndices[InBoneIndex];
+	}
+
+	const int32 NumBones = BoneNames.Num();
+	
+	// if we're asking for root's branch, get the last bone  
+	if (InBoneIndex == 0)
+	{
+		CachedEndOfBranchIndices[InBoneIndex] = NumBones-1;
+		return CachedEndOfBranchIndices[InBoneIndex];
+	}
+
+	CachedEndOfBranchIndices[InBoneIndex] = INDEX_NONE;
+	
+	// store ref parent
+	const int32 RefParentIndex = GetParentIndex(InBoneIndex);
+
+	int32 BoneIndex = InBoneIndex + 1;
+	int32 ParentIndex = GetParentIndex(BoneIndex);
+
+	// if next bellow bone's parent is less than or equal to RefParentIndex, we are leaving the branch
+	// so ne need to go further
+	while (ParentIndex > RefParentIndex && BoneIndex < NumBones)
+	{
+		CachedEndOfBranchIndices[InBoneIndex] = BoneIndex;
+				
+		BoneIndex++;
+		ParentIndex = GetParentIndex(BoneIndex);
+	}
+
+	return CachedEndOfBranchIndices[InBoneIndex];
+}
+
+void FIKRigSkeleton::ConvertLocalPoseToGlobal(
+	const TArray<int32>& InParentIndices,
+	const TArray<FTransform>& InLocalPose,
+	TArray<FTransform>& OutGlobalPose)
+{
+	check(InLocalPose.Num() == InParentIndices.Num());
+	
+	OutGlobalPose.Reset();
+	for (int32 BoneIndex=0; BoneIndex<InLocalPose.Num(); ++BoneIndex)
+	{
+		const int32 ParentIndex = InParentIndices[BoneIndex];
+		if (ParentIndex == INDEX_NONE)
+		{
+			OutGlobalPose.Add(InLocalPose[BoneIndex]);
+			continue; // root bone is always in local space 
+		}
+
+		const FTransform& ChildLocalTransform  = InLocalPose[BoneIndex];
+		const FTransform& ParentGlobalTransform  = OutGlobalPose[ParentIndex];
+		OutGlobalPose.Add(ChildLocalTransform * ParentGlobalTransform);
+	}
+}
+
+void FIKRigSkeleton::UpdateAllGlobalTransformFromLocal()
+{
+	// generate GLOBAL transforms based on LOCAL transforms
+	CurrentPoseGlobal = CurrentPoseLocal;
+	for (int32 BoneIndex=0; BoneIndex<CurrentPoseLocal.Num(); ++BoneIndex)
+	{
+		UpdateGlobalTransformFromLocal(BoneIndex);
+	}
+}
+
+void FIKRigSkeleton::UpdateGlobalTransformFromLocal(const int32 BoneIndex)
+{
+	const int32 ParentIndex = ParentIndices[BoneIndex];
+	if (ParentIndex == INDEX_NONE)
+	{
+		CurrentPoseGlobal[BoneIndex] = CurrentPoseLocal[BoneIndex]; // root bone is always in local space
+		return; 
+	}
+
+	const FTransform& ChildLocalTransform  = CurrentPoseLocal[BoneIndex];
+	const FTransform& ParentGlobalTransform  = CurrentPoseGlobal[ParentIndex];
+	CurrentPoseGlobal[BoneIndex] = ChildLocalTransform * ParentGlobalTransform;
+}
+
+void FIKRigSkeleton::UpdateAllLocalTransformFromGlobal()
+{
+	// generate LOCAL transforms based on GLOBAL transforms
+	CurrentPoseLocal = CurrentPoseGlobal;
+	for (int32 BoneIndex=0; BoneIndex<CurrentPoseGlobal.Num(); ++BoneIndex)
+	{
+		UpdateLocalTransformFromGlobal(BoneIndex);
+	}
+}
+
+void FIKRigSkeleton::UpdateLocalTransformFromGlobal(const int32 BoneIndex)
+{
+	const int32 ParentIndex = ParentIndices[BoneIndex];
+	if (ParentIndex == INDEX_NONE)
+	{
+		CurrentPoseLocal[BoneIndex] = CurrentPoseGlobal[BoneIndex]; // root bone is always in local space
+		return; 
+	}
+
+	const FTransform& ChildGlobal  = CurrentPoseGlobal[BoneIndex];
+	const FTransform& ParentGlobal  = CurrentPoseGlobal[ParentIndex];
+	CurrentPoseLocal[BoneIndex] = ChildGlobal.GetRelativeTransform(ParentGlobal);
+}
+
+void FIKRigSkeleton::PropagateGlobalPoseBelowBone(const int32 StartBoneIndex)
+{
+	const int32 LastBranchIndex = GetCachedEndOfBranchIndex(StartBoneIndex);
+	if (LastBranchIndex != INDEX_NONE)
+	{
+		for (int32 BoneIndex=StartBoneIndex+1; BoneIndex<=LastBranchIndex; ++BoneIndex)
+		{
+			UpdateGlobalTransformFromLocal(BoneIndex);
+			UpdateLocalTransformFromGlobal(BoneIndex);
+		}
+	}
+}
+
+bool FIKRigSkeleton::IsBoneInDirectLineage(const FName& Child, const FName& PotentialParent) const
+{
+	const int32 ChildIndex = GetBoneIndexFromName(Child);
+	if (ChildIndex == INDEX_NONE)
+	{
+		return false;
+	}
+
+	const int32 PotentialParentIndex = GetBoneIndexFromName(PotentialParent);
+	if (PotentialParentIndex == INDEX_NONE)
+	{
+		return false;
+	}
+
+	int32 NextParentIndex = ChildIndex;
+	while(NextParentIndex != INDEX_NONE)
+	{
+		if (NextParentIndex == PotentialParentIndex)
+		{
+			return true;
+		}
+		NextParentIndex = ParentIndices[NextParentIndex];
+	}
+
+	return false;
+}
+
+bool FIKRigSkeleton::IsBoneExcluded(const int32 BoneIndex) const
+{
+	if (ensure(BoneNames.IsValidIndex(BoneIndex)))
+	{
+		return ExcludedBones.Contains(BoneNames[BoneIndex]);	
+	}
+	
+	return false;
+}
+
+void FIKRigSkeleton::NormalizeRotations(TArray<FTransform>& Transforms)
+{
+	for (FTransform& Transform : Transforms)
+	{
+		Transform.NormalizeRotation();
+	}
+}
+
+void FIKRigSkeleton::GetChainsInList(const TArray<int32>& SelectedBones, TArray<FBoneChain>& OutChains) const
+{	
+	// no bones provided, return empty list
+	if (SelectedBones.IsEmpty())
+	{
+		return;
+	}
+
+	// get parents of all the selected bones
+	TSet<int32> SelectedParentIndices;
+	for (const int32 SelectedBone : SelectedBones)
+	{
+		SelectedParentIndices.Add(GetParentIndex(SelectedBone));
+	}
+	
+	// find all selected leaf nodes and make each one a chain
+    for (const int32 SelectedBone : SelectedBones)
+    {
+    	// is this a leaf node in the selection?
+    	if (SelectedParentIndices.Contains(SelectedBone))
+    	{
+    		continue;
+    	}
+
+    	// convert leaf nodes to chains
+    	int32 ChainStart;
+    	const int32 ChainEnd = SelectedBone;
+    	int32 ParentIndex = SelectedBone;
+    	while (true)
+    	{
+    		ChainStart = ParentIndex;
+    		ParentIndex = GetParentIndex(ParentIndex);
+    		if (ParentIndex == INDEX_NONE)
+    		{
+    			break;
+    		}
+    		if (!SelectedBones.Contains(ParentIndex))
+    		{
+    			break;
+    		}
+    	}
+
+    	const FName ChainName = NAME_None;
+    	const FName StartBoneName = GetBoneNameFromIndex(ChainStart);
+    	const FName EndBoneName = GetBoneNameFromIndex(ChainEnd);
+    	const FName GoalName = NAME_None;
+    	OutChains.Add(FBoneChain(ChainName, StartBoneName, EndBoneName, GoalName));	
+    }
+}
+
+bool FIKRigSkeleton::ValidateChainAndGetBones(const FBoneChain& BoneChain, TArray<int32>& OutBoneIndices) const
+{
+	OutBoneIndices.Reset();
+	
+	// validate start and end bones exist and are not the root
+	const int32 StartIndex = GetBoneIndexFromName(BoneChain.StartBone.BoneName);
+	const int32 EndIndex = GetBoneIndexFromName(BoneChain.EndBone.BoneName);
+	const bool bFoundStartBone = StartIndex > INDEX_NONE;
+	const bool bFoundEndBone = EndIndex > INDEX_NONE;
+
+	// no need to build the chain if start/end indices are wrong 
+	const bool bIsWellFormed = bFoundStartBone && bFoundEndBone && EndIndex >= StartIndex;
+	if (!bIsWellFormed)
+	{
+		return false;
+	}
+	
+	// init array with end bone 
+	OutBoneIndices = {EndIndex};
+
+	// if only one bone in the chain
+	if (EndIndex == StartIndex)
+	{
+		return true;
+	}
+
+	// record all bones in chain while walking up the hierarchy (tip to root of chain)
+	int32 ParentIndex = GetParentIndex(EndIndex);
+	while (ParentIndex > INDEX_NONE && ParentIndex >= StartIndex)
+	{
+		OutBoneIndices.Add(ParentIndex);
+		ParentIndex = GetParentIndex(ParentIndex);
+	}
+
+	// if we walked up till the start bone
+	if (OutBoneIndices.Last() == StartIndex)
+	{
+		// reverse the indices (we want root to tip order)
+		Algo::Reverse(OutBoneIndices);
+		return true;
+	}
+
+	// oops, we walked all the way up without finding the start bone
+	OutBoneIndices.Reset();
+	return false;
+}
+
+bool FIKRigSkeleton::GetMirroredBoneIndices(
+	const TArray<int32>& BoneIndices,
+	TArray<int32>& OutMirroredIndices) const
+{
+	auto GetDistancesToAllBones = [this](const FVector& InPoint, TArray<float>& OutDistances)
+	{
+		OutDistances.SetNumUninitialized(RefPoseGlobal.Num());
+		int32 Index = 0;
+		for (const FTransform& BoneTransform : RefPoseGlobal)
+		{
+			OutDistances[Index] = FVector::Distance(InPoint, BoneTransform.GetLocation());
+			++Index;
+		}
+	};
+	
+	auto GetBonesWithinDistance = [this](const float MaxDistance, const TArray<float>& InDistances, TArray<int32>& OutBoneIndices)
+	{
+		OutBoneIndices.Reset();
+		int32 Index = 0;
+		for (const float& Distance : InDistances)
+		{
+			if (Distance < MaxDistance)
+			{
+				OutBoneIndices.Add(Index);
+			}
+			++Index;
+		}
+		return !OutBoneIndices.IsEmpty();
+	};
+
+	auto GetBoneWithClosestName = [this](const FName BaseName, const TArray<int32>& InBoneIndices) -> int32
+	{
+		check(!InBoneIndices.IsEmpty());
+		int32 BoneWithClosestNameToBaseName = InBoneIndices[0];
+		const FString BaseNameStr = BaseName.ToString().ToLower();
+		float HighestScore = -1.f;
+		for (const int32 BoneIndex : InBoneIndices)
+		{
+			FString CurrentNameStr = BoneNames[BoneIndex].ToString().ToLower();
+			float WorstCase = BaseNameStr.Len() + CurrentNameStr.Len();
+			WorstCase = WorstCase < 1.0f ? 1.0f : WorstCase;
+			const float Score = 1.0f - (Algo::LevenshteinDistance(BaseNameStr, CurrentNameStr) / WorstCase);
+			if (Score > HighestScore)
+			{
+				HighestScore = Score;
+				BoneWithClosestNameToBaseName = BoneIndex;
+			}
+		}
+		
+		return BoneWithClosestNameToBaseName;
+	};
+
+	auto GetBoundingSphere = [this]()->FSphere
+	{
+		TArray<FVector> Points;
+		Points.AddUninitialized(RefPoseGlobal.Num());
+		int32 Index=0;
+		for (const FTransform& BoneTransform : RefPoseGlobal)
+		{
+			Points[Index] = BoneTransform.GetLocation();
+			++Index;
+		}
+		return FSphere(&Points[0], Points.Num());
+	};
+
+	// generate bounding sphere of skeleton to establish reasonable search radius
+	const FSphere BoundingSphere = GetBoundingSphere();
+	if (BoundingSphere.W < KINDA_SMALL_NUMBER)
+	{
+		// skeleton is degenerate
+		return false;
+	}
+		
+	const float MinSearchRadius = BoundingSphere.W * 0.01f;
+	const float MaxSearchRadius = BoundingSphere.W;
+	const float SearchStep = MinSearchRadius * 2.f;
+
+	// spin through all the bones in the chain:
+	// 1. iteratively expand a search radius to find bones near the mirrored location
+	// 2. pick the bone near the mirrored location with the most similar name
+	//
+	// NOTE: we can't just do a "find closest bone" here because production skeletons often have lots of "helper" bones
+	// that are located at or near the same position. So we have to use the name to help find the actual mirrored bone.
+	OutMirroredIndices.Reset();
+	for (const int32 BoneIndex : BoneIndices)
+	{
+		// get mirrored location to search around
+		const FVector BoneLocation = RefPoseGlobal[BoneIndex].GetLocation();
+		const FVector MirroredLocation = BoneLocation * FVector(-1.f,1.f,1.f);
+		
+		// get distance from the mirrored location to all other bones in the skeleton
+		TArray<float> DistancesToAllBones;
+		GetDistancesToAllBones(MirroredLocation, DistancesToAllBones);
+		
+		// do an expanding search around the mirrored location until we find bones near it
+		TArray<int32> NearbyBones;
+		float SearchRadius = MinSearchRadius;
+		while(SearchRadius <= MaxSearchRadius)
+		{
+			if (GetBonesWithinDistance(SearchRadius, DistancesToAllBones, NearbyBones))
+			{
+				break;
+			}
+			SearchRadius += SearchStep;
+		};
+		
+		if (NearbyBones.IsEmpty())
+		{
+			OutMirroredIndices.Reset();
+			return false;
+		}
+
+		// found the nearby bone with the best matching name
+		int32 BestMatchBoneIndex = GetBoneWithClosestName(BoneNames[BoneIndex], NearbyBones);
+		OutMirroredIndices.Add(BestMatchBoneIndex);
+	}
+	
+	return true;
+}
diff --git a/Engine/Plugins/Animation/IKRig/Source/IKRig/Private/Rig/Solvers/IKRigSolver.cpp b/Engine/Plugins/Animation/IKRig/Source/IKRig/Private/Rig/Solvers/IKRigSolver.cpp
new file mode 100644
--- /dev/null
+++ b/Engine/Plugins/Animation/IKRig/Source/IKRig/Private/Rig/Solvers/IKRigSolver.cpp
@@ -0,0 +1,23 @@
+// Copyright Epic Games, Inc. All Rights Reserved.
+
+#include "Rig/Solvers/IKRigSolver.h"
+
+#include UE_INLINE_GENERATED_CPP_BY_NAME(IKRigSolver)
+
+#if WITH_EDITORONLY_DATA
+
+void UIKRigSolver::PostLoad()
+{
+	Super::PostLoad();
+	SetFlags(RF_Transactional); // patch old solvers to enable undo/redo
+}
+
+void UIKRigSolver::PostTransacted(const FTransactionObjectEvent& TransactionEvent)
+{
+	Super::PostTransacted(TransactionEvent);
+	IKRigSolverModified.Broadcast(this);
+}
+
+#endif
+
+
diff --git a/Engine/Plugins/Animation/IKRig/Source/IKRig/Private/Rig/Solvers/IKRig_BodyMover.cpp b/Engine/Plugins/Animation/IKRig/Source/IKRig/Private/Rig/Solvers/IKRig_BodyMover.cpp
new file mode 100644
--- /dev/null
+++ b/Engine/Plugins/Animation/IKRig/Source/IKRig/Private/Rig/Solvers/IKRig_BodyMover.cpp
@@ -0,0 +1,229 @@
+// Copyright Epic Games, Inc. All Rights Reserved.
+
+#include "Rig/Solvers/IKRig_BodyMover.h"
+#include "Rig/IKRigDataTypes.h"
+#include "Rig/IKRigSkeleton.h"
+#include "Rig/Solvers/PointsToRotation.h"
+
+#include UE_INLINE_GENERATED_CPP_BY_NAME(IKRig_BodyMover)
+
+#define LOCTEXT_NAMESPACE "UIKRig_BodyMover"
+
+void UIKRig_BodyMover::Initialize(const FIKRigSkeleton& IKRigSkeleton)
+{
+	BodyBoneIndex = IKRigSkeleton.GetBoneIndexFromName(RootBone);
+}
+
+void UIKRig_BodyMover::Solve(FIKRigSkeleton& IKRigSkeleton, const FIKRigGoalContainer& Goals)
+{
+	// no body bone specified
+	if (BodyBoneIndex == INDEX_NONE)
+	{
+		return;
+	}
+
+	// no effectors added
+	if (Effectors.IsEmpty())
+	{
+		return;
+	}
+	
+	// ensure body bone exists
+	check(IKRigSkeleton.RefPoseGlobal.IsValidIndex(BodyBoneIndex));
+	
+	// calculate a "best fit" transform from deformed goal locations
+	TArray<FVector> InitialPoints;
+	TArray<FVector> CurrentPoints;
+	for (UIKRig_BodyMoverEffector* Effector : Effectors)
+	{
+		const FIKRigGoal* Goal = Goals.FindGoalByName(Effector->GoalName);
+		if (!Goal)
+		{
+			return;
+		}
+
+		const int32 BoneIndex = IKRigSkeleton.GetBoneIndexFromName(Effector->BoneName);
+		const FTransform InitialEffector = IKRigSkeleton.CurrentPoseGlobal[BoneIndex];
+
+		InitialPoints.Add(InitialEffector.GetTranslation());
+		CurrentPoints.Add(Goal->FinalBlendedPosition);
+	}
+	
+	FVector InitialCentroid;
+	FVector CurrentCentroid;
+	const FQuat RotationOffset = GetRotationFromDeformedPoints(
+		InitialPoints,
+		CurrentPoints,
+		InitialCentroid,
+		CurrentCentroid);
+
+	// alpha blend the position offset and add it to the current bone location
+	const FVector Offset = (CurrentCentroid - InitialCentroid);
+	const FVector Weight(
+		Offset.X > 0.f ? PositionPositiveX : PositionNegativeX,
+		Offset.Y > 0.f ? PositionPositiveY : PositionNegativeY,
+		Offset.Z > 0.f ? PositionPositiveZ : PositionNegativeZ);
+
+	// the bone transform to modify
+	FTransform& CurrentBodyTransform = IKRigSkeleton.CurrentPoseGlobal[BodyBoneIndex];
+	CurrentBodyTransform.AddToTranslation(Offset * (Weight*PositionAlpha));
+
+	// do per-axis alpha blend
+	FVector Euler = RotationOffset.Euler() * FVector(RotateXAlpha, RotateYAlpha, RotateZAlpha);
+	FQuat FinalRotationOffset = FQuat::MakeFromEuler(Euler);
+	// alpha blend the entire rotation offset
+	FinalRotationOffset = FQuat::FastLerp(FQuat::Identity, FinalRotationOffset, RotationAlpha).GetNormalized();
+	// add rotation offset to original rotation
+	CurrentBodyTransform.SetRotation(FinalRotationOffset * CurrentBodyTransform.GetRotation());
+
+	// do FK update of children
+	IKRigSkeleton.PropagateGlobalPoseBelowBone(BodyBoneIndex);
+}
+
+void UIKRig_BodyMover::UpdateSolverSettings(UIKRigSolver* InSettings)
+{
+	if(UIKRig_BodyMover* Settings = Cast<UIKRig_BodyMover>(InSettings))
+	{
+		// copy solver settings
+		PositionAlpha = Settings->PositionAlpha;
+		PositionPositiveX = Settings->PositionPositiveX;
+		PositionPositiveY = Settings->PositionPositiveY;
+		PositionPositiveZ = Settings->PositionPositiveZ;
+		PositionNegativeX = Settings->PositionNegativeX;
+		PositionNegativeY = Settings->PositionNegativeY;
+		PositionNegativeZ = Settings->PositionNegativeZ;
+		RotationAlpha = Settings->RotationAlpha;
+		RotateXAlpha = Settings->RotateXAlpha;
+		RotateYAlpha = Settings->RotateYAlpha;
+		RotateZAlpha = Settings->RotateZAlpha;
+
+		// copy effector settings
+		for (const UIKRig_BodyMoverEffector* InEffector : Settings->Effectors)
+		{
+			for (UIKRig_BodyMoverEffector* Effector : Effectors)
+			{
+				if (Effector->GoalName == InEffector->GoalName)
+				{
+					Effector->InfluenceMultiplier = InEffector->InfluenceMultiplier;
+					break;
+				}
+			}
+		}
+	}
+}
+
+void UIKRig_BodyMover::RemoveGoal(const FName& GoalName)
+{
+	// ensure goal even exists in this solver
+	const int32 GoalIndex = GetIndexOfGoal(GoalName);
+	if (GoalIndex == INDEX_NONE)
+	{
+		return;
+	}
+
+	// remove it
+	Effectors.RemoveAt(GoalIndex);
+}
+
+#if WITH_EDITOR
+
+FText UIKRig_BodyMover::GetNiceName() const
+{
+	return FText(LOCTEXT("SolverName", "Body Mover"));
+}
+
+bool UIKRig_BodyMover::GetWarningMessage(FText& OutWarningMessage) const
+{
+	if (RootBone == NAME_None)
+	{
+		OutWarningMessage = LOCTEXT("MissingRoot", "Missing root bone.");
+		return true;
+	}
+
+	if (Effectors.IsEmpty())
+	{
+		OutWarningMessage = LOCTEXT("MissingGoal", "Missing goals.");
+		return true;
+	}
+	
+	return false;
+}
+
+void UIKRig_BodyMover::AddGoal(const UIKRigEffectorGoal& NewGoal)
+{
+	UIKRig_BodyMoverEffector* NewEffector = NewObject<UIKRig_BodyMoverEffector>(this, UIKRig_BodyMoverEffector::StaticClass());
+	NewEffector->GoalName = NewGoal.GoalName;
+	NewEffector->BoneName = NewGoal.BoneName;
+	Effectors.Add(NewEffector);
+}
+
+void UIKRig_BodyMover::RenameGoal(const FName& OldName, const FName& NewName)
+{
+	// ensure goal even exists in this solver
+	const int32 GoalIndex = GetIndexOfGoal(OldName);
+	if (GoalIndex == INDEX_NONE)
+	{
+		return;
+	}
+
+	// rename
+	Effectors[GoalIndex]->GoalName = NewName;
+}
+
+void UIKRig_BodyMover::SetGoalBone(const FName& GoalName, const FName& NewBoneName)
+{
+	// ensure goal even exists in this solver
+	const int32 GoalIndex = GetIndexOfGoal(GoalName);
+	if (GoalIndex == INDEX_NONE)
+	{
+		return;
+	}
+
+	// rename
+	Effectors[GoalIndex]->Modify();
+	Effectors[GoalIndex]->BoneName = NewBoneName;
+}
+
+bool UIKRig_BodyMover::IsGoalConnected(const FName& GoalName) const
+{
+	return GetIndexOfGoal(GoalName) != INDEX_NONE;
+}
+
+UObject* UIKRig_BodyMover::GetGoalSettings(const FName& GoalName) const
+{
+	const int32 GoalIndex = GetIndexOfGoal(GoalName);
+	if (GoalIndex == INDEX_NONE)
+	{
+		return nullptr;
+	}
+
+	return Effectors[GoalIndex];
+}
+
+bool UIKRig_BodyMover::IsBoneAffectedBySolver(const FName& BoneName, const FIKRigSkeleton& IKRigSkeleton) const
+{
+	return IKRigSkeleton.IsBoneInDirectLineage(BoneName, RootBone);
+}
+
+void UIKRig_BodyMover::SetRootBone(const FName& RootBoneName)
+{
+	RootBone = RootBoneName;
+}
+
+#endif
+
+int32 UIKRig_BodyMover::GetIndexOfGoal(const FName& OldName) const
+{
+	for (int32 i=0; i<Effectors.Num(); ++i)
+	{
+		if (Effectors[i]->GoalName == OldName)
+		{
+			return i;
+		}
+	}
+
+	return INDEX_NONE;
+}
+
+#undef LOCTEXT_NAMESPACE
+
diff --git a/Engine/Plugins/Animation/IKRig/Source/IKRig/Private/Rig/Solvers/IKRig_LimbSolver.cpp b/Engine/Plugins/Animation/IKRig/Source/IKRig/Private/Rig/Solvers/IKRig_LimbSolver.cpp
new file mode 100644
--- /dev/null
+++ b/Engine/Plugins/Animation/IKRig/Source/IKRig/Private/Rig/Solvers/IKRig_LimbSolver.cpp
@@ -0,0 +1,249 @@
+// Copyright Epic Games, Inc. All Rights Reserved.
+
+
+#include "Rig/Solvers/IKRig_LimbSolver.h"
+
+#include UE_INLINE_GENERATED_CPP_BY_NAME(IKRig_LimbSolver)
+
+#define LOCTEXT_NAMESPACE "IKRig_LimbSolver"
+
+UIKRig_LimbSolver::UIKRig_LimbSolver()
+{
+	Effector = CreateDefaultSubobject<UIKRig_LimbEffector>(TEXT("Effector"));
+}
+
+void UIKRig_LimbSolver::Initialize(const FIKRigSkeleton& InSkeleton)
+{
+	Solver.Reset();
+	ChildrenToUpdate.Empty();
+	
+	if (Effector->GoalName == NAME_None || Effector->BoneName == NAME_None || RootName == NAME_None)
+	{
+		return;
+	}
+
+	int32 BoneIndex = InSkeleton.GetBoneIndexFromName(Effector->BoneName);
+	const int32 RootIndex = InSkeleton.GetBoneIndexFromName(RootName);
+	if (BoneIndex == INDEX_NONE || RootIndex == INDEX_NONE)
+	{
+		return;
+	}
+
+	// populate indices
+	TArray<int32> BoneIndices( {BoneIndex} );
+	BoneIndex = InSkeleton.GetParentIndex(BoneIndex);
+	while (BoneIndex != INDEX_NONE && BoneIndex >= RootIndex)
+	{
+		BoneIndices.Add(BoneIndex);
+		BoneIndex = InSkeleton.GetParentIndex(BoneIndex);
+	};
+
+	// if chain is not long enough
+	if (BoneIndices.Num() < 3)
+	{
+		return;
+	}
+
+	// sort the chain from root to end
+	Algo::Reverse(BoneIndices);
+
+	// initialize solver
+	for (int32 Index: BoneIndices)
+	{
+		const FVector Location = InSkeleton.CurrentPoseGlobal[Index].GetLocation();
+		Solver.AddLink(Location, Index);
+	}
+
+	const bool bInitialized = Solver.Initialize();
+	if (bInitialized)
+	{
+		// store children that needs propagation once solved
+		TArray<int32> Children;
+		for (int32 Index = 0; Index < BoneIndices.Num()-1; ++Index)
+		{
+			// store children if not already handled by the solver (part if the links)
+			InSkeleton.GetChildIndices(BoneIndices[Index], Children);
+			const int32 NextIndex = BoneIndices[Index+1];
+			for (const int32 ChildIndex: Children)
+			{
+				if (ChildIndex != NextIndex)
+				{
+					ChildrenToUpdate.Add(ChildIndex);
+					GatherChildren(ChildIndex, InSkeleton, ChildrenToUpdate);
+				}
+			}
+		}
+		// store end bone children
+		GatherChildren(BoneIndices.Last(), InSkeleton, ChildrenToUpdate);
+	}
+}
+
+void UIKRig_LimbSolver::Solve(FIKRigSkeleton& InOutRigSkeleton, const FIKRigGoalContainer& InGoals)
+{
+	if (Solver.NumLinks() < 3)
+	{
+		return;
+	}
+	
+	const FIKRigGoal* IKGoal = InGoals.FindGoalByName(Effector->GoalName);
+	if (!IKGoal)
+	{
+		return;
+	}
+
+	// update settings
+	FLimbSolverSettings Settings;
+	Settings.ReachPrecision = ReachPrecision;
+	Settings.MaxIterations = MaxIterations;
+	Settings.bEnableLimit = bEnableLimit;
+	Settings.MinRotationAngle = MinRotationAngle;
+	Settings.EndBoneForwardAxis = EndBoneForwardAxis; 
+	Settings.HingeRotationAxis = HingeRotationAxis;
+	Settings.bEnableTwistCorrection = bEnableTwistCorrection;
+	Settings.ReachStepAlpha = ReachStepAlpha;
+	Settings.PullDistribution = PullDistribution;
+	Settings.bAveragePull = bAveragePull;
+
+	// update settings
+	const FVector& GoalLocation = IKGoal->FinalBlendedPosition;
+	const FQuat& GoalRotation = IKGoal->FinalBlendedRotation;
+	const bool bModifiedLimb = Solver.Solve(
+		InOutRigSkeleton.CurrentPoseGlobal,
+		GoalLocation,
+		GoalRotation,
+		Settings);
+
+	// propagate if needed
+	if (bModifiedLimb)
+	{
+		// update chain bones local transform
+		for (int32 Index = 0; Index < Solver.NumLinks(); Index++)
+		{
+			InOutRigSkeleton.UpdateLocalTransformFromGlobal(Solver.GetBoneIndex(Index));
+		}
+
+		// propagate to children
+		for (const int32 ChildIndex: ChildrenToUpdate)
+		{
+			InOutRigSkeleton.UpdateGlobalTransformFromLocal(ChildIndex);
+		}
+	}
+}
+
+void UIKRig_LimbSolver::UpdateSolverSettings(UIKRigSolver* InSettings)
+{
+	if (UIKRig_LimbSolver* Settings = Cast<UIKRig_LimbSolver>(InSettings))
+	{
+		ReachPrecision = Settings->ReachPrecision;
+		HingeRotationAxis = Settings->HingeRotationAxis;
+		MaxIterations = Settings->MaxIterations;
+		bEnableLimit = Settings->bEnableLimit;
+		MinRotationAngle = Settings->MinRotationAngle;
+		bAveragePull = Settings->bAveragePull;
+		PullDistribution = Settings->PullDistribution;
+		ReachStepAlpha = Settings->ReachStepAlpha;
+		bEnableTwistCorrection = Settings->bEnableTwistCorrection;
+		EndBoneForwardAxis = Settings->EndBoneForwardAxis; 
+	}
+}
+
+void UIKRig_LimbSolver::RemoveGoal(const FName& GoalName)
+{
+	if (Effector->GoalName == GoalName)
+	{
+		Effector->Modify();
+		Effector->GoalName = NAME_None;
+		Effector->BoneName = NAME_None;
+	}
+}
+
+#if WITH_EDITOR
+
+FText UIKRig_LimbSolver::GetNiceName() const
+{
+	return FText(LOCTEXT("SolverName", "Limb IK"));
+}
+
+bool UIKRig_LimbSolver::GetWarningMessage(FText& OutWarningMessage) const
+{
+	if (Effector->GoalName == NAME_None)
+	{
+		OutWarningMessage = LOCTEXT("MissingGoal", "Missing goal.");
+		return true;
+	}
+
+	if (RootName == NAME_None)
+	{
+		OutWarningMessage = LOCTEXT("MissingRoot", "Missing root.");
+		return true;
+	}
+
+	if (Solver.NumLinks() < 3)
+	{
+		OutWarningMessage = LOCTEXT("Requires3BonesChain", "Requires at least 3 bones between root and goal.");
+		return true;
+	}
+	
+	return false;
+}
+
+bool UIKRig_LimbSolver::IsBoneAffectedBySolver(const FName& BoneName, const FIKRigSkeleton& IKRigSkeleton) const
+{
+	return IKRigSkeleton.IsBoneInDirectLineage(BoneName, RootName);
+}
+
+void UIKRig_LimbSolver::AddGoal(const UIKRigEffectorGoal& NewGoal)
+{
+	Effector->Modify();
+	Effector->GoalName = NewGoal.GoalName;
+	Effector->BoneName = NewGoal.BoneName;
+}
+
+void UIKRig_LimbSolver::RenameGoal(const FName& OldName, const FName& NewName)
+{
+	if (Effector->GoalName == OldName)
+	{
+		Effector->Modify();
+		Effector->GoalName = NewName;
+	}
+}
+
+void UIKRig_LimbSolver::SetGoalBone(const FName& GoalName, const FName& NewBoneName)
+{
+	if (Effector->GoalName == GoalName)
+	{
+		Effector->Modify();
+		Effector->BoneName = NewBoneName;
+	}
+}
+
+bool UIKRig_LimbSolver::IsGoalConnected(const FName& GoalName) const
+{
+	return (Effector->GoalName == GoalName);
+}
+
+UObject* UIKRig_LimbSolver::GetGoalSettings(const FName& GoalName) const
+{
+	return (Effector->GoalName == GoalName) ? Effector : nullptr;
+}
+
+void UIKRig_LimbSolver::SetRootBone(const FName& RootBoneName)
+{
+	RootName = RootBoneName;
+}
+
+#endif
+
+void UIKRig_LimbSolver::GatherChildren(const int32 BoneIndex, const FIKRigSkeleton& InSkeleton, TArray<int32>& OutChildren)
+{
+	TArray<int32> Children;
+	InSkeleton.GetChildIndices(BoneIndex, Children);
+	for (int32 ChildIndex: Children)
+	{
+		OutChildren.Add(ChildIndex);
+		GatherChildren(ChildIndex, InSkeleton, OutChildren);
+	}
+}
+
+#undef LOCTEXT_NAMESPACE
+
diff --git a/Engine/Plugins/Animation/IKRig/Source/IKRig/Private/Rig/Solvers/IKRig_PBIKSolver.cpp b/Engine/Plugins/Animation/IKRig/Source/IKRig/Private/Rig/Solvers/IKRig_PBIKSolver.cpp
new file mode 100644
--- /dev/null
+++ b/Engine/Plugins/Animation/IKRig/Source/IKRig/Private/Rig/Solvers/IKRig_PBIKSolver.cpp
@@ -0,0 +1,399 @@
+// Copyright Epic Games, Inc. All Rights Reserved.
+
+#include "Rig/Solvers/IKRig_FBIKSolver.h"
+
+#include "Rig/IKRigDataTypes.h"
+#include "Rig/IKRigSkeleton.h"
+
+#include UE_INLINE_GENERATED_CPP_BY_NAME(IKRig_FBIKSolver)
+
+#define LOCTEXT_NAMESPACE "UIKRigFBIKSolver"
+
+void UIKRigFBIKSolver::Initialize(const FIKRigSkeleton& InSkeleton)
+{	
+	// check how many effectors are assigned to a bone
+	int NumEffectors = 0;
+	for (const UIKRig_FBIKEffector* Effector : Effectors)
+	{
+		if (InSkeleton.GetBoneIndexFromName(Effector->BoneName) != INDEX_NONE)
+		{
+			++NumEffectors; // bone is set and exists!
+		}
+	}
+
+	// validate inputs are ready to be initialized
+	const bool bHasEffectors = NumEffectors > 0;
+	const bool bRootIsAssigned = RootBone != NAME_None;
+	if (!(bHasEffectors && bRootIsAssigned))
+	{
+		return; // not setup yet
+	}
+
+	// reset all internal data
+	Solver.Reset();
+
+	// create bones
+	for (int BoneIndex = 0; BoneIndex < InSkeleton.BoneNames.Num(); ++BoneIndex)
+	{
+		const FName& Name = InSkeleton.BoneNames[BoneIndex];
+
+		// get the parent bone solver index
+		const int32 ParentIndex = InSkeleton.GetParentIndexThatIsNotExcluded(BoneIndex);
+		const FTransform OrigTransform = InSkeleton.RefPoseGlobal[BoneIndex];
+		const FVector InOrigPosition = OrigTransform.GetLocation();
+		const FQuat InOrigRotation = OrigTransform.GetRotation();
+		const bool bIsRoot = Name == RootBone;
+		Solver.AddBone(Name, ParentIndex, InOrigPosition, InOrigRotation, bIsRoot);
+	}
+
+	// create effectors
+	for (UIKRig_FBIKEffector* Effector : Effectors)
+	{
+		Effector->IndexInSolver = Solver.AddEffector(Effector->BoneName);
+	}
+		
+	// initialize solver
+	Solver.Initialize();
+}
+
+void UIKRigFBIKSolver::Solve(FIKRigSkeleton& IKRigSkeleton, const FIKRigGoalContainer& Goals)
+{
+	if (!Solver.IsReadyToSimulate())
+	{
+		return;
+	}
+
+	if (Solver.GetNumBones() != IKRigSkeleton.BoneNames.Num())
+	{
+		return;
+	}
+
+	TArray<FTransform>& InOutTransforms = IKRigSkeleton.CurrentPoseGlobal;
+	
+	// set bones to input pose
+	for(int32 BoneIndex = 0; BoneIndex < Solver.GetNumBones(); BoneIndex++)
+	{
+		Solver.SetBoneTransform(BoneIndex, InOutTransforms[BoneIndex]);
+	}
+
+	// update bone settings
+	for (const UIKRig_FBIKBoneSettings* BoneSetting : BoneSettings)
+	{
+		const int32 BoneIndex = Solver.GetBoneIndex(BoneSetting->Bone);
+		if (PBIK::FBoneSettings* InternalSettings = Solver.GetBoneSettings(BoneIndex))
+		{
+			BoneSetting->CopyToCoreStruct(*InternalSettings);
+		}
+	}
+
+	// update effectors
+	for (const UIKRig_FBIKEffector* Effector : Effectors)
+	{
+		if (Effector->IndexInSolver < 0)
+		{
+			continue;
+		}
+		
+		const FIKRigGoal* Goal = Goals.FindGoalByName(Effector->GoalName);
+		if (!Goal)
+		{
+			return;
+		}
+
+		PBIK::FEffectorSettings Settings;
+		Settings.PositionAlpha = 1.0f; // this is constant because IKRig manages offset alphas itself
+		Settings.RotationAlpha = 1.0f; // this is constant because IKRig manages offset alphas itself
+		Settings.StrengthAlpha = Effector->StrengthAlpha;
+		Settings.PullChainAlpha = Effector->PullChainAlpha;
+		Settings.PinRotation = Effector->PinRotation;
+		
+		Solver.SetEffectorGoal(
+			Effector->IndexInSolver,
+			Goal->FinalBlendedPosition,
+			Goal->FinalBlendedRotation,
+			Settings);
+	}
+
+	// update settings
+	FPBIKSolverSettings Settings;
+	Settings.Iterations = Iterations;
+	Settings.MassMultiplier = MassMultiplier;
+	Settings.bAllowStretch = bAllowStretch;
+	Settings.RootBehavior = RootBehavior;
+	Settings.PrePullRootSettings = PrePullRootSettings;
+	Settings.GlobalPullChainAlpha = PullChainAlpha;
+	Settings.MaxAngle = MaxAngle;
+	Settings.OverRelaxation = OverRelaxation;
+	Settings.bStartSolveFromInputPose = bStartSolveFromInputPose;
+
+	// solve
+	Solver.Solve(Settings);
+
+	// copy transforms back
+	for(int32 BoneIndex = 0; BoneIndex < Solver.GetNumBones(); BoneIndex++)
+	{
+		Solver.GetBoneGlobalTransform(BoneIndex, InOutTransforms[BoneIndex]);
+	}
+}
+
+void UIKRigFBIKSolver::GetBonesWithSettings(TSet<FName>& OutBonesWithSettings) const
+{
+	for (UIKRig_FBIKBoneSettings* BoneSetting : BoneSettings)
+	{
+		if (BoneSetting)
+		{
+			OutBonesWithSettings.Add(BoneSetting->Bone);
+		}
+	}
+}
+
+void UIKRigFBIKSolver::UpdateSolverSettings(UIKRigSolver* InSettings)
+{
+	if(UIKRigFBIKSolver* Settings = Cast<UIKRigFBIKSolver>(InSettings))
+	{
+		Iterations = Settings->Iterations;
+		MassMultiplier = Settings->MassMultiplier;
+		bAllowStretch = Settings->bAllowStretch;
+		RootBehavior = Settings->RootBehavior;
+		PrePullRootSettings = Settings->PrePullRootSettings;
+		PullChainAlpha = Settings->PullChainAlpha;
+		MaxAngle = Settings->MaxAngle;
+		OverRelaxation = Settings->OverRelaxation;
+		bStartSolveFromInputPose = Settings->bStartSolveFromInputPose;
+
+		// copy effector settings
+		for (const UIKRig_FBIKEffector* InEffector : Settings->Effectors)
+		{
+			for (UIKRig_FBIKEffector* Effector : Effectors)
+			{
+				if (Effector && Effector->GoalName == InEffector->GoalName)
+				{
+					Effector->CopySettings(InEffector);
+					break;
+				}
+			}
+		}
+
+		// copy bone settings
+		for (const UIKRig_FBIKBoneSettings* InBoneSetting : Settings->BoneSettings)
+		{
+			for (UIKRig_FBIKBoneSettings* BoneSetting : BoneSettings)
+			{
+				if (BoneSetting && BoneSetting->Bone == InBoneSetting->Bone)
+				{
+					BoneSetting->CopySettings(InBoneSetting);
+					break;	
+				}
+			}
+		}
+	}
+}
+
+void UIKRigFBIKSolver::RemoveGoal(const FName& GoalName)
+{
+	// ensure goal even exists in this solver
+	const int32 GoalIndex = GetIndexOfGoal(GoalName);
+	if (GoalIndex == INDEX_NONE)
+	{
+		return;
+	}
+
+	// remove it
+	Effectors.RemoveAt(GoalIndex);
+}
+
+#if WITH_EDITOR
+
+FText UIKRigFBIKSolver::GetNiceName() const
+{
+	return FText(LOCTEXT("SolverName", "Full Body IK"));
+}
+
+bool UIKRigFBIKSolver::GetWarningMessage(FText& OutWarningMessage) const
+{
+	if (RootBone == NAME_None)
+	{
+		OutWarningMessage = LOCTEXT("MissingRoot", "Missing root bone.");
+		return true;
+	}
+
+	if (Effectors.IsEmpty())
+	{
+		OutWarningMessage = LOCTEXT("MissingGoal", "Missing goals.");
+		return true;
+	}
+	
+	return false;
+}
+
+void UIKRigFBIKSolver::AddGoal(const UIKRigEffectorGoal& NewGoal)
+{
+	UIKRig_FBIKEffector* NewEffector = NewObject<UIKRig_FBIKEffector>(this, UIKRig_FBIKEffector::StaticClass());
+	NewEffector->GoalName = NewGoal.GoalName;
+	NewEffector->BoneName = NewGoal.BoneName;
+	Effectors.Add(NewEffector);
+}
+
+void UIKRigFBIKSolver::RenameGoal(const FName& OldName, const FName& NewName)
+{
+	// ensure goal even exists in this solver
+	const int32 GoalIndex = GetIndexOfGoal(OldName);
+	if (GoalIndex == INDEX_NONE)
+	{
+		return;
+	}
+
+	// rename
+	Effectors[GoalIndex]->Modify();
+	Effectors[GoalIndex]->GoalName = NewName;
+}
+
+void UIKRigFBIKSolver::SetGoalBone(const FName& GoalName, const FName& NewBoneName)
+{
+	// ensure goal even exists in this solver
+	const int32 GoalIndex = GetIndexOfGoal(GoalName);
+	if (GoalIndex == INDEX_NONE)
+	{
+		return;
+	}
+
+	// rename
+	Effectors[GoalIndex]->Modify();
+	Effectors[GoalIndex]->BoneName = NewBoneName;
+}
+
+bool UIKRigFBIKSolver::IsGoalConnected(const FName& GoalName) const
+{
+	return GetIndexOfGoal(GoalName) != INDEX_NONE;
+}
+
+void UIKRigFBIKSolver::SetRootBone(const FName& RootBoneName)
+{
+	RootBone = RootBoneName;
+}
+
+UObject* UIKRigFBIKSolver::GetGoalSettings(const FName& GoalName) const
+{
+	const int32 GoalIndex = GetIndexOfGoal(GoalName);
+	if (GoalIndex == INDEX_NONE)
+	{
+		return nullptr;
+	}
+
+	return Effectors[GoalIndex];
+}
+
+void UIKRigFBIKSolver::AddBoneSetting(const FName& BoneName)
+{
+	if (GetBoneSetting(BoneName))
+	{
+		return; // already have settings on this bone
+	}
+
+	UIKRig_FBIKBoneSettings* NewBoneSettings = NewObject<UIKRig_FBIKBoneSettings>(this, UIKRig_FBIKBoneSettings::StaticClass());
+	NewBoneSettings->Bone = BoneName;
+	BoneSettings.Add(NewBoneSettings);
+}
+
+void UIKRigFBIKSolver::RemoveBoneSetting(const FName& BoneName)
+{
+	UIKRig_FBIKBoneSettings* BoneSettingToRemove = nullptr; 
+	for (UIKRig_FBIKBoneSettings* BoneSetting : BoneSettings)
+	{
+		if (BoneSetting->Bone == BoneName)
+		{
+			BoneSettingToRemove = BoneSetting;
+			break; // can only be one with this name
+		}
+	}
+
+	if (BoneSettingToRemove)
+	{
+		BoneSettings.Remove(BoneSettingToRemove);
+	}
+}
+
+UObject* UIKRigFBIKSolver::GetBoneSetting(const FName& BoneName) const
+{
+	for (UIKRig_FBIKBoneSettings* BoneSetting : BoneSettings)
+	{
+		if (BoneSetting && BoneSetting->Bone == BoneName)
+		{
+			return BoneSetting;
+		}
+	}
+	
+	return nullptr;
+}
+
+void UIKRigFBIKSolver::DrawBoneSettings(
+	const FName& BoneName,
+	const FIKRigSkeleton& IKRigSkeleton,
+	FPrimitiveDrawInterface* PDI) const
+{
+	
+}
+
+bool UIKRigFBIKSolver::IsBoneAffectedBySolver(const FName& BoneName, const FIKRigSkeleton& IKRigSkeleton) const
+{
+	// nothing is affected by solver without a root bone assigned or at least 1 effector
+	if (RootBone == NAME_None || Effectors.IsEmpty())
+	{
+		return false;
+	}
+
+	// has to be BELOW root
+	if (!IKRigSkeleton.IsBoneInDirectLineage(BoneName, RootBone))
+	{
+		return false;
+	}
+
+	// has to be ABOVE an effector
+	for (UIKRig_FBIKEffector* Effector : Effectors)
+	{
+		if (IKRigSkeleton.IsBoneInDirectLineage(Effector->BoneName, BoneName))
+		{
+			return true;
+		}
+	}
+	
+	return false;
+}
+
+void UIKRigFBIKSolver::PostLoad()
+{
+	Super::PostLoad();
+	
+	// patch for loading old effectors, we blow them away
+	bool bHasNullEffector = false;
+	for (const UIKRig_FBIKEffector* Effector : Effectors)
+	{
+		if (!Effector)
+		{
+			bHasNullEffector = true;
+		}
+	}
+
+	if (bHasNullEffector)
+	{
+		Effectors.Empty();
+	}
+}
+
+#endif
+
+int32 UIKRigFBIKSolver::GetIndexOfGoal(const FName& OldName) const
+{
+	for (int32 i=0; i<Effectors.Num(); ++i)
+	{
+		if (Effectors[i]->GoalName == OldName)
+		{
+			return i;
+		}
+	}
+
+	return INDEX_NONE;
+}
+
+#undef LOCTEXT_NAMESPACE
+
diff --git a/Engine/Plugins/Animation/IKRig/Source/IKRig/Private/Rig/Solvers/IKRig_PoleSolver.cpp b/Engine/Plugins/Animation/IKRig/Source/IKRig/Private/Rig/Solvers/IKRig_PoleSolver.cpp
new file mode 100644
--- /dev/null
+++ b/Engine/Plugins/Animation/IKRig/Source/IKRig/Private/Rig/Solvers/IKRig_PoleSolver.cpp
@@ -0,0 +1,282 @@
+// Copyright Epic Games, Inc. All Rights Reserved.
+
+#include "Rig/Solvers/IKRig_PoleSolver.h"
+#include "Rig/IKRigDataTypes.h"
+#include "Rig/IKRigSkeleton.h"
+
+#include UE_INLINE_GENERATED_CPP_BY_NAME(IKRig_PoleSolver)
+
+#define LOCTEXT_NAMESPACE "UIKRig_PoleSolver"
+
+UIKRig_PoleSolver::UIKRig_PoleSolver()
+{
+	Effector = CreateDefaultSubobject<UIKRig_PoleSolverEffector>(TEXT("Effector"));
+}
+
+void UIKRig_PoleSolver::Initialize(const FIKRigSkeleton& IKRigSkeleton)
+{
+	Chain.Empty();
+	
+	int32 BoneIndex = IKRigSkeleton.GetBoneIndexFromName(EndName);
+	const int32 RootIndex = IKRigSkeleton.GetBoneIndexFromName(RootName);
+
+	if (BoneIndex == INDEX_NONE || RootIndex == INDEX_NONE)
+	{
+		return;
+	}
+
+	// populate chain
+	Chain.Add(BoneIndex);
+	BoneIndex = IKRigSkeleton.GetParentIndex(BoneIndex);
+	while (BoneIndex != INDEX_NONE && BoneIndex >= RootIndex)
+	{
+		Chain.Add(BoneIndex);
+		BoneIndex = IKRigSkeleton.GetParentIndex(BoneIndex);
+	};
+
+	// if chain is not long enough
+	if (Chain.Num() < 3)
+	{
+		Chain.Empty();
+		return;
+	}
+	
+	// sort the chain from root to end
+	Algo::Reverse(Chain);
+	
+	// store children that needs propagation once solved
+	TArray<int32> Children;
+	for (int32 Index = 0; Index < Chain.Num()-1; ++Index)
+	{
+		// store children if not already handled by the solver
+		IKRigSkeleton.GetChildIndices(Chain[Index], Children);
+		const int32 NextIndex = Chain[Index+1];
+		for (const int32 ChildIndex: Children)
+		{
+			if (ChildIndex != NextIndex)
+			{
+				ChildrenToUpdate.Add(ChildIndex);
+				GatherChildren(ChildIndex, IKRigSkeleton, ChildrenToUpdate);
+			}
+		}
+	}
+}
+
+void UIKRig_PoleSolver::Solve(FIKRigSkeleton& IKRigSkeleton, const FIKRigGoalContainer& Goals)
+{
+	if (Chain.Num() < 3)
+	{
+		return;
+	}
+	
+	const FIKRigGoal* IKGoal = Goals.FindGoalByName(Effector->GoalName);
+	if (!IKGoal)
+	{
+		return;
+	}
+
+	const bool bHasAlpha = Effector->Alpha > KINDA_SMALL_NUMBER;
+	if (!bHasAlpha)
+	{
+		return;
+	}
+
+	const int32 RootIndex = Chain[0];
+	const int32 KneeIndex = Chain[1];
+	const int32 EndIndex = Chain.Last();
+
+	TArray<FTransform>& InOutTransforms = IKRigSkeleton.CurrentPoseGlobal;
+
+	// initial configuration
+	const FVector RootLocation = InOutTransforms[RootIndex].GetLocation();
+	const FVector KneeLocation = InOutTransforms[KneeIndex].GetLocation();
+	const FVector EndLocation = InOutTransforms[EndIndex].GetLocation();
+	
+	const FVector RootToEnd = (EndLocation - RootLocation).GetSafeNormal();
+	const FVector RootToKnee = (KneeLocation - RootLocation).GetSafeNormal();
+	if (RootToEnd.IsZero() || RootToKnee.IsZero())
+	{
+		return;
+	}
+	const FVector InitPlane = FVector::CrossProduct(RootToEnd, RootToKnee).GetSafeNormal();
+
+	// target configuration
+	const FVector& GoalLocation = IKGoal->FinalBlendedPosition;
+	const FVector RootToPole = (GoalLocation - RootLocation).GetSafeNormal();
+	if (RootToPole.IsZero())
+	{
+		return;
+	}
+	const FVector TargetPlane = FVector::CrossProduct(RootToEnd, RootToPole).GetSafeNormal();
+
+	// compute delta rotation from InitialPlane to TargetPlane
+	if (InitPlane.IsZero() || InitPlane.Equals(TargetPlane))
+	{
+    	return;
+    }
+	
+	const FQuat DeltaRotation = FQuat::FindBetweenNormals(InitPlane, TargetPlane);
+	if (!DeltaRotation.IsIdentity())
+	{
+		// update transforms
+		for (int32 Index = 0; Index < Chain.Num()-1; Index++)
+		{
+			int32 BoneIndex = Chain[Index];
+			FTransform& BoneTransform = InOutTransforms[BoneIndex];
+
+			// rotation
+			const FQuat BoneRotation = BoneTransform.GetRotation();
+			const FQuat TargetRotation = FMath::Lerp(BoneRotation, DeltaRotation * BoneRotation, Effector->Alpha);
+			BoneTransform.SetRotation(TargetRotation);
+
+			// translation
+			const FVector BoneTranslation = BoneTransform.GetLocation();
+			const FVector TargetTranslation = FMath::Lerp( BoneTranslation, RootLocation + DeltaRotation.RotateVector(BoneTranslation - RootLocation), Effector->Alpha);
+			BoneTransform.SetTranslation(TargetTranslation);
+		}
+
+		// propagate to children
+		for (const int32 ChildIndex: ChildrenToUpdate)
+		{
+			IKRigSkeleton.UpdateGlobalTransformFromLocal(ChildIndex);
+		}
+	}
+}
+
+void UIKRig_PoleSolver::UpdateSolverSettings(UIKRigSolver* InSettings)
+{
+	if (UIKRig_PoleSolver* Settings = Cast<UIKRig_PoleSolver>(InSettings))
+	{
+		Effector->Alpha = Settings->Effector->Alpha;
+	}
+}
+
+void UIKRig_PoleSolver::RemoveGoal(const FName& GoalName)
+{
+	if (Effector->GoalName == GoalName)
+	{
+		Effector->Modify();
+		Effector->GoalName = NAME_None;
+		Effector->BoneName = NAME_None;
+	}
+}
+
+#if WITH_EDITOR
+
+FText UIKRig_PoleSolver::GetNiceName() const
+{
+	return FText(LOCTEXT("SolverName", "Pole Solver"));
+}
+
+bool UIKRig_PoleSolver::GetWarningMessage(FText& OutWarningMessage) const
+{
+	if (RootName == NAME_None)
+	{
+		OutWarningMessage = LOCTEXT("MissingRoot", "Missing root bone.");
+		return true;
+	}
+	
+	if (EndName == NAME_None)
+	{
+		OutWarningMessage = LOCTEXT("MissingEnd", "Missing end bone.");
+		return true;
+	}
+	
+	if (Effector->GoalName == NAME_None)
+	{
+		OutWarningMessage = LOCTEXT("MissingGoal", "Missing goal.");
+		return true;
+	}
+	
+	if (Chain.Num() < 3)
+	{
+		OutWarningMessage = LOCTEXT("Requires3BonesChain", "Requires at least 3 bones between root and end bones.");
+		return true;
+	}
+
+	return false;
+}
+
+void UIKRig_PoleSolver::AddGoal(const UIKRigEffectorGoal& NewGoal)
+{
+	Effector->Modify();
+	Effector->GoalName = NewGoal.GoalName;
+	Effector->BoneName = NewGoal.BoneName;
+}
+
+void UIKRig_PoleSolver::RenameGoal(const FName& OldName, const FName& NewName)
+{
+	if (Effector->GoalName == OldName)
+	{
+		Effector->Modify();
+		Effector->GoalName = NewName;
+	}
+}
+
+void UIKRig_PoleSolver::SetGoalBone(const FName& GoalName, const FName& NewBoneName)
+{
+	if (Effector->GoalName == GoalName)
+	{
+		Effector->Modify();
+		Effector->BoneName = NewBoneName;
+	}
+}
+
+bool UIKRig_PoleSolver::IsGoalConnected(const FName& GoalName) const
+{
+	return (Effector->GoalName == GoalName);
+}
+
+UObject* UIKRig_PoleSolver::GetGoalSettings(const FName& GoalName) const
+{
+	return (Effector->GoalName == GoalName) ? Effector : nullptr;
+}
+
+bool UIKRig_PoleSolver::IsBoneAffectedBySolver(const FName& BoneName, const FIKRigSkeleton& IKRigSkeleton) const
+{
+	const bool bAffected = IKRigSkeleton.IsBoneInDirectLineage(BoneName, RootName);
+	if (!bAffected)
+	{
+		return false;
+	}
+	
+	const int32 EndIndex = Chain.IsEmpty() ? INDEX_NONE : Chain.Last();
+	if (EndIndex == INDEX_NONE)
+	{
+		return false;
+	}
+
+	const int32 ChildIndex = IKRigSkeleton.GetBoneIndexFromName(BoneName);
+	return ChildIndex <= EndIndex;
+}
+
+void UIKRig_PoleSolver::SetRootBone(const FName& RootBoneName)
+{
+	RootName = RootBoneName;
+}
+
+void UIKRig_PoleSolver::SetEndBone(const FName& EndBoneName)
+{
+	EndName = EndBoneName;
+}
+
+FName UIKRig_PoleSolver::GetEndBone() const
+{
+	return EndName;
+}
+
+#endif
+
+void UIKRig_PoleSolver::GatherChildren(const int32 BoneIndex, const FIKRigSkeleton& InSkeleton, TArray<int32>& OutChildren)
+{
+	TArray<int32> Children;
+	InSkeleton.GetChildIndices(BoneIndex, Children);
+	for (int32 ChildIndex: Children)
+	{
+		OutChildren.Add(ChildIndex);
+		GatherChildren(ChildIndex, InSkeleton, OutChildren);
+	}
+}
+
+#undef LOCTEXT_NAMESPACE
+
diff --git a/Engine/Plugins/Animation/IKRig/Source/IKRig/Private/Rig/Solvers/IKRig_SetTransform.cpp b/Engine/Plugins/Animation/IKRig/Source/IKRig/Private/Rig/Solvers/IKRig_SetTransform.cpp
new file mode 100644
--- /dev/null
+++ b/Engine/Plugins/Animation/IKRig/Source/IKRig/Private/Rig/Solvers/IKRig_SetTransform.cpp
@@ -0,0 +1,136 @@
+// Copyright Epic Games, Inc. All Rights Reserved.
+
+#include "Rig/Solvers/IKRig_SetTransform.h"
+#include "Rig/IKRigDataTypes.h"
+#include "Rig/IKRigSkeleton.h"
+
+#include UE_INLINE_GENERATED_CPP_BY_NAME(IKRig_SetTransform)
+
+#define LOCTEXT_NAMESPACE "UIKRig_SetTransform"
+
+UIKRig_SetTransform::UIKRig_SetTransform()
+{
+	Effector = CreateDefaultSubobject<UIKRig_SetTransformEffector>(TEXT("Effector"));
+}
+
+void UIKRig_SetTransform::Initialize(const FIKRigSkeleton& IKRigSkeleton)
+{
+	BoneIndex = IKRigSkeleton.GetBoneIndexFromName(RootBone);
+}
+
+void UIKRig_SetTransform::Solve(FIKRigSkeleton& IKRigSkeleton, const FIKRigGoalContainer& Goals)
+{
+	// BoneIndex is irrelevant
+	if (BoneIndex == INDEX_NONE)
+	{
+		return;
+	}
+	
+	const FIKRigGoal* InGoal = Goals.FindGoalByName(Goal);
+	if (!InGoal)
+	{
+		return;
+	}
+
+	// check that settings are such that there is anything to do at all
+	const bool bAnythingEnabled = Effector->bEnablePosition || Effector->bEnableRotation;
+	const bool bHasAlpha = Effector->Alpha > KINDA_SMALL_NUMBER;
+	if (!(bAnythingEnabled && bHasAlpha))
+	{
+		return;
+	}
+
+	FTransform& CurrentTransform = IKRigSkeleton.CurrentPoseGlobal[BoneIndex];	
+
+	if (Effector->bEnablePosition)
+	{
+		const FVector TargetPosition = FMath::Lerp(CurrentTransform.GetTranslation(), InGoal->FinalBlendedPosition, Effector->Alpha);
+		CurrentTransform.SetLocation(TargetPosition);
+	}
+	
+	if (Effector->bEnableRotation)
+	{
+		const FQuat TargetRotation = FMath::Lerp(CurrentTransform.GetRotation(), InGoal->FinalBlendedRotation, Effector->Alpha);
+		CurrentTransform.SetRotation(TargetRotation);
+	}
+	
+	IKRigSkeleton.PropagateGlobalPoseBelowBone(BoneIndex);
+}
+
+void UIKRig_SetTransform::UpdateSolverSettings(UIKRigSolver* InSettings)
+{
+	if (UIKRig_SetTransform* Settings = Cast<UIKRig_SetTransform>(InSettings))
+	{
+		Effector->bEnablePosition = Settings->Effector->bEnablePosition;
+		Effector->bEnableRotation = Settings->Effector->bEnableRotation;
+		Effector->Alpha = Settings->Effector->Alpha;
+	}
+}
+
+void UIKRig_SetTransform::RemoveGoal(const FName& GoalName)
+{
+	if (Goal == GoalName)
+	{
+		Goal = NAME_None;
+		RootBone = NAME_None;
+	}	
+}
+
+#if WITH_EDITOR
+
+FText UIKRig_SetTransform::GetNiceName() const
+{
+	return FText(LOCTEXT("SolverName", "Set Transform"));
+}
+
+bool UIKRig_SetTransform::GetWarningMessage(FText& OutWarningMessage) const
+{
+	if (RootBone == NAME_None)
+	{
+		OutWarningMessage = LOCTEXT("MissingGoal", "Missing goal.");
+		return true;
+	}
+	return false;
+}
+
+void UIKRig_SetTransform::AddGoal(const UIKRigEffectorGoal& NewGoal)
+{
+	Goal = NewGoal.GoalName;
+	RootBone = NewGoal.BoneName;
+}
+
+void UIKRig_SetTransform::RenameGoal(const FName& OldName, const FName& NewName)
+{
+	if (Goal == OldName)
+	{
+		Goal = NewName;
+	}
+}
+
+void UIKRig_SetTransform::SetGoalBone(const FName& GoalName, const FName& NewBoneName)
+{
+	if (Goal == GoalName)
+	{
+		RootBone = NewBoneName;
+	}
+}
+
+bool UIKRig_SetTransform::IsGoalConnected(const FName& GoalName) const
+{
+	return Goal == GoalName;
+}
+
+UObject* UIKRig_SetTransform::GetGoalSettings(const FName& GoalName) const
+{
+	return Goal == GoalName ? Effector : nullptr;
+}
+
+bool UIKRig_SetTransform::IsBoneAffectedBySolver(const FName& BoneName, const FIKRigSkeleton& IKRigSkeleton) const
+{
+	return IKRigSkeleton.IsBoneInDirectLineage(BoneName, RootBone);
+}
+
+#endif
+
+#undef LOCTEXT_NAMESPACE
+
diff --git a/Engine/Plugins/Animation/IKRig/Source/IKRig/Private/Rig/Solvers/LimbSolver.cpp b/Engine/Plugins/Animation/IKRig/Source/IKRig/Private/Rig/Solvers/LimbSolver.cpp
new file mode 100644
--- /dev/null
+++ b/Engine/Plugins/Animation/IKRig/Source/IKRig/Private/Rig/Solvers/LimbSolver.cpp
@@ -0,0 +1,630 @@
+// Copyright Epic Games, Inc. All Rights Reserved.
+
+#include "Rig/Solvers/LimbSolver.h"
+
+#include UE_INLINE_GENERATED_CPP_BY_NAME(LimbSolver)
+
+void FLimbSolver::Reset()
+{
+	Length = 0.f;
+	Links.Empty();
+}
+
+bool FLimbSolver::Initialize()
+{
+	Length = 0.f;
+
+	// compute segments' length + total limb's length
+	for (int32 Index = 1; Index < NumLinks(); ++Index)
+	{
+		FLimbLink& Link = Links[Index];
+		Link.Length = FVector::Dist( Link.Location, Links[Index-1].Location);
+		Length += Link.Length;
+	}
+
+	return (NumLinks() > 2);
+}
+
+bool FLimbSolver::Solve(
+	TArray<FTransform>& InOutTransforms,
+	const FVector& InGoalLocation,
+	const FQuat& InGoalRotation,
+	const FLimbSolverSettings& InSettings)
+{
+	// rotate root so the end aligns with the goal
+	const bool bOriented = OrientTransformsTowardsGoal( InOutTransforms, InGoalLocation);
+
+	// update links locations
+	for (int32 Index = 0; Index < NumLinks(); ++Index)
+	{
+		const int32 BoneIndex = Links[Index].BoneIndex;
+		Links[Index].Location = InOutTransforms[BoneIndex].GetLocation();
+	}
+
+	// solve ik
+	const bool bReached = AnimationCore::SolveLimb(
+		Links,
+		Length,
+		InGoalLocation,
+		InSettings.ReachPrecision,
+		InSettings.MaxIterations,
+		InOutTransforms[Links[0].BoneIndex].GetUnitAxis(InSettings.HingeRotationAxis),
+		InSettings.bEnableLimit,
+		InSettings.MinRotationAngle,
+		InSettings.ReachStepAlpha,
+		InSettings.bAveragePull,
+		InSettings.PullDistribution
+		);
+
+	const bool bAdjustedKneeTwist = InSettings.bEnableTwistCorrection ? AdjustKneeTwist(
+															InOutTransforms,
+															FTransform( InGoalRotation, InGoalLocation),
+															InSettings.EndBoneForwardAxis ) : false;
+
+	bool bModifiedLimb = bOriented || bReached || bAdjustedKneeTwist;
+	
+	// adjust end bone orientation
+	const int32 EndIndex = Links.Last().BoneIndex;
+	FTransform& EndTransform = InOutTransforms[EndIndex];
+	if (bModifiedLimb || !EndTransform.GetRotation().Equals(InGoalRotation))
+	{
+		EndTransform.SetRotation(InGoalRotation);
+		bModifiedLimb = true;
+	}
+	
+	if (bModifiedLimb)
+	{
+		// update transforms
+		UpdateTransforms(InOutTransforms);
+	}
+
+	return bModifiedLimb;
+}
+
+int32 FLimbSolver::NumLinks() const
+{
+	return Links.Num();
+}
+
+int32 FLimbSolver::AddLink(const FVector& InLocation, int32 InBoneIndex)
+{
+	return Links.Emplace(InLocation, InBoneIndex);
+}
+
+int32 FLimbSolver::GetBoneIndex(int32 Index) const
+{
+	check(Index >= INDEX_NONE && Index < Links.Num());
+	return Links[Index].BoneIndex; 
+}
+
+bool FLimbSolver::OrientTransformsTowardsGoal(TArray<FTransform>& InOutTransforms, const FVector& InGoalLocation) const
+{
+	const int32 RootIndex = Links[0].BoneIndex;
+	const int32 EndIndex = Links.Last().BoneIndex; 
+	
+	const FVector RootLocation = InOutTransforms[RootIndex].GetLocation();
+	const FVector EndLocation = InOutTransforms[EndIndex].GetLocation();
+
+	const FVector InitialDir = (EndLocation - RootLocation).GetSafeNormal();
+	const FVector TargetDir = (InGoalLocation - RootLocation).GetSafeNormal();
+	
+	// find delta rotation take takes us from Old to New dir
+	if (!InitialDir.IsZero() && !InitialDir.Equals(TargetDir))
+	{
+		const FQuat DeltaRotation = FQuat::FindBetweenNormals(InitialDir, TargetDir);
+		if (!DeltaRotation.IsIdentity())
+		{
+			// rotate Leg so it is aligned with IK Target
+			for (const FLimbLink& Link: Links)
+			{
+				FTransform& BoneTransform = InOutTransforms[Link.BoneIndex]; 
+				
+				BoneTransform.SetRotation(DeltaRotation * BoneTransform.GetRotation());
+
+				const FVector BoneLocation = BoneTransform.GetLocation();
+				BoneTransform.SetLocation(RootLocation + DeltaRotation.RotateVector(BoneLocation - RootLocation));
+			}
+
+			return true;
+		}
+	}
+
+	return false;
+}
+
+void FLimbSolver::UpdateTransforms(TArray<FTransform>& InOutTransforms)
+{
+	// Rotations
+	for (int32 Index = 1; Index < NumLinks(); Index++)
+	{
+		const int32 BoneIndex = Links[Index].BoneIndex;
+		const FTransform& BoneTransform = InOutTransforms[BoneIndex];
+		const FVector& TargetBoneLocation = Links[Index].Location;
+		
+		const int32 ParentIndex = Links[Index-1].BoneIndex;
+		FTransform& ParentTransform = InOutTransforms[ParentIndex];
+		const FVector& TargetParentLocation = Links[Index-1].Location;
+
+		// calculate pre-translation vector between this bone and child
+		const FVector InitialDir = (BoneTransform.GetLocation() - ParentTransform.GetLocation()).GetSafeNormal();
+
+		// Get vector from the post-translation bone to it's child
+		const FVector TargetDir = (TargetBoneLocation - TargetParentLocation).GetSafeNormal();
+
+		const FQuat DeltaRotation = FQuat::FindBetweenNormals(InitialDir, TargetDir);
+		ParentTransform.SetRotation(DeltaRotation * ParentTransform.GetRotation());
+	}
+ 
+	// Translations (do it in one pass?)
+	for (int32 Index = 1; Index < NumLinks(); Index++)
+	{
+		const int32 BoneIndex = Links[Index].BoneIndex;
+		InOutTransforms[BoneIndex].SetTranslation(Links[Index].Location);
+	}
+}
+
+bool FLimbSolver::AdjustKneeTwist(TArray<FTransform>& InOutTransforms, const FTransform& InGoalTransform, const EAxis::Type FootBoneForwardAxis )
+{
+	const int32 EndIndex = Links.Last().BoneIndex;
+	const FVector FootFKLocation = InOutTransforms[EndIndex].GetLocation();
+	const FVector FootIKLocation = InGoalTransform.GetLocation();
+
+	const int32 RootIndex = Links[0].BoneIndex;
+	const FVector HipLocation = InOutTransforms[RootIndex].GetLocation();
+	const FVector FootAxisZ = (FootIKLocation - HipLocation).GetSafeNormal();
+
+	FVector FootFKAxisX = InOutTransforms[EndIndex].GetUnitAxis(FootBoneForwardAxis);
+	FVector FootIKAxisX = InGoalTransform.GetUnitAxis(FootBoneForwardAxis);
+
+	// Reorient X Axis to be perpendicular with FootAxisZ
+	const FVector TmpFKAxisY = FVector::CrossProduct(FootAxisZ, FootFKAxisX);
+	FootFKAxisX = FVector::CrossProduct(TmpFKAxisY, FootAxisZ);
+	
+	const FVector TmpIKAxisY = FVector::CrossProduct(FootAxisZ, FootIKAxisX);
+	FootIKAxisX = FVector::CrossProduct(TmpIKAxisY, FootAxisZ);
+
+	// Compare Axis X to see if we need a rotation to be performed
+	if (!FootFKAxisX.IsZero() && !FootFKAxisX.Equals(FootIKAxisX))
+	{
+		const FQuat DeltaRotation = FQuat::FindBetweenNormals(FootFKAxisX, FootIKAxisX);
+		if (!DeltaRotation.IsIdentity())
+		{
+			// rotate Leg so it is aligned with IK Target
+			for (const FLimbLink& Link: Links)
+			{
+				FTransform& BoneTransform = InOutTransforms[Link.BoneIndex]; 
+				
+				BoneTransform.SetRotation(DeltaRotation * BoneTransform.GetRotation());
+
+				const FVector BoneLocation = BoneTransform.GetLocation();
+				BoneTransform.SetLocation(HipLocation + DeltaRotation.RotateVector(BoneLocation - HipLocation));
+			}
+
+			return true;
+		}
+	}
+
+	return false;
+}
+
+namespace AnimationCore
+{
+	void SolveStraightLimb(TArray<FLimbLink>& InOutLinks, const FVector& GoalLocation)
+	{
+		const FVector& RootLocation = InOutLinks[0].Location;
+		const FVector Direction = (GoalLocation - RootLocation).GetSafeNormal();
+
+		for (int32 Index = 1; Index < InOutLinks.Num(); Index++)
+		{
+			InOutLinks[Index].Location = InOutLinks[Index-1].Location + Direction * InOutLinks[Index].Length;
+		}
+	}
+
+	void SolveTwoBoneIK(TArray<FLimbLink>& InOutLinks, const FVector& InTargetLocation, const FVector& HingeRotationAxis)
+	{
+		check(InOutLinks.Num() == 3);
+
+		FVector& pA = InOutLinks[2].Location; // Foot
+		FVector& pB = InOutLinks[1].Location; // Knee
+		FVector& pC = InOutLinks[0].Location; // Hip / Root
+		
+		// Move foot directly to target.
+		pA = InTargetLocation;
+
+		const FVector HipToFoot = pA - pC;
+
+		// Use Law of Cosines to work out solution.
+		// At this point we know the target location is reachable, and we are already aligned with that location.
+		// So the leg is in the right plane.
+		const float a = InOutLinks[1].Length;	// hip to knee
+		const float b = HipToFoot.Length();		// hip to foot
+		const float c = InOutLinks[2].Length;	// knee to foot
+		
+		const float Two_ab = 2.f * a * b;
+		const float CosC = !FMath::IsNearlyZero(Two_ab) ? (a * a + b * b - c * c) / Two_ab : 0.f;
+ 		const float C = FMath::Acos(CosC);
+		
+		// Project Knee onto Hip to Foot line.
+		const FVector HipToFootDir = !FMath::IsNearlyZero(b) ? HipToFoot / b : FVector::ZeroVector;
+		const FVector HipToKnee = pB - pC;
+		const FVector ProjKnee = pC + HipToKnee.ProjectOnToNormal(HipToFootDir);
+
+		const FVector ProjKneeToKnee = (pB - ProjKnee);
+		FVector BendDir = ProjKneeToKnee.GetSafeNormal(KINDA_SMALL_NUMBER);
+		
+		// If we have a HingeRotationAxis defined, we can cache 'BendDir'
+		// and use it when we can't determine it. (When limb is straight without a bend).
+		// We do this instead of using an explicit one, so we carry over the pole vector that animators use. 
+		// So they can animate it, and we try to extract it from the animation.
+		if ((HingeRotationAxis != FVector::ZeroVector) && (HipToFootDir != FVector::ZeroVector) && !FMath::IsNearlyZero(a))
+		{
+			const FVector HipToKneeDir = HipToKnee / a;
+			const float KneeBendDot = FVector::DotProduct(HipToKneeDir, HipToFootDir);
+
+			FVector& CachedRealBendDir = InOutLinks[1].RealBendDir;
+			FVector& CachedBaseBendDir = InOutLinks[1].BaseBendDir;
+
+			// Valid 'bend', cache 'BendDir'
+			if ((BendDir != FVector::ZeroVector) && (KneeBendDot < 0.99f))
+			{
+				CachedRealBendDir = BendDir;
+				CachedBaseBendDir = FVector::CrossProduct(HingeRotationAxis, HipToFootDir);
+			}
+			// Limb is too straight, can't determine BendDir accurately, so use cached value if possible.
+			else 
+			{
+				// If we have cached 'BendDir', then reorient it based on 'HingeRotationAxis'
+				if (CachedRealBendDir != FVector::ZeroVector)
+				{
+					const FVector CurrentBaseBendDir = FVector::CrossProduct(HingeRotationAxis, HipToFootDir);
+					const FQuat DeltaCachedToCurrBendDir = FQuat::FindBetweenNormals(CachedBaseBendDir, CurrentBaseBendDir);
+					BendDir = DeltaCachedToCurrBendDir.RotateVector(CachedRealBendDir);
+				}
+			}
+		}
+
+		// We just combine both lines into one to save a multiplication.
+		// const FVector NewProjectedKneeLoc = pC + HipToFootDir * a * CosC;
+		// const FVector NewKneeLoc = NewProjectedKneeLoc + Dir_LegLineToKnee * a * FMath::Sin(C);
+		const FVector NewKneeLoc = pC + a * (HipToFootDir * CosC + BendDir * FMath::Sin(C));
+		pB = NewKneeLoc;
+	}
+
+	FVector FindPlaneNormal(const TArray<FLimbLink>& InLinks, const FVector& RootLocation, const FVector& TargetLocation)
+	{
+		const FVector AxisX = (TargetLocation - RootLocation).GetSafeNormal();
+
+		for (int32 Index = 1; Index < InLinks.Num(); ++Index)
+		{
+			const FVector AxisY = (InLinks[Index].Location - RootLocation).GetSafeNormal();
+			const FVector PlaneNormal = FVector::CrossProduct(AxisX, AxisY);
+
+			// Make sure we have a valid normal (Axes were not coplanar).
+			if (PlaneNormal.SizeSquared() > SMALL_NUMBER)
+			{
+				return PlaneNormal.GetUnsafeNormal();
+			}
+		}
+
+		// All links are co-planar?
+		return FVector::UpVector;
+	}
+
+	void FABRIK_ApplyLinkConstraints_Forward(TArray<FLimbLink>& InOutLinks, int32 LinkIndex, float MinRotationAngleRadians)
+	{
+		if ((LinkIndex <= 0) || (LinkIndex >= InOutLinks.Num() - 1))
+		{
+			return;
+		}
+	
+		const FLimbLink& ChildLink = InOutLinks[LinkIndex + 1];
+		const FLimbLink& CurrentLink = InOutLinks[LinkIndex];
+		FLimbLink& ParentLink = InOutLinks[LinkIndex - 1];
+	
+		const FVector ChildAxisX = (ChildLink.Location - CurrentLink.Location).GetSafeNormal();
+		const FVector ChildAxisY = FVector::CrossProduct(CurrentLink.LinkAxisZ, ChildAxisX);
+		const FVector ParentAxisX = (ParentLink.Location - CurrentLink.Location).GetSafeNormal();
+	
+		const float ParentCos = FVector::DotProduct(ParentAxisX, ChildAxisX);
+		const float ParentSin = FVector::DotProduct(ParentAxisX, ChildAxisY);
+	
+		const bool bNeedsReorient = (ParentSin < 0.f) || (ParentCos > FMath::Cos(MinRotationAngleRadians));
+	
+		// Parent Link needs to be reoriented.
+		if (bNeedsReorient)
+		{
+			// folding over itself.
+			if (ParentCos > 0.f)
+			{
+				// Enforce minimum angle.
+				ParentLink.Location = CurrentLink.Location + CurrentLink.Length * (FMath::Cos(MinRotationAngleRadians) * ChildAxisX + FMath::Sin(MinRotationAngleRadians) * ChildAxisY);
+			}
+			else
+			{
+				// When opening up leg, allow it to extend in a full straight line.
+				ParentLink.Location = CurrentLink.Location - ChildAxisX * CurrentLink.Length;
+			}
+		}
+	}
+	
+	void FABRIK_ApplyLinkConstraints_Backward(TArray<FLimbLink>& InOutLinks, int32 LinkIndex, float MinRotationAngleRadians)
+	{
+		if ((LinkIndex <= 0) || (LinkIndex >= InOutLinks.Num() - 1))
+		{
+			return;
+		}
+	
+		FLimbLink& ChildLink = InOutLinks[LinkIndex + 1];
+		const FLimbLink& CurrentLink = InOutLinks[LinkIndex];
+		const FLimbLink& ParentLink = InOutLinks[LinkIndex - 1];
+	
+		const FVector ParentAxisX = (ParentLink.Location - CurrentLink.Location).GetSafeNormal();
+		const FVector ParentAxisY = FVector::CrossProduct(CurrentLink.LinkAxisZ, ParentAxisX);
+		const FVector ChildAxisX = (ChildLink.Location - CurrentLink.Location).GetSafeNormal();
+	
+		const float ChildCos = FVector::DotProduct(ChildAxisX, ParentAxisX);
+		const float ChildSin = FVector::DotProduct(ChildAxisX, ParentAxisY);
+
+		const float CosMinAngle = FMath::Cos(MinRotationAngleRadians);
+		
+		// Parent Link needs to be reoriented.
+		const bool bNeedsReorient = (ChildSin > 0.f) || (ChildCos > CosMinAngle);
+		if (bNeedsReorient)
+		{
+			// folding over itself.
+			if (ChildCos > 0.f)
+			{
+				// Enforce minimum angle.
+				const float SinMinAngle = FMath::Sin(MinRotationAngleRadians);
+				ChildLink.Location = CurrentLink.Location + ChildLink.Length * (CosMinAngle * ParentAxisX - SinMinAngle * ParentAxisY);
+			}
+			else
+			{
+				// When opening up leg, allow it to extend in a full straight line.
+				ChildLink.Location = CurrentLink.Location - ParentAxisX * ChildLink.Length;
+			}
+		}
+	}
+	
+	void  FABRIK_ForwardReach(
+		TArray<FLimbLink>& InOutLinks,
+		const FVector& InTargetLocation,
+		float ReachStepAlpha,
+		bool bEnableRotationLimit,
+		float MinRotationAngleRadians)
+	{
+		const int32 NumLinks = InOutLinks.Num();
+		
+		// Move end effector towards target
+		// If we are compressing the chain, limit displacement.
+		// Due to how FABRIK works, if we push the target past the parent's joint, we flip the bone.
+		{
+			FVector EndEffectorToTarget = InTargetLocation - InOutLinks.Last().Location;
+
+			FVector EndEffectorToTargetDir;
+			float EndEffectToTargetSize;
+			EndEffectorToTarget.ToDirectionAndLength(EndEffectorToTargetDir, EndEffectToTargetSize);
+
+			float Displacement = EndEffectToTargetSize;
+			for (int32 Index = NumLinks-1; Index >= 1; Index--)
+			{
+				FVector EndEffectorToParent = InOutLinks[Index].Location - InOutLinks.Last().Location;
+				float ParentDisplacement = FVector::DotProduct(EndEffectorToParent, EndEffectorToTargetDir);
+
+				Displacement = (ParentDisplacement > 0.f) ? FMath::Min(Displacement, ParentDisplacement * ReachStepAlpha) : Displacement;
+			}
+
+			InOutLinks.Last().Location += EndEffectorToTargetDir * Displacement;
+		}
+
+		// "Forward Reaching" stage - adjust bones from end effector.
+		for (int32 Index = NumLinks-2; Index > 0; Index--)
+		{
+			const FLimbLink& ChildLink = InOutLinks[Index+1];
+			FLimbLink& CurrentLink = InOutLinks[Index];
+
+			CurrentLink.Location = ChildLink.Location + (CurrentLink.Location - ChildLink.Location).GetUnsafeNormal() * ChildLink.Length;
+
+			if (bEnableRotationLimit)
+			{
+				FABRIK_ApplyLinkConstraints_Forward(InOutLinks, Index, MinRotationAngleRadians);
+			}
+		}
+	}
+
+	void FABRIK_BackwardReach(
+		TArray<FLimbLink>& InOutLinks,
+		const FVector& InRootTargetLocation,
+		float ReachStepAlpha,
+		bool bEnableRotationLimit,
+		float MinRotationAngleRadians)
+	{
+		const int32 NumLinks = InOutLinks.Num();
+		
+		// Move Root back towards RootTarget
+		// If we are compressing the chain, limit displacement.
+		// Due to how FABRIK works, if we push the target past the parent's joint, we flip the bone.
+		{
+			FVector RootToRootTarget = InRootTargetLocation - InOutLinks[0].Location;
+			
+			FVector RootToRootTargetDir;
+			float RootToRootTargetSize;
+			RootToRootTarget.ToDirectionAndLength(RootToRootTargetDir, RootToRootTargetSize);
+
+			float Displacement = RootToRootTargetSize;
+			
+			for (int32 Index = 0; Index < NumLinks-1; Index++)
+			{
+				FVector RootToChild = InOutLinks[Index].Location - InOutLinks[0].Location;
+				float ChildDisplacement = FVector::DotProduct(RootToChild, RootToRootTargetDir);
+				
+				Displacement = (ChildDisplacement > 0.f) ? FMath::Min(Displacement, ChildDisplacement * ReachStepAlpha) : Displacement;				
+			}
+
+			InOutLinks[0].Location += RootToRootTargetDir * Displacement;
+		}
+
+		// "Backward Reaching" stage - adjust bones from root.
+		for (int32 Index = 1; Index < NumLinks-1; Index++)
+		{
+			const FLimbLink& ParentLink = InOutLinks[Index-1];
+			FLimbLink& CurrentLink = InOutLinks[Index];
+			CurrentLink.Location = ParentLink.Location + (CurrentLink.Location - ParentLink.Location).GetUnsafeNormal() * CurrentLink.Length;
+
+			if (bEnableRotationLimit)
+			{
+				FABRIK_ApplyLinkConstraints_Backward(InOutLinks, Index, MinRotationAngleRadians);
+			}
+		}
+	}
+	
+	void SolveFABRIK(
+		TArray<FLimbLink>& InOutLinks,
+		const FVector& InTargetLocation,
+		float InReachPrecision,
+		int32 InMaxIterations,
+		bool bUseAngleLimit,
+		float MinRotationAngleRadians,
+		float ReachStepAlpha,
+		float PullDistribution,
+		bool bAveragePull)
+	{
+		const int32 NumLinks = InOutLinks.Num();
+		
+		// Make sure precision is not too small.
+		const float ReachPrecision = FMath::Max(InReachPrecision, KINDA_SMALL_NUMBER);
+
+		const FVector RootTargetLocation = InOutLinks[0].Location;
+
+		// Check distance between foot and foot target location
+		float Slop = FVector::Dist(InOutLinks.Last().Location, InTargetLocation);
+		
+		if (Slop > ReachPrecision || bUseAngleLimit)
+		{
+			if (bUseAngleLimit)
+			{
+				// Since we've previously aligned the foot with the IK Target, we're solving IK in 2D space on a single plane.
+				// Find Plane Normal, to use in rotation constraints.
+				const FVector PlaneNormal = FindPlaneNormal(InOutLinks, RootTargetLocation, InTargetLocation);
+
+				for (int32 Index = NumLinks-2; Index > 0; Index--)
+				{
+					const FLimbLink& ChildLink = InOutLinks[Index + 1];
+					FLimbLink& CurrentLink = InOutLinks[Index];
+					const FLimbLink& ParentLink = InOutLinks[Index - 1];
+
+					const FVector ChildAxisX = (ChildLink.Location - CurrentLink.Location).GetSafeNormal();
+					const FVector ChildAxisY = FVector::CrossProduct(PlaneNormal, ChildAxisX);
+					const FVector ParentAxisX = (ParentLink.Location - CurrentLink.Location).GetSafeNormal();
+
+					// Orient Z, so that ChildAxisY points 'up' and produces positive Sin values.
+					CurrentLink.LinkAxisZ = FVector::DotProduct(ParentAxisX, ChildAxisY) > 0.f ? PlaneNormal : -PlaneNormal;
+				}
+			}
+
+			// Re-position limb to distribute pull
+			const FVector EndToTarget = InTargetLocation - InOutLinks.Last().Location;
+			const FVector RootToTarget = RootTargetLocation - InOutLinks[0].Location;
+			const FVector PullDistributionOffset = PullDistribution * (EndToTarget) + (1.f - PullDistribution) * (RootToTarget);
+			for (int32 Index = 0; Index < NumLinks; Index++)
+			{
+				InOutLinks[Index].Location += PullDistributionOffset;
+			}
+			
+			int32 IterationCount = 1;
+			const int32 MaxIterations = FMath::Max(InMaxIterations, 1);
+			do
+			{
+				const float PreviousSlop = Slop;
+
+				// Pull averaging only has a visual impact when we have more than 2 bones (3 links).
+				const bool bNeedPullAveraging =(NumLinks > 3) && bAveragePull && (Slop > 1.f);
+				if (bNeedPullAveraging)
+				{	
+					TArray<FLimbLink> ForwardPull = InOutLinks;
+					FABRIK_ForwardReach(ForwardPull, InTargetLocation, ReachStepAlpha, bUseAngleLimit, MinRotationAngleRadians);
+
+					TArray<FLimbLink> BackwardPull = InOutLinks;
+					FABRIK_BackwardReach(BackwardPull,RootTargetLocation, ReachStepAlpha, bUseAngleLimit, MinRotationAngleRadians);
+
+					// Average pulls
+					for (int32 LinkIndex = 0; LinkIndex < NumLinks; LinkIndex++)
+					{
+						InOutLinks[LinkIndex].Location = 0.5f * (ForwardPull[LinkIndex].Location + BackwardPull[LinkIndex].Location);
+					}
+				}
+				else
+				{
+					FABRIK_ForwardReach(InOutLinks, InTargetLocation, ReachStepAlpha, bUseAngleLimit, MinRotationAngleRadians);
+					FABRIK_BackwardReach(InOutLinks, RootTargetLocation, ReachStepAlpha, bUseAngleLimit, MinRotationAngleRadians);
+				}
+				
+				Slop = FVector::Dist( InOutLinks.Last().Location, InTargetLocation) + FVector::Dist(InOutLinks[0].Location, RootTargetLocation);
+
+				// Abort if we're not getting closer and enter a deadlock.
+				if (Slop > PreviousSlop)
+				{
+					break;
+				}
+
+			} while ((Slop > ReachPrecision) && (++IterationCount < MaxIterations));
+			
+			// Make sure our root is back at our root target.
+			if (!InOutLinks[0].Location.Equals(RootTargetLocation))
+			{
+				FABRIK_BackwardReach(InOutLinks, RootTargetLocation, ReachStepAlpha, bUseAngleLimit, MinRotationAngleRadians);
+			}
+
+			// Replace end bone based on how close we got to the end target.
+			FLimbLink& EndLink = InOutLinks[NumLinks-1];
+			const FLimbLink& ParentLink = InOutLinks[NumLinks-2];
+			EndLink.Location = ParentLink.Location + (EndLink.Location - ParentLink.Location).GetUnsafeNormal() * EndLink.Length;
+		}
+	}
+	
+	bool SolveLimb(
+		TArray<FLimbLink>& InOutLinks,
+		float LimbLength,
+		const FVector& GoalLocation,
+		float Precision,
+		int32 InMaxIterations,
+		const FVector& HingeRotationAxis,
+		bool bUseAngleLimit,
+		float MinRotationAngle,
+		float ReachStepAlpha,
+		bool bAveragePull,
+		float PullDistribution
+		)
+	{
+		const FVector& EndLocation = InOutLinks.Last().Location;
+		const bool FootAtGoal = EndLocation.Equals(GoalLocation, Precision);
+		
+		if (FootAtGoal && !bUseAngleLimit)
+		{
+			return false;
+		}
+		
+		const bool bDirectSolve = ( InOutLinks.Num() == 3 );
+		if (bDirectSolve)
+		{
+			const FVector& RootLocation = InOutLinks[0].Location;
+			
+			// If we can't reach, we just go in a straight line towards the target,
+			if (FVector::DistSquared(RootLocation, GoalLocation) >= FMath::Square(LimbLength))
+			{
+				SolveStraightLimb(InOutLinks, GoalLocation);
+			}
+			else
+			{
+				SolveTwoBoneIK(InOutLinks, GoalLocation, HingeRotationAxis);
+			}
+		}
+		else
+		{
+			const float MinRotationAngleRadians = FMath::DegreesToRadians(FMath::Clamp(MinRotationAngle, 0.f, 90.f));
+			SolveFABRIK(InOutLinks, GoalLocation, Precision, InMaxIterations, bUseAngleLimit, MinRotationAngleRadians, ReachStepAlpha, PullDistribution, bAveragePull);
+		}
+		
+		return true;
+	}
+}
diff --git a/Engine/Plugins/Animation/IKRig/Source/IKRig/Private/Solvers/IKRig_BodyMover.cpp b/Engine/Plugins/Animation/IKRig/Source/IKRig/Private/Solvers/IKRig_BodyMover.cpp
deleted file mode 100644
diff --git a/Engine/Plugins/Animation/IKRig/Source/IKRig/Private/Solvers/IKRig_LimbSolver.cpp b/Engine/Plugins/Animation/IKRig/Source/IKRig/Private/Solvers/IKRig_LimbSolver.cpp
deleted file mode 100644
diff --git a/Engine/Plugins/Animation/IKRig/Source/IKRig/Private/Solvers/IKRig_PBIKSolver.cpp b/Engine/Plugins/Animation/IKRig/Source/IKRig/Private/Solvers/IKRig_PBIKSolver.cpp
deleted file mode 100644
diff --git a/Engine/Plugins/Animation/IKRig/Source/IKRig/Private/Solvers/IKRig_PoleSolver.cpp b/Engine/Plugins/Animation/IKRig/Source/IKRig/Private/Solvers/IKRig_PoleSolver.cpp
deleted file mode 100644
diff --git a/Engine/Plugins/Animation/IKRig/Source/IKRig/Private/Solvers/IKRig_SetTransform.cpp b/Engine/Plugins/Animation/IKRig/Source/IKRig/Private/Solvers/IKRig_SetTransform.cpp
deleted file mode 100644
diff --git a/Engine/Plugins/Animation/IKRig/Source/IKRig/Private/Solvers/LimbSolver.cpp b/Engine/Plugins/Animation/IKRig/Source/IKRig/Private/Solvers/LimbSolver.cpp
deleted file mode 100644
diff --git a/Engine/Plugins/Animation/IKRig/Source/IKRig/Public/ActorComponents/IKRigComponent.h b/Engine/Plugins/Animation/IKRig/Source/IKRig/Public/ActorComponents/IKRigComponent.h
--- a/Engine/Plugins/Animation/IKRig/Source/IKRig/Public/ActorComponents/IKRigComponent.h
+++ b/Engine/Plugins/Animation/IKRig/Source/IKRig/Public/ActorComponents/IKRigComponent.h
@@ -4,7 +4,7 @@
 
 #include "Components/ActorComponent.h"
 
-#include "IKRigDataTypes.h"
+#include "Rig/IKRigDataTypes.h"
 #include "IKRigInterface.h"
 
 #include "IKRigComponent.generated.h"
diff --git a/Engine/Plugins/Animation/IKRig/Source/IKRig/Public/ActorComponents/IKRigInterface.h b/Engine/Plugins/Animation/IKRig/Source/IKRig/Public/ActorComponents/IKRigInterface.h
--- a/Engine/Plugins/Animation/IKRig/Source/IKRig/Public/ActorComponents/IKRigInterface.h
+++ b/Engine/Plugins/Animation/IKRig/Source/IKRig/Public/ActorComponents/IKRigInterface.h
@@ -2,7 +2,7 @@
 
 #pragma once
 
-#include "IKRigDataTypes.h"
+#include "Rig/IKRigDataTypes.h"
 #include "UObject/Interface.h"
 
 #include "IKRigInterface.generated.h"
diff --git a/Engine/Plugins/Animation/IKRig/Source/IKRig/Public/AnimNodes/AnimNode_IKRig.h b/Engine/Plugins/Animation/IKRig/Source/IKRig/Public/AnimNodes/AnimNode_IKRig.h
--- a/Engine/Plugins/Animation/IKRig/Source/IKRig/Public/AnimNodes/AnimNode_IKRig.h
+++ b/Engine/Plugins/Animation/IKRig/Source/IKRig/Public/AnimNodes/AnimNode_IKRig.h
@@ -4,7 +4,7 @@
 
 #include "CoreMinimal.h"
 
-#include "IKRigDataTypes.h"
+#include "Rig/IKRigDataTypes.h"
 #include "Animation/AnimNode_CustomProperty.h"
 #include "Animation/InputScaleBias.h"
 #include <functional>
diff --git a/Engine/Plugins/Animation/IKRig/Source/IKRig/Public/IKRigDataTypes.h b/Engine/Plugins/Animation/IKRig/Source/IKRig/Public/IKRigDataTypes.h
deleted file mode 100644
diff --git a/Engine/Plugins/Animation/IKRig/Source/IKRig/Public/IKRigDefinition.h b/Engine/Plugins/Animation/IKRig/Source/IKRig/Public/IKRigDefinition.h
deleted file mode 100644
diff --git a/Engine/Plugins/Animation/IKRig/Source/IKRig/Public/IKRigProcessor.h b/Engine/Plugins/Animation/IKRig/Source/IKRig/Public/IKRigProcessor.h
deleted file mode 100644
diff --git a/Engine/Plugins/Animation/IKRig/Source/IKRig/Public/IKRigSkeleton.h b/Engine/Plugins/Animation/IKRig/Source/IKRig/Public/IKRigSkeleton.h
deleted file mode 100644
diff --git a/Engine/Plugins/Animation/IKRig/Source/IKRig/Public/IKRigSolver.h b/Engine/Plugins/Animation/IKRig/Source/IKRig/Public/IKRigSolver.h
deleted file mode 100644
diff --git a/Engine/Plugins/Animation/IKRig/Source/IKRig/Public/Retargeter/IKRetargeter.h b/Engine/Plugins/Animation/IKRig/Source/IKRig/Public/Retargeter/IKRetargeter.h
--- a/Engine/Plugins/Animation/IKRig/Source/IKRig/Public/Retargeter/IKRetargeter.h
+++ b/Engine/Plugins/Animation/IKRig/Source/IKRig/Public/Retargeter/IKRetargeter.h
@@ -3,7 +3,7 @@
 #pragma once
 
 #include "IKRetargetProfile.h"
-#include "IKRigDefinition.h"
+#include "Rig/IKRigDefinition.h"
 #include "IKRetargetSettings.h"
 
 #include "IKRetargeter.generated.h"
diff --git a/Engine/Plugins/Animation/IKRig/Source/IKRig/Public/Rig/IKRigDataTypes.h b/Engine/Plugins/Animation/IKRig/Source/IKRig/Public/Rig/IKRigDataTypes.h
new file mode 100644
--- /dev/null
+++ b/Engine/Plugins/Animation/IKRig/Source/IKRig/Public/Rig/IKRigDataTypes.h
@@ -0,0 +1,181 @@
+// Copyright Epic Games, Inc. All Rights Reserved.
+
+#pragma once
+
+#include "BoneContainer.h"
+#include "IKRigDefinition.h"
+
+#include "IKRigDataTypes.generated.h"
+
+class UIKRigProcessor;
+
+UENUM(BlueprintType)
+enum class EIKRigGoalSpace : uint8
+{
+	Component		UMETA(DisplayName = "Component Space"),
+	Additive		UMETA(DisplayName = "Additive"),
+    World			UMETA(DisplayName = "World Space"),
+};
+
+UENUM(BlueprintType)
+enum class EIKRigGoalTransformSource : uint8
+{
+	Manual			UMETA(DisplayName = "Manual Input"),
+    Bone			UMETA(DisplayName = "Bone"),
+    ActorComponent	UMETA(DisplayName = "Actor Component"),
+};
+
+USTRUCT(Blueprintable)
+struct IKRIG_API FIKRigGoal
+{
+	GENERATED_BODY()
+	
+	/** Name of the IK Goal. Must correspond to the name of a Goal in the target IKRig asset. */
+	UPROPERTY(EditAnywhere, BlueprintReadWrite, Category = Goal)
+	FName Name;
+
+	/** Set the source of the transform data for the Goal.
+	 * "Manual Input" uses the values provided by the blueprint node pin.
+	 * "Bone" uses the transform of the bone provided by OptionalSourceBone.
+	 * "Actor Component" uses the transform supplied by any Actor Components that implements the IIKGoalCreatorInterface*/
+	UPROPERTY(EditAnywhere, BlueprintReadWrite, Category = Goal)
+	EIKRigGoalTransformSource TransformSource = EIKRigGoalTransformSource::Manual;
+
+	/** When TransformSource is set to "Bone" mode, the Position and Rotation will be driven by this Bone's input transform.
+	 * When using a Bone as the transform source, the Position and Rotation Alpha values can still be set as desired.
+	 * But the PositionSpace and RotationSpace are no longer relevant and will not be used.*/
+	UPROPERTY(EditAnywhere, Category = Goal)
+	FBoneReference SourceBone;
+	
+	// UPROPERTY(EditAnywhere, BlueprintReadWrite, Category = Goal, meta = (CustomWidget = "BoneName"))
+	// FName OptionalSourceBone;
+	
+	/** Position of the IK goal in Component Space of target actor component. */
+	UPROPERTY(EditAnywhere, BlueprintReadWrite, Category = Position, meta=(EditCondition="Source == EIKRigGoalSource::Manual"))
+	FVector Position;
+
+	/** Rotation of the IK goal in Component Space of target actor component. */
+	UPROPERTY(EditAnywhere, BlueprintReadWrite, Category = Rotation, meta=(EditCondition="Source == EIKRigGoalSource::Manual"))
+	FRotator Rotation;
+
+	/** Range 0-1, default is 1.0. Smoothly blends the Goal position from the input pose (0.0) to the Goal position (1.0). */
+	UPROPERTY(EditAnywhere, BlueprintReadWrite, Category = Position, meta = (ClampMin = "0.0", ClampMax = "1.0", UIMin = "0.0", UIMax = "1.0"))
+	float PositionAlpha;
+
+	/** Range 0-1, default is 1.0. Smoothly blends the Goal rotation from the input pose (0.0) to the Goal rotation (1.0). */
+	UPROPERTY(EditAnywhere, BlueprintReadWrite, Category = Rotation, meta = (ClampMin = "0.0", ClampMax = "1.0", UIMin = "0.0", UIMax = "1.0"))
+	float RotationAlpha;
+
+	/** The space that the goal position is in.
+	 *"Additive" treats the goal transform as an additive offset relative to the Bone at the effector.
+	 *"Component" treats the goal transform as being in the space of the Skeletal Mesh Actor Component.
+	 *"World" treats the goal transform as being in the space of the World. */
+	UPROPERTY(EditAnywhere, BlueprintReadWrite, Category = Position)
+	EIKRigGoalSpace PositionSpace;
+	
+	/** The space that the goal rotation is in.
+	*"Additive" treats the goal transform as an additive offset relative to the Bone at the effector.
+	*"Component" treats the goal transform as being in the space of the Skeletal Mesh Actor Component.
+	*"World" treats the goal transform as being in the space of the World. */
+	UPROPERTY(EditAnywhere, BlueprintReadWrite, Category = Rotation)
+	EIKRigGoalSpace RotationSpace;
+
+	UPROPERTY(Transient)
+	FVector FinalBlendedPosition;
+
+	UPROPERTY(Transient)
+	FQuat FinalBlendedRotation;
+
+	FIKRigGoal()
+    : Name(NAME_None),
+	Position(ForceInitToZero),
+	Rotation(FRotator::ZeroRotator),
+	PositionAlpha(1.0f),
+    RotationAlpha(1.0f),
+	PositionSpace(EIKRigGoalSpace::Additive),
+	RotationSpace(EIKRigGoalSpace::Additive),
+	FinalBlendedPosition(ForceInitToZero),
+	FinalBlendedRotation(FQuat::Identity){}
+	
+	FIKRigGoal(const FName& GoalName)
+    : Name(GoalName),
+	Position(ForceInitToZero),
+	Rotation(FRotator::ZeroRotator),
+    PositionAlpha(1.0f),
+    RotationAlpha(1.0f),
+	PositionSpace(EIKRigGoalSpace::Additive),
+    RotationSpace(EIKRigGoalSpace::Additive),
+	FinalBlendedPosition(ForceInitToZero),
+    FinalBlendedRotation(FQuat::Identity){}
+
+	FIKRigGoal(
+        const FName& InName,
+        const FVector& InPosition,
+        const FQuat& InRotation,
+        const float InPositionAlpha,
+        const float InRotationAlpha,
+        EIKRigGoalSpace InPositionSpace,
+        EIKRigGoalSpace InRotationSpace)
+        : Name(InName),
+        Position(InPosition),
+        Rotation(InRotation.Rotator()),
+        PositionAlpha(InPositionAlpha),
+        RotationAlpha(InRotationAlpha),
+		PositionSpace(InPositionSpace),
+		RotationSpace(InRotationSpace),
+		FinalBlendedPosition(InPosition),
+		FinalBlendedRotation(InRotation){}
+
+	FIKRigGoal(const UIKRigEffectorGoal* InGoal)
+		: Name(InGoal->GoalName),
+		Position(InGoal->CurrentTransform.GetTranslation()),
+		Rotation(InGoal->CurrentTransform.Rotator()),
+		PositionAlpha(InGoal->PositionAlpha),
+        RotationAlpha(InGoal->RotationAlpha),
+		PositionSpace(EIKRigGoalSpace::Component),
+        RotationSpace(EIKRigGoalSpace::Component),
+        FinalBlendedPosition(ForceInitToZero),
+        FinalBlendedRotation(FQuat::Identity){}
+
+	FString ToString() const
+	{
+		return FString::Printf(TEXT("Name=%s, Pos=(%s, Alpha=%3.3f), Rot=(%s, Alpha=%3.3f)"),
+			*Name.ToString(), *FinalBlendedPosition.ToString(), PositionAlpha, *FinalBlendedRotation.ToString(), RotationAlpha);
+	}
+};
+
+inline uint32 GetTypeHash(FIKRigGoal ObjectRef) { return GetTypeHash(ObjectRef.Name); }
+
+USTRUCT()
+struct IKRIG_API FIKRigGoalContainer
+{
+	GENERATED_BODY()
+
+public:
+
+	/** Set an IK goal to go to a specific location and rotation (in the specified space) blended by alpha.
+	 * Will ADD the goal if none exist with the input name. */
+	void SetIKGoal(const FIKRigGoal& InGoal);
+
+	/** Set an IK goal to go to a specific location and rotation (in the specified space) blended by alpha.
+	* Will ADD the goal if none exist with the input name. */
+	void SetIKGoal(const UIKRigEffectorGoal* InEffectorGoal);
+
+	/** Get an IK goal with the given name. Returns nullptr if no Goal is found in the container with the name. */
+	const FIKRigGoal* FindGoalByName(const FName& GoalName) const;
+
+	/** Clear out all Goals in container. */
+	void Empty() { Goals.Empty(); };
+
+	/** Get read-only access to array of Goals. */
+	const TArray<FIKRigGoal>& GetGoalArray() const { return Goals; };
+	
+private:
+
+	FIKRigGoal* FindGoalWriteable(const FName& GoalName) const;
+	
+	/** Array of Goals. Cannot contain duplicates (name is key). */
+	TArray<FIKRigGoal> Goals;
+
+	friend UIKRigProcessor;
+};
diff --git a/Engine/Plugins/Animation/IKRig/Source/IKRig/Public/Rig/IKRigDefinition.h b/Engine/Plugins/Animation/IKRig/Source/IKRig/Public/Rig/IKRigDefinition.h
new file mode 100644
--- /dev/null
+++ b/Engine/Plugins/Animation/IKRig/Source/IKRig/Public/Rig/IKRigDefinition.h
@@ -0,0 +1,269 @@
+// Copyright Epic Games, Inc. All Rights Reserved.
+
+#pragma once
+
+#include "BoneContainer.h"
+#include "IKRigLogger.h"
+#include "UObject/Object.h"
+#include "Interfaces/Interface_PreviewMeshProvider.h"
+#include "IKRigSkeleton.h"
+#if WITH_EDITOR
+#include "SAdvancedTransformInputBox.h"
+#endif
+
+#include "IKRigDefinition.generated.h"
+
+class IPropertyHandle;
+class UIKRigSolver;
+
+#if WITH_EDITORONLY_DATA
+UENUM(BlueprintType)
+enum class EIKRigGoalPreviewMode : uint8
+{
+	Additive		UMETA(DisplayName = "Additive"),
+	Absolute		UMETA(DisplayName = "Absolute"),
+};
+
+namespace EIKRigTransformType
+{
+	enum Type : int8
+	{
+		Current,
+		Reference,
+	};
+}
+#endif
+
+UCLASS(BlueprintType)
+class IKRIG_API UIKRigEffectorGoal : public UObject
+{
+	GENERATED_BODY()
+
+public:
+
+	/** The name used to refer to this goal from outside systems.
+	 * This is the name to use when referring to this Goal from Blueprint, Anim Graph, Control Rig or IK Retargeter.*/
+	UPROPERTY(BlueprintReadOnly, VisibleAnywhere, Category = "Goal Settings")
+	FName GoalName;
+
+	/** The name of the bone that this Goal is located at.*/
+	UPROPERTY(BlueprintReadOnly, VisibleAnywhere, Category = "Goal Settings")
+	FName BoneName;
+
+	/** Range 0-1, default is 1. Blend between the input bone position (0.0) and the current goal position (1.0).*/
+	UPROPERTY(BlueprintReadWrite, EditAnywhere, Category = "Goal Settings", meta = (ClampMin = "0.0", ClampMax = "1.0", UIMin = "0.0", UIMax = "1.0"))
+	float PositionAlpha = 1.0f;
+	
+	/** Range 0-1, default is 1. Blend between the input bone rotation (0.0) and the current goal rotation (1.0).*/
+	UPROPERTY(BlueprintReadWrite, EditAnywhere, Category = "Goal Settings", meta = (ClampMin = "0.0", ClampMax = "1.0", UIMin = "0.0", UIMax = "1.0"))
+	float RotationAlpha = 1.0f;
+
+	/** The current transform of this Goal, in the Global Space of the character.*/
+	UPROPERTY(BlueprintReadWrite, EditAnywhere, Category = "Goal Settings")
+	FTransform CurrentTransform;
+
+	/** The initial transform of this Goal, as defined by the initial transform of the Goal's bone in the retarget pose.*/
+	UPROPERTY(BlueprintReadOnly, VisibleAnywhere, Category = "Goal Settings")
+	FTransform InitialTransform;
+
+	bool operator==(const UIKRigEffectorGoal& Other) const { return GoalName == Other.GoalName; }
+
+#if WITH_EDITORONLY_DATA
+
+	/** Effects how this Goal transform is previewed in the IK Rig editor.
+	* "Additive" interprets the Goal transform as being relative to the input pose. Useful for previewing animations.
+	* "Absolute" pins the Goal transform to the Gizmo in the viewport.
+	*/
+	UPROPERTY(BlueprintReadWrite, EditAnywhere, Category = "Goal Settings")
+	EIKRigGoalPreviewMode PreviewMode = EIKRigGoalPreviewMode::Additive;
+	
+	/**The size of the Goal gizmo drawing in the editor viewport.*/
+	UPROPERTY(BlueprintReadWrite, EditAnywhere, Category = "Viewport Goal Settings", meta = (ClampMin = "0.1", ClampMax = "1000.0", UIMin = "0.1", UIMax = "100.0"))
+	float SizeMultiplier = 1.0f;
+
+	/**The thickness of the Goal gizmo drawing in the editor viewport.*/
+	UPROPERTY(BlueprintReadWrite, EditAnywhere, Category = "Viewport Goal Settings",  meta = (ClampMin = "0.0", ClampMax = "10.0", UIMin = "0.0", UIMax = "5.0"))
+	float ThicknessMultiplier = 0.7f;
+
+	/** Should position data be exposed in Blueprint */
+	UPROPERTY(BlueprintReadWrite, EditAnywhere, Category = "Exposure")
+	bool bExposePosition = false;
+
+	/** Should rotation data be exposed in Blueprint */
+	UPROPERTY(BlueprintReadWrite, EditAnywhere, Category = "Exposure")
+	bool bExposeRotation = false;
+	
+	virtual void PostLoad() override
+	{
+		Super::PostLoad();
+		SetFlags(RF_Transactional);
+	}
+	
+#endif
+
+#if WITH_EDITOR
+
+	TOptional<FTransform::FReal> GetNumericValue(
+		ESlateTransformComponent::Type Component,
+		ESlateRotationRepresentation::Type Representation,
+		ESlateTransformSubComponent::Type SubComponent,
+		EIKRigTransformType::Type TransformType
+	) const;
+
+	TTuple<FTransform, FTransform> PrepareNumericValueChanged(
+		ESlateTransformComponent::Type Component,
+		ESlateRotationRepresentation::Type Representation,
+		ESlateTransformSubComponent::Type SubComponent,
+		FTransform::FReal Value,
+		EIKRigTransformType::Type TransformType
+	) const;
+
+	void SetTransform( const FTransform& InTransform, EIKRigTransformType::Type InTransformType);
+
+	void OnCopyToClipboard(ESlateTransformComponent::Type Component, EIKRigTransformType::Type TransformType) const;
+	void OnPasteFromClipboard(ESlateTransformComponent::Type Component, EIKRigTransformType::Type TransformType);
+
+	bool TransformDiffersFromDefault(ESlateTransformComponent::Type Component, TSharedPtr<IPropertyHandle> PropertyHandle) const;
+	void ResetTransformToDefault(ESlateTransformComponent::Type Component, TSharedPtr<IPropertyHandle> PropertyHandle);
+	
+#endif
+};
+
+USTRUCT(BlueprintType)
+struct IKRIG_API FBoneChain
+{
+	GENERATED_BODY()
+
+	FBoneChain() = default;
+	
+	FBoneChain(
+		FName InName, 
+		const FName InStartBone,
+		const FName InEndBone,
+		const FName InGoalName)
+		: ChainName(InName)
+		, StartBone(InStartBone)
+		, EndBone(InEndBone)
+		, IKGoalName(InGoalName)
+	{}
+
+	UPROPERTY(BlueprintReadWrite, EditAnywhere, Category = BoneChain)
+	FName ChainName;
+
+	UPROPERTY(VisibleAnywhere, Category = BoneChain)
+	FBoneReference StartBone;
+
+	UPROPERTY(VisibleAnywhere, Category = BoneChain)
+	FBoneReference EndBone;
+	
+	UPROPERTY(BlueprintReadOnly, VisibleAnywhere, Category = BoneChain)
+	FName IKGoalName;
+};
+
+USTRUCT()
+struct IKRIG_API FRetargetDefinition
+{
+	GENERATED_BODY()
+	
+private:
+
+	UPROPERTY()
+	FName RootBone;
+
+	UPROPERTY()
+	TArray<FBoneChain> BoneChains;
+
+	/** read/write access restricted to friends only */
+	FBoneChain* GetEditableBoneChainByName(FName ChainName);
+	
+	friend class UIKRigController;
+	friend class UIKRigDefinition;
+};
+
+UCLASS(Blueprintable)
+class IKRIG_API UIKRigDefinition : public UObject, public IInterface_PreviewMeshProvider
+{
+	GENERATED_BODY()
+
+public:
+
+	/** The skeletal mesh to run the IK solve on (loaded into viewport).
+	* NOTE: you can assign ANY Skeletal Mesh to apply the IK Rig to. Compatibility is determined when a new mesh is assigned
+	* by comparing it's hierarchy with the goals, solvers and bone settings required by the rig. See output log for details. */
+	UPROPERTY(AssetRegistrySearchable, EditAnywhere, Category = PreviewMesh)
+	TSoftObjectPtr<USkeletalMesh> PreviewSkeletalMesh;
+
+	/** runtime, read-only access to skeleton data */
+	const FIKRigSkeleton& GetSkeleton() const { return Skeleton; };
+	
+	/** runtime, read-only access to array of Goals, all modifications must go through UIKRigController */
+	const TArray<UIKRigEffectorGoal*>& GetGoalArray() const { return Goals; };
+
+	/** runtime, read-only access to array of Solvers, all modifications must go through UIKRigController */
+	const TArray<UIKRigSolver*>& GetSolverArray() const { return Solvers; };
+
+	/** runtime, read-only access to Retarget Definition, all modifications must go through UIKRigController */
+	const TArray<FBoneChain>& GetRetargetChains() const { return RetargetDefinition.BoneChains; };
+
+	/** runtime, read-only access to Retarget Root, all modifications must go through UIKRigController */
+	const FName& GetRetargetRoot() const { return RetargetDefinition.RootBone; };
+
+	/** runtime, read-only access to read a bone chain, all modifications must go through UIKRigController */
+	const FBoneChain* GetRetargetChainByName(FName ChainName) const;
+
+	/** UObject */
+	virtual void Serialize(FArchive& Ar) override;
+	/** END UObject */
+	
+	/** IInterface_PreviewMeshProvider interface */
+	virtual void SetPreviewMesh(USkeletalMesh* PreviewMesh, bool bMarkAsDirty = true) override;
+	virtual USkeletalMesh* GetPreviewMesh() const override;
+	/** END IInterface_PreviewMeshProvider interface */
+
+#if WITH_EDITOR
+	/* Get name of Preview Mesh property */
+	static const FName GetPreviewMeshPropertyName();
+#endif
+
+#if WITH_EDITORONLY_DATA
+	/**The size of the Bones in the editor viewport. This is set by callbacks from the viewport Character>Bones menu*/
+	UPROPERTY()
+	float BoneSize = 2.0f;
+
+	/**Draw bones in the viewport.*/
+	UPROPERTY(EditAnywhere, Category = "Viewport Goal Settings")
+	bool DrawGoals = true;
+	
+	/**The size of the Goals in the editor viewport.*/
+	UPROPERTY(EditAnywhere, Category = "Viewport Goal Settings", meta = (ClampMin = "0.01", UIMin = "0.1", UIMax = "100.0"))
+	float GoalSize = 7.0f;
+
+	/** The thickness of the Goals in the editor viewport.*/
+	UPROPERTY(EditAnywhere, Category = "Viewport Goal Settings",  meta = (ClampMin = "0.01", UIMin = "0.0", UIMax = "10.0"))
+	float GoalThickness = 0.7f;
+
+	/** The controller responsible for managing this asset's data (all editor mutation goes through this) */
+	UPROPERTY(Transient)
+	TObjectPtr<UObject> Controller;
+#endif
+
+private:
+	
+	/** hierarchy and bone-pose transforms */
+	UPROPERTY()
+	FIKRigSkeleton Skeleton;
+
+	/** goals, used as effectors by solvers that support them */
+	UPROPERTY()
+	TArray<TObjectPtr<UIKRigEffectorGoal>> Goals;
+	
+	/** polymorphic stack of solvers, executed in serial fashion where output of prior solve is input to the next */
+	UPROPERTY(instanced)
+	TArray<TObjectPtr<UIKRigSolver>> Solvers;
+
+	/** bone chains for IK Retargeter */
+	UPROPERTY()
+	FRetargetDefinition RetargetDefinition;
+
+	friend class UIKRigController;
+};
diff --git a/Engine/Plugins/Animation/IKRig/Source/IKRig/Public/Rig/IKRigProcessor.h b/Engine/Plugins/Animation/IKRig/Source/IKRig/Public/Rig/IKRigProcessor.h
new file mode 100644
--- /dev/null
+++ b/Engine/Plugins/Animation/IKRig/Source/IKRig/Public/Rig/IKRigProcessor.h
@@ -0,0 +1,135 @@
+// Copyright Epic Games, Inc. All Rights Reserved.
+
+#pragma once
+
+#include "CoreMinimal.h"
+#include "IKRigDataTypes.h"
+#include "IKRigLogger.h"
+#include "IKRigSkeleton.h"
+
+#include "IKRigProcessor.generated.h"
+
+class UIKRigDefinition;
+class UIKRigSolver;
+
+USTRUCT()
+struct FGoalBone
+{
+	GENERATED_BODY()
+	
+	FName BoneName;
+	int32 BoneIndex;
+	int32 OptSourceIndex = INDEX_NONE;
+};
+
+UCLASS()
+class IKRIG_API UIKRigProcessor : public UObject
+{
+	GENERATED_BODY()
+	
+public:
+	
+	/** the runtime for an IKRig to convert an input pose into
+	*   a solved output pose given a set of IK Rig Goals:
+	*   
+	* 1. Create a new IKRigProcessor once using MakeNewIKRigProcessor()
+	* 2. Initialize() with an IKRigDefinition asset
+	* 3. each tick, call SetIKGoal() and SetInputPoseGlobal()
+	* 4. Call Solve()
+	* 5. Copy output transforms with CopyOutputGlobalPoseToArray()
+	* 
+	*/
+
+	UIKRigProcessor();
+	
+	/** setup a new processor to run the given IKRig asset
+	 *  NOTE!! this function creates new UObjects and consequently MUST be called from the main thread!!
+	 *  @param InRigAsset - the IK Rig defining the collection of solvers to execute and all the rig settings
+	 *  @param SkeletalMesh - the skeletal mesh you want to solve the IK on
+	 *  @param ExcludedGoals - a list of goal names to exclude from this processor instance (support per-instance removal of goals)
+	 */
+	void Initialize(
+		const UIKRigDefinition* InRigAsset,
+		const USkeletalMesh* SkeletalMesh,
+		const TArray<FName>& ExcludedGoals = TArray<FName>());
+
+	//
+	// BEGIN UPDATE SEQUENCE FUNCTIONS
+	//
+	// This is the general sequence of function calls to run a typical IK solve:
+	//
+	
+	/** Set all bone transforms in global space. This is the pose the IK solve will start from */
+	void SetInputPoseGlobal(const TArray<FTransform>& InGlobalBoneTransforms);
+
+	/** Optionally can be called before Solve() to use the reference pose as start pose */
+	void SetInputPoseToRefPose();
+
+	/** Set a named IK goal to go to a specific location, rotation and space, blended by separate position/rotation alpha (0-1)*/
+	void SetIKGoal(const FIKRigGoal& Goal);
+
+	/** Set a named IK goal to go to a specific location, rotation and space, blended by separate position/rotation alpha (0-1)*/
+	void SetIKGoal(const UIKRigEffectorGoal* Goal);
+
+	/** Run entire stack of solvers.
+	 * If any Goals were supplied in World Space, a valid WorldToComponent transform must be provided.  */
+	void Solve(const FTransform& WorldToComponent = FTransform::Identity);
+
+	/** Get the results after calling Solve() */
+	void CopyOutputGlobalPoseToArray(TArray<FTransform>& OutputPoseGlobal) const;
+
+	/** Reset all internal data structures. Will require re-initialization before solving again. */
+	void Reset();
+
+	//
+	// END UPDATE SEQUENCE FUNCTIONS
+	//
+
+	/** Get access to the internal goal data (read only) */
+	const FIKRigGoalContainer& GetGoalContainer() const;
+	/** Get the bone associated with a goal */
+	const FGoalBone* GetGoalBone(const FName& GoalName) const;
+	
+	/** Get read/write access to the internal skeleton data */
+	FIKRigSkeleton& GetSkeletonWriteable();
+	/** Get read-only access to the internal skeleton data */
+	const FIKRigSkeleton& GetSkeleton() const;
+
+	/** Used to determine if the IK Rig asset is compatible with a given skeleton. */
+	static bool IsIKRigCompatibleWithSkeleton(
+		const UIKRigDefinition* InRigAsset,
+		const FIKRigInputSkeleton& InputSkeleton,
+		const FIKRigLogger* Log);
+
+	bool IsInitialized() const { return bInitialized; };
+
+	void SetNeedsInitialized();
+
+	/** logging system */
+	FIKRigLogger Log;
+
+	/** Used to propagate setting values from the source asset at runtime (settings that do not require re-initialization) */
+	void CopyAllInputsFromSourceAssetAtRuntime(const UIKRigDefinition* SourceAsset);
+	
+private:
+
+	/** Update the final pos/rot of all the goals based on their alpha values and their space settings. */
+	void ResolveFinalGoalTransforms(const FTransform& WorldToComponent);
+
+	/** the stack of solvers to run in order */
+	UPROPERTY(Transient)
+	TArray<TObjectPtr<UIKRigSolver>> Solvers;
+
+	/** the named transforms that solvers use as end effectors */
+	FIKRigGoalContainer GoalContainer;
+
+	/** map of goal names to bone names/indices */
+	TMap<FName, FGoalBone> GoalBones;
+
+	/** storage for hierarchy and bone transforms */
+	FIKRigSkeleton Skeleton;
+
+	/** solving disabled until this flag is true */
+	bool bInitialized = false;
+	bool bTriedToInitialize = false;
+};
\ No newline at end of file
diff --git a/Engine/Plugins/Animation/IKRig/Source/IKRig/Public/Rig/IKRigSkeleton.h b/Engine/Plugins/Animation/IKRig/Source/IKRig/Public/Rig/IKRigSkeleton.h
new file mode 100644
--- /dev/null
+++ b/Engine/Plugins/Animation/IKRig/Source/IKRig/Public/Rig/IKRigSkeleton.h
@@ -0,0 +1,125 @@
+// Copyright Epic Games, Inc. All Rights Reserved.
+
+#pragma once
+
+#include "CoreMinimal.h"
+
+#include "IKRigSkeleton.generated.h"
+
+struct FBoneChain;
+class USkeletalMesh;
+class UIKRigDefinition;
+class UIKRigSolver;
+
+/** Data used just to initialize an IKRigSkeleton from outside systems
+ *
+ * The input skeleton may be different than the skeleton that the IK Rig asset was created for, within some limits.
+ * 1. It must have all the Bones that the IK Rig asset referenced (must be a sub-set)
+ * 2. All the bones must have the same parents (no change in hierarchy)
+ *
+ * You can however add additional bones, change the reference pose (including proportions) and the bone indices.
+ * This allows you to run the same IK Rig asset on somewhat different skeletal meshes.
+ *
+ * To validate compatibility use UIKRigProcess::IsIKRigCompatibleWithSkeleton()
+ */
+USTRUCT()
+struct IKRIG_API FIKRigInputSkeleton
+{
+	GENERATED_BODY()
+	
+	TArray<FName> BoneNames;
+	TArray<int32> ParentIndices;
+	TArray<FTransform> LocalRefPose;
+	FName SkeletalMeshName;
+
+	FIKRigInputSkeleton() = default;
+	
+	FIKRigInputSkeleton(const USkeletalMesh* SkeletalMesh);
+
+	void Initialize(const USkeletalMesh* SkeletalMesh);
+
+	void Reset();
+};
+
+USTRUCT()
+struct IKRIG_API FIKRigSkeleton
+{
+	GENERATED_BODY()
+
+	/** Names of bones. Used to match hierarchy with runtime skeleton. */
+	UPROPERTY(VisibleAnywhere, Category = Skeleton)
+	TArray<FName> BoneNames;
+
+	/** Same length as BoneNames, stores index of parent for each bone. */
+	UPROPERTY(VisibleAnywhere, Category = Skeleton)
+	TArray<int32> ParentIndices;
+
+	/** Sparse array of bones that are to be excluded from any solvers (parented around, treated as FK children). */
+	UPROPERTY(VisibleAnywhere, Category = Skeleton)
+	TArray<FName> ExcludedBones;
+
+	/** The current GLOBAL space pose of each bone. */
+	UPROPERTY(VisibleAnywhere, Category = Skeleton)
+	TArray<FTransform> CurrentPoseGlobal;
+
+	/** The current LOCAL space pose of each bone. */
+	UPROPERTY(VisibleAnywhere, Category = Skeleton)
+	TArray<FTransform> CurrentPoseLocal;
+
+	/** The initial/reference GLOBAL space pose of each bone. */
+	UPROPERTY(VisibleAnywhere, Category = Skeleton)
+	TArray<FTransform> RefPoseGlobal;
+
+	void SetInputSkeleton(const USkeletalMesh* SkeletalMesh, const TArray<FName>& InExcludedBones);
+
+	void SetInputSkeleton(const FIKRigInputSkeleton& InputSkeleton, const TArray<FName>& InExcludedBones);
+
+	void Reset();
+	
+	int32 GetBoneIndexFromName(const FName InName) const;
+
+	const FName& GetBoneNameFromIndex(const int32 BoneIndex) const;
+	
+	int32 GetParentIndex(const int32 BoneIndex) const;
+
+	int32 GetParentIndexThatIsNotExcluded(const int32 BoneIndex) const;
+
+	int32 GetChildIndices(const int32 ParentBoneIndex, TArray<int32>& Children) const;
+	
+	int32 GetCachedEndOfBranchIndex(const int32 InBoneIndex) const;
+
+	static void ConvertLocalPoseToGlobal(
+		const TArray<int32>& InParentIndices,
+		const TArray<FTransform>& InLocalPose,
+		TArray<FTransform>& OutGlobalPose);
+	
+	void UpdateAllGlobalTransformFromLocal();
+	
+	void UpdateAllLocalTransformFromGlobal();
+	
+	void UpdateGlobalTransformFromLocal(const int32 BoneIndex);
+	
+	void UpdateLocalTransformFromGlobal(const int32 BoneIndex);
+	
+	void PropagateGlobalPoseBelowBone(const int32 BoneIndex);
+
+	bool IsBoneInDirectLineage(const FName& Child, const FName& PotentialParent) const;
+
+	bool IsBoneExcluded(const int32 BoneIndex) const;
+
+	static void NormalizeRotations(TArray<FTransform>& Transforms);
+
+	void GetChainsInList(const TArray<int32>& SelectedBones, TArray<FBoneChain>& OutChains) const;
+
+	// get indices of bones in this chain, ordered from tip to root. returns false if chain is invalid.
+	bool ValidateChainAndGetBones(const FBoneChain& Chain, TArray<int32>& OutBoneIndices) const;
+
+	// given a set of bones, returns a set of indices of bones that are mirrored across X axis (YZ plane)
+	// return false if any of the bones do not have a mirrored pair within the MaxDistanceThreshold
+	bool GetMirroredBoneIndices(const TArray<int32>& BoneIndices, TArray<int32>& OutMirroredIndices) const;
+
+private:
+
+	/** The branch bellow a specific bone. It stores the last element of the branch instead of the indices */
+	mutable TArray<int32> CachedEndOfBranchIndices;
+};
diff --git a/Engine/Plugins/Animation/IKRig/Source/IKRig/Public/Rig/Solvers/IKRigSolver.h b/Engine/Plugins/Animation/IKRig/Source/IKRig/Public/Rig/Solvers/IKRigSolver.h
new file mode 100644
--- /dev/null
+++ b/Engine/Plugins/Animation/IKRig/Source/IKRig/Public/Rig/Solvers/IKRigSolver.h
@@ -0,0 +1,124 @@
+// Copyright Epic Games, Inc. All Rights Reserved.
+
+#pragma once
+
+#include "UObject/Object.h"
+#include "Rig/IKRigDataTypes.h"
+
+#include "IKRigSolver.generated.h"
+
+class FPrimitiveDrawInterface;
+struct FIKRigGoalContainer;
+struct FIKRigSkeleton;
+
+// this is the base class for creating your own solver type that integrates into the IK Rig framework/editor.
+UCLASS(abstract, hidecategories = UObject)
+class IKRIG_API UIKRigSolver : public UObject
+{
+	GENERATED_BODY()
+	
+public:
+
+	UIKRigSolver() { SetFlags(RF_Transactional); }
+
+	//** RUNTIME */
+	/** override to setup internal data based on ref pose */
+	virtual void Initialize(const FIKRigSkeleton& IKRigSkeleton) PURE_VIRTUAL("Init");
+	/** override Solve() to evaluate new output pose (InOutGlobalTransform) */
+	virtual void Solve(FIKRigSkeleton& IKRigSkeleton, const FIKRigGoalContainer& Goals) PURE_VIRTUAL("Solve");
+	//** END RUNTIME */
+
+	//** ROOT BONE (optional, implement if your solver requires a root bone) */
+	/** if solver requires a root bone, then override this to return it. */
+	virtual FName GetRootBone() const { return NAME_None; };
+	/** override to support telling outside systems which bones this solver has setting for.
+	* NOTE: This must be overriden on solvers that use bone settings.
+	* NOTE: Only ADD to the incoming set, do not remove from it. */
+	virtual void GetBonesWithSettings(TSet<FName>& OutBonesWithSettings) const {};
+	
+	/** get if this solver is enabled */
+	bool IsEnabled() const { return bIsEnabled; };
+	/** turn solver on/off (will be skipped during execution if disabled) */
+	void SetEnabled(const bool bEnabled){ bIsEnabled = bEnabled; };
+	
+	//** SOLVER SETTINGS */
+	/** override to support RECEIVING modified solver settings from outside systems for editing/UI.
+	 * Note: you can safely cast this to your own solver type and copy any relevant settings at runtime
+	 * This is necessary because at runtime, the IKRigProcessor creates a copy of your solver class
+	 * and the copy must be notified of changes made to the class settings in the source asset.*/
+	virtual void UpdateSolverSettings(UIKRigSolver* InSettings){};
+
+	/** override to support REMOVING a goal from custom solver */
+	virtual void RemoveGoal(const FName& GoalName) PURE_VIRTUAL("RemoveGoal");
+
+#if WITH_EDITORONLY_DATA
+	
+	/** callback whenever this solver is edited */
+	DECLARE_EVENT_OneParam(UIKRigSolver, FIKRigSolverModified, UIKRigSolver*);
+	FIKRigSolverModified& OnSolverModified(){ return IKRigSolverModified; };
+	
+	/** override to give your solver a nice name to display in the UI */
+	virtual FText GetNiceName() const { checkNoEntry() return FText::GetEmpty(); };
+	/** override to provide warning to user during setup of any missing components. return false if no warnings. */
+	virtual bool GetWarningMessage(FText& OutWarningMessage) const { return false; };
+	//** END SOLVER SETTINGS */
+
+	//** GOALS */
+	/** override to support ADDING a new goal to custom solver */
+	virtual void AddGoal(const UIKRigEffectorGoal& NewGoal) PURE_VIRTUAL("AddGoal");
+	/** override to support RENAMING an existing goal */
+	virtual void RenameGoal(const FName& OldName, const FName& NewName) PURE_VIRTUAL("RenameGoal");
+	/** override to support CHANGING BONE for an existing goal */
+	virtual void SetGoalBone(const FName& GoalName, const FName& NewBoneName) PURE_VIRTUAL("SetGoalBone");
+	/** override to support QUERY for a connected goal */
+	virtual bool IsGoalConnected(const FName& GoalName) const {return false;};
+	/** override to support supplying goals settings specific to this solver to outside systems for editing/UI */
+	virtual UObject* GetGoalSettings(const FName& GoalName) const {return nullptr;};
+	//** END GOALS */
+
+	//** ROOT BONE (optional, implement if your solver requires a root bone) */
+	/** override to support SETTING ROOT BONE for the solver */
+	virtual void SetRootBone(const FName& RootBoneName){};
+	virtual bool RequiresRootBone() const { return false; };
+	//** END ROOT BONE */
+
+	//** ROOT BONE (optional, implement if your solver requires a end bone) */
+	/** override to support SETTING END BONE for the solver */
+	virtual void SetEndBone(const FName& EndBoneName){};
+	/** override to support GETTING END BONE of the solver */
+	virtual FName GetEndBone() const { return NAME_None; };
+	/** override and return true if this solver requires an end bone to be specified */
+	virtual bool RequiresEndBone() const { return false; };
+	//** END ROOT BONE */
+
+	//** BONE SETTINGS (optional, implement if your solver supports per-bone settings) */
+	/** override to support ADDING PER-BONE settings for this solver */
+	virtual void AddBoneSetting(const FName& BoneName){};
+	/** override to support ADDING PER-BONE settings for this solver */
+	virtual void RemoveBoneSetting(const FName& BoneName){};
+	/** override to support supplying per-bone settings to outside systems for editing/UI
+	 ** NOTE: This must be overriden on solvers that use bone settings.*/
+	virtual UObject* GetBoneSetting(const FName& BoneName) const { ensure(!UsesBoneSettings()); return nullptr; };
+	
+	/** override to tell systems if this solver supports per-bone settings */
+	virtual bool UsesBoneSettings() const { return false;};
+	/** todo override to draw custom per-bone settings in the editor viewport */
+	virtual void DrawBoneSettings(const FName& BoneName, const FIKRigSkeleton& IKRigSkeleton, FPrimitiveDrawInterface* PDI) const {};
+	/** return true if the supplied Bone is affected by this solver - this provides UI feedback for user */
+	virtual bool IsBoneAffectedBySolver(const FName& BoneName, const FIKRigSkeleton& IKRigSkeleton) const { return false; };
+	//** END ROOT BONE */
+
+	/** UObject interface */
+	virtual void PostLoad() override;
+	virtual void PostTransacted(const FTransactionObjectEvent& TransactionEvent) override;
+	/** END UObject interface */
+
+private:
+
+	/** Register callbacks to update IK Rig when a solver is modified */
+	FIKRigSolverModified IKRigSolverModified;
+#endif
+	
+	UPROPERTY()
+	bool bIsEnabled = true;
+};
\ No newline at end of file
diff --git a/Engine/Plugins/Animation/IKRig/Source/IKRig/Public/Rig/Solvers/IKRig_BodyMover.h b/Engine/Plugins/Animation/IKRig/Source/IKRig/Public/Rig/Solvers/IKRig_BodyMover.h
new file mode 100644
--- /dev/null
+++ b/Engine/Plugins/Animation/IKRig/Source/IKRig/Public/Rig/Solvers/IKRig_BodyMover.h
@@ -0,0 +1,117 @@
+﻿// Copyright Epic Games, Inc. All Rights Reserved.
+
+#pragma once
+
+#include "Rig/IKRigDefinition.h"
+#include "IKRigSolver.h"
+
+#include "IKRig_BodyMover.generated.h"
+
+UCLASS(BlueprintType)
+class IKRIG_API UIKRig_BodyMoverEffector : public UObject
+{
+	GENERATED_BODY()
+
+public:
+	UIKRig_BodyMoverEffector() { SetFlags(RF_Transactional); }
+	
+	UPROPERTY(BlueprintReadOnly, VisibleAnywhere, Category = "Body Mover Effector")
+	FName GoalName;
+
+	UPROPERTY(BlueprintReadOnly, VisibleAnywhere, Category = "Body Mover Effector")
+	FName BoneName;
+
+	/** Scale the influence this effector has on the body. Range is 0-10. Default is 1.0. */
+	UPROPERTY(BlueprintReadWrite, EditAnywhere, Category = "Body Mover Effector", meta = (ClampMin = "0", ClampMax = "10", UIMin = "0.0", UIMax = "10.0"))
+	float InfluenceMultiplier = 1.0f;
+};
+
+UCLASS(BlueprintType, EditInlineNew)
+class IKRIG_API UIKRig_BodyMover : public UIKRigSolver
+{
+	GENERATED_BODY()
+
+public:
+
+	/** The target bone to move with the effectors. */
+	UPROPERTY(VisibleAnywhere, Category = "Body Mover Settings")
+	FName RootBone;
+
+	/** Blend the translational effect of this solver on/off. Range is 0-1. Default is 1.0. */
+	UPROPERTY(BlueprintReadWrite, EditAnywhere, Category = "Body Mover Settings", meta = (UIMin = "0.0", UIMax = "1.0"))
+	float PositionAlpha = 1.0f;
+
+	/** Multiply the POSITIVE X translation. Range is 0-1. Default is 1.0. */
+	UPROPERTY(BlueprintReadWrite, EditAnywhere, Category = "Body Mover Settings", meta = (UIMin = "0.0", UIMax = "1.0"))
+	float PositionPositiveX = 1.0f;
+
+	/** Multiply the NEGATIVE X translation. Range is 0-1. Default is 1.0. */
+	UPROPERTY(BlueprintReadWrite, EditAnywhere, Category = "Body Mover Settings", meta = (UIMin = "0.0", UIMax = "1.0"))
+	float PositionNegativeX = 1.0f;
+
+	/** Multiply the POSITIVE Y translation. Range is 0-1. Default is 1.0. */
+	UPROPERTY(BlueprintReadWrite, EditAnywhere, Category = "Body Mover Settings", meta = (UIMin = "0.0", UIMax = "1.0"))
+	float PositionPositiveY = 1.0f;
+
+	/** Multiply the NEGATIVE Y translation. Range is 0-1. Default is 1.0. */
+	UPROPERTY(BlueprintReadWrite, EditAnywhere, Category = "Body Mover Settings", meta = (UIMin = "0.0", UIMax = "1.0"))
+	float PositionNegativeY = 1.0f;
+
+	/** Multiply the POSITIVE Z translation. Range is 0-1. Default is 1.0. */
+	UPROPERTY(BlueprintReadWrite, EditAnywhere, Category = "Body Mover Settings", meta = (UIMin = "0.0", UIMax = "1.0"))
+	float PositionPositiveZ = 1.0f;
+
+	/** Multiply the NEGATIVE Z translation. Range is 0-1. Default is 1.0. */
+	UPROPERTY(BlueprintReadWrite, EditAnywhere, Category = "Body Mover Settings", meta = (UIMin = "0.0", UIMax = "1.0"))
+	float PositionNegativeZ = 1.0f;
+
+	/** Blend the total rotational effect on/off. Range is 0-1. Default is 1.0. */
+	UPROPERTY(BlueprintReadWrite, EditAnywhere, Category = "Body Mover Settings", meta = (UIMin = "0.0", UIMax = "1.0"))
+	float RotationAlpha = 1.0f;
+
+	/** Blend the X-axis rotational effect on/off. Range is 0-1. Default is 1.0. */
+	UPROPERTY(BlueprintReadWrite, EditAnywhere, Category = "Body Mover Settings", meta = (UIMin = "0.0", UIMax = "1.0"))
+	float RotateXAlpha = 1.0f;
+
+	/** Blend the Y-axis rotational effect on/off. Range is 0-1. Default is 1.0. */
+	UPROPERTY(BlueprintReadWrite, EditAnywhere, Category = "Body Mover Settings", meta = (UIMin = "0.0", UIMax = "1.0"))
+	float RotateYAlpha = 1.0f;
+
+	/** Blend the Z-axis rotational effect on/off. Range is 0-1. Default is 1.0. */
+	UPROPERTY(BlueprintReadWrite, EditAnywhere, Category = "Body Mover Settings", meta = (UIMin = "0.0", UIMax = "1.0"))
+	float RotateZAlpha = 1.0f;
+	
+	UPROPERTY()
+	TArray<TObjectPtr<UIKRig_BodyMoverEffector>> Effectors;
+
+	/** UIKRigSolver interface */
+	virtual void Initialize(const FIKRigSkeleton& IKRigSkeleton) override;
+	virtual void Solve(FIKRigSkeleton& IKRigSkeleton, const FIKRigGoalContainer& Goals) override;
+
+	virtual FName GetRootBone() const override { return RootBone; };
+	virtual void RemoveGoal(const FName& GoalName) override;
+	virtual void UpdateSolverSettings(UIKRigSolver* InSettings) override;
+
+#if WITH_EDITOR
+	virtual FText GetNiceName() const override;
+	virtual bool GetWarningMessage(FText& OutWarningMessage) const override;
+	virtual bool IsBoneAffectedBySolver(const FName& BoneName, const FIKRigSkeleton& IKRigSkeleton) const override;
+	// goals
+	virtual void AddGoal(const UIKRigEffectorGoal& NewGoal) override;
+	virtual void RenameGoal(const FName& OldName, const FName& NewName) override;
+	virtual void SetGoalBone(const FName& GoalName, const FName& NewBoneName) override;
+	virtual bool IsGoalConnected(const FName& GoalName) const override;
+	virtual UObject* GetGoalSettings(const FName& GoalName) const override;
+	// root bone can be set on this solver
+	virtual bool RequiresRootBone() const override { return true; };
+	virtual void SetRootBone(const FName& RootBoneName) override;
+	/** END UIKRigSolver interface */
+#endif
+
+private:
+
+	int32 GetIndexOfGoal(const FName& OldName) const;
+
+	int32 BodyBoneIndex;
+};
+
diff --git a/Engine/Plugins/Animation/IKRig/Source/IKRig/Public/Rig/Solvers/IKRig_FBIKSolver.h b/Engine/Plugins/Animation/IKRig/Source/IKRig/Public/Rig/Solvers/IKRig_FBIKSolver.h
new file mode 100644
--- /dev/null
+++ b/Engine/Plugins/Animation/IKRig/Source/IKRig/Public/Rig/Solvers/IKRig_FBIKSolver.h
@@ -0,0 +1,281 @@
+// Copyright Epic Games, Inc. All Rights Reserved.
+
+#pragma once
+
+#include "Rig/IKRigDefinition.h"
+#include "IKRigSolver.h"
+#include "Core/PBIKSolver.h"
+#include "PBIK_Shared.h"
+
+#include "IKRig_FBIKSolver.generated.h"
+
+UCLASS(BlueprintType)
+class IKRIG_API UIKRig_FBIKEffector : public UObject
+{
+	GENERATED_BODY()
+
+public:
+	UIKRig_FBIKEffector() { SetFlags(RF_Transactional); }
+	
+	/** The Goal that is driving this effector's transform. */
+	UPROPERTY(BlueprintReadOnly, VisibleAnywhere, Category = "Full Body IK Effector")
+	FName GoalName;
+	
+	/** The bone that this effector will pull on. */
+	UPROPERTY(BlueprintReadOnly, VisibleAnywhere, Category = "Full Body IK Effector")
+	FName BoneName;
+
+	/** Range 0-1 (default is 1.0). The strength of the effector when pulling the bone towards it's target location.
+	* At 0.0, the effector does not pull at all, but the bones between the effector and the root will still slightly resist motion from other effectors.
+	* This can thus act as a "stabilizer" for parts of the body that you do not want to behave in a pure FK fashion.
+	*/
+	UPROPERTY(BlueprintReadWrite, EditAnywhere, Category = "Full Body IK Effector", meta = (ClampMin = "0", ClampMax = "1", UIMin = "0.0", UIMax = "1.0"))
+	float StrengthAlpha = 1.0f;
+
+	/** Range 0-1 (default is 1.0). When enabled (greater than 0.0), the solver internally partitions the skeleton into 'chains' which extend from the effector to the nearest fork in the skeleton.
+	*These chains are pre-rotated and translated, as a whole, towards the effector targets.
+	*This can improve the results for sparse bone chains, and significantly improve convergence on dense bone chains.
+	*But it may cause undesirable results in highly constrained bone chains (like robot arms).
+	*/
+	UPROPERTY(BlueprintReadWrite, EditAnywhere, Category = "Full Body IK Effector", meta = (ClampMin = "0", ClampMax = "1", UIMin = "0.0", UIMax = "1.0"))
+	float PullChainAlpha = 1.0f;
+
+	/** Range 0-1 (default is 1.0).
+	*Blends the effector bone rotation between the rotation of the effector transform (1.0) and the rotation of the input bone (0.0).*/
+	UPROPERTY(BlueprintReadWrite, EditAnywhere, Category = "Full Body IK Effector", meta = (ClampMin = "0", ClampMax = "1", UIMin = "0.0", UIMax = "1.0"))
+	float PinRotation = 1.0f;
+
+	UPROPERTY(Transient)
+	int32 IndexInSolver = -1;
+
+	void CopySettings(const UIKRig_FBIKEffector* Other)
+	{
+		StrengthAlpha = Other->StrengthAlpha;
+		PullChainAlpha = Other->PullChainAlpha;
+		PinRotation = Other->PinRotation;
+	}
+};
+
+UCLASS(BlueprintType)
+class IKRIG_API UIKRig_FBIKBoneSettings : public UObject
+{
+	GENERATED_BODY()
+
+public:
+	
+	UIKRig_FBIKBoneSettings()
+		: Bone(NAME_None), 
+		X(EPBIKLimitType::Free),
+		Y(EPBIKLimitType::Free),
+		Z(EPBIKLimitType::Free),
+		PreferredAngles(FVector::ZeroVector)
+	{
+		SetFlags(RF_Transactional);
+	}
+
+	/** The bone these settings are applied to. */
+	UPROPERTY(BlueprintReadOnly, VisibleAnywhere, Category = Bone, meta = (Constant, CustomWidget = "BoneName"))
+	FName Bone;
+	
+	/** Range is 0 to 1 (Default is 0). At higher values, the bone will resist rotating (forcing other bones to compensate). */
+	UPROPERTY(BlueprintReadWrite, EditAnywhere, Category = Stiffness, meta = (ClampMin = "0", ClampMax = "1", UIMin = "0.0", UIMax = "1.0"))
+	float RotationStiffness = 0.0f;
+
+	/** Range is 0 to 1 (Default is 0). At higher values, the bone will resist translational motion (forcing other bones to compensate). */
+	UPROPERTY(BlueprintReadWrite, EditAnywhere, Category = Stiffness, meta = (ClampMin = "0", ClampMax = "1", UIMin = "0.0", UIMax = "1.0"))
+	float PositionStiffness = 0.0f;
+
+	/** Limit the rotation angle of the bone on the X axis.
+	 *Free: can rotate freely in this axis.
+	 *Limited: rotation is clamped between the min/max angles relative to the Skeletal Mesh reference pose.
+	 *Locked: no rotation is allowed in this axis (will remain at reference pose angle). */
+	UPROPERTY(BlueprintReadWrite, EditAnywhere, Category = Limits)
+	EPBIKLimitType X;
+	/**Range is -180 to 0 (Default is 0). Degrees of rotation in the negative X direction to allow when joint is in "Limited" mode. */
+	UPROPERTY(BlueprintReadWrite, EditAnywhere, Category = Limits, meta = (ClampMin = "-180", ClampMax = "0", UIMin = "-180.0", UIMax = "0.0"))
+	float MinX = 0.0f;
+	/**Range is 0 to 180 (Default is 0). Degrees of rotation in the positive X direction to allow when joint is in "Limited" mode. */
+	UPROPERTY(BlueprintReadWrite, EditAnywhere, Category = Limits, meta = (ClampMin = "0", ClampMax = "180", UIMin = "0.0", UIMax = "180.0"))
+	float MaxX = 0.0f;
+
+	/** Limit the rotation angle of the bone on the Y axis.
+	*Free: can rotate freely in this axis.
+	*Limited: rotation is clamped between the min/max angles relative to the Skeletal Mesh reference pose.
+	*Locked: no rotation is allowed in this axis (will remain at input pose angle). */
+	UPROPERTY(BlueprintReadWrite, EditAnywhere, Category = Limits)
+	EPBIKLimitType Y;
+	/**Range is -180 to 0 (Default is 0). Degrees of rotation in the negative Y direction to allow when joint is in "Limited" mode. */
+	UPROPERTY(BlueprintReadWrite, EditAnywhere, Category = Limits, meta = (ClampMin = "-180", ClampMax = "0", UIMin = "-180.0", UIMax = "0.0"))
+	float MinY = 0.0f;
+	/**Range is 0 to 180 (Default is 0). Degrees of rotation in the positive Y direction to allow when joint is in "Limited" mode. */
+	UPROPERTY(BlueprintReadWrite, EditAnywhere, Category = Limits, meta = (ClampMin = "0", ClampMax = "180", UIMin = "0.0", UIMax = "180.0"))
+	float MaxY = 0.0f;
+
+	/** Limit the rotation angle of the bone on the Z axis.
+	*Free: can rotate freely in this axis.
+	*Limited: rotation is clamped between the min/max angles relative to the Skeletal Mesh reference pose.
+	*Locked: no rotation is allowed in this axis (will remain at input pose angle). */
+	UPROPERTY(BlueprintReadWrite, EditAnywhere, Category = Limits)
+	EPBIKLimitType Z;
+	/**Range is -180 to 0 (Default is 0). Degrees of rotation in the negative Z direction to allow when joint is in "Limited" mode. */
+	UPROPERTY(BlueprintReadWrite, EditAnywhere, Category = Limits, meta = (ClampMin = "-180", ClampMax = "0", UIMin = "-180.0", UIMax = "0.0"))
+	float MinZ = 0.0f;
+	/**Range is 0 to 180 (Default is 0). Degrees of rotation in the positive Z direction to allow when joint is in "Limited" mode. */
+	UPROPERTY(BlueprintReadWrite, EditAnywhere, Category = Limits, meta = (ClampMin = "0", ClampMax = "180", UIMin = "0.0", UIMax = "180.0"))
+	float MaxZ = 0.0f;
+
+	/**When true, this bone will "prefer" to rotate in the direction specified by the Preferred Angles when the chain it belongs to is compressed.
+	 * Preferred Angles should be the first method used to fix bones that bend in the wrong direction (rather than limits).
+	 * The resulting angles can be visualized on their own by temporarily setting the Solver iterations to 0 and moving the effectors.*/
+	UPROPERTY(BlueprintReadWrite, EditAnywhere, Category = PreferredAngles)
+	bool bUsePreferredAngles = false;
+	/**The local Euler angles (in degrees) used to rotate this bone when the chain it belongs to is squashed.
+	 * This happens by moving the effector at the tip of the chain towards the root of the chain.
+	 * This can be used to coerce knees and elbows to bend in the anatomically "correct" direction without resorting to limits (which may require more iterations to converge).*/
+	UPROPERTY(BlueprintReadWrite, EditAnywhere, Category = PreferredAngles)
+	FVector PreferredAngles;
+
+	void CopyToCoreStruct(PBIK::FBoneSettings& Settings) const
+	{
+		Settings.RotationStiffness = RotationStiffness;
+		Settings.PositionStiffness = PositionStiffness;
+		Settings.X = static_cast<PBIK::ELimitType>(X);
+		Settings.MinX = MinX;
+		Settings.MaxX = MaxX;
+		Settings.Y = static_cast<PBIK::ELimitType>(Y);
+		Settings.MinY = MinY;
+		Settings.MaxY = MaxY;
+		Settings.Z = static_cast<PBIK::ELimitType>(Z);
+		Settings.MinZ = MinZ;
+		Settings.MaxZ = MaxZ;
+		Settings.bUsePreferredAngles = bUsePreferredAngles;
+		Settings.PreferredAngles.Pitch = PreferredAngles.Y;
+		Settings.PreferredAngles.Yaw = PreferredAngles.Z;
+		Settings.PreferredAngles.Roll = PreferredAngles.X;
+	}
+
+	void CopySettings(const UIKRig_FBIKBoneSettings* Other)
+	{	
+		RotationStiffness = Other->RotationStiffness;
+		PositionStiffness = Other->PositionStiffness;
+		X = Other->X;
+		MinX = Other->MinX;
+		MaxX = Other->MaxX;
+		Y = Other->Y;
+		MinY = Other->MinY;
+		MaxY = Other->MaxY;
+		Z = Other->Z;
+		MinZ = Other->MinZ;
+		MaxZ = Other->MaxZ;
+		bUsePreferredAngles = Other->bUsePreferredAngles;
+		PreferredAngles = Other->PreferredAngles;
+	}
+};
+
+UCLASS(BlueprintType, EditInlineNew, config = Engine, hidecategories = UObject)
+class IKRIG_API UIKRigFBIKSolver : public UIKRigSolver
+{
+	GENERATED_BODY()
+
+public:
+
+	/** All bones above this bone in the hierarchy will be completely ignored by the solver. Typically this is set to
+	 * the top-most skinned bone in the Skeletal Mesh (ie Pelvis, Hips etc), NOT the actual root of the entire skeleton.
+	 *
+	 * If you want to use the solver on a single chain of bones, and NOT translate the chain, ensure that "PinRoot" is
+	 * checked on to disable the root from translating to reach the effector goals.*/
+	UPROPERTY(BlueprintReadWrite, VisibleAnywhere, Category = "Full Body IK Settings")
+	FName RootBone;
+
+	/** High iteration counts can help solve complex joint configurations with competing constraints, but will increase runtime cost. Default is 20. */
+	UPROPERTY(BlueprintReadWrite, EditAnywhere, Category = SolverSettings, meta = (ClampMin = "0", ClampMax = "1000", UIMin = "0.0", UIMax = "200.0"))
+	int32 Iterations = 20;
+
+	/** A global mass multiplier; higher values will make the joints more stiff, but require more iterations. Typical range is 0.0 to 10.0. */
+	UPROPERTY(BlueprintReadWrite, EditAnywhere, Category = SolverSettings, meta = (ClampMin = "0", UIMin = "0.0", UIMax = "10.0"))
+	float MassMultiplier = 1.0f;
+
+	/** If true, joints will translate to reach the effectors; causing bones to lengthen if necessary. Good for cartoon effects. Default is false. */
+	UPROPERTY(BlueprintReadWrite, EditAnywhere, Category = SolverSettings)
+	bool bAllowStretch = false;
+
+	/** (Default is PrePull) Set the behavior of the solver root.
+	*Pre Pull: translates and rotates the root (and all children) by the average motion of the stretched effectors to help achieve faster convergence when reaching far.
+	*Pin to Input: locks the translation and rotation of the root bone to the input pose. Overrides any bone settings applied to the root. Good for partial-body solves.
+	*Free: treats the root bone like any other and allows it to move freely or according to any bone settings applied to it. */
+	UPROPERTY(BlueprintReadWrite, EditAnywhere, Category = RootBehavior)
+	EPBIKRootBehavior RootBehavior = EPBIKRootBehavior::PrePull;
+
+	/** Settings only applicable when Root Behavior is set to "PrePull". Use these values to adjust the gross movement and orientation of the entire skeleton. */
+	UPROPERTY(BlueprintReadWrite, EditAnywhere, Category = RootBehavior)
+	FRootPrePullSettings PrePullRootSettings;
+
+	/** A global multiplier for all Pull Chain Alpha values on all effectors. Range is 0.0 to 1.0. Default is 1.0. */
+	UPROPERTY(BlueprintReadWrite, EditAnywhere, Category = AdvancedSettings, meta = (ClampMin = "0", ClampMax = "1", UIMin = "0.0", UIMax = "1.0"))
+	float PullChainAlpha = 1.0f;
+
+	/** Maximum angle that a joint can be rotated per constraint iteration. Lower this value if the solve is diverging. Range is 0.0 to 180.0. Default is 30. */
+	UPROPERTY(BlueprintReadWrite, EditAnywhere, Category = AdvancedSettings, meta = (ClampMin = "0", ClampMax = "45", UIMin = "0.0", UIMax = "180.0"))
+	float MaxAngle = 30.f;
+
+	/** Pushes constraints beyond their normal amount to speed up convergence. Increasing this may speed up convergence, but at the cost of stability. Range is 1.0 - 2.0. Default is 1.3. */
+	UPROPERTY(BlueprintReadWrite, EditAnywhere, Category = AdvancedSettings, meta = (ClampMin = "1", UIMin = "1.0", UIMax = "10.0"))
+	float OverRelaxation = 1.3f;
+	
+	/** When true, the solver is reset each tick to start from the current input pose. Default is true.
+	 * If false, incoming animated poses are ignored and the solver starts from the results of the previous solve.
+	 * In very limited circumstances, it can be beneficial to use the pose from the previous frame.*/
+	UPROPERTY(BlueprintReadWrite, EditAnywhere, Category = AdvancedSettings)
+	bool bStartSolveFromInputPose = true;
+	
+	UPROPERTY()
+	TArray<TObjectPtr<UIKRig_FBIKEffector>> Effectors;
+
+	UPROPERTY()
+	TArray<TObjectPtr<UIKRig_FBIKBoneSettings>> BoneSettings;
+
+	/** UIKRigSolver interface */
+	// runtime
+	virtual void Initialize(const FIKRigSkeleton& IKRigSkeleton) override;
+	virtual void Solve(FIKRigSkeleton& IKRigSkeleton, const FIKRigGoalContainer& Goals) override;
+
+	virtual FName GetRootBone() const override { return RootBone; };
+	virtual void GetBonesWithSettings(TSet<FName>& OutBonesWithSettings) const override;
+
+	virtual void UpdateSolverSettings(UIKRigSolver* InSettings) override;
+	virtual void RemoveGoal(const FName& GoalName) override;
+	
+#if WITH_EDITOR
+	virtual FText GetNiceName() const override;
+	virtual bool GetWarningMessage(FText& OutWarningMessage) const override;
+	// goals
+	virtual void AddGoal(const UIKRigEffectorGoal& NewGoal) override;
+	virtual void RenameGoal(const FName& OldName, const FName& NewName) override;
+	virtual void SetGoalBone(const FName& GoalName, const FName& NewBoneName) override;
+	virtual bool IsGoalConnected(const FName& GoalName) const override;
+	virtual UObject* GetGoalSettings(const FName& GoalName) const override;
+	// bone settings
+	virtual void AddBoneSetting(const FName& BoneName) override;
+	virtual void RemoveBoneSetting(const FName& BoneName) override;
+	virtual UObject* GetBoneSetting(const FName& BoneName) const override;
+	virtual bool UsesBoneSettings() const override{ return true;};
+	virtual void DrawBoneSettings(const FName& BoneName, const FIKRigSkeleton& IKRigSkeleton, FPrimitiveDrawInterface* PDI) const override;
+	// root bone can be set on this solver
+	virtual bool RequiresRootBone() const override { return true; };
+	virtual void SetRootBone(const FName& RootBoneName) override;
+	virtual bool IsBoneAffectedBySolver(const FName& BoneName, const FIKRigSkeleton& IKRigSkeleton) const override;
+	/** END UIKRigSolver interface */
+
+	//** UObject */
+	virtual void PostLoad() override;
+	//** END UObject */
+#endif
+
+private:
+
+	FPBIKSolver Solver;
+
+	int32 GetIndexOfGoal(const FName& GoalName) const;
+};
+
+
diff --git a/Engine/Plugins/Animation/IKRig/Source/IKRig/Public/Rig/Solvers/IKRig_LimbSolver.h b/Engine/Plugins/Animation/IKRig/Source/IKRig/Public/Rig/Solvers/IKRig_LimbSolver.h
new file mode 100644
--- /dev/null
+++ b/Engine/Plugins/Animation/IKRig/Source/IKRig/Public/Rig/Solvers/IKRig_LimbSolver.h
@@ -0,0 +1,120 @@
+﻿// Copyright Epic Games, Inc. All Rights Reserved.
+
+#pragma once
+
+#include "Rig/IKRigDefinition.h"
+#include "IKRigSolver.h"
+
+#include "Rig/Solvers/LimbSolver.h"
+
+#include "IKRig_LimbSolver.generated.h"
+
+UCLASS(BlueprintType)
+class IKRIG_API UIKRig_LimbEffector : public UObject
+{
+	GENERATED_BODY()
+
+public:
+	UIKRig_LimbEffector() { SetFlags(RF_Transactional); }
+
+	UPROPERTY(BlueprintReadOnly, VisibleAnywhere, Category = "Limb IK Effector")
+	FName GoalName = NAME_None;
+
+	UPROPERTY(BlueprintReadOnly, VisibleAnywhere, Category = "Limb IK Effector")
+	FName BoneName = NAME_None;
+};
+
+UCLASS(BlueprintType, EditInlineNew)
+class IKRIG_API UIKRig_LimbSolver : public UIKRigSolver
+{
+	GENERATED_BODY()
+
+public:
+
+	UIKRig_LimbSolver();
+	
+	UPROPERTY(BlueprintReadOnly, VisibleAnywhere, DisplayName = "Root Bone", Category = "Limb IK Settings")
+	FName RootName = NAME_None;
+
+	/** Precision (distance to the target) */
+	UPROPERTY(BlueprintReadWrite, EditAnywhere, Category = "Limb IK Settings", meta = (ClampMin = "0.0"))
+	float ReachPrecision = 0.01f;
+	
+	// TWO BONES SETTINGS
+	
+	/** Hinge Bones Rotation Axis. This is essentially the plane normal for (hip - knee - foot). */
+	UPROPERTY(BlueprintReadWrite, EditAnywhere, Category = "Limb IK Settings|Two Bones")
+	TEnumAsByte<EAxis::Type> HingeRotationAxis = EAxis::None;
+
+	// FABRIK SETTINGS
+
+	/** Number of Max Iterations to reach the target */
+	UPROPERTY(BlueprintReadWrite, EditAnywhere, Category = "Limb IK Settings|FABRIK", meta = (ClampMin = "0", UIMin = "0", UIMax = "100"))
+	int32 MaxIterations = 12;
+
+	/** Enable/Disable rotational limits */
+	UPROPERTY(BlueprintReadWrite, EditAnywhere, Category = "Limb IK Settings|FABRIK|Limits")
+	bool bEnableLimit = false;
+
+	/** Only used if bEnableRotationLimit is enabled. Prevents the leg from folding onto itself,
+	* and forces at least this angle between Parent and Child bone. */
+	UPROPERTY(BlueprintReadWrite, EditAnywhere, Category = "Limb IK Settings|FABRIK|Limits", meta = (EditCondition="bEnableLimit", ClampMin = "0.0", ClampMax = "90.0", UIMin = "0.0", UIMax = "90.0"))
+	float MinRotationAngle = 15.f;
+
+	/** Pull averaging only has a visual impact when we have more than 2 bones (3 links). */
+	UPROPERTY(BlueprintReadWrite, EditAnywhere, Category = "Limb IK Settings|FABRIK|Pull Averaging")
+	bool bAveragePull = true;
+
+	/** Re-position limb to distribute pull: 0 = foot, 0.5 = balanced, 1.f = hip */
+	UPROPERTY(BlueprintReadWrite, EditAnywhere, Category = "Limb IK Settings|FABRIK|Pull Averaging", meta = (ClampMin = "0.0", ClampMax = "1.0",UIMin = "0.0", UIMax = "1.0"))
+	float PullDistribution = 0.5f;
+	
+	/** Move end effector towards target. If we are compressing the chain, limit displacement.*/
+	UPROPERTY(BlueprintReadWrite, EditAnywhere, Category = "Limb IK Settings|FABRIK|Pull Averaging", meta = (ClampMin = "0.0", ClampMax = "1.0",UIMin = "0.0", UIMax = "1.0"))
+	float ReachStepAlpha = 0.7f;
+
+	// TWIST SETTINGS
+	
+	/** Enable Knee Twist correction, by comparing Foot FK with Foot IK orientation. */
+	UPROPERTY(BlueprintReadWrite, EditAnywhere, Category = "Limb IK Settings|Twist")
+	bool bEnableTwistCorrection = false;
+	
+	/** Forward Axis for Foot bone. */
+	UPROPERTY(BlueprintReadWrite, EditAnywhere, Category = "Limb IK Settings|Twist", meta = (EditCondition="bEnableTwistCorrection"))
+	TEnumAsByte<EAxis::Type> EndBoneForwardAxis = EAxis::Y;
+	
+	/** UIKRigSolver interface */
+	virtual void Initialize(const FIKRigSkeleton& IKRigSkeleton) override;
+	virtual void Solve(FIKRigSkeleton& IKRigSkeleton, const FIKRigGoalContainer& Goals) override;
+
+	virtual void UpdateSolverSettings(UIKRigSolver* InSettings) override;
+	virtual void RemoveGoal(const FName& GoalName) override;
+
+	virtual FName GetRootBone() const override { return RootName; };
+	
+#if WITH_EDITOR
+	virtual FText GetNiceName() const override;
+	virtual bool GetWarningMessage(FText& OutWarningMessage) const override;
+	virtual bool IsBoneAffectedBySolver(const FName& BoneName, const FIKRigSkeleton& IKRigSkeleton) const override;
+	// goals
+	virtual void AddGoal(const UIKRigEffectorGoal& NewGoal) override;
+	virtual void RenameGoal(const FName& OldName, const FName& NewName) override;
+	virtual void SetGoalBone(const FName& GoalName, const FName& NewBoneName) override;
+	virtual bool IsGoalConnected(const FName& GoalName) const override;
+	virtual UObject* GetGoalSettings(const FName& GoalName) const override;
+	// root bone can be set on this solver
+	virtual bool RequiresRootBone() const override { return true; };
+	virtual void SetRootBone(const FName& RootBoneName) override;
+#endif
+	/** END UIKRigSolver interface */
+	
+private:
+
+	static void GatherChildren(const int32 BoneIndex, const FIKRigSkeleton& InSkeleton, TArray<int32>& OutChildren);
+	
+	UPROPERTY()
+	TObjectPtr<UIKRig_LimbEffector> Effector;
+	
+	FLimbSolver						Solver;
+	TArray<int32>					ChildrenToUpdate;
+};
diff --git a/Engine/Plugins/Animation/IKRig/Source/IKRig/Public/Rig/Solvers/IKRig_PoleSolver.h b/Engine/Plugins/Animation/IKRig/Source/IKRig/Public/Rig/Solvers/IKRig_PoleSolver.h
new file mode 100644
--- /dev/null
+++ b/Engine/Plugins/Animation/IKRig/Source/IKRig/Public/Rig/Solvers/IKRig_PoleSolver.h
@@ -0,0 +1,84 @@
+﻿// Copyright Epic Games, Inc. All Rights Reserved.
+
+#pragma once
+
+#include "Rig/IKRigDefinition.h"
+#include "Rig/Solvers/IKRigSolver.h"
+
+#include "IKRig_PoleSolver.generated.h"
+
+UCLASS(BlueprintType)
+class IKRIG_API UIKRig_PoleSolverEffector : public UObject
+{
+	GENERATED_BODY()
+
+public:
+
+	UIKRig_PoleSolverEffector() { SetFlags(RF_Transactional); }
+	
+	UPROPERTY(BlueprintReadOnly, VisibleAnywhere, Category = "Pole Solver Effector")
+	FName GoalName = NAME_None;
+
+	UPROPERTY(BlueprintReadOnly, VisibleAnywhere, Category = "Pole Solver Effector")
+	FName BoneName = NAME_None;
+
+	/** Blend the effector on/off. Range is 0-1. Default is 1.0. */
+	UPROPERTY(BlueprintReadWrite, EditAnywhere, Category = "Pole Solver Effector", meta = (ClampMin = "0.0", ClampMax = "1.0", UIMin = "0.0", UIMax = "1.0"))
+	float Alpha = 1.0f;
+};
+
+UCLASS(BlueprintType, EditInlineNew)
+class IKRIG_API UIKRig_PoleSolver : public UIKRigSolver
+{
+	GENERATED_BODY()
+
+public:
+
+	UIKRig_PoleSolver();
+	
+	UPROPERTY(VisibleAnywhere, DisplayName = "Root Bone", Category = "Pole Solver Settings")
+	FName RootName = NAME_None;
+
+	UPROPERTY(VisibleAnywhere, DisplayName = "End Bone", Category = "Pole Solver Settings")
+	FName EndName = NAME_None;
+	
+	/** UIKRigSolver interface */
+	virtual void Initialize(const FIKRigSkeleton& IKRigSkeleton) override;
+	virtual void Solve(FIKRigSkeleton& IKRigSkeleton, const FIKRigGoalContainer& Goals) override;
+	
+	virtual void UpdateSolverSettings(UIKRigSolver* InSettings) override;
+	virtual void RemoveGoal(const FName& GoalName) override;
+
+	virtual FName GetRootBone() const override { return RootName; };
+
+#if WITH_EDITOR
+	virtual FText GetNiceName() const override;
+	virtual bool GetWarningMessage(FText& OutWarningMessage) const override;
+	virtual bool IsBoneAffectedBySolver(const FName& BoneName, const FIKRigSkeleton& IKRigSkeleton) const override;
+	// goals
+	virtual void AddGoal(const UIKRigEffectorGoal& NewGoal) override;
+	virtual void RenameGoal(const FName& OldName, const FName& NewName) override;
+	virtual void SetGoalBone(const FName& GoalName, const FName& NewBoneName) override;
+	virtual bool IsGoalConnected(const FName& GoalName) const override;
+	virtual UObject* GetGoalSettings(const FName& GoalName) const override;
+	// root bone can be set on this solver
+	virtual bool RequiresRootBone() const override { return true; };
+	virtual void SetRootBone(const FName& RootBoneName) override;
+	// end bone can be set on this solver
+	virtual void SetEndBone(const FName& EndBoneName) override;
+	virtual FName GetEndBone() const override;
+	virtual bool RequiresEndBone() const override { return true; };
+	/** END UIKRigSolver interface */
+#endif
+
+private:
+
+	static void GatherChildren(const int32 BoneIndex, const FIKRigSkeleton& InSkeleton, TArray<int32>& OutChildren);
+	
+	UPROPERTY()
+	TObjectPtr<UIKRig_PoleSolverEffector>	Effector;
+	
+	TArray<int32>							Chain;
+	TArray<int32>							ChildrenToUpdate;
+};
+
diff --git a/Engine/Plugins/Animation/IKRig/Source/IKRig/Public/Rig/Solvers/IKRig_SetTransform.h b/Engine/Plugins/Animation/IKRig/Source/IKRig/Public/Rig/Solvers/IKRig_SetTransform.h
new file mode 100644
--- /dev/null
+++ b/Engine/Plugins/Animation/IKRig/Source/IKRig/Public/Rig/Solvers/IKRig_SetTransform.h
@@ -0,0 +1,70 @@
+// Copyright Epic Games, Inc. All Rights Reserved.
+
+#pragma once
+
+#include "Rig/IKRigDefinition.h"
+#include "IKRigSolver.h"
+
+#include "IKRig_SetTransform.generated.h"
+
+UCLASS()
+class IKRIG_API UIKRig_SetTransformEffector : public UObject
+{
+	GENERATED_BODY()
+
+public:
+
+	/** If true, Goal will drive the translation of the target bone. Default is true. */
+	UPROPERTY(EditAnywhere, Category = "Set Transform Effector")
+	bool bEnablePosition = true;
+
+	/** If true, Goal will drive the rotation of the target bone. Default is true. */
+	UPROPERTY(EditAnywhere, Category = "Set Transform Effector")
+	bool bEnableRotation = true;
+
+	/** Blend the effector on/off. Range is 0-1. Default is 1.0. */
+	UPROPERTY(EditAnywhere, Category = "Set Transform Effector", meta = (ClampMin = "0.0", ClampMax = "1.0", UIMin = "0.0", UIMax = "1.0"))
+	float Alpha = 1.0f;
+};
+
+UCLASS(BlueprintType, EditInlineNew)
+class IKRIG_API UIKRig_SetTransform : public UIKRigSolver
+{
+	GENERATED_BODY()
+
+public:
+	
+	UIKRig_SetTransform();
+
+	UPROPERTY(VisibleAnywhere, Category = "Set Transform Settings")
+	FName Goal;
+
+	UPROPERTY(VisibleAnywhere, Category = "Set Transform Settings")
+	FName RootBone;
+	
+	UPROPERTY()
+	TObjectPtr<UIKRig_SetTransformEffector> Effector;
+
+	/** UIKRigSolver interface */
+	virtual void Initialize(const FIKRigSkeleton& IKRigSkeleton) override;
+	virtual void Solve(FIKRigSkeleton& IKRigSkeleton, const FIKRigGoalContainer& Goals) override;
+	virtual void UpdateSolverSettings(UIKRigSolver* InSettings) override;
+	virtual void RemoveGoal(const FName& GoalName) override;
+	virtual FName GetRootBone() const override { return RootBone; };
+#if WITH_EDITOR
+	virtual FText GetNiceName() const override;
+	virtual bool GetWarningMessage(FText& OutWarningMessage) const override;
+	virtual void AddGoal(const UIKRigEffectorGoal& NewGoal) override;
+	virtual void RenameGoal(const FName& OldName, const FName& NewName) override;
+	virtual void SetGoalBone(const FName& GoalName, const FName& NewBoneName) override;
+	virtual bool IsGoalConnected(const FName& GoalName) const override;
+	virtual UObject* GetGoalSettings(const FName& GoalName) const override;
+	virtual bool IsBoneAffectedBySolver(const FName& BoneName, const FIKRigSkeleton& IKRigSkeleton) const override;
+#endif
+	/** END UIKRigSolver interface */
+
+private:
+
+	int32 BoneIndex;
+};
+
diff --git a/Engine/Plugins/Animation/IKRig/Source/IKRig/Public/Rig/Solvers/LimbSolver.h b/Engine/Plugins/Animation/IKRig/Source/IKRig/Public/Rig/Solvers/LimbSolver.h
new file mode 100644
--- /dev/null
+++ b/Engine/Plugins/Animation/IKRig/Source/IKRig/Public/Rig/Solvers/LimbSolver.h
@@ -0,0 +1,159 @@
+﻿// Copyright Epic Games, Inc. All Rights Reserved.
+
+#pragma once
+
+#include "CoreMinimal.h"
+#include "LimbSolver.generated.h"
+
+USTRUCT()
+struct IKRIG_API FLimbSolverSettings
+{
+	GENERATED_BODY()
+
+	/** Precision (distance to the target) */
+	UPROPERTY(EditAnywhere, Category = "Limb Solver Settings", meta = (ClampMin = "0.0"))
+	float ReachPrecision = 0.01f;
+	
+	// TWO BONES SETTINGS
+	
+	/** Hinge Bones Rotation Axis. This is essentially the plane normal for (hip - knee - foot). */
+	UPROPERTY(EditAnywhere, Category = "Limb Solver Settings|Two Bones")
+	TEnumAsByte<EAxis::Type> HingeRotationAxis = EAxis::None;
+
+	// FABRIK SETTINGS
+
+	/** Number of Max Iterations to reach the target */
+	UPROPERTY(EditAnywhere, Category = "Limb Solver Settings|FABRIK", meta = (ClampMin = "0", UIMin = "0", UIMax = "100"))
+	int32 MaxIterations = 12;
+
+	/** Enable/Disable rotational limits */
+	UPROPERTY(EditAnywhere, Category = "Limb Solver Settings|FABRIK|Limits")
+	bool bEnableLimit = false;
+
+	/** Only used if bEnableRotationLimit is enabled. Prevents the leg from folding onto itself,
+	* and forces at least this angle between Parent and Child bone. */
+	UPROPERTY(EditAnywhere, Category = "Limb Solver Settings|FABRIK|Limits", meta = (EditCondition="bEnableLimit", ClampMin = "0.0", ClampMax = "90.0", UIMin = "0.0", UIMax = "90.0"))
+	float MinRotationAngle = 15.f;
+
+	/** Pull averaging only has a visual impact when we have more than 2 bones (3 links). */
+	UPROPERTY(EditAnywhere, Category = "Limb Solver Settings|FABRIK|Pull Averaging")
+	bool bAveragePull = true;
+
+	/** Re-position limb to distribute pull: 0 = foot, 0.5 = balanced, 1.f = hip */
+	UPROPERTY(EditAnywhere, Category = "Limb Solver Settings|FABRIK|Pull Averaging", meta = (ClampMin = "0.0", ClampMax = "1.0",UIMin = "0.0", UIMax = "1.0"))
+	float PullDistribution = 0.5f;
+	
+	/** Move end effector towards target. If we are compressing the chain, limit displacement.*/
+	UPROPERTY(EditAnywhere, Category = "Limb Solver Settings|FABRIK|Pull Averaging", meta = (ClampMin = "0.0", ClampMax = "1.0",UIMin = "0.0", UIMax = "1.0"))
+	float ReachStepAlpha = 0.7f;
+
+	// TWIST SETTINGS
+	
+	/** Enable Knee Twist correction, by comparing Foot FK with Foot IK orientation. */
+	UPROPERTY(EditAnywhere, Category = "Limb Solver Settings|Twist")
+	bool bEnableTwistCorrection = false;
+	
+	/** Forward Axis for Foot bone. */
+	UPROPERTY(EditAnywhere, Category = "Limb Solver Settings|Twist", meta = (EditCondition="bEnableTwistCorrection"))
+	TEnumAsByte<EAxis::Type> EndBoneForwardAxis = EAxis::Y;
+};
+
+USTRUCT()
+struct FLimbLink
+{
+	GENERATED_USTRUCT_BODY()
+
+	// Location of bone in component space.
+	FVector Location;
+
+	// Distance to its parent
+	float Length;
+
+	// Bone Index in SkeletalMesh
+	int32 BoneIndex;
+
+	// Axis utilites
+	FVector LinkAxisZ;
+	FVector RealBendDir;
+	FVector BaseBendDir;
+
+	FLimbLink()
+		: Location(FVector::ZeroVector)
+		, Length(0.f)
+		, BoneIndex( INDEX_NONE )
+		, LinkAxisZ(FVector::ZeroVector)
+		, RealBendDir(FVector::ZeroVector)
+		, BaseBendDir(FVector::ZeroVector)
+	{}
+
+	FLimbLink(const FVector& InLocation, int32 InBoneIndex)
+		: Location(InLocation)
+		, Length(0.f)
+		, BoneIndex(InBoneIndex)
+		, LinkAxisZ(FVector::ZeroVector)
+		, RealBendDir(FVector::ZeroVector)
+		, BaseBendDir(FVector::ZeroVector)
+	{}
+};
+
+USTRUCT()
+struct FLimbSolver
+{
+	GENERATED_USTRUCT_BODY()
+
+	void Reset();
+	bool Initialize();
+	bool Solve(
+		TArray<FTransform>& InOutTransforms,
+		const FVector& InGoalLocation,
+		const FQuat& InGoalRotation,
+		const FLimbSolverSettings& InSettings);
+
+	int32 AddLink(const FVector& InLocation, int32 InBoneIndex);
+	int32 GetBoneIndex(const int32 Index) const;
+	int32 NumLinks() const;
+	
+private:
+	
+	void UpdateTransforms(TArray<FTransform>& InOutTransforms);
+	bool OrientTransformsTowardsGoal(TArray<FTransform>& InOutTransforms, const FVector& InGoalLocation) const;
+	bool AdjustKneeTwist(TArray<FTransform>& InOutTransforms, const FTransform& InGoalTransform, const EAxis::Type FootBoneForwardAxis );
+
+	float Length = 0.f;
+	
+	TArray<FLimbLink> Links;
+};
+
+namespace AnimationCore
+{
+	/**
+	* Limb solver
+	*
+	* This solves a hybrid TwoBoneIK / FABRIK based on the number of links within the chain
+	*
+	* @param	InOutLinks			Array of chain links
+	* @param	LimbLength			The full length on the limb (sum of all links' length)
+	* @param	GoalLocation		The location of the IK target
+	* @param	Precision			Precision 
+	* @param	InMaxIterations		Number of Max Iteration (used for FABRIK only)
+	* @param	HingeRotationAxis	Hinge Bones Rotation Axis. This is essentially the plane normal for (hip - knee - foot). (used for TwoBones only)
+	* @param	bUseAngleLimit		Enable/Disable rotational limits (used for FABRIK only)
+	* @param	MinRotationAngle	Prevents the leg from folding onto itself, and forces at least this angle between Parent and Child bone (used for FABRIK only)
+	* @param	ReachStepAlpha		Move end effector towards target. If we are compressing the chain, limit displacement. (used for FABRIK only)
+	* * @param	bAveragePull		Pull averaging only has a visual impact when we have more than 2 bones (3 links). (used for FABRIK only)
+	* @param	PullDistribution	Re-position limb to distribute pull: 0 = foot, 0.5 = balanced, 1.f = hip. (used for FABRIK only)
+	* @return  true if modified. False if not.
+	*/
+	bool SolveLimb(
+		TArray<FLimbLink>& InOutLinks,
+		float LimbLength,
+		const FVector& GoalLocation,
+		float Precision,
+		int32 InMaxIterations,
+		const FVector& HingeRotationAxis,
+		bool bUseAngleLimit,
+		float MinRotationAngle,
+		float ReachStepAlpha,
+		bool bAveragePull,
+		float PullDistribution);
+}
\ No newline at end of file
diff --git a/Engine/Plugins/Animation/IKRig/Source/IKRig/Public/Rig/Solvers/PointsToRotation.h b/Engine/Plugins/Animation/IKRig/Source/IKRig/Public/Rig/Solvers/PointsToRotation.h
new file mode 100644
--- /dev/null
+++ b/Engine/Plugins/Animation/IKRig/Source/IKRig/Public/Rig/Solvers/PointsToRotation.h
@@ -0,0 +1,83 @@
+// Copyright Epic Games, Inc. All Rights Reserved.
+
+#pragma once
+
+#include "Math/Quat.h"
+#include "Math/RotationMatrix.h"
+
+// given a set of points in an initial configuration and the same set of points in a deformed configuration,
+// this function outputs a quaternion that represents the "best fit" rotation that rotates the initial points to the
+// current points.
+static FQuat GetRotationFromDeformedPoints(
+	const TArray<FVector>& InInitialPoints,
+	const TArray<FVector>& InCurrentPoints,
+	FVector& OutInitialCentroid,
+	FVector& OutCurrentCentroid)
+{
+	check(InInitialPoints.Num() == InCurrentPoints.Num());
+
+	// must have more than 1 point to generate a gradient
+	if (InInitialPoints.Num() <= 1)
+	{
+		return FQuat::Identity;
+	}
+	
+	// calculate initial and current centroids
+	OutInitialCentroid = FVector::ZeroVector;
+	OutCurrentCentroid = FVector::ZeroVector;
+	for (int32 PointIndex=0; PointIndex<InInitialPoints.Num(); ++PointIndex)
+	{
+		OutInitialCentroid += InInitialPoints[PointIndex];
+		OutCurrentCentroid += InCurrentPoints[PointIndex];
+	}
+
+	// average centroids
+	const float InvNumEffectors = 1.0f / static_cast<float>(InInitialPoints.Num());
+	OutInitialCentroid *= InvNumEffectors;
+	OutCurrentCentroid *= InvNumEffectors;
+
+	// accumulate deformation gradient to extract a rotation from
+	FVector DX = FVector::ZeroVector; // DX, DY, DZ are rows of 3x3 deformation gradient tensor
+	FVector DY = FVector::ZeroVector;
+	FVector DZ = FVector::ZeroVector;
+	for (int32 PointIndex=0; PointIndex<InInitialPoints.Num(); ++PointIndex)
+	{
+		// accumulate the deformation gradient tensor for all points
+		// "Meshless Deformations Based on Shape Matching"
+		// Equation 7 describes accumulation of deformation gradient from points
+		//
+		// P is normalized vector from INITIAL centroid to INITIAL point
+		// Q is normalized vector from CURRENT centroid to CURRENT point
+		FVector P = (InInitialPoints[PointIndex] - OutInitialCentroid).GetSafeNormal();
+		FVector Q = (InCurrentPoints[PointIndex] - OutCurrentCentroid).GetSafeNormal();
+		// PQ^T is the outer product of P and Q which is a 3x3 matrix
+		// https://en.m.wikipedia.org/wiki/Outer_product
+		DX += FVector(P[0]*Q[0], P[0]*Q[1], P[0]*Q[2]);
+		DY += FVector(P[1]*Q[0], P[1]*Q[1], P[1]*Q[2]);
+		DZ += FVector(P[2]*Q[0], P[2]*Q[1], P[2]*Q[2]);
+	}
+
+	// extract "best fit" rotation from deformation gradient
+	FQuat Q = FQuat::Identity;
+	constexpr int32 MaxIter = 50;
+	// "A Robust Method to Extract the Rotational Part of Deformations" equation 7
+	// https://matthias-research.github.io/pages/publications/stablePolarDecomp.pdf
+	for (unsigned int Iter = 0; Iter < MaxIter; Iter++)
+	{
+		FMatrix R = FRotationMatrix::Make(Q);
+		FVector RCol0(R.M[0][0], R.M[0][1], R.M[0][2]);
+		FVector RCol1(R.M[1][0], R.M[1][1], R.M[1][2]);
+		FVector RCol2(R.M[2][0], R.M[2][1], R.M[2][2]);
+		FVector Omega = RCol0.Cross(DX) + RCol1.Cross(DY) + RCol2.Cross(DZ);
+		Omega *= 1.0f / (fabs(RCol0.Dot(DX) + RCol1.Dot(DY) + RCol2.Dot(DZ)) + SMALL_NUMBER);
+		const float W = Omega.Size();
+		if (W < SMALL_NUMBER)
+		{
+			break;
+		}
+		Q = FQuat(FQuat((1.0 / W) * Omega, W)) * Q;
+		Q.Normalize();
+	}
+
+	return Q;
+}
\ No newline at end of file
diff --git a/Engine/Plugins/Animation/IKRig/Source/IKRig/Public/Solvers/IKRig_BodyMover.h b/Engine/Plugins/Animation/IKRig/Source/IKRig/Public/Solvers/IKRig_BodyMover.h
deleted file mode 100644
diff --git a/Engine/Plugins/Animation/IKRig/Source/IKRig/Public/Solvers/IKRig_FBIKSolver.h b/Engine/Plugins/Animation/IKRig/Source/IKRig/Public/Solvers/IKRig_FBIKSolver.h
deleted file mode 100644
diff --git a/Engine/Plugins/Animation/IKRig/Source/IKRig/Public/Solvers/IKRig_LimbSolver.h b/Engine/Plugins/Animation/IKRig/Source/IKRig/Public/Solvers/IKRig_LimbSolver.h
deleted file mode 100644
diff --git a/Engine/Plugins/Animation/IKRig/Source/IKRig/Public/Solvers/IKRig_PoleSolver.h b/Engine/Plugins/Animation/IKRig/Source/IKRig/Public/Solvers/IKRig_PoleSolver.h
deleted file mode 100644
diff --git a/Engine/Plugins/Animation/IKRig/Source/IKRig/Public/Solvers/IKRig_SetTransform.h b/Engine/Plugins/Animation/IKRig/Source/IKRig/Public/Solvers/IKRig_SetTransform.h
deleted file mode 100644
diff --git a/Engine/Plugins/Animation/IKRig/Source/IKRig/Public/Solvers/LimbSolver.h b/Engine/Plugins/Animation/IKRig/Source/IKRig/Public/Solvers/LimbSolver.h
deleted file mode 100644
diff --git a/Engine/Plugins/Animation/IKRig/Source/IKRig/Public/Solvers/PointsToRotation.h b/Engine/Plugins/Animation/IKRig/Source/IKRig/Public/Solvers/PointsToRotation.h
deleted file mode 100644
diff --git a/Engine/Plugins/Animation/IKRig/Source/IKRigDeveloper/Private/AnimGraphNode_IKRig.cpp b/Engine/Plugins/Animation/IKRig/Source/IKRigDeveloper/Private/AnimGraphNode_IKRig.cpp
--- a/Engine/Plugins/Animation/IKRig/Source/IKRigDeveloper/Private/AnimGraphNode_IKRig.cpp
+++ b/Engine/Plugins/Animation/IKRig/Source/IKRigDeveloper/Private/AnimGraphNode_IKRig.cpp
@@ -3,8 +3,8 @@
 #include "AnimGraphNode_IKRig.h"
 #include "Animation/AnimInstance.h"
 #include "Components/SkeletalMeshComponent.h"
-#include "IKRigDefinition.h"
-#include "IKRigSolver.h"
+#include "Rig/IKRigDefinition.h"
+#include "Rig/Solvers/IKRigSolver.h"
 #include "Kismet2/CompilerResultsLog.h"
 #include "AnimationGraphSchema.h"
 #include "ScopedTransaction.h"
diff --git a/Engine/Plugins/Animation/IKRig/Source/IKRigEditor/Private/IKRigEditor.cpp b/Engine/Plugins/Animation/IKRig/Source/IKRigEditor/Private/IKRigEditor.cpp
--- a/Engine/Plugins/Animation/IKRig/Source/IKRigEditor/Private/IKRigEditor.cpp
+++ b/Engine/Plugins/Animation/IKRig/Source/IKRigEditor/Private/IKRigEditor.cpp
@@ -7,7 +7,7 @@
 #include "PropertyEditorModule.h"
 #include "PropertyEditorDelegates.h"
 #include "EditorModeRegistry.h"
-#include "IKRigDefinition.h"
+#include "Rig/IKRigDefinition.h"
 #include "RigEditor/AssetTypeActions_IKRigDefinition.h"
 #include "RetargetEditor/AssetTypeActions_IKRetargeter.h"
 #include "RetargetEditor/IKRetargetCommands.h"
diff --git a/Engine/Plugins/Animation/IKRig/Source/IKRigEditor/Private/RetargetEditor/IKRetargetEditPoseMode.cpp b/Engine/Plugins/Animation/IKRig/Source/IKRigEditor/Private/RetargetEditor/IKRetargetEditPoseMode.cpp
--- a/Engine/Plugins/Animation/IKRig/Source/IKRigEditor/Private/RetargetEditor/IKRetargetEditPoseMode.cpp
+++ b/Engine/Plugins/Animation/IKRig/Source/IKRigEditor/Private/RetargetEditor/IKRetargetEditPoseMode.cpp
@@ -5,7 +5,7 @@
 #include "EditorViewportClient.h"
 #include "AssetEditorModeManager.h"
 #include "IKRigDebugRendering.h"
-#include "IKRigProcessor.h"
+#include "Rig/IKRigProcessor.h"
 #include "IPersonaPreviewScene.h"
 #include "Animation/DebugSkelMeshComponent.h"
 #include "Preferences/PersonaOptions.h"
diff --git a/Engine/Plugins/Animation/IKRig/Source/IKRigEditor/Private/RigEditor/AssetTypeActions_IKRigDefinition.cpp b/Engine/Plugins/Animation/IKRig/Source/IKRigEditor/Private/RigEditor/AssetTypeActions_IKRigDefinition.cpp
--- a/Engine/Plugins/Animation/IKRig/Source/IKRigEditor/Private/RigEditor/AssetTypeActions_IKRigDefinition.cpp
+++ b/Engine/Plugins/Animation/IKRig/Source/IKRigEditor/Private/RigEditor/AssetTypeActions_IKRigDefinition.cpp
@@ -1,7 +1,7 @@
 // Copyright Epic Games, Inc. All Rights Reserved.
 
 #include "RigEditor/AssetTypeActions_IKRigDefinition.h"
-#include "IKRigDefinition.h"
+#include "Rig/IKRigDefinition.h"
 #include "RigEditor/IKRigToolkit.h"
 #include "ThumbnailRendering/SceneThumbnailInfo.h"
 
diff --git a/Engine/Plugins/Animation/IKRig/Source/IKRigEditor/Private/RigEditor/IKRigController.cpp b/Engine/Plugins/Animation/IKRig/Source/IKRigEditor/Private/RigEditor/IKRigController.cpp
--- a/Engine/Plugins/Animation/IKRig/Source/IKRigEditor/Private/RigEditor/IKRigController.cpp
+++ b/Engine/Plugins/Animation/IKRig/Source/IKRigEditor/Private/RigEditor/IKRigController.cpp
@@ -3,9 +3,9 @@
 #include "RigEditor/IKRigController.h"
 
 #include "IKRigEditor.h"
-#include "IKRigDefinition.h"
-#include "IKRigProcessor.h"
-#include "IKRigSolver.h"
+#include "Rig/IKRigDefinition.h"
+#include "Rig/IKRigProcessor.h"
+#include "Rig/Solvers/IKRigSolver.h"
 
 #include "Engine/SkeletalMesh.h"
 #include "ScopedTransaction.h"
diff --git a/Engine/Plugins/Animation/IKRig/Source/IKRigEditor/Private/RigEditor/IKRigDefinitionFactory.cpp b/Engine/Plugins/Animation/IKRig/Source/IKRigEditor/Private/RigEditor/IKRigDefinitionFactory.cpp
--- a/Engine/Plugins/Animation/IKRig/Source/IKRigEditor/Private/RigEditor/IKRigDefinitionFactory.cpp
+++ b/Engine/Plugins/Animation/IKRig/Source/IKRigEditor/Private/RigEditor/IKRigDefinitionFactory.cpp
@@ -1,7 +1,7 @@
 // Copyright Epic Games, Inc. All Rights Reserved.
 
 #include "RigEditor/IKRigDefinitionFactory.h"
-#include "IKRigDefinition.h"
+#include "Rig/IKRigDefinition.h"
 #include "AssetTypeCategories.h"
 #include "ContentBrowserModule.h"
 #include "IContentBrowserSingleton.h"
diff --git a/Engine/Plugins/Animation/IKRig/Source/IKRigEditor/Private/RigEditor/IKRigEditorController.cpp b/Engine/Plugins/Animation/IKRig/Source/IKRigEditor/Private/RigEditor/IKRigEditorController.cpp
--- a/Engine/Plugins/Animation/IKRig/Source/IKRigEditor/Private/RigEditor/IKRigEditorController.cpp
+++ b/Engine/Plugins/Animation/IKRig/Source/IKRigEditor/Private/RigEditor/IKRigEditorController.cpp
@@ -9,7 +9,7 @@
 #include "RigEditor/IKRigToolkit.h"
 #include "RigEditor/SIKRigAssetBrowser.h"
 #include "RigEditor/SIKRigOutputLog.h"
-#include "Solvers/IKRig_FBIKSolver.h"
+#include "Rig/Solvers/IKRig_FBIKSolver.h"
 #include "Widgets/Input/SComboBox.h"
 #include "ContentBrowserModule.h"
 #include "IContentBrowserSingleton.h"
diff --git a/Engine/Plugins/Animation/IKRig/Source/IKRigEditor/Private/RigEditor/IKRigOutputLogTabSummoner.cpp b/Engine/Plugins/Animation/IKRig/Source/IKRigEditor/Private/RigEditor/IKRigOutputLogTabSummoner.cpp
--- a/Engine/Plugins/Animation/IKRig/Source/IKRigEditor/Private/RigEditor/IKRigOutputLogTabSummoner.cpp
+++ b/Engine/Plugins/Animation/IKRig/Source/IKRigEditor/Private/RigEditor/IKRigOutputLogTabSummoner.cpp
@@ -2,7 +2,7 @@
 #include "RigEditor/IKRigOutputLogTabSummoner.h"
 
 #include "IDocumentation.h"
-#include "IKRigProcessor.h"
+#include "Rig/IKRigProcessor.h"
 #include "RigEditor/IKRigToolkit.h"
 #include "RigEditor/IKRigEditorStyle.h"
 #include "RigEditor/SIKRigOutputLog.h"
diff --git a/Engine/Plugins/Animation/IKRig/Source/IKRigEditor/Private/RigEditor/IKRigToolkit.cpp b/Engine/Plugins/Animation/IKRig/Source/IKRigEditor/Private/RigEditor/IKRigToolkit.cpp
--- a/Engine/Plugins/Animation/IKRig/Source/IKRigEditor/Private/RigEditor/IKRigToolkit.cpp
+++ b/Engine/Plugins/Animation/IKRig/Source/IKRigEditor/Private/RigEditor/IKRigToolkit.cpp
@@ -16,7 +16,7 @@
 #include "Animation/DebugSkelMeshComponent.h"
 #include "Framework/MultiBox/MultiBoxBuilder.h"
 
-#include "IKRigDefinition.h"
+#include "Rig/IKRigDefinition.h"
 #include "IPersonaViewport.h"
 #include "PersonaPreviewSceneDescription.h"
 #include "RigEditor/IKRigAnimInstance.h"
diff --git a/Engine/Plugins/Animation/IKRig/Source/IKRigEditor/Private/RigEditor/SIKRigHierarchy.cpp b/Engine/Plugins/Animation/IKRig/Source/IKRigEditor/Private/RigEditor/SIKRigHierarchy.cpp
--- a/Engine/Plugins/Animation/IKRig/Source/IKRigEditor/Private/RigEditor/SIKRigHierarchy.cpp
+++ b/Engine/Plugins/Animation/IKRig/Source/IKRigEditor/Private/RigEditor/SIKRigHierarchy.cpp
@@ -2,8 +2,8 @@
 
 #include "RigEditor/SIKRigHierarchy.h"
 
-#include "IKRigProcessor.h"
-#include "IKRigSolver.h"
+#include "Rig/IKRigProcessor.h"
+#include "Rig/Solvers/IKRigSolver.h"
 #include "Engine/SkeletalMesh.h"
 #include "IPersonaToolkit.h"
 #include "SKismetInspector.h"
diff --git a/Engine/Plugins/Animation/IKRig/Source/IKRigEditor/Private/RigEditor/SIKRigSolverStack.cpp b/Engine/Plugins/Animation/IKRig/Source/IKRigEditor/Private/RigEditor/SIKRigSolverStack.cpp
--- a/Engine/Plugins/Animation/IKRig/Source/IKRigEditor/Private/RigEditor/SIKRigSolverStack.cpp
+++ b/Engine/Plugins/Animation/IKRig/Source/IKRigEditor/Private/RigEditor/SIKRigSolverStack.cpp
@@ -5,7 +5,7 @@
 #include "RigEditor/IKRigEditorStyle.h"
 #include "RigEditor/IKRigToolkit.h"
 #include "RigEditor/IKRigEditorController.h"
-#include "IKRigSolver.h"
+#include "Rig/Solvers/IKRigSolver.h"
 
 #include "Framework/Commands/UICommandList.h"
 #include "Framework/MultiBox/MultiBoxBuilder.h"
diff --git a/Engine/Plugins/Animation/IKRig/Source/IKRigEditor/Public/IKRigDebugRendering.h b/Engine/Plugins/Animation/IKRig/Source/IKRigEditor/Public/IKRigDebugRendering.h
--- a/Engine/Plugins/Animation/IKRig/Source/IKRigEditor/Public/IKRigDebugRendering.h
+++ b/Engine/Plugins/Animation/IKRig/Source/IKRigEditor/Public/IKRigDebugRendering.h
@@ -2,7 +2,7 @@
 
 #pragma once
 
-#include "IKRigDefinition.h"
+#include "Rig/IKRigDefinition.h"
 #include "SceneManagement.h"
 #include "SkeletalDebugRendering.h"
 
diff --git a/Engine/Plugins/Animation/IKRig/Source/IKRigEditor/Public/RigEditor/IKRigAnimInstance.h b/Engine/Plugins/Animation/IKRig/Source/IKRigEditor/Public/RigEditor/IKRigAnimInstance.h
--- a/Engine/Plugins/Animation/IKRig/Source/IKRigEditor/Public/RigEditor/IKRigAnimInstance.h
+++ b/Engine/Plugins/Animation/IKRig/Source/IKRigEditor/Public/RigEditor/IKRigAnimInstance.h
@@ -3,7 +3,7 @@
 #pragma once
 
 #include "AnimPreviewInstance.h"
-#include "IKRigProcessor.h"
+#include "Rig/IKRigProcessor.h"
 #include "AnimNodes/AnimNode_IKRig.h"
 
 #include "IKRigAnimInstance.generated.h"
diff --git a/Engine/Plugins/Animation/IKRig/Source/IKRigEditor/Public/RigEditor/IKRigEditMode.h b/Engine/Plugins/Animation/IKRig/Source/IKRigEditor/Public/RigEditor/IKRigEditMode.h
--- a/Engine/Plugins/Animation/IKRig/Source/IKRigEditor/Public/RigEditor/IKRigEditMode.h
+++ b/Engine/Plugins/Animation/IKRig/Source/IKRigEditor/Public/RigEditor/IKRigEditMode.h
@@ -4,7 +4,7 @@
 
 #include "CoreMinimal.h"
 
-#include "IKRigDefinition.h"
+#include "Rig/IKRigDefinition.h"
 #include "IPersonaEditMode.h"
 #include "SkeletalDebugRendering.h"
 
diff --git a/Engine/Plugins/Animation/IKRig/Source/IKRigEditor/Public/RigEditor/IKRigEditorController.h b/Engine/Plugins/Animation/IKRig/Source/IKRigEditor/Public/RigEditor/IKRigEditorController.h
--- a/Engine/Plugins/Animation/IKRig/Source/IKRigEditor/Public/RigEditor/IKRigEditorController.h
+++ b/Engine/Plugins/Animation/IKRig/Source/IKRigEditor/Public/RigEditor/IKRigEditorController.h
@@ -2,7 +2,7 @@
 
 #pragma once
 
-#include "IKRigSolver.h"
+#include "Rig/Solvers/IKRigSolver.h"
 #include "SIKRigRetargetChainList.h"
 #include "Templates/SubclassOf.h"
 #include "Animation/AnimationAsset.h"
diff --git a/Engine/Plugins/Animation/IKRig/Source/IKRigEditor/Public/RigEditor/SIKRigRetargetChainList.h b/Engine/Plugins/Animation/IKRig/Source/IKRigEditor/Public/RigEditor/SIKRigRetargetChainList.h
--- a/Engine/Plugins/Animation/IKRig/Source/IKRigEditor/Public/RigEditor/SIKRigRetargetChainList.h
+++ b/Engine/Plugins/Animation/IKRig/Source/IKRigEditor/Public/RigEditor/SIKRigRetargetChainList.h
@@ -4,7 +4,7 @@
 
 #include "CoreMinimal.h"
 #include "EditorUndoClient.h"
-#include "IKRigDefinition.h"
+#include "Rig/IKRigDefinition.h"
 #include "DragAndDrop/DecoratedDragDropOp.h"
 #include "Widgets/SCompoundWidget.h"
 #include "Widgets/Views/STableViewBase.h"
diff --git a/Engine/Plugins/Experimental/Animation/ContextualAnimation/Source/ContextualAnimation/Private/ContextualAnimSceneActorComponent.cpp b/Engine/Plugins/Experimental/Animation/ContextualAnimation/Source/ContextualAnimation/Private/ContextualAnimSceneActorComponent.cpp
--- a/Engine/Plugins/Experimental/Animation/ContextualAnimation/Source/ContextualAnimation/Private/ContextualAnimSceneActorComponent.cpp
+++ b/Engine/Plugins/Experimental/Animation/ContextualAnimation/Source/ContextualAnimation/Private/ContextualAnimSceneActorComponent.cpp
@@ -13,7 +13,7 @@
 #include "AnimNotifyState_IKWindow.h"
 #include "GameFramework/Pawn.h"
 #include "Net/UnrealNetwork.h"
-#include "IKRigDataTypes.h"
+#include "Rig/IKRigDataTypes.h"
 #include "Net/Core/PushModel/PushModel.h"
 #include "Physics/PhysicsInterfaceTypes.h"
 #include "PrimitiveSceneProxy.h"
