IK Retargeter toolbar redesign for 5.3

diff --git a/Engine/Plugins/Animation/IKRig/Source/IKRigEditor/Private/RetargetEditor/IKRetargetAnimInstanceProxy.cpp b/Engine/Plugins/Animation/IKRig/Source/IKRigEditor/Private/RetargetEditor/IKRetargetAnimInstanceProxy.cpp
--- a/Engine/Plugins/Animation/IKRig/Source/IKRigEditor/Private/RetargetEditor/IKRetargetAnimInstanceProxy.cpp
+++ b/Engine/Plugins/Animation/IKRig/Source/IKRigEditor/Private/RetargetEditor/IKRetargetAnimInstanceProxy.cpp
@@ -41,21 +41,29 @@ void FIKRetargetAnimInstanceProxy::CacheBones()
 }
 
 bool FIKRetargetAnimInstanceProxy::Evaluate(FPoseContext& Output)
-{	
-	if (OutputMode == ERetargeterOutputMode::RunRetarget)
+{
+	switch (OutputMode)
 	{
-		if (SourceOrTarget == ERetargetSourceOrTarget::Source)
+	case ERetargeterOutputMode::RunRetarget:
 		{
-			FAnimPreviewInstanceProxy::Evaluate(Output);
+			if (SourceOrTarget == ERetargetSourceOrTarget::Source)
+			{
+				FAnimPreviewInstanceProxy::Evaluate(Output);
+			}
+			else
+			{
+				RetargetNode->Evaluate_AnyThread(Output);
+			}
+			break;
 		}
-		else
+	case ERetargeterOutputMode::EditRetargetPose:
+	case ERetargeterOutputMode::ShowRetargetPose:
 		{
-			RetargetNode->Evaluate_AnyThread(Output);
+			PreviewPoseNode->Evaluate_AnyThread(Output);
+			break;
 		}
-	}
-	else
-	{
-		PreviewPoseNode->Evaluate_AnyThread(Output);
+	default:
+		checkNoEntry();
 	}
 	
 	return true;
diff --git a/Engine/Plugins/Animation/IKRig/Source/IKRigEditor/Private/RetargetEditor/IKRetargetCommands.cpp b/Engine/Plugins/Animation/IKRig/Source/IKRigEditor/Private/RetargetEditor/IKRetargetCommands.cpp
--- a/Engine/Plugins/Animation/IKRig/Source/IKRigEditor/Private/RetargetEditor/IKRetargetCommands.cpp
+++ b/Engine/Plugins/Animation/IKRig/Source/IKRigEditor/Private/RetargetEditor/IKRetargetCommands.cpp
@@ -6,9 +6,13 @@
 
 void FIKRetargetCommands::RegisterCommands()
 {
-	UI_COMMAND(ShowRetargetPose, "Show Retarget Pose", "Pause playback and set Source and Target meshes to the Retarget Pose.", EUserInterfaceActionType::ToggleButton, FInputChord());
+	UI_COMMAND(RunRetargeter, "Run Retargeter", "Resume playback of animation and retarget to the target skeleton.", EUserInterfaceActionType::Button, FInputChord());
+	UI_COMMAND(EditRetargetPose, "Edit Retarget Pose", "Enter into mode allowing manual editing of the target skeleton pose in the viewport.", EUserInterfaceActionType::Button, FInputChord());
+	UI_COMMAND(ShowRetargetPose, "Show Retarget Pose", "Display the retarget pose and retarget to the target skeleton.", EUserInterfaceActionType::Button, FInputChord());
+
+	UI_COMMAND(ShowGlobalSettings, "Global Settings", "Open global settings in detail tab.", EUserInterfaceActionType::Button, FInputChord());
+	UI_COMMAND(ShowRootSettings, "Root Settings", "Open root settings in detail tab.", EUserInterfaceActionType::Button, FInputChord());
 	
-	UI_COMMAND(EditRetargetPose, "Edit Mode", "Enter into mode allowing manual editing of the target skeleton pose in the viewport.", EUserInterfaceActionType::ToggleButton, FInputChord());
 	UI_COMMAND(ResetAllBones, "Reset All", "Sets the selected bones to the mesh reference pose.", EUserInterfaceActionType::Button, FInputChord());
 	UI_COMMAND(ResetSelectedBones, "Reset Selected Bones", "Sets the selected bones to the mesh reference pose.", EUserInterfaceActionType::Button, FInputChord());
 	UI_COMMAND(ResetSelectedAndChildrenBones, "Reset Selected and Children Bones", "Sets the selected bones (and all children recursively) to the mesh reference pose.", EUserInterfaceActionType::Button, FInputChord());
diff --git a/Engine/Plugins/Animation/IKRig/Source/IKRigEditor/Private/RetargetEditor/IKRetargetEditor.cpp b/Engine/Plugins/Animation/IKRig/Source/IKRigEditor/Private/RetargetEditor/IKRetargetEditor.cpp
--- a/Engine/Plugins/Animation/IKRig/Source/IKRigEditor/Private/RetargetEditor/IKRetargetEditor.cpp
+++ b/Engine/Plugins/Animation/IKRig/Source/IKRigEditor/Private/RetargetEditor/IKRetargetEditor.cpp
@@ -119,14 +119,46 @@ void FIKRetargetEditor::UnregisterTabSpawners(const TSharedRef<FTabManager>& InT
 void FIKRetargetEditor::BindCommands()
 {
 	const FIKRetargetCommands& Commands = FIKRetargetCommands::Get();
-	
+
+	//
+	// Retarget output modes
+	//
 	ToolkitCommands->MapAction(
-        Commands.EditRetargetPose,
-        FExecuteAction::CreateSP(EditorController, &FIKRetargetEditorController::HandleEditPose),
-        FCanExecuteAction::CreateSP(EditorController, &FIKRetargetEditorController::CanEditPose),
-        FIsActionChecked::CreateSP(EditorController,  &FIKRetargetEditorController::IsEditingPose),
-		EUIActionRepeatMode::RepeatDisabled);
+		Commands.RunRetargeter,
+		FExecuteAction::CreateSP(EditorController, &FIKRetargetEditorController::SetRetargeterMode, ERetargeterOutputMode::RunRetarget),
+		FCanExecuteAction::CreateSP(EditorController, &FIKRetargetEditorController::IsReadyToRetarget));
+
+	ToolkitCommands->MapAction(
+		Commands.ShowRetargetPose,
+		FExecuteAction::CreateSP(EditorController, &FIKRetargetEditorController::SetRetargeterMode, ERetargeterOutputMode::ShowRetargetPose),
+		FCanExecuteAction::CreateSP(EditorController, &FIKRetargetEditorController::IsCurrentMeshLoaded),
+		FIsActionChecked());
 
+	ToolkitCommands->MapAction(
+		Commands.EditRetargetPose,
+		FExecuteAction::CreateSP(EditorController, &FIKRetargetEditorController::SetRetargeterMode, ERetargeterOutputMode::EditRetargetPose),
+		FCanExecuteAction::CreateSP(EditorController, &FIKRetargetEditorController::IsCurrentMeshLoaded),
+		FIsActionChecked());
+
+
+	//
+	// Show global / root settings in details panel
+	//
+	ToolkitCommands->MapAction(
+		Commands.ShowGlobalSettings,
+		FExecuteAction::CreateSP(EditorController, &FIKRetargetEditorController::ShowGlobalSettings),
+		FCanExecuteAction(),
+		FIsActionChecked::CreateSP(EditorController,  &FIKRetargetEditorController::IsShowingGlobalSettings));
+	ToolkitCommands->MapAction(
+		Commands.ShowRootSettings,
+		FExecuteAction::CreateSP(EditorController, &FIKRetargetEditorController::ShowRootSettings),
+		FCanExecuteAction(),
+		FIsActionChecked::CreateSP(EditorController,  &FIKRetargetEditorController::IsShowingRootSettings));
+
+	//
+	// Edit pose commands
+	//
+	
 	ToolkitCommands->MapAction(
 		Commands.ResetAllBones,
 		FExecuteAction::CreateSP(EditorController, &FIKRetargetEditorController::HandleResetAllBones),
@@ -215,12 +247,148 @@ void FIKRetargetEditor::FillToolbar(FToolBarBuilder& ToolbarBuilder)
 {
 	ToolbarBuilder.BeginSection("Show Retarget Pose");
 	{
-		// TODO debug section
+		ToolbarBuilder.AddToolBarButton(
+			FExecuteAction::CreateLambda([this]{ EditorController->SetRetargetModeToPreviousMode(); }),
+			NAME_None, 
+			TAttribute<FText>::Create(TAttribute<FText>::FGetter::CreateSP(EditorController, &FIKRetargetEditorController::GetRetargeterModeLabel)),
+			TAttribute<FText>(), 
+			TAttribute<FSlateIcon>::Create(TAttribute<FSlateIcon>::FGetter::CreateSP(EditorController, &FIKRetargetEditorController::GetCurrentRetargetModeIcon))
+		);
+		
+		ToolbarBuilder.AddComboButton(
+			FUIAction(),
+			FOnGetContent::CreateSP(this, &FIKRetargetEditor::GenerateRetargetModesMenu),
+			LOCTEXT("RetargetMode_Label", "UI Modes"),
+			LOCTEXT("RetargetMode_ToolTip", "Choose which mode to display in the viewport."),
+			FSlateIcon(FAppStyle::GetAppStyleSetName(), "LevelEditor.Recompile"),
+			true);
 	}
+	ToolbarBuilder.EndSection();
+
+	ToolbarBuilder.AddSeparator();
+	ToolbarBuilder.AddWidget(SNew(SSpacer), NAME_None, true, HAlign_Right);
+
+	ToolbarBuilder.BeginSection("Show Settings");
+	{
+		ToolbarBuilder.AddToolBarButton(
+		FIKRetargetCommands::Get().ShowGlobalSettings,
+		NAME_None,
+		TAttribute<FText>(),
+		TAttribute<FText>(),
+		FSlateIcon(FAppStyle::Get().GetStyleSetName(),"Icons.Settings"));
+
+		ToolbarBuilder.AddToolBarButton(
+		FIKRetargetCommands::Get().ShowRootSettings,
+		NAME_None,
+		TAttribute<FText>(),
+		TAttribute<FText>(),
+		FSlateIcon(FAppStyle::Get().GetStyleSetName(),"Icons.Settings"));
+	}
+	ToolbarBuilder.EndSection();
+
+	ToolbarBuilder.AddSeparator();
+
+	FLinearColor OffColor = FLinearColor::White;
+	FLinearColor OnColor = FLinearColor::Green;
 	
+	TSharedPtr<SVerticalBox> Box = SNew(SVerticalBox)
+	+ SVerticalBox::Slot()
+	.VAlign(VAlign_Center)
+	.HAlign(HAlign_Center)
+	.AutoHeight()
+	[
+		SNew(SHorizontalBox)
+		+ SHorizontalBox::Slot()
+		.FillWidth(1.f)
+		.HAlign(HAlign_Center)
+		[
+			SNew(STextBlock)
+			.Font(FCoreStyle::Get().GetFontStyle("SmallFont"))
+			.Text(FText::FromString("Toggle Retarget Phases"))
+		]
+	]
+	+ SVerticalBox::Slot()
+	.AutoHeight()
+	[
+		SNew(SHorizontalBox)
+		+ SHorizontalBox::Slot()
+		.AutoWidth()
+		.HAlign(HAlign_Center)
+		[
+			SNew(SButton)
+			.OnClicked_Lambda([this]()
+			{
+				EditorController->ToggleRootRetargetPass();
+				return FReply::Handled();
+			})
+			.ButtonColorAndOpacity_Lambda([this, OffColor, OnColor]() -> FLinearColor
+			{
+				return EditorController->IsRootRetargetOn() ? OnColor : OffColor;
+			})
+			[
+				SNew(STextBlock).Text(FText::FromString("Root"))
+			]
+		]
+
+		+ SHorizontalBox::Slot()
+		.AutoWidth()
+		.HAlign(HAlign_Center)
+		[
+			SNew(SButton)
+			.OnClicked_Lambda([this]()
+			{
+				EditorController->ToggleFKRetargetPass();
+				return FReply::Handled();
+			})
+			.ButtonColorAndOpacity_Lambda([this, OffColor, OnColor]() -> FLinearColor
+			{
+				return EditorController->IsFKRetargetOn() ? OnColor : OffColor;
+			})
+			[
+				SNew(STextBlock).Text(FText::FromString("FK"))
+			]
+		]
+		
+		+ SHorizontalBox::Slot()
+		.AutoWidth()
+		.HAlign(HAlign_Center)
+		[
+			SNew(SButton)
+			.OnClicked_Lambda([this]()
+			{
+				EditorController->ToggleIKRetargetPass();
+				return FReply::Handled();
+			})
+			.ButtonColorAndOpacity_Lambda([this, OffColor, OnColor]() -> FLinearColor
+			{
+				return EditorController->IsIKRetargetOn() ? OnColor : OffColor;
+			})
+			[
+				SNew(STextBlock).Text(FText::FromString("IK"))
+			]
+		]
+	];
+
+	ToolbarBuilder.BeginSection("Toggle Retarget Passes");
+	{
+		ToolbarBuilder.AddWidget(Box.ToSharedRef());
+	}
 	ToolbarBuilder.EndSection();
 }
 
+TSharedRef<SWidget> FIKRetargetEditor::GenerateRetargetModesMenu()
+{
+	FMenuBuilder MenuBuilder(true, GetToolkitCommands());
+	
+	MenuBuilder.BeginSection(TEXT("Retarget Modes"));
+	MenuBuilder.AddMenuEntry(FIKRetargetCommands::Get().RunRetargeter, TEXT("Run Retargeter"), TAttribute<FText>(), TAttribute<FText>(),  EditorController->GetRetargeterModeIcon(ERetargeterOutputMode::RunRetarget));
+	MenuBuilder.AddMenuEntry(FIKRetargetCommands::Get().EditRetargetPose, TEXT("Edit Retarget Pose"), TAttribute<FText>(), TAttribute<FText>(), EditorController->GetRetargeterModeIcon(ERetargeterOutputMode::EditRetargetPose));
+	MenuBuilder.AddMenuEntry(FIKRetargetCommands::Get().ShowRetargetPose, TEXT("Show Retarget Pose"), TAttribute<FText>(), TAttribute<FText>(), EditorController->GetRetargeterModeIcon(ERetargeterOutputMode::ShowRetargetPose));
+	MenuBuilder.EndSection();
+	
+	return MenuBuilder.MakeWidget();
+}
+
 FName FIKRetargetEditor::GetToolkitFName() const
 {
 	return FName("IKRetargetEditor");
diff --git a/Engine/Plugins/Animation/IKRig/Source/IKRigEditor/Private/RetargetEditor/IKRetargetEditorController.cpp b/Engine/Plugins/Animation/IKRig/Source/IKRigEditor/Private/RetargetEditor/IKRetargetEditorController.cpp
--- a/Engine/Plugins/Animation/IKRig/Source/IKRigEditor/Private/RetargetEditor/IKRetargetEditorController.cpp
+++ b/Engine/Plugins/Animation/IKRig/Source/IKRigEditor/Private/RetargetEditor/IKRetargetEditorController.cpp
@@ -17,6 +17,7 @@
 #include "RetargetEditor/IKRetargetDefaultMode.h"
 #include "RetargetEditor/IKRetargetEditPoseMode.h"
 #include "RetargetEditor/IKRetargetEditor.h"
+#include "RetargetEditor/IKRetargetEditorStyle.h"
 #include "RetargetEditor/IKRetargetHitProxies.h"
 #include "RetargetEditor/SIKRetargetChainMapList.h"
 #include "RetargetEditor/SIKRetargetHierarchy.h"
@@ -679,6 +680,30 @@ void FIKRetargetEditorController::ClearOutputLog() const
 	}
 }
 
+bool FIKRetargetEditorController::IsObjectInDetailsView(const UObject* Object)
+{
+	if (!DetailsView.IsValid())
+	{
+		return false;
+	}
+
+	if (!Object)
+	{
+		return false;
+	}
+	
+	TArray<TWeakObjectPtr<UObject>> SelectedObjects = DetailsView->GetSelectedObjects();
+	for (TWeakObjectPtr<UObject> SelectedObject : SelectedObjects)
+	{
+		if (SelectedObject.Get() == Object)
+		{
+			return true;
+		}
+	}
+
+	return false;
+}
+
 void FIKRetargetEditorController::RefreshAllViews() const
 {
 	//if (!Editor.IsValid())
@@ -969,16 +994,9 @@ void FIKRetargetEditorController::SetRootSelected(const bool bIsSelected)
 	{
 		return;
 	}
-	
-	URetargetRootSettings* RootSettings = AssetController->GetAsset()->GetRootSettingsUObject();
-	if (RootSettings->EditorController.Get() != this)
-	{
-		RootSettings->EditorController = SharedThis(this);	
-	}
-	
-	SetDetailsObject(RootSettings);
 
 	LastSelectedItem = ERetargetSelectionType::ROOT;
+	ShowRootSettings();
 }
 
 void FIKRetargetEditorController::CleanSelection(ERetargetSourceOrTarget SourceOrTarget)
@@ -1085,46 +1103,49 @@ void FIKRetargetEditorController::SetRetargeterMode(ERetargeterOutputMode Mode)
 	RefreshDetailsView();
 }
 
-FReply FIKRetargetEditorController::HandleShowRetargetPose()
+FText FIKRetargetEditorController::GetRetargeterModeLabel()
 {
-	const ERetargeterOutputMode CurrentMode = GetRetargeterMode();
-	if (CurrentMode == ERetargeterOutputMode::ShowRetargetPose || CurrentMode == ERetargeterOutputMode::EditRetargetPose)
+	switch (GetRetargeterMode())
 	{
-		SetRetargeterMode(ERetargeterOutputMode::RunRetarget);
-	}
-	else
-	{
-		SetRetargeterMode(ERetargeterOutputMode::ShowRetargetPose);
+	case ERetargeterOutputMode::RunRetarget:
+		return FText::FromString("Running Retarget");
+	case ERetargeterOutputMode::EditRetargetPose:
+		return FText::FromString("Editing Retarget Pose");
+	case ERetargeterOutputMode::ShowRetargetPose:
+		return FText::FromString("Showing Retarget Pose");
+	default:
+		checkNoEntry();
+		return FText::FromString("Unknown Mode.");
 	}
-	
-	return FReply::Handled();
 }
 
-bool FIKRetargetEditorController::CanShowRetargetPose() const
+FSlateIcon FIKRetargetEditorController::GetCurrentRetargetModeIcon() const
 {
-	return GetRetargeterMode() != ERetargeterOutputMode::ShowRetargetPose;
+	return GetRetargeterModeIcon(GetRetargeterMode());
 }
 
-bool FIKRetargetEditorController::IsShowingRetargetPose() const
+FSlateIcon FIKRetargetEditorController::GetRetargeterModeIcon(ERetargeterOutputMode Mode) const
 {
-	return GetRetargeterMode() == ERetargeterOutputMode::ShowRetargetPose;
+	switch (Mode)
+	{
+	case ERetargeterOutputMode::RunRetarget:
+		return FSlateIcon(FAppStyle::Get().GetStyleSetName(),"GenericPlay");
+	case ERetargeterOutputMode::EditRetargetPose:
+		return FSlateIcon(FAppStyle::Get().GetStyleSetName(),"Icons.Edit");
+	case ERetargeterOutputMode::ShowRetargetPose:
+		return FSlateIcon(FAppStyle::Get().GetStyleSetName(),"GenericPause");
+	default:
+		checkNoEntry();
+		return FSlateIcon(FAppStyle::Get().GetStyleSetName(),"GenericPause");
+	}
 }
 
-void FIKRetargetEditorController::HandleEditPose()
+bool FIKRetargetEditorController::IsReadyToRetarget() const
 {
-	if (IsEditingPose())
-	{
-		// stop pose editing
-		SetRetargeterMode(PreviousMode);
-	}
-	else
-	{
-		// start pose editing
-		SetRetargeterMode(ERetargeterOutputMode::EditRetargetPose);
-	}
+	return GetRetargetProcessor()->IsInitialized();
 }
 
-bool FIKRetargetEditorController::CanEditPose() const
+bool FIKRetargetEditorController::IsCurrentMeshLoaded() const
 {
 	return GetSkeletalMesh(GetSourceOrTarget()) != nullptr;
 }
@@ -1651,4 +1672,71 @@ void FIKRetargetEditorController::OnPoseSelected(TSharedPtr<FName> InPose, ESele
 	}
 }
 
+void FIKRetargetEditorController::ShowGlobalSettings()
+{
+	UIKRetargetGlobalSettings* GlobalSettings = AssetController->GetAsset()->GetGlobalSettingsUObject();
+	if (GlobalSettings->EditorController.Get() != this)
+	{
+		GlobalSettings->EditorController = SharedThis(this);
+	}
+	
+	SetDetailsObject(GlobalSettings);
+}
+
+bool FIKRetargetEditorController::IsShowingGlobalSettings()
+{
+	const UIKRetargetGlobalSettings* GlobalSettings = AssetController->GetAsset()->GetGlobalSettingsUObject();
+	return IsObjectInDetailsView(GlobalSettings);
+}
+
+bool FIKRetargetEditorController::IsShowingRootSettings()
+{
+	const URetargetRootSettings* RootSettings = AssetController->GetAsset()->GetRootSettingsUObject();
+	return IsObjectInDetailsView(RootSettings);
+}
+
+void FIKRetargetEditorController::ToggleRootRetargetPass()
+{
+	UIKRetargetGlobalSettings* GlobalSettings = AssetController->GetAsset()->GetGlobalSettingsUObject();
+	GlobalSettings->Settings.bEnableRoot = !GlobalSettings->Settings.bEnableRoot;
+}
+
+bool FIKRetargetEditorController::IsRootRetargetOn()
+{
+	return AssetController->GetAsset()->GetGlobalSettingsUObject()->Settings.bEnableRoot;
+}
+
+void FIKRetargetEditorController::ToggleFKRetargetPass()
+{
+	UIKRetargetGlobalSettings* GlobalSettings = AssetController->GetAsset()->GetGlobalSettingsUObject();
+	GlobalSettings->Settings.bEnableFK = !GlobalSettings->Settings.bEnableFK;
+}
+
+bool FIKRetargetEditorController::IsFKRetargetOn()
+{
+	return AssetController->GetAsset()->GetGlobalSettingsUObject()->Settings.bEnableFK;
+}
+
+void FIKRetargetEditorController::ToggleIKRetargetPass()
+{
+	UIKRetargetGlobalSettings* GlobalSettings = AssetController->GetAsset()->GetGlobalSettingsUObject();
+	GlobalSettings->Settings.bEnableIK = !GlobalSettings->Settings.bEnableIK;
+}
+
+bool FIKRetargetEditorController::IsIKRetargetOn()
+{
+	return AssetController->GetAsset()->GetGlobalSettingsUObject()->Settings.bEnableIK;
+}
+
+void FIKRetargetEditorController::ShowRootSettings()
+{
+	URetargetRootSettings* RootSettings = AssetController->GetAsset()->GetRootSettingsUObject();
+	if (RootSettings->EditorController.Get() != this)
+	{
+		RootSettings->EditorController = SharedThis(this);	
+	}
+	
+	SetDetailsObject(RootSettings);
+}
+
 #undef LOCTEXT_NAMESPACE
diff --git a/Engine/Plugins/Animation/IKRig/Source/IKRigEditor/Private/RetargetEditor/SIKRetargetAssetBrowser.cpp b/Engine/Plugins/Animation/IKRig/Source/IKRigEditor/Private/RetargetEditor/SIKRetargetAssetBrowser.cpp
--- a/Engine/Plugins/Animation/IKRig/Source/IKRigEditor/Private/RetargetEditor/SIKRetargetAssetBrowser.cpp
+++ b/Engine/Plugins/Animation/IKRig/Source/IKRigEditor/Private/RetargetEditor/SIKRetargetAssetBrowser.cpp
@@ -30,23 +30,32 @@ void SIKRetargetAssetBrowser::Construct(
 	ChildSlot
     [
         SNew(SVerticalBox)
-        
-        + SVerticalBox::Slot()
-		.AutoHeight()
-		.Padding(5)
-		[
+        +SVerticalBox::Slot()
+        .AutoHeight()
+        [
 			SNew(SPositiveActionButton)
 			.IsEnabled(this, &SIKRetargetAssetBrowser::IsExportButtonEnabled)
 			.Icon(FAppStyle::Get().GetBrush("Icons.Save"))
 			.Text(LOCTEXT("ExportButtonLabel", "Export Selected Animations"))
 			.ToolTipText(LOCTEXT("ExportButtonToolTip", "Generate new retargeted sequence assets on target skeletal mesh (uses current retargeting configuration)."))
 			.OnClicked(this, &SIKRetargetAssetBrowser::OnExportButtonClicked)
-		]
+        ]        
 
 		+SVerticalBox::Slot()
 		[
 			SAssignNew(AssetBrowserBox, SBox)
 		]
+		
+		+SVerticalBox::Slot()
+		.AutoHeight()
+		[
+			SNew(SPositiveActionButton)
+			.IsEnabled(this, &SIKRetargetAssetBrowser::IsPlayRefPoseEnabled)
+			.Icon(FAppStyle::Get().GetBrush("GenericStop"))
+			.Text(LOCTEXT("PlayRefPoseLabel", "Play Ref Pose"))
+			.ToolTipText(LOCTEXT("PlayRefPoseTooltip", "Stop playback and set source to reference pose."))
+			.OnClicked(this, &SIKRetargetAssetBrowser::OnPlayRefPoseClicked)
+		]
     ];
 
 	RefreshView();
@@ -203,6 +212,32 @@ bool SIKRetargetAssetBrowser::IsExportButtonEnabled() const
 	return true;
 }
 
+FReply SIKRetargetAssetBrowser::OnPlayRefPoseClicked()
+{
+	FIKRetargetEditorController* Controller = EditorController.Pin().Get();
+	if (!Controller)
+	{
+		checkNoEntry();
+		return FReply::Handled();
+	}
+
+	Controller->SetRetargeterMode(ERetargeterOutputMode::RunRetarget);
+	Controller->PlaybackManager->StopPlayback();
+	return FReply::Handled();
+}
+
+bool SIKRetargetAssetBrowser::IsPlayRefPoseEnabled() const
+{
+	FIKRetargetEditorController* Controller = EditorController.Pin().Get();
+	if (!Controller)
+	{
+		checkNoEntry();
+		return false;
+	}
+	
+	return Controller->GetSkeletalMesh(ERetargetSourceOrTarget::Source) != nullptr;
+}
+
 void SIKRetargetAssetBrowser::OnAssetDoubleClicked(const FAssetData& AssetData)
 {
 	if (!AssetData.GetAsset())
diff --git a/Engine/Plugins/Animation/IKRig/Source/IKRigEditor/Private/RetargetEditor/SIKRetargetChainMapList.cpp b/Engine/Plugins/Animation/IKRig/Source/IKRigEditor/Private/RetargetEditor/SIKRetargetChainMapList.cpp
--- a/Engine/Plugins/Animation/IKRig/Source/IKRigEditor/Private/RetargetEditor/SIKRetargetChainMapList.cpp
+++ b/Engine/Plugins/Animation/IKRig/Source/IKRigEditor/Private/RetargetEditor/SIKRetargetChainMapList.cpp
@@ -230,40 +230,14 @@ void SIKRetargetChainMapList::Construct(
 		.AutoHeight()
 		[
 			SNew(SHorizontalBox)
-
-			+ SHorizontalBox::Slot()
-			.AutoWidth()
-			.Padding(5, 0)
-			.HAlign(HAlign_Left)
-			.VAlign(VAlign_Center)
-			[
-				SNew(SPositiveActionButton)
-				.Icon(FAppStyle::Get().GetBrush("Icons.Settings"))
-				.Text(LOCTEXT("EditGlobalButtonLabel", "Global Settings"))
-				.ToolTipText(LOCTEXT("EditGlobalButtonToolTip", "Edit the global retarget settings."))
-				.OnClicked(this, &SIKRetargetChainMapList::OnGlobalSettingsButtonClicked)
-			]
 			
-			+ SHorizontalBox::Slot()
-			.AutoWidth()
-			.Padding(5, 0)
-			.HAlign(HAlign_Left)
-			.VAlign(VAlign_Center)
-			[
-				SNew(SPositiveActionButton)
-				.Icon(FAppStyle::Get().GetBrush("Icons.Settings"))
-				.Text(LOCTEXT("EditRootButtonLabel", "Root Settings"))
-				.ToolTipText(LOCTEXT("EditRootButtonToolTip", "Edit the root retarget settings."))
-				.OnClicked(this, &SIKRetargetChainMapList::OnRootSettingsButtonClicked)
-			]
-
 			+ SHorizontalBox::Slot()
 			.AutoWidth()
 			.HAlign(HAlign_Left)
 			.VAlign(VAlign_Center)
 			[
 				SNew(STextBlock)
-				.Text(LOCTEXT("TargetRootLabel", "Target: "))
+				.Text(LOCTEXT("TargetRootLabel", "Target Root: "))
 				.TextStyle(FAppStyle::Get(), "NormalText")
 			]
 
@@ -284,7 +258,7 @@ void SIKRetargetChainMapList::Construct(
 			.VAlign(VAlign_Center)
 			[
 				SNew(STextBlock)
-				.Text(LOCTEXT("SourceRootLabel", "Source: "))
+				.Text(LOCTEXT("SourceRootLabel", "Source Root: "))
 				.TextStyle(FAppStyle::Get(), "NormalText")
 			]
 
@@ -772,49 +746,5 @@ void SIKRetargetChainMapList::AutoMapChains(const EAutoMapChainType AutoMapType,
 	Controller->HandleRetargeterNeedsInitialized();
 }
 
-FReply SIKRetargetChainMapList::OnGlobalSettingsButtonClicked() const
-{
-	const TSharedPtr<FIKRetargetEditorController> Controller = EditorController.Pin();
-	if (!Controller.IsValid())
-	{
-		return FReply::Unhandled();
-	}
-
-	UIKRetargeterController* RetargeterController = GetRetargetController();
-	if (!RetargeterController)
-	{
-		return FReply::Unhandled();
-	}
-
-	UIKRetargetGlobalSettings* GlobalSettings = Controller->AssetController->GetAsset()->GetGlobalSettingsUObject();
-	if (GlobalSettings->EditorController != Controller)
-	{
-		GlobalSettings->EditorController = Controller;
-	}
-	
-	Controller->SetDetailsObject(GlobalSettings);
-	
-	return FReply::Handled();
-}
-
-FReply SIKRetargetChainMapList::OnRootSettingsButtonClicked() const
-{
-	const TSharedPtr<FIKRetargetEditorController> Controller = EditorController.Pin();
-	if (!Controller.IsValid())
-	{
-		return FReply::Unhandled();
-	}
-
-	UIKRetargeterController* RetargeterController = GetRetargetController();
-	if (!RetargeterController)
-	{
-		return FReply::Unhandled();
-	}
-	
-	Controller->SetRootSelected(true);
-	return FReply::Handled();
-}
-
-
 #undef LOCTEXT_NAMESPACE
 
diff --git a/Engine/Plugins/Animation/IKRig/Source/IKRigEditor/Private/RetargetEditor/SIKRetargetPoseEditor.cpp b/Engine/Plugins/Animation/IKRig/Source/IKRigEditor/Private/RetargetEditor/SIKRetargetPoseEditor.cpp
--- a/Engine/Plugins/Animation/IKRig/Source/IKRigEditor/Private/RetargetEditor/SIKRetargetPoseEditor.cpp
+++ b/Engine/Plugins/Animation/IKRig/Source/IKRigEditor/Private/RetargetEditor/SIKRetargetPoseEditor.cpp
@@ -26,98 +26,6 @@ void SIKRetargetPoseEditor::Construct(
 	[
 		SNew(SVerticalBox)
 
-		// run/pause retargeter
-		+SVerticalBox::Slot()
-		.Padding(2.0f)
-		.AutoHeight()
-		.HAlign(HAlign_Center)
-		[
-			SNew(SHorizontalBox)
-			+SHorizontalBox::Slot()
-			.AutoWidth()
-			.Padding(4,0)
-			[
-				SNew(SButton)
-				.OnClicked(Controller, &FIKRetargetEditorController::HandleShowRetargetPose)
-				.IsEnabled_Lambda([Controller]
-				{
-					return Controller->GetRetargeterMode() != ERetargeterOutputMode::EditRetargetPose;
-				})
-				.ButtonStyle(FAppStyle::Get(), "FlatButton.Success")
-				.ForegroundColor(FLinearColor::Black)
-				.Content()
-				[
-					SNew(SHorizontalBox)
-					+SHorizontalBox::Slot()
-					.AutoWidth()
-					.HAlign(HAlign_Left)
-					[
-						SNew(SImage)
-						.ColorAndOpacity(FLinearColor::Black)
-						.Image_Lambda([Controller]() -> const FSlateBrush*
-						{
-							if (Controller->GetRetargeterMode() == ERetargeterOutputMode::RunRetarget)
-							{
-								return FAppStyle::GetBrush("GenericPause");
-							}
-							return FAppStyle::GetBrush("GenericPlay");
-						})
-					]
-			
-					+SHorizontalBox::Slot()
-					.HAlign(HAlign_Right)
-					.Padding(4,0)
-					[
-						SNew(STextBlock)
-						.ToolTipText(LOCTEXT("BlendPoseToolTip", "Pause playback and show the retarget pose. Use slider below to blend between reference pose (0) and retarget pose (1)."))
-						.Text_Lambda([Controller]() -> const FText
-						{
-							if (Controller->GetRetargeterMode() == ERetargeterOutputMode::RunRetarget)
-							{
-								return LOCTEXT("ShowPoseText", "Show Retarget Pose");
-							}
-							return LOCTEXT("RunRetargetText", "Run Retargeter");
-						})
-					]
-				]
-			]
-
-			// stop playback button (contextually disabled)
-			+SHorizontalBox::Slot()
-			.Padding(4,0)
-			.AutoWidth()
-			[
-				SNew(SButton)
-				.OnClicked_Lambda([Controller]
-				{
-					Controller->PlaybackManager->StopPlayback();
-					return FReply::Handled();
-				})
-				.IsEnabled_Lambda([Controller]
-				{
-					return !Controller->PlaybackManager->IsStopped();
-				})
-				.Content()
-				[
-					SNew(SImage)
-					.ToolTipText(LOCTEXT("StopPlaybackToolTip", "Stop playback of current animation and show the reference pose."))
-					.Image(FAppStyle::GetBrush("GenericStop"))
-					.ColorAndOpacity_Lambda([Controller]
-					{
-						const bool bRunning = !Controller->PlaybackManager->IsStopped();
-						return bRunning ? FLinearColor::Red : FSlateColor::UseForeground();
-					})
-				]
-			]
-		]
-		
-		+SVerticalBox::Slot()
-		.Padding(2.0f)
-		.AutoHeight()
-		[
-			SNew(SSeparator)
-		]
-
 		+SVerticalBox::Slot()
 		.Padding(2.0f)
 		.AutoHeight()
@@ -204,13 +112,6 @@ TSharedRef<SWidget>  SIKRetargetPoseEditor::MakeToolbar(TSharedPtr<FUICommandLis
 	
 	ToolbarBuilder.BeginSection("Edit Current Pose");
 
-	ToolbarBuilder.AddToolBarButton(
-		FIKRetargetCommands::Get().EditRetargetPose,
-		NAME_None,
-		TAttribute<FText>(),
-		TAttribute<FText>(),
-		FSlateIcon(FAppStyle::Get().GetStyleSetName(),"Icons.Edit"));
-
 	ToolbarBuilder.AddComboButton(
 		FUIAction(),
 		FOnGetContent::CreateSP(this, &SIKRetargetPoseEditor::GenerateResetMenuContent, Commands),
diff --git a/Engine/Plugins/Animation/IKRig/Source/IKRigEditor/Public/RetargetEditor/IKRetargetCommands.h b/Engine/Plugins/Animation/IKRig/Source/IKRigEditor/Public/RetargetEditor/IKRetargetCommands.h
--- a/Engine/Plugins/Animation/IKRig/Source/IKRigEditor/Public/RetargetEditor/IKRetargetCommands.h
+++ b/Engine/Plugins/Animation/IKRig/Source/IKRigEditor/Public/RetargetEditor/IKRetargetCommands.h
@@ -17,44 +17,51 @@ public:
 		FIKRetargetEditorStyle::Get().GetStyleSetName() // Icon Style Set
 	)
 	{}
-
-	/** reset retarget pose */
-	TSharedPtr< FUICommandInfo > ShowRetargetPose;
 	
-	/** edit retarget pose */
+	// run the retargeter
+	TSharedPtr< FUICommandInfo > RunRetargeter;
+	// show current retarget pose
+	TSharedPtr< FUICommandInfo > ShowRetargetPose;
+	// edit retarget pose
 	TSharedPtr< FUICommandInfo > EditRetargetPose;
 
-	/** reset retarget pose to ref pose */
+	// open global settings
+	TSharedPtr< FUICommandInfo > ShowGlobalSettings;
+	// open global settings
+	TSharedPtr< FUICommandInfo > ShowRootSettings;
+
+	// enable Root retarget pass
+	TSharedPtr< FUICommandInfo > EnableRoot;
+	// enable FK retarget pass
+	TSharedPtr< FUICommandInfo > EnableFK;
+	// enable IK retarget pass
+	TSharedPtr< FUICommandInfo > EnableIK;
+
+	// reset retarget pose to ref pose
 	TSharedPtr< FUICommandInfo > ResetAllBones;
-	/** reset retarget pose to ref pose */
+	// reset retarget pose to ref pose
 	TSharedPtr< FUICommandInfo > ResetSelectedBones;
-	/** reset retarget pose to ref pose */
+	// reset retarget pose to ref pose
 	TSharedPtr< FUICommandInfo > ResetSelectedAndChildrenBones;
 
-	/** delete retarget pose */
+	// delete retarget pose
 	TSharedPtr< FUICommandInfo > DeleteRetargetPose;
-
-	/** rename retarget pose */
+	// rename retarget pose
 	TSharedPtr< FUICommandInfo > RenameRetargetPose;
-
-	/** new retarget pose */
+	// new retarget pose
 	TSharedPtr< FUICommandInfo > NewRetargetPose;
-
-	/** duplicate retarget pose */
+	// duplicate retarget pose
 	TSharedPtr< FUICommandInfo > DuplicateRetargetPose;
 
-	/** import retarget pose */
+	// import retarget pose
 	TSharedPtr< FUICommandInfo > ImportRetargetPose;
-
-	/** import retarget pose from anim */
+	// import retarget pose from anim
 	TSharedPtr< FUICommandInfo > ImportRetargetPoseFromAnim;
-
-	/** export retarget pose */
+	// export retarget pose
 	TSharedPtr< FUICommandInfo > ExportRetargetPose;
-
-	/** export animation */
+	// export animation
 	TSharedPtr< FUICommandInfo > ExportAnimation;
 
-	/** initialize commands */
+	// initialize commands
 	virtual void RegisterCommands() override;
 };
diff --git a/Engine/Plugins/Animation/IKRig/Source/IKRigEditor/Public/RetargetEditor/IKRetargetEditor.h b/Engine/Plugins/Animation/IKRig/Source/IKRigEditor/Public/RetargetEditor/IKRetargetEditor.h
--- a/Engine/Plugins/Animation/IKRig/Source/IKRigEditor/Public/RetargetEditor/IKRetargetEditor.h
+++ b/Engine/Plugins/Animation/IKRig/Source/IKRigEditor/Public/RetargetEditor/IKRetargetEditor.h
@@ -88,6 +88,7 @@ private:
 	void BindCommands();
 	void ExtendToolbar();
 	void FillToolbar(FToolBarBuilder& ToolbarBuilder);
+	TSharedRef<SWidget> GenerateRetargetModesMenu();
 	/** END toolbar */
 	
 	/** preview scene setup */
diff --git a/Engine/Plugins/Animation/IKRig/Source/IKRigEditor/Public/RetargetEditor/IKRetargetEditorController.h b/Engine/Plugins/Animation/IKRig/Source/IKRigEditor/Public/RetargetEditor/IKRetargetEditorController.h
--- a/Engine/Plugins/Animation/IKRig/Source/IKRigEditor/Public/RetargetEditor/IKRetargetEditorController.h
+++ b/Engine/Plugins/Animation/IKRig/Source/IKRigEditor/Public/RetargetEditor/IKRetargetEditorController.h
@@ -144,6 +144,7 @@ public:
 	void SetAssetBrowserView(TSharedPtr<SIKRetargetAssetBrowser> InAssetBrowserView) { AssetBrowserView = InAssetBrowserView; };
 	void SetOutputLogView(TSharedPtr<SIKRigOutputLog> InOutputLogView) { OutputLogView = InOutputLogView; };
 	void SetHierarchyView(TSharedPtr<SIKRetargetHierarchy> InHierarchyView) { HierarchyView = InHierarchyView; };
+	bool IsObjectInDetailsView(const UObject* Object);
 	
 	/** force refresh all views in the editor */
 	void RefreshAllViews() const;
@@ -155,6 +156,26 @@ public:
 	void SetDetailsObject(UObject* DetailsObject) const;
 	void SetDetailsObjects(const TArray<UObject*>& DetailsObjects) const;
 
+	/** retargeter state */
+	bool IsReadyToRetarget() const;
+	bool IsCurrentMeshLoaded() const;
+	bool IsEditingPose() const;
+
+	/** display global settings in the details panel */
+	void ShowGlobalSettings();
+	bool IsShowingGlobalSettings();
+	/** display root settings in the details panel */
+	void ShowRootSettings();
+	bool IsShowingRootSettings();
+
+	/** toggle retarget passes */
+	void ToggleRootRetargetPass();
+	bool IsRootRetargetOn();
+	void ToggleFKRetargetPass();
+	bool IsFKRetargetOn();
+	void ToggleIKRetargetPass();
+	bool IsIKRetargetOn();
+
 	/** clear the output log */
 	void ClearOutputLog() const;
 
@@ -190,8 +211,11 @@ public:
 
 	/** Set viewport / editor tool mode */
 	void SetRetargeterMode(ERetargeterOutputMode Mode);
+	void SetRetargetModeToPreviousMode() { SetRetargeterMode(PreviousMode); };
 	ERetargeterOutputMode GetRetargeterMode() const { return OutputMode; }
-	ERetargeterOutputMode PreviousMode;
+	FText GetRetargeterModeLabel();
+	FSlateIcon GetCurrentRetargetModeIcon() const;
+	FSlateIcon GetRetargeterModeIcon(ERetargeterOutputMode Mode) const;
 	float GetRetargetPoseAmount() const;
 	void SetRetargetPoseAmount(float InValue);
 	/** END viewport / editor tool mode */
@@ -246,20 +270,10 @@ public:
 
 	/** ------------------------- RETARGET POSES -----------------------------*/
 	
-	/** go to retarget pose */
-	FReply HandleShowRetargetPose();
-	bool CanShowRetargetPose() const;
-	bool IsShowingRetargetPose() const;
-	
 	/** toggle current retarget pose */
 	TArray<TSharedPtr<FName>> PoseNames;
 	FText GetCurrentPoseName() const;
 	void OnPoseSelected(TSharedPtr<FName> InPoseName, ESelectInfo::Type SelectInfo) const;
-	
-	/** edit retarget poses */
-	void HandleEditPose();
-	bool CanEditPose() const;
-	bool IsEditingPose() const;
 
 	/** reset retarget pose */
 	void HandleResetAllBones() const;
@@ -324,6 +338,7 @@ private:
 
 	/** the current output mode of the retargeter */
 	ERetargeterOutputMode OutputMode = ERetargeterOutputMode::RunRetarget;
+	ERetargeterOutputMode PreviousMode;
 	/** slider value to blend between reference pose and retarget pose */
 	float RetargetPosePreviewBlend = 1.0f;
 	
diff --git a/Engine/Plugins/Animation/IKRig/Source/IKRigEditor/Public/RetargetEditor/SIKRetargetAssetBrowser.h b/Engine/Plugins/Animation/IKRig/Source/IKRigEditor/Public/RetargetEditor/SIKRetargetAssetBrowser.h
--- a/Engine/Plugins/Animation/IKRig/Source/IKRigEditor/Public/RetargetEditor/SIKRetargetAssetBrowser.h
+++ b/Engine/Plugins/Animation/IKRig/Source/IKRigEditor/Public/RetargetEditor/SIKRetargetAssetBrowser.h
@@ -26,6 +26,9 @@ private:
 
 	FReply OnExportButtonClicked();
 	bool IsExportButtonEnabled() const;
+
+	FReply OnPlayRefPoseClicked();
+	bool IsPlayRefPoseEnabled() const;
 	
 	void OnAssetDoubleClicked(const FAssetData& AssetData);
 	bool OnShouldFilterAsset(const struct FAssetData& AssetData);
diff --git a/Engine/Plugins/Animation/IKRig/Source/IKRigEditor/Public/RetargetEditor/SIKRetargetChainMapList.h b/Engine/Plugins/Animation/IKRig/Source/IKRigEditor/Public/RetargetEditor/SIKRetargetChainMapList.h
--- a/Engine/Plugins/Animation/IKRig/Source/IKRigEditor/Public/RetargetEditor/SIKRetargetChainMapList.h
+++ b/Engine/Plugins/Animation/IKRig/Source/IKRigEditor/Public/RetargetEditor/SIKRetargetChainMapList.h
@@ -122,10 +122,6 @@ private:
 	bool IsChainMapEnabled() const;
 	/** when a chain is clicked on in the table view */
 	void OnItemClicked(TSharedPtr<FRetargetChainMapElement> InItem) const;
-	/** when edit global settings button clicked */
-	FReply OnGlobalSettingsButtonClicked() const;
-	/** when edit root settings button clicked */
-	FReply OnRootSettingsButtonClicked() const;
 
 	/** auto-map chain button*/
 	TSharedRef<SWidget> CreateChainMapMenuWidget();
