Make retarget controller thread safe.

diff --git a/Engine/Plugins/Animation/IKRig/Source/IKRigEditor/Private/RetargetEditor/IKRetargeterController.cpp b/Engine/Plugins/Animation/IKRig/Source/IKRigEditor/Private/RetargetEditor/IKRetargeterController.cpp
--- a/Engine/Plugins/Animation/IKRig/Source/IKRigEditor/Private/RetargetEditor/IKRetargeterController.cpp
+++ b/Engine/Plugins/Animation/IKRig/Source/IKRigEditor/Private/RetargetEditor/IKRetargeterController.cpp
@@ -38,6 +38,7 @@ UIKRetargeter* UIKRetargeterController::GetAsset() const
 
 void UIKRetargeterController::CleanAsset() const
 {
+	FScopeLock Lock(&ControllerLock);
 	CleanChainMapping();
 	CleanPoseList(ERetargetSourceOrTarget::Source);
 	CleanPoseList(ERetargetSourceOrTarget::Target);
@@ -45,6 +46,8 @@ void UIKRetargeterController::CleanAsset() const
 
 void UIKRetargeterController::SetIKRig(const ERetargetSourceOrTarget SourceOrTarget, UIKRigDefinition* IKRig) const
 {
+	FScopeLock Lock(&ControllerLock);
+	
 	if (SourceOrTarget == ERetargetSourceOrTarget::Source)
 	{
 		Asset->SourceIKRigAsset = IKRig;
@@ -75,18 +78,22 @@ void UIKRetargeterController::SetIKRig(const ERetargetSourceOrTarget SourceOrTar
 
 const UIKRigDefinition* UIKRetargeterController::GetIKRig(const ERetargetSourceOrTarget SourceOrTarget) const
 {
+	FScopeLock Lock(&ControllerLock);
 	return SourceOrTarget == ERetargetSourceOrTarget::Source ? Asset->GetSourceIKRig() : Asset->GetTargetIKRig();
 }
 
 UIKRigDefinition* UIKRetargeterController::GetIKRigWriteable(const ERetargetSourceOrTarget SourceOrTarget) const
 {
+	FScopeLock Lock(&ControllerLock);
 	return SourceOrTarget == ERetargetSourceOrTarget::Source ? Asset->GetSourceIKRigWriteable() : Asset->GetTargetIKRigWriteable();
 }
 
 void UIKRetargeterController::SetPreviewMesh(
 	const ERetargetSourceOrTarget SourceOrTarget,
 	USkeletalMesh* PreviewMesh) const
-{	
+{
+	FScopeLock Lock(&ControllerLock);
+	
 	if (SourceOrTarget == ERetargetSourceOrTarget::Source)
 	{
 		Asset->SourcePreviewMesh = PreviewMesh;
@@ -106,6 +113,8 @@ void UIKRetargeterController::SetPreviewMesh(
 
 USkeletalMesh* UIKRetargeterController::GetPreviewMesh(const ERetargetSourceOrTarget SourceOrTarget) const
 {
+	FScopeLock Lock(&ControllerLock);
+	
 	// can't preview anything if target IK Rig is null
 	const UIKRigDefinition* IKRig = GetIKRig(SourceOrTarget);
 	if (!IKRig)
@@ -126,27 +135,33 @@ USkeletalMesh* UIKRetargeterController::GetPreviewMesh(const ERetargetSourceOrTa
 
 FTargetRootSettings UIKRetargeterController::GetRootSettings() const
 {
+	FScopeLock Lock(&ControllerLock);
 	return GetAsset()->GetRootSettingsUObject()->Settings;
 }
 
 void UIKRetargeterController::SetRootSettings(const FTargetRootSettings& RootSettings) const
 {
+	FScopeLock Lock(&ControllerLock);
 	GetAsset()->GetRootSettingsUObject()->Settings = RootSettings;
 }
 
 FRetargetGlobalSettings UIKRetargeterController::GetGlobalSettings() const
 {
+	FScopeLock Lock(&ControllerLock);
 	return GetAsset()->GetGlobalSettingsUObject()->Settings;
 }
 
 void UIKRetargeterController::SetGlobalSettings(const FRetargetGlobalSettings& GlobalSettings) const
 {
+	FScopeLock Lock(&ControllerLock);
 	GetAsset()->GetGlobalSettingsUObject()->Settings = GlobalSettings;
 }
 
 FTargetChainSettings UIKRetargeterController::GetRetargetChainSettings(const FName& TargetChainName) const
 {
-	URetargetChainSettings* ChainSettings = GetChainSettings(TargetChainName);
+	FScopeLock Lock(&ControllerLock);
+	
+	const URetargetChainSettings* ChainSettings = GetChainSettings(TargetChainName);
 	if (!ChainSettings)
 	{
 		return FTargetChainSettings();
@@ -157,6 +172,8 @@ FTargetChainSettings UIKRetargeterController::GetRetargetChainSettings(const FNa
 
 bool UIKRetargeterController::SetRetargetChainSettings(const FName& TargetChainName, const FTargetChainSettings& Settings) const
 {
+	FScopeLock Lock(&ControllerLock);
+	
 	if (URetargetChainSettings* ChainSettings = GetChainSettings(TargetChainName))
 	{
 		ChainSettings->Settings = Settings;
@@ -173,6 +190,7 @@ bool UIKRetargeterController::GetAskedToFixRootHeightForMesh(USkeletalMesh* Mesh
 
 void UIKRetargeterController::SetAskedToFixRootHeightForMesh(USkeletalMesh* Mesh, bool InAsked) const
 {
+	FScopeLock Lock(&ControllerLock);
 	if (InAsked)
 	{
 		GetAsset()->MeshesAskedToFixRootHeightFor.Add(Mesh);
@@ -183,9 +201,9 @@ void UIKRetargeterController::SetAskedToFixRootHeightForMesh(USkeletalMesh* Mesh
 	}
 }
 
-FName UIKRetargeterController::GetRetargetRootBone(
-	const ERetargetSourceOrTarget SourceOrTarget) const
+FName UIKRetargeterController::GetRetargetRootBone(const ERetargetSourceOrTarget SourceOrTarget) const
 {
+	FScopeLock Lock(&ControllerLock);
 	const UIKRigDefinition* IKRig = GetIKRig(SourceOrTarget);
 	return IKRig ? IKRig->GetRetargetRoot() : FName("None");
 }
@@ -298,6 +316,7 @@ void UIKRetargeterController::CleanPoseList(const ERetargetSourceOrTarget Source
 
 void UIKRetargeterController::AutoMapChains(const EAutoMapChainType AutoMapType, const bool bForceRemap) const
 {
+	FScopeLock Lock(&ControllerLock);
 	FScopedTransaction Transaction(LOCTEXT("AutoMapRetargetChains", "Auto-Map Retarget Chains"));
 	
 	CleanChainMapping();
@@ -442,6 +461,7 @@ void UIKRetargeterController::HandleRetargetChainRemoved(UIKRigDefinition* IKRig
 
 bool UIKRetargeterController::SetSourceChain(FName SourceChainName, FName TargetChainName) const
 {
+	FScopeLock Lock(&ControllerLock);
 	URetargetChainSettings* ChainSettings = GetChainSettings(TargetChainName);
 	if (!ChainSettings)
 	{
@@ -506,6 +526,7 @@ bool UIKRetargeterController::IsChainGoalConnectedToASolver(const FName& GoalNam
 
 FName UIKRetargeterController::CreateRetargetPose(const FName& NewPoseName, const ERetargetSourceOrTarget SourceOrTarget) const
 {
+	FScopeLock Lock(&ControllerLock);
 	FScopedTransaction Transaction(LOCTEXT("CreateRetargetPose", "Create Retarget Pose"));
 	Asset->Modify();
 
@@ -535,6 +556,7 @@ bool UIKRetargeterController::RemoveRetargetPose(const FName& PoseToRemove, cons
 		return false; // cannot remove pose that doesn't exist
 	}
 
+	FScopeLock Lock(&ControllerLock);
 	FScopedTransaction Transaction(LOCTEXT("RemoveRetargetPose", "Remove Retarget Pose"));
 	Asset->Modify();
 
@@ -559,7 +581,8 @@ FName UIKRetargeterController::DuplicateRetargetPose( const FName PoseToDuplicat
 		UE_LOG(LogIKRigEditor, Warning, TEXT("Trying to duplicate pose that does not exist, %s."), *PoseToDuplicate.ToString());
 		return NAME_None; // cannot duplicate pose that doesn't exist
 	}
-	
+
+	FScopeLock Lock(&ControllerLock);
 	FScopedTransaction Transaction(LOCTEXT("DuplicateRetargetPose", "Duplicate Retarget Pose"));
 	Asset->Modify();
 
@@ -581,6 +604,8 @@ FName UIKRetargeterController::DuplicateRetargetPose( const FName PoseToDuplicat
 
 bool UIKRetargeterController::RenameRetargetPose(const FName OldPoseName, const FName NewPoseName, const ERetargetSourceOrTarget SourceOrTarget) const
 {
+	FScopeLock Lock(&ControllerLock);
+	
 	// does the old pose exist?
 	if (!GetRetargetPoses(SourceOrTarget).Contains(OldPoseName))
 	{
@@ -627,6 +652,8 @@ void UIKRetargeterController::ResetRetargetPose(
 	const TArray<FName>& BonesToReset,
 	const ERetargetSourceOrTarget SourceOrTarget) const
 {
+	FScopeLock Lock(&ControllerLock);
+	
 	TMap<FName, FIKRetargetPose>& Poses = GetRetargetPoses(SourceOrTarget);
 	if (!Poses.Contains(PoseToReset))
 	{
@@ -668,18 +695,21 @@ void UIKRetargeterController::ResetRetargetPose(
 
 FName UIKRetargeterController::GetCurrentRetargetPoseName(const ERetargetSourceOrTarget SourceOrTarget) const
 {
+	FScopeLock Lock(&ControllerLock);
 	return SourceOrTarget == ERetargetSourceOrTarget::Source ? GetAsset()->CurrentSourceRetargetPose : GetAsset()->CurrentTargetRetargetPose;
 }
 
 bool UIKRetargeterController::SetCurrentRetargetPose(FName NewCurrentPose, const ERetargetSourceOrTarget SourceOrTarget) const
 {
+	FScopeLock Lock(&ControllerLock);
+	
 	const TMap<FName, FIKRetargetPose>& Poses = GetRetargetPoses(SourceOrTarget);
 	if (!Poses.Contains(NewCurrentPose))
 	{
 		UE_LOG(LogIKRigEditor, Warning, TEXT("Trying to set current pose to a pose that does not exist, %s."), *NewCurrentPose.ToString());
 		return false;
 	}
-
+	
 	FScopedTransaction Transaction(LOCTEXT("SetCurrentPose", "Set Current Pose"));
 	Asset->Modify();
 	FName& CurrentPose = SourceOrTarget == ERetargetSourceOrTarget::Source ? Asset->CurrentSourceRetargetPose : Asset->CurrentTargetRetargetPose;
@@ -692,11 +722,13 @@ bool UIKRetargeterController::SetCurrentRetargetPose(FName NewCurrentPose, const
 
 TMap<FName, FIKRetargetPose>& UIKRetargeterController::GetRetargetPoses(const ERetargetSourceOrTarget SourceOrTarget) const
 {
+	FScopeLock Lock(&ControllerLock);
 	return SourceOrTarget == ERetargetSourceOrTarget::Source ? GetAsset()->SourceRetargetPoses : GetAsset()->TargetRetargetPoses;
 }
 
 FIKRetargetPose& UIKRetargeterController::GetCurrentRetargetPose(const ERetargetSourceOrTarget SourceOrTarget) const
 {
+	FScopeLock Lock(&ControllerLock);
 	return GetRetargetPoses(SourceOrTarget)[GetCurrentRetargetPoseName(SourceOrTarget)];
 }
 
@@ -705,6 +737,7 @@ void UIKRetargeterController::SetRotationOffsetForRetargetPoseBone(
 	const FQuat& RotationOffset,
 	const ERetargetSourceOrTarget SourceOrTarget) const
 {
+	FScopeLock Lock(&ControllerLock);
 	const UIKRigDefinition* IKRig = SourceOrTarget == ERetargetSourceOrTarget::Source ? GetAsset()->GetSourceIKRig() : GetAsset()->GetTargetIKRig();
 	FIKRetargetPose& Pose = GetCurrentRetargetPose(SourceOrTarget);
 	Pose.SetDeltaRotationForBone(BoneName, RotationOffset);
@@ -715,6 +748,8 @@ FQuat UIKRetargeterController::GetRotationOffsetForRetargetPoseBone(
 	const FName& BoneName,
 	const ERetargetSourceOrTarget SourceOrTarget) const
 {
+	FScopeLock Lock(&ControllerLock);
+	
 	TMap<FName, FQuat>& BoneOffsets = GetCurrentRetargetPose(SourceOrTarget).BoneRotationOffsets;
 	if (!BoneOffsets.Contains(BoneName))
 	{
@@ -728,12 +763,14 @@ void UIKRetargeterController::SetRootOffsetInRetargetPose(
 	const FVector& TranslationOffset,
 	const ERetargetSourceOrTarget SourceOrTarget) const
 {
+	FScopeLock Lock(&ControllerLock);
 	GetCurrentRetargetPose(SourceOrTarget).AddToRootTranslationDelta(TranslationOffset);
 }
 
 FVector UIKRetargeterController::GetRootOffsetInRetargetPose(
 	const ERetargetSourceOrTarget SourceOrTarget) const
 {
+	FScopeLock Lock(&ControllerLock);
 	return GetCurrentRetargetPose(SourceOrTarget).GetRootTranslationDelta();
 }
 
diff --git a/Engine/Plugins/Animation/IKRig/Source/IKRigEditor/Public/RetargetEditor/IKRetargeterController.h b/Engine/Plugins/Animation/IKRig/Source/IKRigEditor/Public/RetargetEditor/IKRetargeterController.h
--- a/Engine/Plugins/Animation/IKRig/Source/IKRigEditor/Public/RetargetEditor/IKRetargeterController.h
+++ b/Engine/Plugins/Animation/IKRig/Source/IKRigEditor/Public/RetargetEditor/IKRetargeterController.h
@@ -50,23 +50,23 @@ private:
 public:
 	
 	/** Use this to get the controller for the given retargeter asset */
-	UFUNCTION(BlueprintCallable, Category=IKRetargeter)
+	UFUNCTION(BlueprintCallable, Category=IKRetargeter, meta = (BlueprintThreadSafe))
 	static UIKRetargeterController* GetController(const UIKRetargeter* InRetargeterAsset);
 	
 	// Set the IK Rig to use as the source or target (to copy animation FROM/TO) 
-	UFUNCTION(BlueprintCallable, BlueprintPure=false, Category=IKRetargeter)
+	UFUNCTION(BlueprintCallable, BlueprintPure=false, Category=IKRetargeter, meta = (BlueprintThreadSafe))
 	void SetIKRig(const ERetargetSourceOrTarget SourceOrTarget, UIKRigDefinition* IKRig) const;
 	
 	// Get either source or target IK Rig 
-	UFUNCTION(BlueprintCallable, Category=IKRetargeter)
+	UFUNCTION(BlueprintCallable, Category=IKRetargeter, meta = (BlueprintThreadSafe))
 	const UIKRigDefinition* GetIKRig(const ERetargetSourceOrTarget SourceOrTarget) const;
 
 	// Set the preview skeletal mesh for either source or target
-	UFUNCTION(BlueprintCallable, BlueprintPure=false, Category=IKRetargeter)
+	UFUNCTION(BlueprintCallable, BlueprintPure=false, Category=IKRetargeter, meta = (BlueprintThreadSafe))
 	void SetPreviewMesh(const ERetargetSourceOrTarget SourceOrTarget, USkeletalMesh* InPreviewMesh) const;
 	
 	// Get the preview skeletal mesh
-	UFUNCTION(BlueprintCallable, Category=IKRetargeter)
+	UFUNCTION(BlueprintCallable, Category=IKRetargeter, meta = (BlueprintThreadSafe))
 	USkeletalMesh* GetPreviewMesh(const ERetargetSourceOrTarget SourceOrTarget) const;
 
 	//
@@ -74,27 +74,27 @@ public:
 	//
 	
 	// Get a copy of the retarget root settings for this asset.
-	UFUNCTION(BlueprintCallable, BlueprintPure=false, Category=IKRetargeter)
+	UFUNCTION(BlueprintCallable, BlueprintPure=false, Category=IKRetargeter, meta = (BlueprintThreadSafe))
 	FTargetRootSettings GetRootSettings() const;
 	
 	// Set the retarget root settings for this asset.
-	UFUNCTION(BlueprintCallable, BlueprintPure=false, Category=IKRetargeter)
+	UFUNCTION(BlueprintCallable, BlueprintPure=false, Category=IKRetargeter, meta = (BlueprintThreadSafe))
 	void SetRootSettings(const FTargetRootSettings& RootSettings) const;
 
 	// Get a copy of the global settings for this asset.
-	UFUNCTION(BlueprintCallable, BlueprintPure=false, Category=IKRetargeter)
+	UFUNCTION(BlueprintCallable, BlueprintPure=false, Category=IKRetargeter, meta = (BlueprintThreadSafe))
 	FRetargetGlobalSettings GetGlobalSettings() const;
 
 	// Get a copy of the global settings for this asset.
-	UFUNCTION(BlueprintCallable, BlueprintPure=false, Category=IKRetargeter)
+	UFUNCTION(BlueprintCallable, BlueprintPure=false, Category=IKRetargeter, meta = (BlueprintThreadSafe))
 	void SetGlobalSettings(const FRetargetGlobalSettings& GlobalSettings) const;
 
 	// Get a copy of the settings for the target chain by name.
-	UFUNCTION(BlueprintCallable, BlueprintPure=false, Category=IKRetargeter)
+	UFUNCTION(BlueprintCallable, BlueprintPure=false, Category=IKRetargeter, meta = (BlueprintThreadSafe))
 	FTargetChainSettings GetRetargetChainSettings(const FName& TargetChainName) const;
 
 	// Set the settings for the target chain by name. Returns true if the chain settings were applied, false otherwise.
-	UFUNCTION(BlueprintCallable, BlueprintPure=false, Category=IKRetargeter)
+	UFUNCTION(BlueprintCallable, BlueprintPure=false, Category=IKRetargeter, meta = (BlueprintThreadSafe))
 	bool SetRetargetChainSettings(const FName& TargetChainName, const FTargetChainSettings& Settings) const;
 
 	//
@@ -123,19 +123,19 @@ public:
 	//
 	
 	// Use fuzzy string search to find "best" Source chain to map to each Target chain
-	UFUNCTION(BlueprintCallable, BlueprintPure=false, Category=IKRetargeter)
+	UFUNCTION(BlueprintCallable, BlueprintPure=false, Category=IKRetargeter, meta = (BlueprintThreadSafe))
 	void AutoMapChains(const EAutoMapChainType AutoMapType, const bool bForceRemap) const;
 
 	// Assign a source chain to the given target chain. Animation will be copied from the source to the target.
-	UFUNCTION(BlueprintCallable, BlueprintPure=false, Category=IKRetargeter)
+	UFUNCTION(BlueprintCallable, BlueprintPure=false, Category=IKRetargeter, meta = (BlueprintThreadSafe))
 	bool SetSourceChain(FName SourceChainName, FName TargetChainName) const;
 
 	// Get the name of the source chain mapped to a given target chain (the chain animation is copied FROM).
-	UFUNCTION(BlueprintCallable, Category=IKRetargeter)
+	UFUNCTION(BlueprintCallable, Category=IKRetargeter, meta = (BlueprintThreadSafe))
 	FName GetSourceChain(const FName& TargetChainName) const;
 	
 	// Get read-only access to the list of settings for each target chain
-	UFUNCTION(BlueprintCallable, Category=IKRetargeter)
+	UFUNCTION(BlueprintCallable, Category=IKRetargeter, meta = (BlueprintThreadSafe))
 	const TArray<URetargetChainSettings*>& GetAllChainSettings() const;
 
 	//
@@ -179,63 +179,63 @@ private:
 public:
 	
 	// Add new retarget pose. Returns the name of the new retarget pose.
-	UFUNCTION(BlueprintCallable, BlueprintPure=false, Category=IKRetargeter)
+	UFUNCTION(BlueprintCallable, BlueprintPure=false, Category=IKRetargeter, meta =(BlueprintThreadSafe))
 	FName CreateRetargetPose(const FName& NewPoseName, const ERetargetSourceOrTarget SourceOrTarget) const;
 
 	// Remove a retarget pose. Returns true if the pose was found and removed.
-	UFUNCTION(BlueprintCallable, BlueprintPure=false, Category=IKRetargeter)
+	UFUNCTION(BlueprintCallable, BlueprintPure=false, Category=IKRetargeter, meta =(BlueprintThreadSafe))
 	bool RemoveRetargetPose(const FName& PoseToRemove, const ERetargetSourceOrTarget SourceOrTarget) const;
 
 	// Duplicate a retarget pose. Returns the name of the new, duplicate pose (or NAME_None if PoseToDuplicate is not found).
-	UFUNCTION(BlueprintCallable, BlueprintPure=false, Category=IKRetargeter)
+	UFUNCTION(BlueprintCallable, BlueprintPure=false, Category=IKRetargeter, meta =(BlueprintThreadSafe))
 	FName DuplicateRetargetPose(const FName PoseToDuplicate, const FName NewName, const ERetargetSourceOrTarget SourceOrTarget) const;
 	
 	// Rename current retarget pose. Returns true if a pose was renamed.
-	UFUNCTION(BlueprintCallable, BlueprintPure=false, Category=IKRetargeter)
+	UFUNCTION(BlueprintCallable, BlueprintPure=false, Category=IKRetargeter, meta =(BlueprintThreadSafe))
 	bool RenameRetargetPose(const FName OldPoseName, const FName NewPoseName, const ERetargetSourceOrTarget SourceOrTarget) const;
 
 	// Change which retarget pose is used by the retargeter at runtime
-	UFUNCTION(BlueprintCallable, BlueprintPure=false, Category=IKRetargeter)
+	UFUNCTION(BlueprintCallable, BlueprintPure=false, Category=IKRetargeter, meta =(BlueprintThreadSafe))
 	bool SetCurrentRetargetPose(FName CurrentPose, const ERetargetSourceOrTarget SourceOrTarget) const;
 	
 	// Get the current retarget pose
-	UFUNCTION(BlueprintCallable, Category=IKRetargeter)
+	UFUNCTION(BlueprintCallable, Category=IKRetargeter, meta =(BlueprintThreadSafe))
     FName GetCurrentRetargetPoseName(const ERetargetSourceOrTarget SourceOrTarget) const;
 	
 	// Get access to array of retarget poses
-	UFUNCTION(BlueprintCallable, Category=IKRetargeter)
+	UFUNCTION(BlueprintCallable, Category=IKRetargeter, meta =(BlueprintThreadSafe))
 	TMap<FName, FIKRetargetPose>& GetRetargetPoses(const ERetargetSourceOrTarget SourceOrTarget) const;
 	
 	// Get the current retarget pose
-	UFUNCTION(BlueprintCallable, Category=IKRetargeter)
+	UFUNCTION(BlueprintCallable, Category=IKRetargeter, meta =(BlueprintThreadSafe))
 	FIKRetargetPose& GetCurrentRetargetPose(const ERetargetSourceOrTarget SourceOrTarget) const;
 
 	// Reset a retarget pose for the specified bones.
 	// If BonesToReset is Empty, will removes all stored deltas, returning pose to reference pose
-	UFUNCTION(BlueprintCallable, BlueprintPure=false, Category=IKRetargeter)
+	UFUNCTION(BlueprintCallable, BlueprintPure=false, Category=IKRetargeter, meta =(BlueprintThreadSafe))
 	void ResetRetargetPose(const FName& PoseToReset, const TArray<FName>& BonesToReset, const ERetargetSourceOrTarget SourceOrTarget) const;
 	
 	// Set a delta rotation for a given bone for the current retarget pose
-	UFUNCTION(BlueprintCallable, BlueprintPure=false, Category=IKRetargeter)
+	UFUNCTION(BlueprintCallable, BlueprintPure=false, Category=IKRetargeter, meta =(BlueprintThreadSafe))
 	void SetRotationOffsetForRetargetPoseBone(
 		const FName& BoneName,
 		const FQuat& RotationOffset,
 		const ERetargetSourceOrTarget SkeletonMode) const;
 	
 	// Get a delta rotation for a given bone for the current retarget pose
-	UFUNCTION(BlueprintCallable, BlueprintPure=false, Category=IKRetargeter)
+	UFUNCTION(BlueprintCallable, BlueprintPure=false, Category=IKRetargeter, meta =(BlueprintThreadSafe))
 	FQuat GetRotationOffsetForRetargetPoseBone(
 		const FName& BoneName,
 		const ERetargetSourceOrTarget SourceOrTarget) const;
 	
 	// Set the translation offset on the retarget root bone for the current retarget pose
-	UFUNCTION(BlueprintCallable, BlueprintPure=false, Category=IKRetargeter)
+	UFUNCTION(BlueprintCallable, BlueprintPure=false, Category=IKRetargeter, meta =(BlueprintThreadSafe))
 	void SetRootOffsetInRetargetPose(
 		const FVector& TranslationOffset,
 		const ERetargetSourceOrTarget SourceOrTarget) const;
 
 	// Get the translation offset on the retarget root bone for the current retarget pose
-	UFUNCTION(BlueprintCallable, Category=IKRetargeter)
+	UFUNCTION(BlueprintCallable, Category=IKRetargeter, meta = (BlueprintThreadSafe))
 	FVector GetRootOffsetInRetargetPose(const ERetargetSourceOrTarget SourceOrTarget) const;
 
 	//
@@ -282,6 +282,9 @@ private:
 	{
 		RetargeterNeedsInitialized.Broadcast();
 	}
+
+	// only allow modifications to data model from one thread at a time
+	mutable FCriticalSection ControllerLock;
 	
 public:
 
