Fixing bug in IK Rig editor where adding a new goal adds it to every solver in the stack.

diff --git a/Engine/Plugins/Animation/IKRig/Source/IKRig/Private/Rig/IKRigProcessor.cpp b/Engine/Plugins/Animation/IKRig/Source/IKRig/Private/Rig/IKRigProcessor.cpp
--- a/Engine/Plugins/Animation/IKRig/Source/IKRig/Private/Rig/IKRigProcessor.cpp
+++ b/Engine/Plugins/Animation/IKRig/Source/IKRig/Private/Rig/IKRigProcessor.cpp
@@ -108,6 +108,23 @@ void UIKRigProcessor::Initialize(
 				return;
 			}
 		}
+
+		// warn if goal is not connected to any solvers
+		bool bIsGoalConnectedToAnySolver = false;
+		const TArray<UIKRigSolver*>& AssetSolvers = InRigAsset->GetSolverArray();
+		for (const UIKRigSolver* IKRigSolver : AssetSolvers)
+		{
+			if (IKRigSolver->IsGoalConnected(EffectorGoal->GoalName))
+			{
+				bIsGoalConnectedToAnySolver = true;
+			}
+		}
+		if (!bIsGoalConnectedToAnySolver)
+		{
+			Log.LogWarning(FText::Format(
+				LOCTEXT("DisconnectedGoal", "IK Rig, '{0}' has a Goal, '{1}' that is not connected to any solvers. It will have no effect."),
+				FText::FromString(InRigAsset->GetName()), FText::FromName(EffectorGoal->GoalName)));
+		}
 		
 		GoalBones.Add(EffectorGoal->GoalName, NewGoalBone);
 	}
diff --git a/Engine/Plugins/Animation/IKRig/Source/IKRigEditor/Private/RigEditor/IKRigController.cpp b/Engine/Plugins/Animation/IKRig/Source/IKRigEditor/Private/RigEditor/IKRigController.cpp
--- a/Engine/Plugins/Animation/IKRig/Source/IKRigEditor/Private/RigEditor/IKRigController.cpp
+++ b/Engine/Plugins/Animation/IKRig/Source/IKRigEditor/Private/RigEditor/IKRigController.cpp
@@ -917,12 +917,6 @@ FName UIKRigController::AddNewGoal(const FName GoalName, const FName BoneName) c
 	// set initial transform
 	NewGoal->InitialTransform = GetRefPoseTransformOfBone(NewGoal->BoneName);
 	NewGoal->CurrentTransform = NewGoal->InitialTransform;
-
-	// connect the new goal to all the solvers
-	for (int32 SolverIndex=0; SolverIndex<GetNumSolvers(); ++SolverIndex)
-	{
-		ConnectGoalToSolver(NewGoal->GoalName, SolverIndex);
-	}
  
 	BroadcastNeedsReinitialized();
 	BroadcastGoalsChange();
diff --git a/Engine/Plugins/Animation/IKRig/Source/IKRigEditor/Private/RigEditor/IKRigEditorController.cpp b/Engine/Plugins/Animation/IKRig/Source/IKRigEditor/Private/RigEditor/IKRigEditorController.cpp
--- a/Engine/Plugins/Animation/IKRig/Source/IKRigEditor/Private/RigEditor/IKRigEditorController.cpp
+++ b/Engine/Plugins/Animation/IKRig/Source/IKRigEditor/Private/RigEditor/IKRigEditorController.cpp
@@ -506,6 +506,14 @@ void FIKRigEditorController::AddNewGoals(const TArray<FName>& GoalNames, const T
 		// ask user if they want to assign this goal to a chain (if there is one on this bone)
 		PromptToAssignGoalToChain(NewGoalName);
 
+		// connect the new goal to selected solvers
+		TArray<TSharedPtr<FSolverStackElement>> SelectedSolvers;
+		GetSelectedSolvers(SelectedSolvers);
+		for (TSharedPtr<FSolverStackElement>& Solver : SelectedSolvers)
+		{
+			AssetController->ConnectGoalToSolver(NewGoalName, Solver.Get()->IndexInStack);
+		}
+
 		LastCreatedGoalName = GoalName;
 	}
 	
@@ -553,7 +561,7 @@ void FIKRigEditorController::HandleBoneSelectedInViewport(const FName& BoneName,
 	ShowDetailsForBone(BoneName);
 }
 
-void FIKRigEditorController::GetSelectedSolvers(TArray<TSharedPtr<FSolverStackElement>>& OutSelectedSolvers)
+void FIKRigEditorController::GetSelectedSolvers(TArray<TSharedPtr<FSolverStackElement>>& OutSelectedSolvers) const
 {
 	if (SolverStackView.IsValid())
 	{
@@ -1068,6 +1076,14 @@ FName FIKRigEditorController::PromptToAddNewRetargetChain(FBoneChain& BoneChain)
 			// create new goal
 			const FName NewGoalName = FName(FText::Format(LOCTEXT("GoalOnNewChainName", "{0}_Goal"), FText::FromName(BoneChain.ChainName)).ToString());
 			GoalName = AssetController->AddNewGoal(NewGoalName, BoneChain.EndBone.BoneName);
+
+			// connect the new goal to selected solvers
+			TArray<TSharedPtr<FSolverStackElement>> SelectedSolvers;
+			GetSelectedSolvers(SelectedSolvers);
+			for (TSharedPtr<FSolverStackElement>& Solver : SelectedSolvers)
+			{
+				AssetController->ConnectGoalToSolver(GoalName, Solver.Get()->IndexInStack);
+			}
 		}
 		
 		// assign the existing goal to it
diff --git a/Engine/Plugins/Animation/IKRig/Source/IKRigEditor/Public/RigEditor/IKRigEditorController.h b/Engine/Plugins/Animation/IKRig/Source/IKRigEditor/Public/RigEditor/IKRigEditorController.h
--- a/Engine/Plugins/Animation/IKRig/Source/IKRigEditor/Public/RigEditor/IKRigEditorController.h
+++ b/Engine/Plugins/Animation/IKRig/Source/IKRigEditor/Public/RigEditor/IKRigEditorController.h
@@ -142,7 +142,7 @@ public:
 	void ClearOutputLog() const;
 
 	// return list of those solvers in the stack that are selected by user 
-	void GetSelectedSolvers(TArray<TSharedPtr<FSolverStackElement> >& OutSelectedSolvers);
+	void GetSelectedSolvers(TArray<TSharedPtr<FSolverStackElement> >& OutSelectedSolvers) const;
 	// get index of the first selected solver, return INDEX_None if nothing selected 
 	int32 GetSelectedSolverIndex();
 	// get names of all goals that are selected 
