Adding asset action menu to create IK Rig assets from Skeletal Meshes. And to create IK Retargeter from an IK Rig.

diff --git a/Engine/Plugins/Animation/IKRig/Source/IKRigEditor/Private/IKRigEditor.cpp b/Engine/Plugins/Animation/IKRig/Source/IKRigEditor/Private/IKRigEditor.cpp
--- a/Engine/Plugins/Animation/IKRig/Source/IKRigEditor/Private/IKRigEditor.cpp
+++ b/Engine/Plugins/Animation/IKRig/Source/IKRigEditor/Private/IKRigEditor.cpp
@@ -43,7 +43,9 @@ void FIKRigEditor::StartupModule()
 	ToolsModule.RegisterAssetTypeActions(IKRetargeterAssetAction.ToSharedRef());
 
 	// extend the content browser menu
-	FAssetTypeActions_IKRetargeter::ExtendAnimSequenceToolMenu();
+	FAssetTypeActions_IKRetargeter::ExtendAnimAssetMenusForBatchRetargeting();
+	FAssetTypeActions_IKRetargeter::ExtendIKRigMenuToMakeRetargeter();
+	FAssetTypeActions_IKRigDefinition::ExtendSkeletalMeshMenuToMakeIKRig();
 
 	// register custom editor modes
 	FEditorModeRegistry::Get().RegisterMode<FIKRigEditMode>(FIKRigEditMode::ModeName, LOCTEXT("IKRigEditMode", "IKRig"), FSlateIcon(), false);
diff --git a/Engine/Plugins/Animation/IKRig/Source/IKRigEditor/Private/RetargetEditor/AssetTypeActions_IKRetargeter.cpp b/Engine/Plugins/Animation/IKRig/Source/IKRigEditor/Private/RetargetEditor/AssetTypeActions_IKRetargeter.cpp
--- a/Engine/Plugins/Animation/IKRig/Source/IKRigEditor/Private/RetargetEditor/AssetTypeActions_IKRetargeter.cpp
+++ b/Engine/Plugins/Animation/IKRig/Source/IKRigEditor/Private/RetargetEditor/AssetTypeActions_IKRetargeter.cpp
@@ -7,9 +7,11 @@
 #include "ToolMenus.h"
 #include "RetargetEditor/IKRetargetEditor.h"
 #include "RetargetEditor/IKRetargetEditorStyle.h"
+#include "RetargetEditor/IKRetargetFactory.h"
 #include "RetargetEditor/SRetargetAnimAssetsWindow.h"
 #include "Retargeter/IKRetargeter.h"
 #include "ThumbnailRendering/SceneThumbnailInfo.h"
+#include "RigEditor/IKRigEditorStyle.h"
 
 #define LOCTEXT_NAMESPACE "AssetTypeActions"
 
@@ -40,7 +42,88 @@ UThumbnailInfo* FAssetTypeActions_IKRetargeter::GetThumbnailInfo(UObject* Asset)
 	return NewObject<USceneThumbnailInfo>(IKRetargeter, NAME_None, RF_Transactional);
 }
 
-void FAssetTypeActions_IKRetargeter::ExtendAnimSequenceToolMenu()
+void FAssetTypeActions_IKRetargeter::ExtendIKRigMenuToMakeRetargeter()
+{
+	UToolMenu* Menu = UToolMenus::Get()->ExtendMenu("ContentBrowser.AssetContextMenu.IKRigDefinition");
+	if (Menu == nullptr)
+	{
+		return;
+	}
+
+	FToolMenuSection& Section = Menu->FindOrAddSection("GetAssetActions");
+	Section.AddDynamicEntry("CreateIKRetargeter", FNewToolMenuSectionDelegate::CreateLambda([](FToolMenuSection& InSection)
+	{
+		UContentBrowserAssetContextMenuContext* Context = InSection.FindContext<UContentBrowserAssetContextMenuContext>();
+		if (Context)
+		{
+			TArray<UObject*> SelectedObjects = Context->GetSelectedObjects();
+			if (SelectedObjects.Num() > 0)
+			{
+				InSection.AddMenuEntry(
+					"CreateRetargeter",
+					LOCTEXT("CreateIKRetargeter", "Create IK Retargeter"),
+					LOCTEXT("CreateIKRetargeter_ToolTip", "Creates an IK Retargeter using this IK Rig as the source."),
+					FSlateIcon(FIKRigEditorStyle::Get().GetStyleSetName(), "IKRig", "ClassIcon.IKRigDefinition"),
+					FExecuteAction::CreateLambda([SelectedObjects]()
+					{
+						for (UObject* SelectedObject : SelectedObjects)
+						{
+							CreateNewIKRetargeterFromIKRig(SelectedObject);
+						}
+					})
+				);
+			}
+		}
+	}));
+}
+
+void FAssetTypeActions_IKRetargeter::CreateNewIKRetargeterFromIKRig(UObject* InSelectedObject)
+{
+	// validate ik rig
+	UIKRigDefinition* IKRigAsset = Cast<UIKRigDefinition>(InSelectedObject);
+	if(!IKRigAsset)
+	{
+		UE_LOG(LogTemp, Error, TEXT("CreateNewIKRetargeterFromIKRig: Provided object has to be an IK Rig."));
+		return;
+	}
+	
+	// remove the standard "IK_" prefix which is standard for these asset types (convenience)
+	FString AssetName = InSelectedObject->GetName();
+	if (AssetName.Len() > 3)
+	{
+		AssetName.RemoveFromStart("IK_");	
+	}
+
+	// create unique package and asset names
+	FString UniquePackageName;
+	FString UniqueAssetName = FString::Printf(TEXT("RTG_%s"), *AssetName);
+	
+	FString PackagePath = InSelectedObject->GetPathName(); 
+	int32 LastSlashPos = INDEX_NONE;
+	if (PackagePath.FindLastChar('/', LastSlashPos))
+	{
+		PackagePath = PackagePath.Left(LastSlashPos);
+	}
+	
+	PackagePath = PackagePath / UniqueAssetName;
+	const FAssetToolsModule& AssetToolsModule = FModuleManager::LoadModuleChecked<FAssetToolsModule>("AssetTools");
+	AssetToolsModule.Get().CreateUniqueAssetName(PackagePath, TEXT(""), UniquePackageName, UniqueAssetName);
+	if (UniquePackageName.EndsWith(UniqueAssetName))
+	{
+		UniquePackageName = UniquePackageName.LeftChop(UniqueAssetName.Len() + 1);
+	}
+	
+	// create the new IK Retargeter asset
+	UIKRetargetFactory* Factory = NewObject<UIKRetargetFactory>();
+	UObject* NewAsset = AssetToolsModule.Get().CreateAsset(*UniqueAssetName, *UniquePackageName, nullptr, Factory);
+	const UIKRetargeter* RetargetAsset = Cast<UIKRetargeter>(NewAsset);
+
+	// assign ik rig to the retargeter
+	const UIKRetargeterController* Controller = UIKRetargeterController::GetController(RetargetAsset);
+	Controller->SetIKRig(ERetargetSourceOrTarget::Source, IKRigAsset);
+}
+
+void FAssetTypeActions_IKRetargeter::ExtendAnimAssetMenusForBatchRetargeting()
 {
 	static const TArray<FName> MenusToExtend 
 	{
diff --git a/Engine/Plugins/Animation/IKRig/Source/IKRigEditor/Private/RigEditor/AssetTypeActions_IKRigDefinition.cpp b/Engine/Plugins/Animation/IKRig/Source/IKRigEditor/Private/RigEditor/AssetTypeActions_IKRigDefinition.cpp
--- a/Engine/Plugins/Animation/IKRig/Source/IKRigEditor/Private/RigEditor/AssetTypeActions_IKRigDefinition.cpp
+++ b/Engine/Plugins/Animation/IKRig/Source/IKRigEditor/Private/RigEditor/AssetTypeActions_IKRigDefinition.cpp
@@ -4,6 +4,10 @@
 #include "Rig/IKRigDefinition.h"
 #include "RigEditor/IKRigToolkit.h"
 #include "ThumbnailRendering/SceneThumbnailInfo.h"
+#include "ContentBrowserMenuContexts.h"
+#include "RigEditor/IKRigDefinitionFactory.h"
+#include "RigEditor/IKRigEditorStyle.h"
+#include "Engine/SkeletalMesh.h"
 
 #define LOCTEXT_NAMESPACE "AssetTypeActions"
 
@@ -19,6 +23,87 @@ UThumbnailInfo* FAssetTypeActions_IKRigDefinition::GetThumbnailInfo(UObject* Ass
 	return NewObject<USceneThumbnailInfo>(IKRig, NAME_None, RF_Transactional);
 }
 
+void FAssetTypeActions_IKRigDefinition::ExtendSkeletalMeshMenuToMakeIKRig()
+{
+	UToolMenu* Menu = UToolMenus::Get()->ExtendMenu("ContentBrowser.AssetContextMenu.SkeletalMesh.CreateSkeletalMeshSubmenu");
+	if (Menu == nullptr)
+	{
+		return;
+	}
+
+	FToolMenuSection& Section = Menu->AddSection("IKRig", LOCTEXT("IKRigSectionName", "IK Rig"));
+	Section.AddDynamicEntry("CreateIKRig", FNewToolMenuSectionDelegate::CreateLambda([](FToolMenuSection& InSection)
+	{
+		UContentBrowserAssetContextMenuContext* Context = InSection.FindContext<UContentBrowserAssetContextMenuContext>();
+		if (Context)
+		{
+			TArray<UObject*> SelectedObjects = Context->GetSelectedObjects();
+			if (SelectedObjects.Num() > 0)
+			{
+				InSection.AddMenuEntry(
+					"CreateIKRig",
+					LOCTEXT("CreateIKRig", "IK Rig"),
+					LOCTEXT("CreateIKRig_ToolTip", "Creates an IK rig for this skeletal mesh."),
+					FSlateIcon(FIKRigEditorStyle::Get().GetStyleSetName(), "IKRig", "ClassIcon.IKRigDefinition"),
+					FExecuteAction::CreateLambda([SelectedObjects]()
+					{
+						for (UObject* SelectedObject : SelectedObjects)
+						{
+							CreateNewIKRigFromSkeletalMesh(SelectedObject);
+						}
+					})
+				);
+			}
+		}
+	}));
+}
+
+void FAssetTypeActions_IKRigDefinition::CreateNewIKRigFromSkeletalMesh(UObject* InSelectedObject)
+{
+	// validate skeletal mesh
+	USkeletalMesh* SkeletalMesh = Cast<USkeletalMesh>(InSelectedObject);
+	if(!SkeletalMesh)
+	{
+		UE_LOG(LogTemp, Error, TEXT("CreateNewIKRigFromSkeletalMesh: Provided object has to be a SkeletalMesh."));
+		return;
+	}
+
+	// remove the standard "SK_" prefix which is standard for these asset types (convenience)
+	FString AssetName = InSelectedObject->GetName();
+	if (AssetName.Len() > 3)
+	{
+		AssetName.RemoveFromStart("SK_");	
+	}
+
+	// create unique package and asset names
+	FString UniquePackageName;
+	FString UniqueAssetName = FString::Printf(TEXT("IK_%s"), *AssetName);
+	
+	FString PackagePath = InSelectedObject->GetPathName(); 
+	int32 LastSlashPos = INDEX_NONE;
+	if (PackagePath.FindLastChar('/', LastSlashPos))
+	{
+		PackagePath = PackagePath.Left(LastSlashPos);
+	}
+	
+	PackagePath = PackagePath / UniqueAssetName;
+	const FAssetToolsModule& AssetToolsModule = FModuleManager::LoadModuleChecked<FAssetToolsModule>("AssetTools");
+	AssetToolsModule.Get().CreateUniqueAssetName(PackagePath, TEXT(""), UniquePackageName, UniqueAssetName);
+	if (UniquePackageName.EndsWith(UniqueAssetName))
+	{
+		UniquePackageName = UniquePackageName.LeftChop(UniqueAssetName.Len() + 1);
+	}
+	
+	// create the new IK Rig asset
+	UIKRigDefinitionFactory* Factory = NewObject<UIKRigDefinitionFactory>();
+	UObject* NewAsset = AssetToolsModule.Get().CreateAsset(*UniqueAssetName, *UniquePackageName, nullptr, Factory);
+	const UIKRigDefinition* IKRig = Cast<UIKRigDefinition>(NewAsset);
+
+	// assign skeletal mesh to the ik rig
+	const UIKRigController* Controller = UIKRigController::GetController(IKRig);
+	Controller->SetSkeletalMesh(SkeletalMesh);
+}
+
 void FAssetTypeActions_IKRigDefinition::OpenAssetEditor(
 	const TArray<UObject*>& InObjects,
 	TSharedPtr<IToolkitHost> EditWithinLevelEditor)
diff --git a/Engine/Plugins/Animation/IKRig/Source/IKRigEditor/Public/RetargetEditor/AssetTypeActions_IKRetargeter.h b/Engine/Plugins/Animation/IKRig/Source/IKRigEditor/Public/RetargetEditor/AssetTypeActions_IKRetargeter.h
--- a/Engine/Plugins/Animation/IKRig/Source/IKRigEditor/Public/RetargetEditor/AssetTypeActions_IKRetargeter.h
+++ b/Engine/Plugins/Animation/IKRig/Source/IKRigEditor/Public/RetargetEditor/AssetTypeActions_IKRetargeter.h
@@ -21,6 +21,9 @@ public:
 	virtual const TArray<FText>& GetSubMenus() const override;
 	virtual UThumbnailInfo* GetThumbnailInfo(UObject* Asset) const override;
 
-	static void ExtendAnimSequenceToolMenu();
+	static void ExtendIKRigMenuToMakeRetargeter();
+	static void CreateNewIKRetargeterFromIKRig(UObject* InSelectedObject);
+	
+	static void ExtendAnimAssetMenusForBatchRetargeting();
 	static void CreateRetargetSubMenu(FToolMenuSection& InSection);
 };
diff --git a/Engine/Plugins/Animation/IKRig/Source/IKRigEditor/Public/RigEditor/AssetTypeActions_IKRigDefinition.h b/Engine/Plugins/Animation/IKRig/Source/IKRigEditor/Public/RigEditor/AssetTypeActions_IKRigDefinition.h
--- a/Engine/Plugins/Animation/IKRig/Source/IKRigEditor/Public/RigEditor/AssetTypeActions_IKRigDefinition.h
+++ b/Engine/Plugins/Animation/IKRig/Source/IKRigEditor/Public/RigEditor/AssetTypeActions_IKRigDefinition.h
@@ -18,4 +18,7 @@ public:
 	virtual const TArray<FText>& GetSubMenus() const override;
 	virtual UThumbnailInfo* GetThumbnailInfo(UObject* Asset) const override;
 	// END IAssetTypeActions Implementation
+
+	static void ExtendSkeletalMeshMenuToMakeIKRig();
+	static void CreateNewIKRigFromSkeletalMesh(UObject* InSelectedObject);
 };
