Adding thumbnail rendering to the IK Rig/Retarget assets.

diff --git a/Engine/Plugins/Animation/IKRig/Source/IKRigEditor/Private/IKRigEditor.cpp b/Engine/Plugins/Animation/IKRig/Source/IKRigEditor/Private/IKRigEditor.cpp
--- a/Engine/Plugins/Animation/IKRig/Source/IKRigEditor/Private/IKRigEditor.cpp
+++ b/Engine/Plugins/Animation/IKRig/Source/IKRigEditor/Private/IKRigEditor.cpp
@@ -14,11 +14,13 @@
 #include "RetargetEditor/IKRetargetDefaultMode.h"
 #include "RetargetEditor/IKRetargetDetails.h"
 #include "RetargetEditor/IKRetargetEditPoseMode.h"
+#include "RetargetEditor/IKRetargeterThumbnailRenderer.h"
 #include "RigEditor/IKRigCommands.h"
 #include "RigEditor/IKRigEditMode.h"
 #include "RigEditor/IKRigSkeletonCommands.h"
 #include "RigEditor/IKRigDetailCustomizations.h"
 #include "RigEditor/IKRigEditorController.h"
+#include "RigEditor/IKRigThumbnailRenderer.h"
 
 DEFINE_LOG_CATEGORY(LogIKRigEditor);
 
@@ -72,6 +74,10 @@ void FIKRigEditor::StartupModule()
 	// custom retargeter global details
 	PropertyEditorModule.RegisterCustomClassLayout(UIKRetargetGlobalSettings::StaticClass()->GetFName(), FOnGetDetailCustomizationInstance::CreateStatic(&FRetargetGlobalSettingsCustomization::MakeInstance));
 	ClassesToUnregisterOnShutdown.Add(UIKRetargetGlobalSettings::StaticClass()->GetFName());
+
+	// register a thumbnail renderer for the assets
+	UThumbnailManager::Get().RegisterCustomRenderer(UIKRigDefinition::StaticClass(), UIKRigThumbnailRenderer::StaticClass());
+	UThumbnailManager::Get().RegisterCustomRenderer(UIKRetargeter::StaticClass(), UIKRetargeterThumbnailRenderer::StaticClass());
 }
 
 void FIKRigEditor::ShutdownModule()
@@ -99,6 +105,9 @@ void FIKRigEditor::ShutdownModule()
 			ToolsModule.UnregisterAssetTypeActions(IKRetargeterAssetAction.ToSharedRef());
 		}
 	}
+
+	UThumbnailManager::Get().UnregisterCustomRenderer(UIKRigDefinition::StaticClass());
+	UThumbnailManager::Get().UnregisterCustomRenderer(UIKRetargeter::StaticClass());
 }
 
 #undef LOCTEXT_NAMESPACE
\ No newline at end of file
diff --git a/Engine/Plugins/Animation/IKRig/Source/IKRigEditor/Private/RetargetEditor/IKRetargeterThumbnailRenderer.cpp b/Engine/Plugins/Animation/IKRig/Source/IKRigEditor/Private/RetargetEditor/IKRetargeterThumbnailRenderer.cpp
new file mode 100644
--- /dev/null
+++ b/Engine/Plugins/Animation/IKRig/Source/IKRigEditor/Private/RetargetEditor/IKRetargeterThumbnailRenderer.cpp
@@ -0,0 +1,55 @@
+﻿// Copyright Epic Games, Inc. All Rights Reserved.
+
+#include "RetargetEditor/IKRetargeterThumbnailRenderer.h"
+
+#include "Engine/SkeletalMesh.h"
+#include "Retargeter/IKRetargeter.h"
+
+#include UE_INLINE_GENERATED_CPP_BY_NAME(IKRetargeterThumbnailRenderer)
+
+bool UIKRetargeterThumbnailRenderer::CanVisualizeAsset(UObject* Object)
+{
+	return GetPreviewMeshFromAsset(Object) != nullptr;
+}
+
+void UIKRetargeterThumbnailRenderer::Draw(UObject* Object, int32 X, int32 Y, uint32 Width, uint32 Height, FRenderTarget* RenderTarget, FCanvas* Canvas, bool bAdditionalViewFamily)
+{
+	if (USkeletalMesh* MeshToDraw = GetPreviewMeshFromAsset(Object))
+	{
+		Super::Draw(MeshToDraw, X, Y, Width, Height, RenderTarget, Canvas, bAdditionalViewFamily);	
+	}
+}
+
+EThumbnailRenderFrequency UIKRetargeterThumbnailRenderer::GetThumbnailRenderFrequency(UObject* Object) const
+{
+	return GetPreviewMeshFromAsset(Object) ? EThumbnailRenderFrequency::Realtime : EThumbnailRenderFrequency::OnPropertyChange;
+}
+
+USkeletalMesh* UIKRetargeterThumbnailRenderer::GetPreviewMeshFromAsset(UObject* Object) const
+{
+	const UIKRetargeter* InRetargeter = Cast<UIKRetargeter>(Object);
+	if (!InRetargeter)
+	{
+		return nullptr;
+	}
+
+	// prefer target mesh for thumbnail
+	if (const UIKRigDefinition* TargetIKRig =  InRetargeter->GetTargetIKRig())
+	{
+		if (USkeletalMesh* TargetMesh = TargetIKRig->GetPreviewMesh())
+		{
+			return TargetMesh;
+		}
+	}
+
+	// fallback to source mesh if no target applied
+	if (const UIKRigDefinition* SourceIKRig =  InRetargeter->GetSourceIKRig())
+	{
+		if (USkeletalMesh* SourceMesh = SourceIKRig->GetPreviewMesh())
+		{
+			return SourceMesh;
+		}
+	}
+
+	return nullptr;
+}
diff --git a/Engine/Plugins/Animation/IKRig/Source/IKRigEditor/Private/RigEditor/IKRigThumbnailRenderer.cpp b/Engine/Plugins/Animation/IKRig/Source/IKRigEditor/Private/RigEditor/IKRigThumbnailRenderer.cpp
new file mode 100644
--- /dev/null
+++ b/Engine/Plugins/Animation/IKRig/Source/IKRigEditor/Private/RigEditor/IKRigThumbnailRenderer.cpp
@@ -0,0 +1,37 @@
+﻿// Copyright Epic Games, Inc. All Rights Reserved.
+
+#include "RigEditor/IKRigThumbnailRenderer.h"
+
+#include "Engine/SkeletalMesh.h"
+#include "Rig/IKRigDefinition.h"
+
+#include UE_INLINE_GENERATED_CPP_BY_NAME(IKRigThumbnailRenderer)
+
+bool UIKRigThumbnailRenderer::CanVisualizeAsset(UObject* Object)
+{
+	return GetPreviewMeshFromRig(Object) != nullptr;
+}
+
+void UIKRigThumbnailRenderer::Draw(UObject* Object, int32 X, int32 Y, uint32 Width, uint32 Height, FRenderTarget* RenderTarget, FCanvas* Canvas, bool bAdditionalViewFamily)
+{
+	if (USkeletalMesh* MeshToDraw = GetPreviewMeshFromRig(Object))
+	{
+		Super::Draw(MeshToDraw, X, Y, Width, Height, RenderTarget, Canvas, bAdditionalViewFamily);	
+	}
+}
+
+EThumbnailRenderFrequency UIKRigThumbnailRenderer::GetThumbnailRenderFrequency(UObject* Object) const
+{
+	return GetPreviewMeshFromRig(Object) ? EThumbnailRenderFrequency::Realtime : EThumbnailRenderFrequency::OnPropertyChange;
+}
+
+USkeletalMesh* UIKRigThumbnailRenderer::GetPreviewMeshFromRig(UObject* Object) const
+{
+	const UIKRigDefinition* InRig = Cast<UIKRigDefinition>(Object);
+	if (!InRig)
+	{
+		return nullptr;
+	}
+
+	return InRig->GetPreviewMesh();
+}
diff --git a/Engine/Plugins/Animation/IKRig/Source/IKRigEditor/Public/RetargetEditor/IKRetargeterThumbnailRenderer.h b/Engine/Plugins/Animation/IKRig/Source/IKRigEditor/Public/RetargetEditor/IKRetargeterThumbnailRenderer.h
new file mode 100644
--- /dev/null
+++ b/Engine/Plugins/Animation/IKRig/Source/IKRigEditor/Public/RetargetEditor/IKRetargeterThumbnailRenderer.h
@@ -0,0 +1,26 @@
+﻿// Copyright Epic Games, Inc. All Rights Reserved.
+
+#pragma once
+
+#include "UObject/ObjectMacros.h"
+#include "ThumbnailRendering/SkeletalMeshThumbnailRenderer.h"
+#include "IKRetargeterThumbnailRenderer.generated.h"
+
+class FCanvas;
+class FRenderTarget;
+
+// this thumbnail renderer displays a given IK Rig in the asset icon
+UCLASS(config=Editor, MinimalAPI)
+class UIKRetargeterThumbnailRenderer : public USkeletalMeshThumbnailRenderer
+{
+	GENERATED_BODY()
+
+	// Begin UThumbnailRenderer Object
+	IKRIGEDITOR_API virtual bool CanVisualizeAsset(UObject* Object) override;
+	IKRIGEDITOR_API virtual void Draw(UObject* Object, int32 X, int32 Y, uint32 Width, uint32 Height, FRenderTarget*, FCanvas* Canvas, bool bAdditionalViewFamily) override;
+	IKRIGEDITOR_API virtual EThumbnailRenderFrequency GetThumbnailRenderFrequency(UObject* Object) const override;
+	// End UThumbnailRenderer Object
+	
+	USkeletalMesh* GetPreviewMeshFromAsset(UObject* Object) const;
+};
+
diff --git a/Engine/Plugins/Animation/IKRig/Source/IKRigEditor/Public/RigEditor/IKRigThumbnailRenderer.h b/Engine/Plugins/Animation/IKRig/Source/IKRigEditor/Public/RigEditor/IKRigThumbnailRenderer.h
new file mode 100644
--- /dev/null
+++ b/Engine/Plugins/Animation/IKRig/Source/IKRigEditor/Public/RigEditor/IKRigThumbnailRenderer.h
@@ -0,0 +1,27 @@
+﻿// Copyright Epic Games, Inc. All Rights Reserved.
+
+#pragma once
+
+#include "UObject/ObjectMacros.h"
+#include "ThumbnailRendering/SkeletalMeshThumbnailRenderer.h"
+#include "IKRigThumbnailRenderer.generated.h"
+
+class FCanvas;
+class FRenderTarget;
+
+// this thumbnail renderer displays a given IK Rig in the asset icon
+UCLASS(config=Editor, MinimalAPI)
+class UIKRigThumbnailRenderer : public USkeletalMeshThumbnailRenderer
+{
+	GENERATED_BODY()
+
+	// Begin UThumbnailRenderer Object
+	IKRIGEDITOR_API virtual bool CanVisualizeAsset(UObject* Object) override;
+	IKRIGEDITOR_API virtual void Draw(UObject* Object, int32 X, int32 Y, uint32 Width, uint32 Height, FRenderTarget*, FCanvas* Canvas, bool bAdditionalViewFamily) override;
+	IKRIGEDITOR_API virtual EThumbnailRenderFrequency GetThumbnailRenderFrequency(UObject* Object) const override;
+	// End UThumbnailRenderer Object
+
+private:
+	USkeletalMesh* GetPreviewMeshFromRig(UObject* Object) const;
+};
+
