Fixes crash when unloading IK Rig module when unregistering thumbnail renderer without UObject system being initialized.

diff --git a/Engine/Plugins/Animation/IKRig/Source/IKRigEditor/Private/IKRigEditor.cpp b/Engine/Plugins/Animation/IKRig/Source/IKRigEditor/Private/IKRigEditor.cpp
--- a/Engine/Plugins/Animation/IKRig/Source/IKRigEditor/Private/IKRigEditor.cpp
+++ b/Engine/Plugins/Animation/IKRig/Source/IKRigEditor/Private/IKRigEditor.cpp
@@ -106,8 +106,11 @@ void FIKRigEditor::ShutdownModule()
 		}
 	}
 
-	UThumbnailManager::Get().UnregisterCustomRenderer(UIKRigDefinition::StaticClass());
-	UThumbnailManager::Get().UnregisterCustomRenderer(UIKRetargeter::StaticClass());
+	if (UObjectInitialized())
+	{
+		UThumbnailManager::Get().UnregisterCustomRenderer(UIKRigDefinition::StaticClass());
+		UThumbnailManager::Get().UnregisterCustomRenderer(UIKRetargeter::StaticClass());
+	}
 }
 
 #undef LOCTEXT_NAMESPACE
\ No newline at end of file
