Rendering both source and target in IK Retargeter thumbnail render.

diff --git a/Engine/Plugins/Animation/IKRig/Source/IKRigEditor/Private/RetargetEditor/IKRetargeterThumbnailRenderer.cpp b/Engine/Plugins/Animation/IKRig/Source/IKRigEditor/Private/RetargetEditor/IKRetargeterThumbnailRenderer.cpp
--- a/Engine/Plugins/Animation/IKRig/Source/IKRigEditor/Private/RetargetEditor/IKRetargeterThumbnailRenderer.cpp
+++ b/Engine/Plugins/Animation/IKRig/Source/IKRigEditor/Private/RetargetEditor/IKRetargeterThumbnailRenderer.cpp
@@ -2,30 +2,151 @@
 
 #include "RetargetEditor/IKRetargeterThumbnailRenderer.h"
 
+#include "SceneView.h"
+#include "Components/SkeletalMeshComponent.h"
+#include "Animation/SkeletalMeshActor.h"
 #include "Engine/SkeletalMesh.h"
+#include "RetargetEditor/IKRetargeterController.h"
 #include "Retargeter/IKRetargeter.h"
+#include "ThumbnailRendering/SceneThumbnailInfo.h"
 
 #include UE_INLINE_GENERATED_CPP_BY_NAME(IKRetargeterThumbnailRenderer)
 
-bool UIKRetargeterThumbnailRenderer::CanVisualizeAsset(UObject* Object)
+FIKRetargeterThumbnailScene::FIKRetargeterThumbnailScene()
 {
-	return GetPreviewMeshFromAsset(Object) != nullptr;
+	bForceAllUsedMipsResident = false;
+
+	FActorSpawnParameters SpawnInfo;
+	SpawnInfo.SpawnCollisionHandlingOverride = ESpawnActorCollisionHandlingMethod::AlwaysSpawn;
+	SpawnInfo.bNoFail = true;
+	SpawnInfo.ObjectFlags = RF_Transient;
+	
+	SourceActor = GetWorld()->SpawnActor<ASkeletalMeshActor>(SpawnInfo);
+	SourceActor->SetActorEnableCollision(false);
+	
+	TargetActor = GetWorld()->SpawnActor<ASkeletalMeshActor>(SpawnInfo);
+	TargetActor->SetActorEnableCollision(false);
 }
 
-void UIKRetargeterThumbnailRenderer::Draw(UObject* Object, int32 X, int32 Y, uint32 Width, uint32 Height, FRenderTarget* RenderTarget, FCanvas* Canvas, bool bAdditionalViewFamily)
+void FIKRetargeterThumbnailScene::SetSkeletalMeshes(USkeletalMesh* SourceMesh, USkeletalMesh* TargetMesh) const
 {
-	if (USkeletalMesh* MeshToDraw = GetPreviewMeshFromAsset(Object))
+	auto AssignMeshToActor = [this](USkeletalMesh* InMesh, ERetargetSourceOrTarget SourceOrTarget, bool bApplyOffset)
 	{
-		Super::Draw(MeshToDraw, X, Y, Width, Height, RenderTarget, Canvas, bAdditionalViewFamily);	
+		const bool bIsSourceMesh = SourceOrTarget == ERetargetSourceOrTarget::Source;
+		ASkeletalMeshActor* PreviewActor = bIsSourceMesh ? SourceActor : TargetActor;
+		PreviewActor->GetSkeletalMeshComponent()->OverrideMaterials.Empty();
+		PreviewActor->GetSkeletalMeshComponent()->SetSkeletalMesh(InMesh, false);
+
+		if (InMesh)
+		{
+			// center the mesh at the world origin
+			PreviewActor->SetActorLocation(FVector(0,0,0), false);
+			PreviewActor->GetSkeletalMeshComponent()->UpdateBounds();
+
+			// X offset it sideways to put source/target side-by-side
+			FVector Offset = FVector::ZeroVector;
+			if (bApplyOffset)
+			{
+				const FBoxSphereBounds Bounds = InMesh->GetImportedBounds();
+				const float HalfWidth = (Bounds.BoxExtent.X - Bounds.Origin.X) * 0.5f;
+				Offset.X = bIsSourceMesh ? -HalfWidth : HalfWidth;
+			}
+		
+			// Z offset to put it on top of the plane
+			Offset.Z = GetBoundsZOffset(PreviewActor->GetSkeletalMeshComponent()->Bounds);
+			const FVector MoveToOrigin = -PreviewActor->GetSkeletalMeshComponent()->Bounds.Origin;
+			PreviewActor->SetActorLocation(MoveToOrigin + Offset, false);
+			PreviewActor->GetSkeletalMeshComponent()->UpdateBounds();
+			PreviewActor->GetSkeletalMeshComponent()->RecreateRenderState_Concurrent();
+		}
+		else
+		{
+			PreviewActor->GetSkeletalMeshComponent()->ClearAnimScriptInstance();
+		}
+	};
+
+	const bool bApplyOffset = SourceMesh && TargetMesh;
+	AssignMeshToActor(SourceMesh, ERetargetSourceOrTarget::Source, bApplyOffset);
+	AssignMeshToActor(TargetMesh, ERetargetSourceOrTarget::Target, bApplyOffset);
+}
+
+void FIKRetargeterThumbnailScene::GetViewMatrixParameters(
+	const float InFOVDegrees,
+	FVector& OutOrigin,
+	float& OutOrbitPitch,
+	float& OutOrbitYaw, 
+	float& OutOrbitZoom) const
+{
+	const USkeletalMeshComponent* SourceComponent = SourceActor->GetSkeletalMeshComponent();
+	const USkeletalMeshComponent* TargetComponent = TargetActor->GetSkeletalMeshComponent();
+	check(SourceComponent);
+	check(TargetComponent);
+
+	// get bounds of both meshes combined
+	const FBoxSphereBounds SourceBounds = SourceComponent->Bounds;
+	const FBoxSphereBounds TargetBounds = TargetComponent->Bounds;
+	const FBoxSphereBounds TotalBounds = TargetBounds + SourceBounds;
+
+	// calc target distance
+	const float HalfBoundSize = TotalBounds.GetBox().GetSize().Size() * 0.5f;
+	const float HalfFOVRadians = FMath::DegreesToRadians<float>(InFOVDegrees) * 0.5f;
+	const float TargetDistance = HalfBoundSize / FMath::Tan(HalfFOVRadians);
+
+	// get thumbnail info from mesh (preferring target if available)
+	const USceneThumbnailInfo* ThumbnailInfo = nullptr;
+	if(const USkeletalMesh* TargetMesh = TargetComponent->GetSkeletalMeshAsset())
+	{
+		ThumbnailInfo = Cast<USceneThumbnailInfo>(TargetMesh->GetThumbnailInfo());
 	}
+	else if (const USkeletalMesh* SourceMesh = SourceComponent->GetSkeletalMeshAsset())
+	{
+		ThumbnailInfo = Cast<USceneThumbnailInfo>(SourceMesh->GetThumbnailInfo());
+	}
+	else
+	{
+		ThumbnailInfo = USceneThumbnailInfo::StaticClass()->GetDefaultObject<USceneThumbnailInfo>();
+	}
+
+	OutOrigin = -TotalBounds.GetBox().GetCenter();
+	OutOrbitPitch = ThumbnailInfo->OrbitPitch;
+	OutOrbitYaw = ThumbnailInfo->OrbitYaw;
+	OutOrbitZoom = TargetDistance + ThumbnailInfo->OrbitZoom;
+}
+
+bool UIKRetargeterThumbnailRenderer::CanVisualizeAsset(UObject* Object)
+{
+	return HasSourceOrTargetMesh(Object);
+}
+
+void UIKRetargeterThumbnailRenderer::Draw(UObject* Object, int32 X, int32 Y, uint32 Width, uint32 Height, FRenderTarget* RenderTarget, FCanvas* Canvas, bool bAdditionalViewFamily)
+{
+	TSharedRef<FIKRetargeterThumbnailScene> ThumbnailScene = ThumbnailSceneCache.EnsureThumbnailScene(Object);
+
+	USkeletalMesh* SourceMesh = GetPreviewMeshFromAsset(Object, ERetargetSourceOrTarget::Source);
+	USkeletalMesh* TargetMesh = GetPreviewMeshFromAsset(Object, ERetargetSourceOrTarget::Target);
+
+	ThumbnailScene->SetSkeletalMeshes(SourceMesh, TargetMesh);
+
+	FSceneViewFamilyContext ViewFamily( FSceneViewFamily::ConstructionValues( RenderTarget, ThumbnailScene->GetScene(), FEngineShowFlags(ESFIM_Game) )
+	.SetTime(UThumbnailRenderer::GetTime())
+	.SetAdditionalViewFamily(bAdditionalViewFamily));
+
+	ViewFamily.EngineShowFlags.DisableAdvancedFeatures();
+	ViewFamily.EngineShowFlags.MotionBlur = 0;
+	ViewFamily.EngineShowFlags.LOD = 0;
+
+	RenderViewFamily(Canvas, &ViewFamily, ThumbnailScene->CreateView(&ViewFamily, X, Y, Width, Height));
+	
+	ThumbnailScene->SetSkeletalMeshes(nullptr, nullptr);
 }
 
 EThumbnailRenderFrequency UIKRetargeterThumbnailRenderer::GetThumbnailRenderFrequency(UObject* Object) const
 {
-	return GetPreviewMeshFromAsset(Object) ? EThumbnailRenderFrequency::Realtime : EThumbnailRenderFrequency::OnPropertyChange;
+	const bool bHasSourceOrTarget = HasSourceOrTargetMesh(Object);
+	return bHasSourceOrTarget ? EThumbnailRenderFrequency::Realtime : EThumbnailRenderFrequency::OnPropertyChange;
 }
 
-USkeletalMesh* UIKRetargeterThumbnailRenderer::GetPreviewMeshFromAsset(UObject* Object) const
+USkeletalMesh* UIKRetargeterThumbnailRenderer::GetPreviewMeshFromAsset(UObject* Object, ERetargetSourceOrTarget SourceOrTarget) const
 {
 	const UIKRetargeter* InRetargeter = Cast<UIKRetargeter>(Object);
 	if (!InRetargeter)
@@ -33,23 +154,13 @@ USkeletalMesh* UIKRetargeterThumbnailRenderer::GetPreviewMeshFromAsset(UObject*
 		return nullptr;
 	}
 
-	// prefer target mesh for thumbnail
-	if (const UIKRigDefinition* TargetIKRig =  InRetargeter->GetTargetIKRig())
-	{
-		if (USkeletalMesh* TargetMesh = TargetIKRig->GetPreviewMesh())
-		{
-			return TargetMesh;
-		}
-	}
-
-	// fallback to source mesh if no target applied
-	if (const UIKRigDefinition* SourceIKRig =  InRetargeter->GetSourceIKRig())
-	{
-		if (USkeletalMesh* SourceMesh = SourceIKRig->GetPreviewMesh())
-		{
-			return SourceMesh;
-		}
-	}
+	const UIKRetargeterController* Controller = UIKRetargeterController::GetController(InRetargeter);
+	return Controller->GetPreviewMesh(SourceOrTarget);
+}
 
-	return nullptr;
+bool UIKRetargeterThumbnailRenderer::HasSourceOrTargetMesh(UObject* Object) const
+{
+	const USkeletalMesh* SourceMesh = GetPreviewMeshFromAsset(Object, ERetargetSourceOrTarget::Source);
+	const USkeletalMesh* TargetMesh = GetPreviewMeshFromAsset(Object, ERetargetSourceOrTarget::Target);
+	return SourceMesh || TargetMesh;
 }
diff --git a/Engine/Plugins/Animation/IKRig/Source/IKRigEditor/Public/RetargetEditor/IKRetargeterThumbnailRenderer.h b/Engine/Plugins/Animation/IKRig/Source/IKRigEditor/Public/RetargetEditor/IKRetargeterThumbnailRenderer.h
--- a/Engine/Plugins/Animation/IKRig/Source/IKRigEditor/Public/RetargetEditor/IKRetargeterThumbnailRenderer.h
+++ b/Engine/Plugins/Animation/IKRig/Source/IKRigEditor/Public/RetargetEditor/IKRetargeterThumbnailRenderer.h
@@ -8,10 +8,32 @@
 
 class FCanvas;
 class FRenderTarget;
+class ASkeletalMeshActor;
+enum class ERetargetSourceOrTarget : uint8;
+
+class IKRIGEDITOR_API FIKRetargeterThumbnailScene : public FThumbnailPreviewScene
+{
+public:
+	
+	FIKRetargeterThumbnailScene();
+
+	// sets the retargeter to use in the next CreateView()
+	void SetSkeletalMeshes(USkeletalMesh* SourceMesh, USkeletalMesh* TargetMesh) const;
+
+protected:
+	// FThumbnailPreviewScene implementation
+	virtual void GetViewMatrixParameters(const float InFOVDegrees, FVector& OutOrigin, float& OutOrbitPitch, float& OutOrbitYaw, float& OutOrbitZoom) const override;
+
+private:
+	// the skeletal mesh actors used to display all skeletal mesh thumbnails
+	ASkeletalMeshActor* SourceActor;
+	ASkeletalMeshActor* TargetActor;
+};
+
 
 // this thumbnail renderer displays a given IK Rig in the asset icon
 UCLASS(config=Editor, MinimalAPI)
-class UIKRetargeterThumbnailRenderer : public USkeletalMeshThumbnailRenderer
+class UIKRetargeterThumbnailRenderer : public UDefaultSizedThumbnailRenderer
 {
 	GENERATED_BODY()
 
@@ -21,6 +43,10 @@ class UIKRetargeterThumbnailRenderer : public USkeletalMeshThumbnailRenderer
 	IKRIGEDITOR_API virtual EThumbnailRenderFrequency GetThumbnailRenderFrequency(UObject* Object) const override;
 	// End UThumbnailRenderer Object
 	
-	USkeletalMesh* GetPreviewMeshFromAsset(UObject* Object) const;
+	USkeletalMesh* GetPreviewMeshFromAsset(UObject* Object, ERetargetSourceOrTarget SourceOrTarget) const;
+	bool HasSourceOrTargetMesh(UObject* Object) const;
+
+protected:
+	TObjectInstanceThumbnailScene<FIKRetargeterThumbnailScene, 128> ThumbnailSceneCache;
 };
 
