Fixes issue when duplicating IK Retargeter asset loses references to IK rig assets.

diff --git a/Engine/Plugins/Animation/IKRig/Source/IKRig/Private/Retargeter/IKRetargeter.cpp b/Engine/Plugins/Animation/IKRig/Source/IKRig/Private/Retargeter/IKRetargeter.cpp
--- a/Engine/Plugins/Animation/IKRig/Source/IKRig/Private/Retargeter/IKRetargeter.cpp
+++ b/Engine/Plugins/Animation/IKRig/Source/IKRig/Private/Retargeter/IKRetargeter.cpp
@@ -38,6 +38,14 @@ UIKRetargeter::UIKRetargeter(const FObjectInitializer& ObjectInitializer)
 	CleanAndInitialize();
 }
 
+void UIKRetargeter::PostDuplicate(bool bDuplicateForPIE)
+{
+	Super::PostDuplicate(bDuplicateForPIE);
+#if WITH_EDITORONLY_DATA
+	Controller = nullptr;
+#endif
+}
+
 void UIKRetargeter::Serialize(FArchive& Ar)
 {
 	Super::Serialize(Ar);
diff --git a/Engine/Plugins/Animation/IKRig/Source/IKRig/Private/Rig/IKRigDefinition.cpp b/Engine/Plugins/Animation/IKRig/Source/IKRig/Private/Rig/IKRigDefinition.cpp
--- a/Engine/Plugins/Animation/IKRig/Source/IKRig/Private/Rig/IKRigDefinition.cpp
+++ b/Engine/Plugins/Animation/IKRig/Source/IKRig/Private/Rig/IKRigDefinition.cpp
@@ -268,6 +268,14 @@ FBoneChain* FRetargetDefinition::GetEditableBoneChainByName(FName ChainName)
 	return nullptr;
 }
 
+void UIKRigDefinition::PostDuplicate(bool bDuplicateForPIE)
+{
+	Super::PostDuplicate(bDuplicateForPIE);
+#if WITH_EDITORONLY_DATA
+	Controller = nullptr;
+#endif
+}
+
 void UIKRigDefinition::Serialize(FArchive& Ar)
 {
 	Super::Serialize(Ar);
diff --git a/Engine/Plugins/Animation/IKRig/Source/IKRig/Public/Retargeter/IKRetargeter.h b/Engine/Plugins/Animation/IKRig/Source/IKRig/Public/Retargeter/IKRetargeter.h
--- a/Engine/Plugins/Animation/IKRig/Source/IKRig/Public/Retargeter/IKRetargeter.h
+++ b/Engine/Plugins/Animation/IKRig/Source/IKRig/Public/Retargeter/IKRetargeter.h
@@ -332,7 +332,8 @@ public:
 		const FTargetChainSpeedPlantSettings& SpeedPlantSettings,
 		const FName TargetChainName);
 
-	// UObject 
+	// UObject
+	virtual void PostDuplicate(bool bDuplicateForPIE) override;
 	virtual void Serialize(FArchive& Ar) override;
 	virtual void PostLoad() override;
 	// END UObject
@@ -428,7 +429,7 @@ public:
 	float BoneDrawSize = 1.0f;
 
 	/** The controller responsible for managing this asset's data (all editor mutation goes through this) */
-	UPROPERTY(Transient)
+	UPROPERTY(Transient, DuplicateTransient, NonTransactional )
 	TObjectPtr<UObject> Controller;
 	
 private:
diff --git a/Engine/Plugins/Animation/IKRig/Source/IKRig/Public/Rig/IKRigDefinition.h b/Engine/Plugins/Animation/IKRig/Source/IKRig/Public/Rig/IKRigDefinition.h
--- a/Engine/Plugins/Animation/IKRig/Source/IKRig/Public/Rig/IKRigDefinition.h
+++ b/Engine/Plugins/Animation/IKRig/Source/IKRig/Public/Rig/IKRigDefinition.h
@@ -186,7 +186,7 @@ class IKRIG_API UIKRigDefinition : public UObject, public IInterface_PreviewMesh
 	GENERATED_BODY()
 
 public:
-
+	
 	/** The skeletal mesh to run the IK solve on (loaded into viewport).
 	* NOTE: you can assign ANY Skeletal Mesh to apply the IK Rig to. Compatibility is determined when a new mesh is assigned
 	* by comparing it's hierarchy with the goals, solvers and bone settings required by the rig. See output log for details. */
@@ -212,6 +212,7 @@ public:
 	const FBoneChain* GetRetargetChainByName(FName ChainName) const;
 
 	/** UObject */
+	virtual void PostDuplicate(bool bDuplicateForPIE) override;
 	virtual void Serialize(FArchive& Ar) override;
 	/** END UObject */
 	
@@ -243,7 +244,7 @@ public:
 	float GoalThickness = 0.7f;
 
 	/** The controller responsible for managing this asset's data (all editor mutation goes through this) */
-	UPROPERTY(Transient)
+	UPROPERTY(Transient, DuplicateTransient, NonTransactional )
 	TObjectPtr<UObject> Controller;
 #endif
 
