Adding custon icons to IK Rig editors

diff --git a/Engine/Plugins/Animation/IKRig/Source/IKRigEditor/Private/RetargetEditor/IKRetargetAssetBrowserTabSummoner.cpp b/Engine/Plugins/Animation/IKRig/Source/IKRigEditor/Private/RetargetEditor/IKRetargetAssetBrowserTabSummoner.cpp
--- a/Engine/Plugins/Animation/IKRig/Source/IKRigEditor/Private/RetargetEditor/IKRetargetAssetBrowserTabSummoner.cpp
+++ b/Engine/Plugins/Animation/IKRig/Source/IKRigEditor/Private/RetargetEditor/IKRetargetAssetBrowserTabSummoner.cpp
@@ -17,7 +17,7 @@ FIKRetargetAssetBrowserTabSummoner::FIKRetargetAssetBrowserTabSummoner(const TSh
 	bIsSingleton = true; // only allow a single instance of this tab
 	
 	TabLabel = LOCTEXT("IKRetargetAssetBrowserTabLabel", "Asset Browser");
-	TabIcon = FSlateIcon(FIKRigEditorStyle::Get().GetStyleSetName(), "IKRig.TabIcon");
+	TabIcon = FSlateIcon(FAppStyle::GetAppStyleSetName(), "ContentBrowser.TabIcon");
 
 	ViewMenuDescription = LOCTEXT("IKRetargetAssetBrowser_ViewMenu_Desc", "Asset Browser");
 	ViewMenuTooltip = LOCTEXT("IKRetargetAssetBrowser_ViewMenu_ToolTip", "Show the Asset Browser Tab");
diff --git a/Engine/Plugins/Animation/IKRig/Source/IKRigEditor/Private/RetargetEditor/IKRetargetChainTabSummoner.cpp b/Engine/Plugins/Animation/IKRig/Source/IKRigEditor/Private/RetargetEditor/IKRetargetChainTabSummoner.cpp
--- a/Engine/Plugins/Animation/IKRig/Source/IKRigEditor/Private/RetargetEditor/IKRetargetChainTabSummoner.cpp
+++ b/Engine/Plugins/Animation/IKRig/Source/IKRigEditor/Private/RetargetEditor/IKRetargetChainTabSummoner.cpp
@@ -4,6 +4,7 @@
 #include "IDocumentation.h"
 #include "RigEditor/IKRigEditorStyle.h"
 #include "RetargetEditor/IKRetargetEditor.h"
+#include "RetargetEditor/IKRetargetEditorStyle.h"
 #include "RetargetEditor/SIKRetargetChainMapList.h"
 
 #define LOCTEXT_NAMESPACE "RetargetChainTabSummoner"
@@ -17,7 +18,7 @@ FIKRetargetChainTabSummoner::FIKRetargetChainTabSummoner(const TSharedRef<FIKRet
 	bIsSingleton = true; // only allow a single instance of this tab
 	
 	TabLabel = LOCTEXT("IKRetargetChainTabLabel", "Chain Mapping");
-	TabIcon = FSlateIcon(FIKRigEditorStyle::Get().GetStyleSetName(), "IKRig.TabIcon");
+	TabIcon = FSlateIcon(FIKRetargetEditorStyle::Get().GetStyleSetName(), "IKRetarget.ChainMapping");
 
 	ViewMenuDescription = LOCTEXT("IKRetargetChain_ViewMenu_Desc", "Chain Mapping");
 	ViewMenuTooltip = LOCTEXT("IKRetargetChain_ViewMenu_ToolTip", "Show the Chain Mapping Tab");
diff --git a/Engine/Plugins/Animation/IKRig/Source/IKRigEditor/Private/RetargetEditor/IKRetargetEditor.cpp b/Engine/Plugins/Animation/IKRig/Source/IKRigEditor/Private/RetargetEditor/IKRetargetEditor.cpp
--- a/Engine/Plugins/Animation/IKRig/Source/IKRigEditor/Private/RetargetEditor/IKRetargetEditor.cpp
+++ b/Engine/Plugins/Animation/IKRig/Source/IKRigEditor/Private/RetargetEditor/IKRetargetEditor.cpp
@@ -275,21 +275,21 @@ void FIKRetargetEditor::FillToolbar(FToolBarBuilder& ToolbarBuilder)
 		NAME_None,
 		TAttribute<FText>(),
 		TAttribute<FText>(),
-		FSlateIcon(FAppStyle::Get().GetStyleSetName(),"Icons.Settings"));
+		FSlateIcon(FIKRetargetEditorStyle::Get().GetStyleSetName(),"IKRetarget.GlobalSettings"));
 
 		ToolbarBuilder.AddToolBarButton(
 		FIKRetargetCommands::Get().ShowRootSettings,
 		NAME_None,
 		TAttribute<FText>(),
 		TAttribute<FText>(),
-		FSlateIcon(FAppStyle::Get().GetStyleSetName(),"Icons.Settings"));
+		FSlateIcon(FIKRetargetEditorStyle::Get().GetStyleSetName(),"IKRetarget.RootSettings"));
 	}
 	ToolbarBuilder.EndSection();
 
 	ToolbarBuilder.AddSeparator();
 
 	FLinearColor OffColor = FLinearColor::White;
-	FLinearColor OnColor = FLinearColor::Green;
+	FLinearColor OnColor = FLinearColor(.32f, .66f, .32f, 1.f);
 	
 	TSharedPtr<SVerticalBox> Box = SNew(SVerticalBox)
 	+ SVerticalBox::Slot()
diff --git a/Engine/Plugins/Animation/IKRig/Source/IKRigEditor/Private/RetargetEditor/IKRetargetHierarchyTabSummoner.cpp b/Engine/Plugins/Animation/IKRig/Source/IKRigEditor/Private/RetargetEditor/IKRetargetHierarchyTabSummoner.cpp
--- a/Engine/Plugins/Animation/IKRig/Source/IKRigEditor/Private/RetargetEditor/IKRetargetHierarchyTabSummoner.cpp
+++ b/Engine/Plugins/Animation/IKRig/Source/IKRigEditor/Private/RetargetEditor/IKRetargetHierarchyTabSummoner.cpp
@@ -19,7 +19,7 @@ FIKRetargetHierarchyTabSummoner::FIKRetargetHierarchyTabSummoner(const TSharedRe
 	bIsSingleton = true; // only allow a single instance of this tab
 	
 	TabLabel = LOCTEXT("IKRetargetHierarchy_TabLabel", "Hierarchy");
-	TabIcon = FSlateIcon(FIKRigEditorStyle::Get().GetStyleSetName(), "IKRig.TabIcon");
+	TabIcon = FSlateIcon(FIKRigEditorStyle::Get().GetStyleSetName(), "IKRig.Hierarchy");
 
 	ViewMenuDescription = LOCTEXT("IKRetargetHierarchy_ViewMenu_Desc", "Hierarchy");
 	ViewMenuTooltip = LOCTEXT("IKRetargetHierarchy_ViewMenu_ToolTip", "Show the Retarget Hierarchy Tab");
diff --git a/Engine/Plugins/Animation/IKRig/Source/IKRigEditor/Private/RetargetEditor/IKRetargetOutputLogTabSummoner.cpp b/Engine/Plugins/Animation/IKRig/Source/IKRigEditor/Private/RetargetEditor/IKRetargetOutputLogTabSummoner.cpp
--- a/Engine/Plugins/Animation/IKRig/Source/IKRigEditor/Private/RetargetEditor/IKRetargetOutputLogTabSummoner.cpp
+++ b/Engine/Plugins/Animation/IKRig/Source/IKRigEditor/Private/RetargetEditor/IKRetargetOutputLogTabSummoner.cpp
@@ -2,7 +2,6 @@
 #include "RetargetEditor/IKRetargetOutputLogTabSummoner.h"
 
 #include "IDocumentation.h"
-#include "RigEditor/IKRigEditorStyle.h"
 #include "RetargetEditor/IKRetargetEditor.h"
 #include "RigEditor/SIKRigOutputLog.h"
 #include "Retargeter/IKRetargetProcessor.h"
@@ -19,7 +18,7 @@ FIKRetargetOutputLogTabSummoner::FIKRetargetOutputLogTabSummoner(
 	bIsSingleton = true; // only allow a single instance of this tab
 	
 	TabLabel = LOCTEXT("IKRetargetOutputLogTabLabel", "Retarget Output Log");
-	TabIcon = FSlateIcon(FIKRigEditorStyle::Get().GetStyleSetName(), "IKRig.TabIcon");
+	TabIcon = FSlateIcon(FAppStyle::GetAppStyleSetName(), "Kismet.Tabs.CompilerResults");
 
 	ViewMenuDescription = LOCTEXT("IKRetargetOutputLog_ViewMenu_Desc", "Retarget Output Log");
 	ViewMenuTooltip = LOCTEXT("IKRetargetOutputLog_ViewMenu_ToolTip", "Show the Retargeting Output Log Tab");
diff --git a/Engine/Plugins/Animation/IKRig/Source/IKRigEditor/Private/RigEditor/IKRigAssetBrowserTabSummoner.cpp b/Engine/Plugins/Animation/IKRig/Source/IKRigEditor/Private/RigEditor/IKRigAssetBrowserTabSummoner.cpp
--- a/Engine/Plugins/Animation/IKRig/Source/IKRigEditor/Private/RigEditor/IKRigAssetBrowserTabSummoner.cpp
+++ b/Engine/Plugins/Animation/IKRig/Source/IKRigEditor/Private/RigEditor/IKRigAssetBrowserTabSummoner.cpp
@@ -3,7 +3,6 @@
 
 #include "IDocumentation.h"
 #include "RigEditor/IKRigToolkit.h"
-#include "RigEditor/IKRigEditorStyle.h"
 #include "RigEditor/SIKRigAssetBrowser.h"
 
 #define LOCTEXT_NAMESPACE "RetargetAssetBrowserTabSummoner"
@@ -17,7 +16,7 @@ FIKRigAssetBrowserTabSummoner::FIKRigAssetBrowserTabSummoner(const TSharedRef<FI
 	bIsSingleton = true; // only allow a single instance of this tab
 	
 	TabLabel = LOCTEXT("IKRigAssetBrowserTabLabel", "Asset Browser");
-	TabIcon = FSlateIcon(FIKRigEditorStyle::Get().GetStyleSetName(), "IKRig.TabIcon");
+	TabIcon = FSlateIcon(FAppStyle::GetAppStyleSetName(), "ContentBrowser.TabIcon");
 
 	ViewMenuDescription = LOCTEXT("IKRigAssetBrowser_ViewMenu_Desc", "Asset Browser");
 	ViewMenuTooltip = LOCTEXT("IKRigAssetBrowser_ViewMenu_ToolTip", "Show the Asset Browser Tab");
diff --git a/Engine/Plugins/Animation/IKRig/Source/IKRigEditor/Private/RigEditor/IKRigOutputLogTabSummoner.cpp b/Engine/Plugins/Animation/IKRig/Source/IKRigEditor/Private/RigEditor/IKRigOutputLogTabSummoner.cpp
--- a/Engine/Plugins/Animation/IKRig/Source/IKRigEditor/Private/RigEditor/IKRigOutputLogTabSummoner.cpp
+++ b/Engine/Plugins/Animation/IKRig/Source/IKRigEditor/Private/RigEditor/IKRigOutputLogTabSummoner.cpp
@@ -4,7 +4,6 @@
 #include "IDocumentation.h"
 #include "Rig/IKRigProcessor.h"
 #include "RigEditor/IKRigToolkit.h"
-#include "RigEditor/IKRigEditorStyle.h"
 #include "RigEditor/SIKRigOutputLog.h"
 
 #define LOCTEXT_NAMESPACE "IKRigOutputLogTabSummoner"
@@ -19,7 +18,7 @@ FIKRigOutputLogTabSummoner::FIKRigOutputLogTabSummoner(
 	bIsSingleton = true; // only allow a single instance of this tab
 	
 	TabLabel = LOCTEXT("IKRigOutputLogTabLabel", "IK Rig Output Log");
-	TabIcon = FSlateIcon(FIKRigEditorStyle::Get().GetStyleSetName(), "IKRig.TabIcon");
+	TabIcon = FSlateIcon(FAppStyle::GetAppStyleSetName(), "Kismet.Tabs.CompilerResults");
 
 	ViewMenuDescription = LOCTEXT("IKRigOutputLog_ViewMenu_Desc", "IK Rig Output Log");
 	ViewMenuTooltip = LOCTEXT("IKRigOutputLog_ViewMenu_ToolTip", "Show the IK Rig Output Log Tab");
diff --git a/Engine/Plugins/Animation/IKRig/Source/IKRigEditor/Private/RigEditor/IKRigRetargetChainTabSummoner.cpp b/Engine/Plugins/Animation/IKRig/Source/IKRigEditor/Private/RigEditor/IKRigRetargetChainTabSummoner.cpp
--- a/Engine/Plugins/Animation/IKRig/Source/IKRigEditor/Private/RigEditor/IKRigRetargetChainTabSummoner.cpp
+++ b/Engine/Plugins/Animation/IKRig/Source/IKRigEditor/Private/RigEditor/IKRigRetargetChainTabSummoner.cpp
@@ -18,7 +18,7 @@ FIKRigRetargetChainTabSummoner::FIKRigRetargetChainTabSummoner(const TSharedRef<
 	bIsSingleton = true; // only allow a single instance of this tab
 	
 	TabLabel = LOCTEXT("IKRigRetargetChainTabLabel", "IK Retargeting");
-	TabIcon = FSlateIcon(FIKRigEditorStyle::Get().GetStyleSetName(), "IKRig.TabIcon");
+	TabIcon = FSlateIcon(FIKRigEditorStyle::Get().GetStyleSetName(), "IKRig.IKRetargeting");
 
 	ViewMenuDescription = LOCTEXT("IKRigRetargetChain_ViewMenu_Desc", "IK Rig Retarget Chains");
 	ViewMenuTooltip = LOCTEXT("IKRigRetargetChain_ViewMenu_ToolTip", "Show the IK Rig Retarget Chains Tab");
diff --git a/Engine/Plugins/Animation/IKRig/Source/IKRigEditor/Private/RigEditor/IKRigSkeletonTabSummoner.cpp b/Engine/Plugins/Animation/IKRig/Source/IKRigEditor/Private/RigEditor/IKRigSkeletonTabSummoner.cpp
--- a/Engine/Plugins/Animation/IKRig/Source/IKRigEditor/Private/RigEditor/IKRigSkeletonTabSummoner.cpp
+++ b/Engine/Plugins/Animation/IKRig/Source/IKRigEditor/Private/RigEditor/IKRigSkeletonTabSummoner.cpp
@@ -18,7 +18,7 @@ FIKRigSkeletonTabSummoner::FIKRigSkeletonTabSummoner(const TSharedRef<FIKRigEdit
 	bIsSingleton = true; // only allow a single instance of this tab
 	
 	TabLabel = LOCTEXT("IKRigSkeletonTabLabel", "Hierarchy");
-	TabIcon = FSlateIcon(FIKRigEditorStyle::Get().GetStyleSetName(), "IKRig.TabIcon");
+	TabIcon = FSlateIcon(FIKRigEditorStyle::Get().GetStyleSetName(), "IKRig.Hierarchy");
 
 	ViewMenuDescription = LOCTEXT("IKRigSkeleton_ViewMenu_Desc", "Hierarchy");
 	ViewMenuTooltip = LOCTEXT("IKRigSkeleton_ViewMenu_ToolTip", "Show the IK Rig Hierarchy Tab");
diff --git a/Engine/Plugins/Animation/IKRig/Source/IKRigEditor/Private/RigEditor/IKRigSolverStackTabSummoner.cpp b/Engine/Plugins/Animation/IKRig/Source/IKRigEditor/Private/RigEditor/IKRigSolverStackTabSummoner.cpp
--- a/Engine/Plugins/Animation/IKRig/Source/IKRigEditor/Private/RigEditor/IKRigSolverStackTabSummoner.cpp
+++ b/Engine/Plugins/Animation/IKRig/Source/IKRigEditor/Private/RigEditor/IKRigSolverStackTabSummoner.cpp
@@ -18,7 +18,7 @@ FIKRigSolverStackTabSummoner::FIKRigSolverStackTabSummoner(const TSharedRef<FIKR
 	bIsSingleton = true; // only allow a single instance of this tab
 	
 	TabLabel = LOCTEXT("IKRigSolverStackTabLabel", "Solver Stack");
-	TabIcon = FSlateIcon(FIKRigEditorStyle::Get().GetStyleSetName(), "IKRig.TabIcon");
+	TabIcon = FSlateIcon(FIKRigEditorStyle::Get().GetStyleSetName(), "IKRig.SolverStack");
 
 	ViewMenuDescription = LOCTEXT("IKRigSolverStack_ViewMenu_Desc", "IK Rig Solver Stack");
 	ViewMenuTooltip = LOCTEXT("IKRigSolverStack_ViewMenu_ToolTip", "Show the IK Rig Solver Stack Tab");
diff --git a/Engine/Plugins/Animation/IKRig/Source/IKRigEditor/Public/RetargetEditor/IKRetargetEditorStyle.h b/Engine/Plugins/Animation/IKRig/Source/IKRigEditor/Public/RetargetEditor/IKRetargetEditorStyle.h
--- a/Engine/Plugins/Animation/IKRig/Source/IKRigEditor/Public/RetargetEditor/IKRetargetEditorStyle.h
+++ b/Engine/Plugins/Animation/IKRig/Source/IKRigEditor/Public/RetargetEditor/IKRetargetEditorStyle.h
@@ -19,6 +19,14 @@ public:
 		Set("IKRetarget.Tree.Bone", new IMAGE_BRUSH("Slate/Bone_16x", Icon16x16));
 		Set("ClassIcon.IKRetargeter", new IMAGE_BRUSH_SVG("Slate/IKRigRetargeter", Icon16x16));
 		Set("ClassThumbnail.IKRetargeter", new IMAGE_BRUSH_SVG("Slate/IKRigRetargeter_64", Icon64x64));
+		
+		Set("IKRetarget.GlobalSettings", new IMAGE_BRUSH_SVG("Slate/GlobalSettings", Icon64x64));
+		Set("IKRetarget.RootSettings", new IMAGE_BRUSH_SVG("Slate/RootSettings", Icon64x64));
+		Set("IKRetarget.ChainMapping", new IMAGE_BRUSH_SVG("Slate/ChainMapping", Icon64x64));
+
+		Set("IKRetarget.RunRetargeter", new IMAGE_BRUSH_SVG("Slate/RunRetargeter", Icon64x64));
+		Set("IKRetarget.EditRetargetPose", new IMAGE_BRUSH_SVG("Slate/EditRetargetPose", Icon64x64));
+		Set("IKRetarget.ShowRetargetPose", new IMAGE_BRUSH_SVG("Slate/ShowRetargetPose", Icon64x64));
 
 		SetContentRoot(FPaths::EngineContentDir() / TEXT("Editor/Slate"));
 		Set( "IKRetarget.Viewport.Border", new BOX_BRUSH( "Old/Window/ViewportDebugBorder", 0.8f, FLinearColor(1.0f,1.0f,1.0f,1.0f) ) );
diff --git a/Engine/Plugins/Animation/IKRig/Source/IKRigEditor/Public/RigEditor/IKRigEditorStyle.h b/Engine/Plugins/Animation/IKRig/Source/IKRigEditor/Public/RigEditor/IKRigEditorStyle.h
--- a/Engine/Plugins/Animation/IKRig/Source/IKRigEditor/Public/RigEditor/IKRigEditorStyle.h
+++ b/Engine/Plugins/Animation/IKRig/Source/IKRigEditor/Public/RigEditor/IKRigEditorStyle.h
@@ -34,6 +34,10 @@ public:
 
 		Set("ClassIcon.IKRigDefinition", new IMAGE_BRUSH_SVG("Slate/IKRig", Icon16x16));
 		Set("ClassThumbnail.IKRigDefinition", new IMAGE_BRUSH_SVG("Slate/IKRig_64", Icon64x64));
+
+		Set("IKRig.SolverStack", new IMAGE_BRUSH_SVG("Slate/SolverStack", Icon64x64));
+		Set("IKRig.IKRetargeting", new IMAGE_BRUSH_SVG("Slate/IKRetargeting", Icon64x64));
+		Set("IKRig.Hierarchy", new IMAGE_BRUSH_SVG("Slate/Hierarchy", Icon64x64));
 		
 		FTextBlockStyle NormalText = FAppStyle::GetWidgetStyle<FTextBlockStyle>("SkeletonTree.NormalFont");
 		Set( "IKRig.Tree.NormalText", FTextBlockStyle(NormalText));
