Adding ability to edit local space of retarget pose bones.

diff --git a/Engine/Plugins/Animation/IKRig/Source/IKRigEditor/Private/RetargetEditor/IKRetargetDetails.cpp b/Engine/Plugins/Animation/IKRig/Source/IKRigEditor/Private/RetargetEditor/IKRetargetDetails.cpp
--- a/Engine/Plugins/Animation/IKRig/Source/IKRigEditor/Private/RetargetEditor/IKRetargetDetails.cpp
+++ b/Engine/Plugins/Animation/IKRig/Source/IKRigEditor/Private/RetargetEditor/IKRetargetDetails.cpp
@@ -78,7 +78,7 @@ FEulerTransform UIKRetargetBoneDetails::GetTransform(EIKRetargetTransformType Tr
 				return FEulerTransform(FAnimationRuntime::GetComponentSpaceTransform(RefSkeleton, RefSkeleton.GetRefBonePose(), BoneIndex));
 			}
 		}
-
+		
 	case EIKRetargetTransformType::RelativeOffset:
 		{
 			// this is the only stored data we have for bone pose offsets
@@ -112,15 +112,43 @@ FEulerTransform UIKRetargetBoneDetails::GetTransform(EIKRetargetTransformType Tr
 					ParentGlobalTransform = Mesh->GetBoneTransform(ParentIndex, FTransform::Identity);
 				}
 				FTransform LocalRefTransform = RefSkeleton.GetRefBonePose()[BoneIndex];
-				FTransform GlobalTransform = LocalRefTransform * ParentGlobalTransform;
+				FTransform GlobalTransformNoOffset = LocalRefTransform * ParentGlobalTransform;
 				// get global rotation plus delta
-				FQuat GlobalRotationPlusDelta = GlobalTransform.GetRotation() * LocalRotationDelta.Quaternion();
+				FQuat GlobalRotationPlusDelta = GlobalTransformNoOffset.GetRotation() * LocalRotationDelta.Quaternion();
 				// get global delta rotation
-				FQuat GlobalDeltaRotation = GlobalRotationPlusDelta * GlobalTransform.GetRotation().Inverse();
+				FQuat GlobalDeltaRotation = GlobalRotationPlusDelta * GlobalTransformNoOffset.GetRotation().Inverse();
 				// combine with translation delta
 				return FEulerTransform(GlobalTranslationDelta, GlobalDeltaRotation.Rotator(), FVector::OneVector);
 			}
 		}
+
+	case EIKRetargetTransformType::Bone:
+		{
+			ensure(bLocalSpace);
+			
+			// this is the only stored data we have for bone pose offsets
+			const ERetargetSourceOrTarget SourceOrTarget = EditorController->GetSourceOrTarget();
+			const FQuat LocalRotationOffset = EditorController->AssetController->GetRotationOffsetForRetargetPoseBone(SelectedBone, SourceOrTarget);
+			const FVector GlobalTranslationDelta = IsRootBone() ? EditorController->AssetController->GetCurrentRetargetPose(SourceOrTarget).GetRootTranslationDelta() : FVector::Zero();
+			const int32 ParentIndex = RefSkeleton.GetParentIndex(BoneIndex);
+			
+			// combine the local space offset from ref pose with the recorded offset in the retarget pose
+			FTransform LocalRefTransform = RefSkeleton.GetRefBonePose()[BoneIndex];
+			FQuat CombinedLocalRotation = LocalRefTransform.GetRotation() * LocalRotationOffset;
+
+			// get parent global transform to calculate local translation delta
+			FTransform ParentRefGlobalTransform = FTransform::Identity;
+			if (ParentIndex != INDEX_NONE)
+			{
+				ParentRefGlobalTransform = FAnimationRuntime::GetComponentSpaceTransform(RefSkeleton, RefSkeleton.GetRefBonePose(), ParentIndex);
+			}
+			// now calculate local translational delta from global
+			FVector LocalTranslation = ParentRefGlobalTransform.InverseTransformVector(GlobalTranslationDelta);
+            
+			// create an Euler transform of local space
+            return FEulerTransform(LocalTranslation, CombinedLocalRotation.Rotator(), FVector::OneVector);
+		}
+		
 	default:
 		checkNoEntry();
 	}
@@ -146,6 +174,12 @@ bool UIKRetargetBoneDetails::IsComponentRelative(
 		{
 			return RelativeOffsetTransformRelative[(int32)Component];
 		}
+	case EIKRetargetTransformType::Bone:
+		{
+			return BoneRelative[(int32)Component];
+		}
+	default:
+		checkNoEntry();
 	}
 	return true;
 }
@@ -172,6 +206,13 @@ void UIKRetargetBoneDetails::OnComponentRelativeChanged(
 			RelativeOffsetTransformRelative[(int32)Component] = bIsRelative;
 			break; 
 		}
+	case EIKRetargetTransformType::Bone:
+		{
+			BoneRelative[(int32)Component] = true;
+			break; 
+		}
+	default:
+		checkNoEntry();
 	}
 }
 
@@ -198,6 +239,11 @@ void UIKRetargetBoneDetails::OnCopyToClipboard(
 			bIsRelative = RelativeOffsetTransformRelative[(int32)Component];
 			break;
 		}
+	case EIKRetargetTransformType::Bone:
+		{
+			bIsRelative = BoneRelative[(int32)Component];
+			break;
+		}
 	default:
 		checkNoEntry();
 	}
@@ -357,6 +403,11 @@ TOptional<FVector::FReal> UIKRetargetBoneDetails::GetNumericValue(
 			Transform = GetTransform(TransformType, RelativeOffsetTransformRelative[(int32)Component]);
 			break;
 		}
+	case EIKRetargetTransformType::Bone:
+		{
+			Transform = GetTransform(TransformType, BoneRelative[(int32)Component]);
+			break;
+		}
 	default:
 		checkNoEntry();
 	}
@@ -373,7 +424,9 @@ void UIKRetargetBoneDetails::OnNumericValueCommitted(
 	EIKRetargetTransformType TransformType,
 	bool bIsCommit)
 {
-	if(TransformType != EIKRetargetTransformType::RelativeOffset)
+	const bool bIsRelativeOffset = TransformType == EIKRetargetTransformType::RelativeOffset;
+	const bool bIsLocal = TransformType == EIKRetargetTransformType::Bone;
+	if(!(bIsRelativeOffset || bIsLocal))
 	{
 		return;
 	}
@@ -402,6 +455,53 @@ void UIKRetargetBoneDetails::OnNumericValueCommitted(
 
 	const ERetargetSourceOrTarget SourceOrTarget = EditorController->GetSourceOrTarget();
 
+	switch (TransformType)
+	{
+	case EIKRetargetTransformType::RelativeOffset:
+		{
+			CommitValueAsRelativeOffset(
+				AssetController,
+				RefSkeleton,
+				SourceOrTarget,
+				BoneIndex,
+				Mesh,
+				Component,
+				Representation,
+				SubComponent,
+				Value);
+			break;
+		}
+	case EIKRetargetTransformType::Bone:
+		{
+			CommitValueAsBoneSpace(
+				AssetController,
+				RefSkeleton,
+				SourceOrTarget,
+				BoneIndex,
+				Mesh,
+				Component,
+				Representation,
+				SubComponent,
+				Value);
+			break;
+		}
+	default:
+		// cannot edit any other transform types
+		checkNoEntry();
+	}
+}
+
+void UIKRetargetBoneDetails::CommitValueAsRelativeOffset(
+	UIKRetargeterController* AssetController,
+	const FReferenceSkeleton& RefSkeleton,
+	const ERetargetSourceOrTarget SourceOrTarget,
+	const int32 BoneIndex,
+	UDebugSkelMeshComponent* Mesh,
+	ESlateTransformComponent::Type Component,
+	ESlateRotationRepresentation::Type Representation,
+	ESlateTransformSubComponent::Type SubComponent,
+	FVector::FReal Value)
+{
 	switch(Component)
 	{
 	case ESlateTransformComponent::Location:
@@ -444,12 +544,13 @@ void UIKRetargetBoneDetails::OnNumericValueCommitted(
 	case ESlateTransformComponent::Rotation:
 		{
 			const bool bIsRotationLocal = RelativeOffsetTransformRelative[1];
-			const FQuat LocalRotationDelta = AssetController->GetRotationOffsetForRetargetPoseBone(SelectedBone, SourceOrTarget);
-			FEulerTransform LocalDeltaTransform = FEulerTransform(FVector::ZeroVector, LocalRotationDelta.Rotator(), FVector::OneVector);
 			FQuat NewLocalRotationDelta;
 			
 			if (bIsRotationLocal)
 			{
+				const FQuat LocalRotationDelta = AssetController->GetRotationOffsetForRetargetPoseBone(SelectedBone, SourceOrTarget);
+				FEulerTransform LocalDeltaTransform = FEulerTransform(FVector::ZeroVector, LocalRotationDelta.Rotator(), FVector::OneVector);
+				
 				// rotations are stored in local space, so just apply the edit
 				SAdvancedTransformInputBox<FEulerTransform>::ApplyNumericValueChange(LocalDeltaTransform, Value, Component, Representation, SubComponent);
 				NewLocalRotationDelta = LocalDeltaTransform.GetRotation();
@@ -463,23 +564,23 @@ void UIKRetargetBoneDetails::OnNumericValueCommitted(
 					ParentGlobalTransform = Mesh->GetBoneTransform(ParentIndex, FTransform::Identity);
 				}
 				FTransform LocalRefTransform = RefSkeleton.GetRefBonePose()[BoneIndex];
-				FTransform GlobalTransform = LocalRefTransform * ParentGlobalTransform;
+				FTransform CurrentGlobalTransformNoDelta = LocalRefTransform * ParentGlobalTransform;
 				
 				// get reference global transform
-				//FTransform GlobalTransform = Mesh->GetBoneTransform(BoneIndex, FTransform::Identity);
 				// get offset global transform
-				FQuat GlobalRefRotationPlusDelta = GlobalTransform.GetRotation() * LocalRotationDelta;
+				const FQuat LocalRotationDelta = AssetController->GetRotationOffsetForRetargetPoseBone(SelectedBone, SourceOrTarget);
+				FQuat GlobalRefRotationPlusDelta = CurrentGlobalTransformNoDelta.GetRotation() * LocalRotationDelta;
 				// get global delta
-				FQuat GlobalDeltaRotation = GlobalRefRotationPlusDelta * GlobalTransform.GetRotation().Inverse();
+				FQuat GlobalRotationOffset = GlobalRefRotationPlusDelta * CurrentGlobalTransformNoDelta.GetRotation().Inverse();
 				// apply edit to global delta
-				FEulerTransform GlobalDeltaTransform = FEulerTransform(FVector::ZeroVector, GlobalDeltaRotation.Rotator(), FVector::OneVector);
+				FEulerTransform GlobalDeltaTransform = FEulerTransform(FVector::ZeroVector, GlobalRotationOffset.Rotator(), FVector::OneVector);
 				SAdvancedTransformInputBox<FEulerTransform>::ApplyNumericValueChange(GlobalDeltaTransform, Value, Component, Representation, SubComponent);
 
 				// convert world space delta quaternion to bone-space
-				GlobalDeltaRotation = GlobalDeltaTransform.GetRotation();
-				const FVector RotationAxis = GlobalDeltaRotation.GetRotationAxis();
-				const FVector UnRotatedAxis = GlobalTransform.InverseTransformVector(RotationAxis);
-				NewLocalRotationDelta = FQuat(UnRotatedAxis, GlobalDeltaRotation.GetAngle());
+				FQuat NewGlobalDeltaRotation = GlobalDeltaTransform.GetRotation();
+				const FVector RotationAxis = NewGlobalDeltaRotation.GetRotationAxis();
+				const FVector UnRotatedAxis = CurrentGlobalTransformNoDelta.InverseTransformVector(RotationAxis);
+				NewLocalRotationDelta = FQuat(UnRotatedAxis, NewGlobalDeltaRotation.GetAngle());
 			}
 			
 			// store the new rotation in the retarget pose
@@ -493,6 +594,87 @@ void UIKRetargetBoneDetails::OnNumericValueCommitted(
 	}
 }
 
+void UIKRetargetBoneDetails::CommitValueAsBoneSpace(
+	UIKRetargeterController* AssetController,
+	const FReferenceSkeleton& RefSkeleton,
+	const ERetargetSourceOrTarget SourceOrTarget,
+	const int32 BoneIndex,
+	UDebugSkelMeshComponent* Mesh,
+	ESlateTransformComponent::Type Component,
+	ESlateRotationRepresentation::Type Representation,
+	ESlateTransformSubComponent::Type SubComponent,
+	FVector::FReal Value)
+{
+	switch(Component)
+	{
+	case ESlateTransformComponent::Location:
+		{
+			const bool bIsTranslationLocal = BoneRelative[0];
+			FTransform CurrentGlobalOffset = FTransform::Identity;
+			CurrentGlobalOffset.SetTranslation(AssetController->GetCurrentRetargetPose(SourceOrTarget).GetRootTranslationDelta());
+			
+			if (bIsTranslationLocal)
+			{	
+				// get the current LOCAL offset
+				FTransform CurrentLocalOffset = CurrentGlobalOffset;
+				FTransform ParentGlobalRefTransform = FTransform::Identity;
+				const int32 ParentIndex = RefSkeleton.GetParentIndex(BoneIndex);
+				if (ParentIndex != INDEX_NONE)
+				{
+					ParentGlobalRefTransform = FAnimationRuntime::GetComponentSpaceTransform(RefSkeleton, RefSkeleton.GetRefBonePose(), ParentIndex);
+				}
+				CurrentLocalOffset = CurrentLocalOffset.GetRelativeTransform(ParentGlobalRefTransform);
+
+				// apply the numerical value to the local space values
+				SAdvancedTransformInputBox<FTransform>::ApplyNumericValueChange(CurrentLocalOffset, Value, Component, Representation, SubComponent);
+
+				// convert back to global space for storage in the pose
+				CurrentGlobalOffset = CurrentLocalOffset * ParentGlobalRefTransform;
+			}
+			else
+			{
+				// apply the edit
+				SAdvancedTransformInputBox<FTransform>::ApplyNumericValueChange(CurrentGlobalOffset, Value, Component, Representation, SubComponent);
+			}
+			
+			// store the new transform in the retarget pose
+			FScopedTransaction Transaction(LOCTEXT("EditRootTranslation", "Edit Retarget Root Pose Translation"));
+			EditorController->AssetController->GetAsset()->Modify();
+			EditorController->AssetController->GetCurrentRetargetPose(SourceOrTarget).SetRootTranslationDelta(CurrentGlobalOffset.GetTranslation());
+			
+			break;
+		}
+	case ESlateTransformComponent::Rotation:
+		{
+			// only allow local bone space pose editing for now
+			const bool bIsRotationLocal = BoneRelative[1];
+			if (ensure(bIsRotationLocal))
+			{
+				// combine the local space offset from ref pose with the recorded offset in the retarget pose
+				FTransform LocalRefTransform = RefSkeleton.GetRefBonePose()[BoneIndex];
+				const FQuat LocalRotationOffset = AssetController->GetRotationOffsetForRetargetPoseBone(SelectedBone, SourceOrTarget);
+				FQuat CombinedLocalRotation = LocalRefTransform.GetRotation() * LocalRotationOffset;
+				FEulerTransform CombinedLocalDeltaTransform = FEulerTransform(FVector::ZeroVector, CombinedLocalRotation.Rotator(), FVector::OneVector);
+			
+				// rotations are stored in local space, so just apply the edit
+				SAdvancedTransformInputBox<FEulerTransform>::ApplyNumericValueChange(CombinedLocalDeltaTransform, Value, Component, Representation, SubComponent);
+
+				// subtract the local space from the result to be left with JUST the retarget pose offset
+				FQuat NewLocalRotationDelta = LocalRefTransform.GetRotation().Inverse() * CombinedLocalDeltaTransform.GetRotation();
+			
+				// store the new rotation in the retarget pose
+				FScopedTransaction Transaction(LOCTEXT("EditRootRotation", "Edit Retarget Pose Rotation"));
+				EditorController->AssetController->GetAsset()->Modify();
+				EditorController->AssetController->SetRotationOffsetForRetargetPoseBone(SelectedBone, NewLocalRotationDelta, SourceOrTarget);
+			}
+			
+			break;
+		}
+	default:
+		checkNoEntry();
+	}
+}
+
 bool UIKRetargetBoneDetails::IsRootBone() const
 {
 	const FName RootBone = EditorController->AssetController->GetRetargetRootBone(EditorController->GetSourceOrTarget());
@@ -612,7 +794,9 @@ void FIKRetargetBoneDetailCustomization::CustomizeDetails(IDetailLayoutBuilder&
 		const EIKRetargetTransformType TransformType = UIData.TransformTypes[PropertyIndex];
 
 		// only enable editing of the relative offset transform type while in edit mode
-		const bool bIsEditable = bIsEditingPose && TransformType == EIKRetargetTransformType::RelativeOffset;
+		const bool bIsRelativeOffset = TransformType == EIKRetargetTransformType::RelativeOffset;
+		const bool bIsBoneOffset = TransformType == EIKRetargetTransformType::Bone;
+		const bool bIsEditable = bIsEditingPose && (bIsRelativeOffset || bIsBoneOffset);
 		
 		TransformWidgetArgs.IsEnabled(bIsEditable);
 		// edit transform
@@ -691,64 +875,66 @@ void FIKRetargetBoneDetailCustomization::GetTransformUIData(
 	const IDetailLayoutBuilder& DetailBuilder,
 	FIKRetargetTransformUIData& OutData) const
 {
-	// read-only
-	static const TArray<EIKRetargetTransformType> ReadOnlyTransformTypes
+	// read-only transform meta data
+	static TArray<FIKRetargetTransformWidgetData> ReadOnlyTransformMetaData =
 	{
-		EIKRetargetTransformType::Current,
-		EIKRetargetTransformType::Reference
-	};
-	static const TArray<FText> ReadOnlyButtonLabels
-	{
-		LOCTEXT("CurrentTransform", "Current"),
-		LOCTEXT("ReferenceTransform", "Reference")
-	};
-	static const TArray<FText> ReadOnlyButtonTooltips
-	{
-		LOCTEXT("CurrentBoneTransformTooltip", "The current transform of the bone."),
-		LOCTEXT("ReferenceBoneTransformTooltip", "The reference transform of the bone.")
+		FIKRetargetTransformWidgetData(
+			EIKRetargetTransformType::Current,
+			LOCTEXT("CurrentTransform", "Current"),
+			LOCTEXT("CurrentBoneTransformTooltip", "The current transform of the bone.")),
+		FIKRetargetTransformWidgetData(
+			EIKRetargetTransformType::Reference,
+			LOCTEXT("ReferenceTransform", "Reference"),
+			LOCTEXT("ReferenceBoneTransformTooltip", "The reference transform of the bone."))
 	};
 	static TAttribute<TArray<EIKRetargetTransformType>> ReadOnlyVisibleTransforms =
 		TArray<EIKRetargetTransformType>({EIKRetargetTransformType::Current});
 
-
-	// editable
-	static const TArray<EIKRetargetTransformType> EditableTransformTypes
-	{
-		EIKRetargetTransformType::RelativeOffset,
-		EIKRetargetTransformType::Reference
-	};
-	static const TArray<FText> EditableButtonLabels
+	// editable transform meta data
+	static TArray<FIKRetargetTransformWidgetData> EditableTransformMetaData =
 	{
-		LOCTEXT("EditableRelativeOffsetTransform", "Relative Offset"),
-		LOCTEXT("EditableReferenceTransform", "Reference")
-	};
-	static const TArray<FText> EditableButtonTooltips
-	{
-		LOCTEXT("RelativeOffsetBoneTransformTooltip", "The offset transform in the current retarget pose, relative to the reference pose."),
-		LOCTEXT("EditableReferenceBoneTransformTooltip", "The reference transform of the bone.")
+		FIKRetargetTransformWidgetData(
+			EIKRetargetTransformType::RelativeOffset,
+			LOCTEXT("EditableRelativeOffsetTransform", "Relative Offset"),
+			LOCTEXT("RelativeOffsetBoneTransformTooltip", "The offset transform in the current retarget pose, relative to the reference pose.")),
+		FIKRetargetTransformWidgetData(
+			EIKRetargetTransformType::Bone,
+			LOCTEXT("EditableRelativeOffsetTransform", "Bone"),
+			LOCTEXT("RelativeOffsetBoneTransformTooltip", "The offset transform in the current retarget pose, relative to the parent bone.")),
+		FIKRetargetTransformWidgetData(
+			EIKRetargetTransformType::Reference,
+			LOCTEXT("EditableReferenceTransform", "Reference"),
+			LOCTEXT("EditableReferenceBoneTransformTooltip", "The transform of the bone in the reference pose."))
 	};
 	static TAttribute<TArray<EIKRetargetTransformType>> EditableVisibleTransforms =
 		TArray<EIKRetargetTransformType>({EIKRetargetTransformType::RelativeOffset});
 
-
 	if (bIsEditingPose)
 	{
-		OutData.TransformTypes = EditableTransformTypes;
-		OutData.ButtonLabels = EditableButtonLabels;
-		OutData.ButtonTooltips = EditableButtonTooltips;
+		for (const FIKRetargetTransformWidgetData& TransformData : EditableTransformMetaData)
+		{
+			OutData.TransformTypes.Add(TransformData.TransformType);
+			OutData.ButtonLabels.Add(TransformData.ButtonLabel);
+			OutData.ButtonTooltips.Add(TransformData.ButtonTooltip);
+		}
+
 		OutData.VisibleTransforms = EditableVisibleTransforms;
 
 		OutData.Properties.Append(
 			{
 			DetailBuilder.GetProperty(GET_MEMBER_NAME_CHECKED(UIKRetargetBoneDetails, OffsetTransform)),
+			DetailBuilder.GetProperty(GET_MEMBER_NAME_CHECKED(UIKRetargetBoneDetails, LocalTransform)),
 			DetailBuilder.GetProperty(GET_MEMBER_NAME_CHECKED(UIKRetargetBoneDetails, ReferenceTransform))
 			});
 	}
 	else
 	{
-		OutData.TransformTypes = ReadOnlyTransformTypes;
-		OutData.ButtonLabels = ReadOnlyButtonLabels;
-		OutData.ButtonTooltips = ReadOnlyButtonTooltips;
+		for (const FIKRetargetTransformWidgetData& TransformData : ReadOnlyTransformMetaData)
+		{
+			OutData.TransformTypes.Add(TransformData.TransformType);
+			OutData.ButtonLabels.Add(TransformData.ButtonLabel);
+			OutData.ButtonTooltips.Add(TransformData.ButtonTooltip);
+		}
 		OutData.VisibleTransforms = ReadOnlyVisibleTransforms;
 
 		OutData.Properties.Append(
diff --git a/Engine/Plugins/Animation/IKRig/Source/IKRigEditor/Public/RetargetEditor/IKRetargetDetails.h b/Engine/Plugins/Animation/IKRig/Source/IKRigEditor/Public/RetargetEditor/IKRetargetDetails.h
--- a/Engine/Plugins/Animation/IKRig/Source/IKRigEditor/Public/RetargetEditor/IKRetargetDetails.h
+++ b/Engine/Plugins/Animation/IKRig/Source/IKRigEditor/Public/RetargetEditor/IKRetargetDetails.h
@@ -9,11 +9,14 @@
 
 #include "IKRetargetDetails.generated.h"
 
+class UDebugSkelMeshComponent;
+
 enum class EIKRetargetTransformType : int8
 {
 	Current,
 	Reference,
 	RelativeOffset,
+	Bone,
 };
 
 UCLASS(config = Engine, hidecategories = UObject)
@@ -26,6 +29,9 @@ public:
 	UPROPERTY(VisibleAnywhere, Category = "Selection")
 	FName SelectedBone;
 
+	UPROPERTY()
+	FTransform LocalTransform;
+	
 	UPROPERTY()
 	FTransform OffsetTransform;
 	
@@ -76,15 +82,45 @@ public:
 
 private:
 
+	void CommitValueAsRelativeOffset(
+		UIKRetargeterController* AssetController,
+		const FReferenceSkeleton& RefSkeleton,
+		const ERetargetSourceOrTarget SourceOrTarget,
+		const int32 BoneIndex,
+		UDebugSkelMeshComponent* Mesh,
+		ESlateTransformComponent::Type Component,
+		ESlateRotationRepresentation::Type Representation,
+		ESlateTransformSubComponent::Type SubComponent,
+		FVector::FReal Value);
+
+	void CommitValueAsBoneSpace(
+		UIKRetargeterController* AssetController,
+		const FReferenceSkeleton& RefSkeleton,
+		const ERetargetSourceOrTarget SourceOrTarget,
+		const int32 BoneIndex,
+		UDebugSkelMeshComponent* Mesh,
+		ESlateTransformComponent::Type Component,
+		ESlateRotationRepresentation::Type Representation,
+		ESlateTransformSubComponent::Type SubComponent,
+		FVector::FReal Value);
+
 	static TOptional<FVector::FReal> CleanRealValue(TOptional<FVector::FReal> InValue);
 	
 	bool IsRootBone() const;
-	
+
+	bool BoneRelative[3] = { false, true, false };
 	bool RelativeOffsetTransformRelative[3] = { false, true, false };
 	bool CurrentTransformRelative[3];
 	bool ReferenceTransformRelative[3];
 };
 
+struct FIKRetargetTransformWidgetData
+{
+	EIKRetargetTransformType TransformType;
+	FText ButtonLabel;
+	FText ButtonTooltip;
+};
+
 struct FIKRetargetTransformUIData
 {
 	TArray<EIKRetargetTransformType> TransformTypes;
