Fixing compilation issue caused by 25204224 in UE5/Main.

diff --git a/Engine/Plugins/Animation/IKRig/Source/IKRigEditor/Public/RetargetEditor/IKRetargetDetails.h b/Engine/Plugins/Animation/IKRig/Source/IKRigEditor/Public/RetargetEditor/IKRetargetDetails.h
--- a/Engine/Plugins/Animation/IKRig/Source/IKRigEditor/Public/RetargetEditor/IKRetargetDetails.h
+++ b/Engine/Plugins/Animation/IKRig/Source/IKRigEditor/Public/RetargetEditor/IKRetargetDetails.h
@@ -116,6 +116,13 @@ private:
 
 struct FIKRetargetTransformWidgetData
 {
+	FIKRetargetTransformWidgetData(EIKRetargetTransformType InType, FText InLabel, FText InTooltip)
+	{
+		TransformType = InType;
+		ButtonLabel = InLabel;
+		ButtonTooltip = InTooltip;
+	};
+	
 	EIKRetargetTransformType TransformType;
 	FText ButtonLabel;
 	FText ButtonTooltip;
