Missing file from 25193036

diff --git a/Engine/Plugins/Animation/IKRig/Source/IKRigEditor/Private/RetargetEditor/IKRetargetEditorController.cpp b/Engine/Plugins/Animation/IKRig/Source/IKRigEditor/Private/RetargetEditor/IKRetargetEditorController.cpp
--- a/Engine/Plugins/Animation/IKRig/Source/IKRigEditor/Private/RetargetEditor/IKRetargetEditorController.cpp
+++ b/Engine/Plugins/Animation/IKRig/Source/IKRigEditor/Private/RetargetEditor/IKRetargetEditorController.cpp
@@ -601,13 +601,17 @@ FTransform FIKRetargetEditorController::GetGlobalRetargetPoseOfBone(
 	const float& Scale,
 	const FVector& Offset) const
 {
-	if (BoneIndex == INDEX_NONE)
+	const UIKRetargetAnimInstance* AnimInstance = GetAnimInstance(SourceOrTarget);
+	if (!AnimInstance)
 	{
 		return FTransform::Identity;
 	}
 	
-	const UIKRetargetAnimInstance* AnimInstance = GetAnimInstance(SourceOrTarget);
 	const TArray<FTransform>& GlobalRetargetPose = AnimInstance->GetGlobalRetargetPose();
+	if (!GlobalRetargetPose.IsValidIndex(BoneIndex))
+	{
+		return FTransform::Identity;
+	}
 	
 	// get transform of bone
 	FTransform BoneTransform = GlobalRetargetPose[BoneIndex];
@@ -1129,14 +1133,14 @@ FSlateIcon FIKRetargetEditorController::GetRetargeterModeIcon(ERetargeterOutputM
 	switch (Mode)
 	{
 	case ERetargeterOutputMode::RunRetarget:
-		return FSlateIcon(FAppStyle::Get().GetStyleSetName(),"GenericPlay");
+		return FSlateIcon(FIKRetargetEditorStyle::Get().GetStyleSetName(), "IKRetarget.RunRetargeter");
 	case ERetargeterOutputMode::EditRetargetPose:
-		return FSlateIcon(FAppStyle::Get().GetStyleSetName(),"Icons.Edit");
+		return FSlateIcon(FIKRetargetEditorStyle::Get().GetStyleSetName(), "IKRetarget.EditRetargetPose");
 	case ERetargeterOutputMode::ShowRetargetPose:
-		return FSlateIcon(FAppStyle::Get().GetStyleSetName(),"GenericPause");
+		return FSlateIcon(FIKRetargetEditorStyle::Get().GetStyleSetName(), "IKRetarget.ShowRetargetPose");
 	default:
 		checkNoEntry();
-		return FSlateIcon(FAppStyle::Get().GetStyleSetName(),"GenericPause");
+		return FSlateIcon(FIKRetargetEditorStyle::Get().GetStyleSetName(), "IKRetarget.ShowRetargetPose");
 	}
 }
 
