Exposing IK Rig asset factories to linker API so outside code can procedurally create the assets.

diff --git a/Engine/Plugins/Animation/IKRig/Source/IKRigEditor/Public/RetargetEditor/IKRetargetFactory.h b/Engine/Plugins/Animation/IKRig/Source/IKRigEditor/Public/RetargetEditor/IKRetargetFactory.h
--- a/Engine/Plugins/Animation/IKRig/Source/IKRigEditor/Public/RetargetEditor/IKRetargetFactory.h
+++ b/Engine/Plugins/Animation/IKRig/Source/IKRigEditor/Public/RetargetEditor/IKRetargetFactory.h
@@ -9,8 +9,8 @@
 
 class SWindow;
 
-UCLASS(hidecategories=Object, MinimalAPI)
-class UIKRetargetFactory : public UFactory
+UCLASS(hidecategories=Object)
+class IKRIGEDITOR_API UIKRetargetFactory : public UFactory
 {
 	GENERATED_BODY()
 
diff --git a/Engine/Plugins/Animation/IKRig/Source/IKRigEditor/Public/RigEditor/IKRigDefinitionFactory.h b/Engine/Plugins/Animation/IKRig/Source/IKRigEditor/Public/RigEditor/IKRigDefinitionFactory.h
--- a/Engine/Plugins/Animation/IKRig/Source/IKRigEditor/Public/RigEditor/IKRigDefinitionFactory.h
+++ b/Engine/Plugins/Animation/IKRig/Source/IKRigEditor/Public/RigEditor/IKRigDefinitionFactory.h
@@ -10,7 +10,7 @@
 class SWindow;
 
 UCLASS(BlueprintType, hidecategories=Object)
-class UIKRigDefinitionFactory : public UFactory
+class IKRIGEDITOR_API UIKRigDefinitionFactory : public UFactory
 {
 	GENERATED_BODY()
 
