Fixing CIS error in localization labels.

diff --git a/Engine/Plugins/Animation/IKRig/Source/IKRigEditor/Private/RetargetEditor/IKRetargetDetails.cpp b/Engine/Plugins/Animation/IKRig/Source/IKRigEditor/Private/RetargetEditor/IKRetargetDetails.cpp
--- a/Engine/Plugins/Animation/IKRig/Source/IKRigEditor/Private/RetargetEditor/IKRetargetDetails.cpp
+++ b/Engine/Plugins/Animation/IKRig/Source/IKRigEditor/Private/RetargetEditor/IKRetargetDetails.cpp
@@ -899,8 +899,8 @@ void FIKRetargetBoneDetailCustomization::GetTransformUIData(
 			LOCTEXT("RelativeOffsetBoneTransformTooltip", "The offset transform in the current retarget pose, relative to the reference pose.")),
 		FIKRetargetTransformWidgetData(
 			EIKRetargetTransformType::Bone,
-			LOCTEXT("EditableRelativeOffsetTransform", "Bone"),
-			LOCTEXT("RelativeOffsetBoneTransformTooltip", "The offset transform in the current retarget pose, relative to the parent bone.")),
+			LOCTEXT("EditableBoneTransform", "Bone"),
+			LOCTEXT("EditableBoneTransformTooltip", "The offset transform in the current retarget pose, relative to the parent bone.")),
 		FIKRetargetTransformWidgetData(
 			EIKRetargetTransformType::Reference,
 			LOCTEXT("EditableReferenceTransform", "Reference"),
