Fixes potential NaNs in IK Retargeter when using only 1 IK goal.

diff --git a/Engine/Plugins/Animation/IKRig/Source/IKRig/Public/Rig/Solvers/PointsToRotation.h b/Engine/Plugins/Animation/IKRig/Source/IKRig/Public/Rig/Solvers/PointsToRotation.h
--- a/Engine/Plugins/Animation/IKRig/Source/IKRig/Public/Rig/Solvers/PointsToRotation.h
+++ b/Engine/Plugins/Animation/IKRig/Source/IKRig/Public/Rig/Solvers/PointsToRotation.h
@@ -19,6 +19,8 @@ static FQuat GetRotationFromDeformedPoints(
 	// must have more than 1 point to generate a gradient
 	if (InInitialPoints.Num() <= 1)
 	{
+		OutInitialCentroid = InInitialPoints.Num() > 0 ? InInitialPoints[0] : FVector::ZeroVector;
+		OutCurrentCentroid = InCurrentPoints.Num() > 0 ? InCurrentPoints[0] : FVector::ZeroVector;
 		return FQuat::Identity;
 	}
 	
