IKRetargeterThumbnailRenderer: Add missing cache clear in BeginDestroy (which will cleanup its prevew scenes)

#rnx

diff --git a/Engine/Plugins/Animation/IKRig/Source/IKRigEditor/Private/RetargetEditor/IKRetargeterThumbnailRenderer.cpp b/Engine/Plugins/Animation/IKRig/Source/IKRigEditor/Private/RetargetEditor/IKRetargeterThumbnailRenderer.cpp
--- a/Engine/Plugins/Animation/IKRig/Source/IKRigEditor/Private/RetargetEditor/IKRetargeterThumbnailRenderer.cpp
+++ b/Engine/Plugins/Animation/IKRig/Source/IKRigEditor/Private/RetargetEditor/IKRetargeterThumbnailRenderer.cpp
@@ -1,4 +1,4 @@
-﻿// Copyright Epic Games, Inc. All Rights Reserved.
+// Copyright Epic Games, Inc. All Rights Reserved.
 
 #include "RetargetEditor/IKRetargeterThumbnailRenderer.h"
 
@@ -113,6 +113,13 @@ void FIKRetargeterThumbnailScene::GetViewMatrixParameters(
 	OutOrbitZoom = TargetDistance + ThumbnailInfo->OrbitZoom;
 }
 
+void UIKRetargeterThumbnailRenderer::BeginDestroy()
+{
+	ThumbnailSceneCache.Clear();
+
+	Super::BeginDestroy();
+}
+
 bool UIKRetargeterThumbnailRenderer::CanVisualizeAsset(UObject* Object)
 {
 	return HasSourceOrTargetMesh(Object);
diff --git a/Engine/Plugins/Animation/IKRig/Source/IKRigEditor/Public/RetargetEditor/IKRetargeterThumbnailRenderer.h b/Engine/Plugins/Animation/IKRig/Source/IKRigEditor/Public/RetargetEditor/IKRetargeterThumbnailRenderer.h
--- a/Engine/Plugins/Animation/IKRig/Source/IKRigEditor/Public/RetargetEditor/IKRetargeterThumbnailRenderer.h
+++ b/Engine/Plugins/Animation/IKRig/Source/IKRigEditor/Public/RetargetEditor/IKRetargeterThumbnailRenderer.h
@@ -1,4 +1,4 @@
-﻿// Copyright Epic Games, Inc. All Rights Reserved.
+// Copyright Epic Games, Inc. All Rights Reserved.
 
 #pragma once
 
@@ -43,6 +43,10 @@ class UIKRetargeterThumbnailRenderer : public UDefaultSizedThumbnailRenderer
 	IKRIGEDITOR_API virtual EThumbnailRenderFrequency GetThumbnailRenderFrequency(UObject* Object) const override;
 	// End UThumbnailRenderer Object
 	
+	// UObject implementation
+	virtual void BeginDestroy() override;
+	// End UObject implementation
+
 	USkeletalMesh* GetPreviewMeshFromAsset(UObject* Object, ERetargetSourceOrTarget SourceOrTarget) const;
 	bool HasSourceOrTargetMesh(UObject* Object) const;
 
