Prevents loading IK Rig assets from worker thread when running editor in -game mode.

From this UDN:
https://epicgames.lightning.force.com/lightning/r/Case/5004z00001osjQDAAY/view

diff --git a/Engine/Plugins/Animation/IKRig/Source/IKRig/Private/Retargeter/IKRetargeter.cpp b/Engine/Plugins/Animation/IKRig/Source/IKRig/Private/Retargeter/IKRetargeter.cpp
--- a/Engine/Plugins/Animation/IKRig/Source/IKRig/Private/Retargeter/IKRetargeter.cpp
+++ b/Engine/Plugins/Animation/IKRig/Source/IKRig/Private/Retargeter/IKRetargeter.cpp
@@ -38,6 +38,46 @@ UIKRetargeter::UIKRetargeter(const FObjectInitializer& ObjectInitializer)
 	CleanAndInitialize();
 }
 
+const UIKRigDefinition* UIKRetargeter::GetSourceIKRig() const
+{
+	if (IsInGameThread())
+	{
+		return SourceIKRigAsset.LoadSynchronous();	
+	}
+	
+	return nullptr;
+}
+
+const UIKRigDefinition* UIKRetargeter::GetTargetIKRig() const
+{
+	if (IsInGameThread())
+	{
+		return TargetIKRigAsset.LoadSynchronous();	
+	}
+	
+	return nullptr;
+}
+
+UIKRigDefinition* UIKRetargeter::GetSourceIKRigWriteable() const
+{
+	if (IsInGameThread())
+	{
+		return SourceIKRigAsset.LoadSynchronous();	
+	}
+	
+	return nullptr;
+}
+
+UIKRigDefinition* UIKRetargeter::GetTargetIKRigWriteable() const
+{
+	if (IsInGameThread())
+	{
+		return TargetIKRigAsset.LoadSynchronous();	
+	}
+	
+	return nullptr;
+}
+
 void UIKRetargeter::PostDuplicate(bool bDuplicateForPIE)
 {
 	Super::PostDuplicate(bDuplicateForPIE);
diff --git a/Engine/Plugins/Animation/IKRig/Source/IKRig/Public/Retargeter/IKRetargeter.h b/Engine/Plugins/Animation/IKRig/Source/IKRig/Public/Retargeter/IKRetargeter.h
--- a/Engine/Plugins/Animation/IKRig/Source/IKRig/Public/Retargeter/IKRetargeter.h
+++ b/Engine/Plugins/Animation/IKRig/Source/IKRig/Public/Retargeter/IKRetargeter.h
@@ -210,15 +210,15 @@ public:
 	UIKRetargeter(const FObjectInitializer& ObjectInitializer);
 
 	// Get read-only access to the source IK Rig asset 
-	const UIKRigDefinition* GetSourceIKRig() const { return SourceIKRigAsset.LoadSynchronous(); };
+	const UIKRigDefinition* GetSourceIKRig() const;
 	// Get read-only access to the target IK Rig asset 
-	const UIKRigDefinition* GetTargetIKRig() const { return TargetIKRigAsset.LoadSynchronous(); };
+	const UIKRigDefinition* GetTargetIKRig() const;
 	// Get read-write access to the source IK Rig asset.
 	// WARNING: do not use for editing the data model. Use Controller class instead. 
-	UIKRigDefinition* GetSourceIKRigWriteable() const { return SourceIKRigAsset.LoadSynchronous(); };
+	UIKRigDefinition* GetSourceIKRigWriteable() const;
 	// Get read-write access to the target IK Rig asset.
 	// WARNING: do not use for editing the data model. Use Controller class instead. 
-	UIKRigDefinition* GetTargetIKRigWriteable() const { return TargetIKRigAsset.LoadSynchronous(); };
+	UIKRigDefinition* GetTargetIKRigWriteable() const;
 
 	// Get read-only access to the chain mapping 
 	const TArray<TObjectPtr<URetargetChainSettings>>& GetAllChainSettings() const { return ChainSettings; };
