Deprecating "Start Solve from Input Pose" in FBIK.

diff --git a/Engine/Plugins/Animation/IKRig/Source/IKRig/Private/Rig/Solvers/IKRig_PBIKSolver.cpp b/Engine/Plugins/Animation/IKRig/Source/IKRig/Private/Rig/Solvers/IKRig_PBIKSolver.cpp
--- a/Engine/Plugins/Animation/IKRig/Source/IKRig/Private/Rig/Solvers/IKRig_PBIKSolver.cpp
+++ b/Engine/Plugins/Animation/IKRig/Source/IKRig/Private/Rig/Solvers/IKRig_PBIKSolver.cpp
@@ -124,7 +124,7 @@ void UIKRigFBIKSolver::Solve(FIKRigSkeleton& IKRigSkeleton, const FIKRigGoalCont
 	Settings.GlobalPullChainAlpha = PullChainAlpha;
 	Settings.MaxAngle = MaxAngle;
 	Settings.OverRelaxation = OverRelaxation;
-	Settings.bStartSolveFromInputPose = bStartSolveFromInputPose;
+	Settings.bStartSolveFromInputPose_DEPRECATED = bStartSolveFromInputPose_DEPRECATED;
 
 	// solve
 	Solver.Solve(Settings);
@@ -159,7 +159,7 @@ void UIKRigFBIKSolver::UpdateSolverSettings(UIKRigSolver* InSettings)
 		PullChainAlpha = Settings->PullChainAlpha;
 		MaxAngle = Settings->MaxAngle;
 		OverRelaxation = Settings->OverRelaxation;
-		bStartSolveFromInputPose = Settings->bStartSolveFromInputPose;
+		bStartSolveFromInputPose_DEPRECATED = Settings->bStartSolveFromInputPose_DEPRECATED;
 
 		// copy effector settings
 		for (const UIKRig_FBIKEffector* InEffector : Settings->Effectors)
diff --git a/Engine/Plugins/Animation/IKRig/Source/IKRig/Public/Rig/Solvers/IKRig_FBIKSolver.h b/Engine/Plugins/Animation/IKRig/Source/IKRig/Public/Rig/Solvers/IKRig_FBIKSolver.h
--- a/Engine/Plugins/Animation/IKRig/Source/IKRig/Public/Rig/Solvers/IKRig_FBIKSolver.h
+++ b/Engine/Plugins/Animation/IKRig/Source/IKRig/Public/Rig/Solvers/IKRig_FBIKSolver.h
@@ -222,11 +222,11 @@ public:
 	UPROPERTY(BlueprintReadWrite, EditAnywhere, Category = AdvancedSettings, meta = (ClampMin = "1", UIMin = "1.0", UIMax = "10.0"))
 	float OverRelaxation = 1.3f;
 	
-	/** When true, the solver is reset each tick to start from the current input pose. Default is true.
+	/** DEPRECATED: When true, the solver is reset each tick to start from the current input pose. Default is true.
 	 * If false, incoming animated poses are ignored and the solver starts from the results of the previous solve.
 	 * In very limited circumstances, it can be beneficial to use the pose from the previous frame.*/
-	UPROPERTY(BlueprintReadWrite, EditAnywhere, Category = AdvancedSettings)
-	bool bStartSolveFromInputPose = true;
+	UPROPERTY()
+	bool bStartSolveFromInputPose_DEPRECATED = true;
 	
 	UPROPERTY()
 	TArray<TObjectPtr<UIKRig_FBIKEffector>> Effectors;
diff --git a/Engine/Plugins/Experimental/FullBodyIK/Source/PBIK/Private/Core/PBIKBody.cpp b/Engine/Plugins/Experimental/FullBodyIK/Source/PBIK/Private/Core/PBIKBody.cpp
--- a/Engine/Plugins/Experimental/FullBodyIK/Source/PBIK/Private/Core/PBIKBody.cpp
+++ b/Engine/Plugins/Experimental/FullBodyIK/Source/PBIK/Private/Core/PBIKBody.cpp
@@ -72,10 +72,7 @@ void FRigidBody::Initialize(const FBone* SolverRoot)
 
 void FRigidBody::UpdateFromInputs(const FPBIKSolverSettings& Settings)
 {
-	if (Settings.bStartSolveFromInputPose)
-	{
-		UpdateTransformAndMassFromBones();
-	}
+	UpdateTransformAndMassFromBones();
 
 	// update InvMass based on global mass multiplier
 	constexpr float MinMass = 0.5f; // prevent mass ever hitting zero
diff --git a/Engine/Plugins/Experimental/FullBodyIK/Source/PBIK/Public/Core/PBIKSolver.h b/Engine/Plugins/Experimental/FullBodyIK/Source/PBIK/Public/Core/PBIKSolver.h
--- a/Engine/Plugins/Experimental/FullBodyIK/Source/PBIK/Public/Core/PBIKSolver.h
+++ b/Engine/Plugins/Experimental/FullBodyIK/Source/PBIK/Public/Core/PBIKSolver.h
@@ -181,11 +181,11 @@ struct PBIK_API FPBIKSolverSettings
 	UPROPERTY(BlueprintReadWrite, EditAnywhere, Category = AdvancedSettings, meta = (ClampMin = "1", UIMin = "1.0", UIMax = "2.0"))
 	float OverRelaxation = 1.3f;
 	
-	/** When true, the solver is reset each tick to start from the current input pose. Default is true.
+	/** DEPRECATED: When true, the solver is reset each tick to start from the current input pose. Default is true.
 	 * If false, incoming animated poses are ignored and the solver starts from the results of the previous solve.
 	 * In very limited circumstances, it can be beneficial to use the pose from the previous frame.*/
-	UPROPERTY(BlueprintReadWrite, EditAnywhere, Category = AdvancedSettings)
-	bool bStartSolveFromInputPose = true;
+	UPROPERTY()
+	bool bStartSolveFromInputPose_DEPRECATED = true;
 };
 
 USTRUCT()
