Fixing crash when undo'ing creation of new retarget chain while both IK rig and retarget editors are open.

diff --git a/Engine/Plugins/Animation/IKRig/Source/IKRigEditor/Private/RetargetEditor/IKRetargetEditor.cpp b/Engine/Plugins/Animation/IKRig/Source/IKRigEditor/Private/RetargetEditor/IKRetargetEditor.cpp
--- a/Engine/Plugins/Animation/IKRig/Source/IKRigEditor/Private/RetargetEditor/IKRetargetEditor.cpp
+++ b/Engine/Plugins/Animation/IKRig/Source/IKRigEditor/Private/RetargetEditor/IKRetargetEditor.cpp
@@ -446,12 +446,14 @@ TStatId FIKRetargetEditor::GetStatId() const
 
 void FIKRetargetEditor::PostUndo(bool bSuccess)
 {
+	EditorController->AssetController->CleanAsset();
 	EditorController->HandlePreviewMeshReplaced(ERetargetSourceOrTarget::Source);
 	EditorController->HandleRetargeterNeedsInitialized();
 }
 
 void FIKRetargetEditor::PostRedo(bool bSuccess)
 {
+	EditorController->AssetController->CleanAsset();
 	EditorController->HandlePreviewMeshReplaced(ERetargetSourceOrTarget::Source);
 	EditorController->HandleRetargeterNeedsInitialized();
 }
