Fixes crash when changing preview mesh in retargeter after closing and reopening the editor.

diff --git a/Engine/Plugins/Animation/IKRig/Source/IKRigEditor/Private/RetargetEditor/IKRetargetEditorController.cpp b/Engine/Plugins/Animation/IKRig/Source/IKRigEditor/Private/RetargetEditor/IKRetargetEditorController.cpp
--- a/Engine/Plugins/Animation/IKRig/Source/IKRigEditor/Private/RetargetEditor/IKRetargetEditorController.cpp
+++ b/Engine/Plugins/Animation/IKRig/Source/IKRigEditor/Private/RetargetEditor/IKRetargetEditorController.cpp
@@ -169,6 +169,7 @@ void FIKRetargetEditorController::Close()
 {
 	AssetController->OnRetargeterNeedsInitialized().Remove(RetargeterReInitDelegateHandle);
 	AssetController->OnIKRigReplaced().Remove(IKRigReplacedDelegateHandle);
+	AssetController->OnPreviewMeshReplaced().Remove(PreviewMeshReplacedDelegateHandle);
 
 	for (const FBoundIKRig& BoundIKRig : BoundIKRigs)
 	{
