Fixing issue with IK Rig editor not initializing. Crash when editing retarget pose. Unable to playback animations.

diff --git a/Engine/Plugins/Animation/IKRig/Source/IKRigEditor/Private/RetargetEditor/IKRetargetEditor.cpp b/Engine/Plugins/Animation/IKRig/Source/IKRigEditor/Private/RetargetEditor/IKRetargetEditor.cpp
--- a/Engine/Plugins/Animation/IKRig/Source/IKRigEditor/Private/RetargetEditor/IKRetargetEditor.cpp
+++ b/Engine/Plugins/Animation/IKRig/Source/IKRigEditor/Private/RetargetEditor/IKRetargetEditor.cpp
@@ -530,25 +530,31 @@ void FIKRetargetEditor::HandlePreviewSceneCreated(const TSharedRef<IPersonaPrevi
 	EditorController->SourceSkelMeshComponent->bSelectable = false;
 	EditorController->TargetSkelMeshComponent->bSelectable = false;
 	
-	// setup an apply an anim instance to the skeletal mesh component
+	// setup and apply an anim instance to the skeletal mesh component
 	EditorController->SourceAnimInstance = NewObject<UIKRetargetAnimInstance>(EditorController->SourceSkelMeshComponent, TEXT("IKRetargetSourceAnimScriptInstance"));
 	EditorController->TargetAnimInstance = NewObject<UIKRetargetAnimInstance>(EditorController->TargetSkelMeshComponent, TEXT("IKRetargetTargetAnimScriptInstance"));
 	SetupAnimInstance();
 	
-	// set the source and target skeletal meshes on the component
-	// NOTE: this must be done AFTER setting the AnimInstance so that the correct root anim node is loaded
-	USkeletalMesh* SourceMesh = EditorController->GetSkeletalMesh(ERetargetSourceOrTarget::Source);
-	USkeletalMesh* TargetMesh = EditorController->GetSkeletalMesh(ERetargetSourceOrTarget::Target);
-	EditorController->SourceSkelMeshComponent->SetSkeletalMesh(SourceMesh);
-	EditorController->TargetSkelMeshComponent->SetSkeletalMesh(TargetMesh);
-
-	// apply mesh to the preview scene
-	InPersonaPreviewScene->SetPreviewMeshComponent(EditorController->SourceSkelMeshComponent);
-	InPersonaPreviewScene->SetPreviewMesh(SourceMesh);
-	InPersonaPreviewScene->SetAdditionalMeshesSelectable(false);
+	// set components to use custom animation mode
+	EditorController->SourceSkelMeshComponent->SetAnimationMode(EAnimationMode::AnimationCustomMode);
+	EditorController->TargetSkelMeshComponent->SetAnimationMode(EAnimationMode::AnimationCustomMode);
 	
+	// must call AddComponent() BEFORE assigning the mesh to prevent auto-assignment of a default anim instance
 	InPersonaPreviewScene->AddComponent(EditorController->SourceSkelMeshComponent, FTransform::Identity);
-	InPersonaPreviewScene->AddComponent(EditorController->TargetSkelMeshComponent, FTransform::Identity);
+    InPersonaPreviewScene->AddComponent(EditorController->TargetSkelMeshComponent, FTransform::Identity);
+    
+    // apply component to the preview scene (must be done BEFORE setting mesh)
+    InPersonaPreviewScene->SetPreviewMeshComponent(EditorController->SourceSkelMeshComponent);
+    InPersonaPreviewScene->SetAdditionalMeshesSelectable(false);
+    
+    // assign source mesh to the preview scene (applies the mesh to the source component)
+	// (must be done AFTER adding the component to prevent InitAnim() from overriding anim instance)
+    USkeletalMesh* SourceMesh = EditorController->GetSkeletalMesh(ERetargetSourceOrTarget::Source);
+    InPersonaPreviewScene->SetPreviewMesh(SourceMesh);
+    
+    // assign target mesh directly to the target component
+    USkeletalMesh* TargetMesh = EditorController->GetSkeletalMesh(ERetargetSourceOrTarget::Target);
+    EditorController->TargetSkelMeshComponent->SetSkeletalMesh(TargetMesh);
 	
 	EditorController->FixZeroHeightRetargetRoot(ERetargetSourceOrTarget::Source);
 	EditorController->FixZeroHeightRetargetRoot(ERetargetSourceOrTarget::Target);
@@ -565,9 +571,6 @@ void FIKRetargetEditor::SetupAnimInstance()
 
 	EditorController->SourceSkelMeshComponent->PreviewInstance = EditorController->SourceAnimInstance.Get();
 	EditorController->TargetSkelMeshComponent->PreviewInstance = EditorController->TargetAnimInstance.Get();
-
-	EditorController->SourceAnimInstance->InitializeAnimation();
-	EditorController->TargetAnimInstance->InitializeAnimation();
 }
 
 void FIKRetargetEditor::HandleOnPreviewSceneSettingsCustomized(IDetailLayoutBuilder& DetailBuilder) const
diff --git a/Engine/Plugins/Animation/IKRig/Source/IKRigEditor/Private/RigEditor/IKRigToolkit.cpp b/Engine/Plugins/Animation/IKRig/Source/IKRigEditor/Private/RigEditor/IKRigToolkit.cpp
--- a/Engine/Plugins/Animation/IKRig/Source/IKRigEditor/Private/RigEditor/IKRigToolkit.cpp
+++ b/Engine/Plugins/Animation/IKRig/Source/IKRigEditor/Private/RigEditor/IKRigToolkit.cpp
@@ -207,25 +207,29 @@ void FIKRigEditorToolkit::HandlePreviewSceneCreated(const TSharedRef<IPersonaPre
 	// turn off default bone rendering (we do our own in the IK Rig editor)
 	EditorController->SkelMeshComponent->SkeletonDrawMode = ESkeletonDrawMode::Hidden;
 
-	// setup an apply an anim instance to the skeletal mesh component
+	// setup and apply an anim instance to the skeletal mesh component
 	UIKRigAnimInstance* AnimInstance = NewObject<UIKRigAnimInstance>(EditorController->SkelMeshComponent, TEXT("IKRigAnimScriptInstance"));
 	EditorController->AnimInstance = AnimInstance;
 	AnimInstance->SetIKRigAsset(EditorController->AssetController->GetAsset());
-	EditorController->SkelMeshComponent->PreviewInstance = AnimInstance;
-	AnimInstance->InitializeAnimation();
 
-	// set the skeletal mesh on the component
-	// NOTE: this must be done AFTER setting the AnimInstance so that the correct root anim node is loaded
-	USkeletalMesh* Mesh = EditorController->AssetController->GetSkeletalMesh();
-	EditorController->SkelMeshComponent->SetSkeletalMesh(Mesh);
+	// must set the animation mode to "AnimationCustomMode" to prevent USkeletalMeshComponent::InitAnim() from
+	// replacing the custom ik rig anim instance with a generic preview anim instance.
+	EditorController->SkelMeshComponent->PreviewInstance = AnimInstance;
+    EditorController->SkelMeshComponent->SetAnimationMode(EAnimationMode::AnimationCustomMode);
+    
+    // must call AddComponent() BEFORE assigning the mesh to prevent auto-assignment of a default anim instance
+    InPersonaPreviewScene->AddComponent(EditorController->SkelMeshComponent, FTransform::Identity);
 
-	// apply mesh to the preview scene
+	// apply mesh component to the preview scene
 	InPersonaPreviewScene->SetPreviewMeshComponent(EditorController->SkelMeshComponent);
 	InPersonaPreviewScene->SetAllowMeshHitProxies(false);
 	InPersonaPreviewScene->SetAdditionalMeshesSelectable(false);
 	EditorController->SkelMeshComponent->bSelectable = false;
-	InPersonaPreviewScene->SetPreviewMesh(Mesh);
-	InPersonaPreviewScene->AddComponent(EditorController->SkelMeshComponent, FTransform::Identity);
+	
+	// set the skeletal mesh on the component
+    // NOTE: this must be done AFTER setting the PreviewInstance so that it assigns it as the main anim instance
+    USkeletalMesh* Mesh = EditorController->AssetController->GetSkeletalMesh();
+    InPersonaPreviewScene->SetPreviewMesh(Mesh);
 }
 
 void FIKRigEditorToolkit::HandleDetailsCreated(const TSharedRef<class IDetailsView>& InDetailsView) const
