Fixes crash when closing IK Retargeter referencing a just deleted IK Rig.

diff --git a/Engine/Plugins/Animation/IKRig/Source/IKRigEditor/Private/RetargetEditor/IKRetargetEditorController.cpp b/Engine/Plugins/Animation/IKRig/Source/IKRigEditor/Private/RetargetEditor/IKRetargetEditorController.cpp
--- a/Engine/Plugins/Animation/IKRig/Source/IKRigEditor/Private/RetargetEditor/IKRetargetEditorController.cpp
+++ b/Engine/Plugins/Animation/IKRig/Source/IKRigEditor/Private/RetargetEditor/IKRetargetEditorController.cpp
@@ -42,7 +42,11 @@ FBoundIKRig::FBoundIKRig(UIKRigDefinition* InIKRig, const FIKRetargetEditorContr
 
 void FBoundIKRig::UnBind() const
 {
-	check(IKRig);
+	if (!IsValid(IKRig))
+	{
+		return;
+	}
+	
 	UIKRigController* IKRigController = UIKRigController::GetController(IKRig);
 	IKRigController->OnIKRigNeedsInitialized().Remove(ReInitIKDelegateHandle);
 	IKRigController->OnRetargetChainAdded().Remove(AddedChainDelegateHandle);
