Fix IK Retargeter & IK Rig crashing in non-shipping builds

PR #10518: Fix IK Retargeter & IK Rig crashing in non-shipping builds
#10518

diff --git a/Engine/Plugins/Animation/IKRig/Source/IKRig/Private/AnimNodes/AnimNode_RetargetPoseFromMesh.cpp b/Engine/Plugins/Animation/IKRig/Source/IKRig/Private/AnimNodes/AnimNode_RetargetPoseFromMesh.cpp
--- a/Engine/Plugins/Animation/IKRig/Source/IKRig/Private/AnimNodes/AnimNode_RetargetPoseFromMesh.cpp
+++ b/Engine/Plugins/Animation/IKRig/Source/IKRig/Private/AnimNodes/AnimNode_RetargetPoseFromMesh.cpp
@@ -101,9 +101,12 @@ void FAnimNode_RetargetPoseFromMesh::Evaluate_AnyThread(FPoseContext& Output)
 	}
 
 	#if WITH_EDITOR
-	// live preview source asset settings in the retarget, editor only
-	// NOTE: this copies goal targets as well, but these are overwritten by IK chain goals
-	Processor->ApplySettingsFromAsset();
+	if (GIsEditor)
+	{
+		// live preview source asset settings in the retarget, editor only
+		// NOTE: this copies goal targets as well, but these are overwritten by IK chain goals
+		Processor->ApplySettingsFromAsset();
+	}
 	#endif
 
 	// apply custom profile settings to the processor
