Fixes IK Retarget editor bug, undo queue broken after exporting animation

diff --git a/Engine/Plugins/Animation/AnimationData/Source/AnimationData/Private/AnimSequencerController.cpp b/Engine/Plugins/Animation/AnimationData/Source/AnimationData/Private/AnimSequencerController.cpp
--- a/Engine/Plugins/Animation/AnimationData/Source/AnimationData/Private/AnimSequencerController.cpp
+++ b/Engine/Plugins/Animation/AnimationData/Source/AnimationData/Private/AnimSequencerController.cpp
@@ -2204,7 +2204,7 @@ void UAnimSequencerController::UpdateWithSkeleton(USkeleton* TargetSkeleton, boo
 {
 	if (ModelInterface->HasBeenPopulated())
 	{
-		OpenBracket(LOCTEXT("SettingNewskeleton", "Updating Skeleton for Animation Data Model"));
+		OpenBracket(LOCTEXT("SettingNewskeleton", "Updating Skeleton for Animation Data Model"), bShouldTransact);
 		{
 			// (re-)generate the rig hierarchy
 			Model->InitializeFKControlRig(CastChecked<UFKControlRig>(Model->GetControlRig()), TargetSkeleton);
diff --git a/Engine/Plugins/Animation/IKRig/Source/IKRigEditor/Private/RetargetEditor/IKRetargetBatchOperation.cpp b/Engine/Plugins/Animation/IKRig/Source/IKRigEditor/Private/RetargetEditor/IKRetargetBatchOperation.cpp
--- a/Engine/Plugins/Animation/IKRig/Source/IKRigEditor/Private/RetargetEditor/IKRetargetBatchOperation.cpp
+++ b/Engine/Plugins/Animation/IKRig/Source/IKRigEditor/Private/RetargetEditor/IKRetargetBatchOperation.cpp
@@ -179,17 +179,17 @@ void UIKRetargetBatchOperation::RetargetAssets(
 			UAnimationBlueprintLibrary::CopyAnimationCurveNamesToSkeleton(OldSkeleton, NewSkeleton, AnimSequenceToRetarget, ERawCurveTrackTypes::RCT_Float);	
 			// clear transform curves since those curves won't work in new skeleton
 			IAnimationDataController& Controller = AnimSequenceToRetarget->GetController();
-			const bool ShouldTransactAnimEdits = false;
-			Controller.OpenBracket(FText::FromString("Generating Retargeted Animation Data"), ShouldTransactAnimEdits);
-			Controller.RemoveAllCurvesOfType(ERawCurveTrackTypes::RCT_Transform);
+			constexpr bool bShouldTransact = false;
+			Controller.OpenBracket(FText::FromString("Generating Retargeted Animation Data"), bShouldTransact);
+			Controller.RemoveAllCurvesOfType(ERawCurveTrackTypes::RCT_Transform, bShouldTransact);
 			// clear bone tracks to prevent recompression
-			Controller.RemoveAllBoneTracks(false);
+			Controller.RemoveAllBoneTracks(bShouldTransact);
 			// set the retarget source to the target skeletal mesh
 			AnimSequenceToRetarget->RetargetSource = NAME_None;
 			AnimSequenceToRetarget->RetargetSourceAsset = Context.TargetMesh;
-			Controller.UpdateWithSkeleton(NewSkeleton);
+			Controller.UpdateWithSkeleton(NewSkeleton, bShouldTransact);
 			// done editing sequence data, close bracket
-			Controller.CloseBracket(ShouldTransactAnimEdits);
+			Controller.CloseBracket(bShouldTransact);
 		}
 
 		// replace references to other animation
@@ -299,9 +299,9 @@ void UIKRetargetBatchOperation::ConvertAnimation(
 
 		// remove all keys from the destination animation sequence
 		IAnimationDataController& TargetSeqController = DestinationSequence->GetController();
-		const bool ShouldTransactAnimEdits = false;
-		TargetSeqController.OpenBracket(FText::FromString("Generating Retargeted Animation Data"), ShouldTransactAnimEdits);
-		TargetSeqController.RemoveAllBoneTracks();
+		constexpr bool bShouldTransact = false;
+		TargetSeqController.OpenBracket(FText::FromString("Generating Retargeted Animation Data"), bShouldTransact);
+		TargetSeqController.RemoveAllBoneTracks(bShouldTransact);
 
 		// number of frames in this animation
 		const int32 NumFrames = SourceSequence->GetNumberOfSampledKeys();
@@ -376,18 +376,14 @@ void UIKRetargetBatchOperation::ConvertAnimation(
 		}
 
 		// add keys to bone tracks
-		const bool bShouldTransact = false;
 		for (int32 TargetBoneIndex=0; TargetBoneIndex<NumTargetBones; ++TargetBoneIndex)
 		{
 			const FName& TargetBoneName = TargetBoneNames[TargetBoneIndex];
 
 			const FRawAnimSequenceTrack& RawTrack = BoneTracks[TargetBoneIndex];
 			TargetSeqController.AddBoneCurve(TargetBoneName, bShouldTransact);
-			TargetSeqController.SetBoneTrackKeys(TargetBoneName, RawTrack.PosKeys, RawTrack.RotKeys, RawTrack.ScaleKeys);
+			TargetSeqController.SetBoneTrackKeys(TargetBoneName, RawTrack.PosKeys, RawTrack.RotKeys, RawTrack.ScaleKeys, bShouldTransact);
 		}
-
-		// done editing sequence data, close bracket
-		TargetSeqController.CloseBracket(ShouldTransactAnimEdits);
 	}
 }
 
