Fixed crash when undo'ing replacement of IK Rig in IK Retarget editor.

diff --git a/Engine/Plugins/Animation/IKRig/Source/IKRigEditor/Private/RetargetEditor/IKRetargetEditorController.cpp b/Engine/Plugins/Animation/IKRig/Source/IKRigEditor/Private/RetargetEditor/IKRetargetEditorController.cpp
--- a/Engine/Plugins/Animation/IKRig/Source/IKRigEditor/Private/RetargetEditor/IKRetargetEditorController.cpp
+++ b/Engine/Plugins/Animation/IKRig/Source/IKRigEditor/Private/RetargetEditor/IKRetargetEditorController.cpp
@@ -269,13 +269,7 @@ void FIKRetargetEditorController::BindToIKRigAssets()
 void FIKRetargetEditorController::HandleIKRigNeedsInitialized(UIKRigDefinition* ModifiedIKRig) const
 {
 	UIKRetargeter* Retargeter = AssetController->GetAsset();
-	check(ModifiedIKRig && Retargeter)
-	 
-	const bool bIsSource = ModifiedIKRig == Retargeter->GetSourceIKRig();
-	const bool bIsTarget = ModifiedIKRig == Retargeter->GetTargetIKRig();
-	check(bIsSource || bIsTarget);
-
-	// the target anim instance has the RetargetPoseFromMesh node which needs reinitialized with new asset version
+	check(Retargeter)
 	HandleRetargeterNeedsInitialized();
 }
 
