Fixes bug where undo/redo loses window focus in IK Rig / Retarget editors.

diff --git a/Engine/Plugins/Animation/IKRig/Source/IKRigEditor/Private/RetargetEditor/SIKRetargetAssetBrowser.cpp b/Engine/Plugins/Animation/IKRig/Source/IKRigEditor/Private/RetargetEditor/SIKRetargetAssetBrowser.cpp
--- a/Engine/Plugins/Animation/IKRig/Source/IKRigEditor/Private/RetargetEditor/SIKRetargetAssetBrowser.cpp
+++ b/Engine/Plugins/Animation/IKRig/Source/IKRigEditor/Private/RetargetEditor/SIKRetargetAssetBrowser.cpp
@@ -86,6 +86,7 @@ void SIKRetargetAssetBrowser::RefreshView()
 	AssetPickerConfig.OnGetAssetContextMenu = FOnGetAssetContextMenu::CreateSP(this, &SIKRetargetAssetBrowser::OnGetAssetContextMenu);
 	AssetPickerConfig.GetCurrentSelectionDelegates.Add(&GetCurrentSelectionDelegate);
 	AssetPickerConfig.bAllowNullSelection = false;
+	AssetPickerConfig.bFocusSearchBoxWhenOpened = false;
 
 	// hide all asset registry columns by default (we only really want the name and path)
 	TArray<UObject::FAssetRegistryTag> AssetRegistryTags;
diff --git a/Engine/Plugins/Animation/IKRig/Source/IKRigEditor/Private/RigEditor/SIKRigAssetBrowser.cpp b/Engine/Plugins/Animation/IKRig/Source/IKRigEditor/Private/RigEditor/SIKRigAssetBrowser.cpp
--- a/Engine/Plugins/Animation/IKRig/Source/IKRigEditor/Private/RigEditor/SIKRigAssetBrowser.cpp
+++ b/Engine/Plugins/Animation/IKRig/Source/IKRigEditor/Private/RigEditor/SIKRigAssetBrowser.cpp
@@ -58,6 +58,7 @@ void SIKRigAssetBrowser::RefreshView()
 	AssetPickerConfig.OnGetAssetContextMenu = FOnGetAssetContextMenu::CreateSP(this, &SIKRigAssetBrowser::OnGetAssetContextMenu);
 	AssetPickerConfig.GetCurrentSelectionDelegates.Add(&GetCurrentSelectionDelegate);
 	AssetPickerConfig.bAllowNullSelection = false;
+	AssetPickerConfig.bFocusSearchBoxWhenOpened = false;
 
 	// hide all asset registry columns by default (we only really want the name and path)
 	TArray<UObject::FAssetRegistryTag> AssetRegistryTags;
