Fixed bug in retarget editor, playback did not resume after showing retarget pose.

diff --git a/Engine/Plugins/Animation/IKRig/Source/IKRigEditor/Private/RetargetEditor/IKRetargetEditorController.cpp b/Engine/Plugins/Animation/IKRig/Source/IKRigEditor/Private/RetargetEditor/IKRetargetEditorController.cpp
--- a/Engine/Plugins/Animation/IKRig/Source/IKRigEditor/Private/RetargetEditor/IKRetargetEditorController.cpp
+++ b/Engine/Plugins/Animation/IKRig/Source/IKRigEditor/Private/RetargetEditor/IKRetargetEditorController.cpp
@@ -81,14 +81,15 @@ void FRetargetPlaybackManager::PlayAnimationAsset(UAnimationAsset* AssetToPlay)
 	}
 }
 
-void FRetargetPlaybackManager::StopPlayback() const
+void FRetargetPlaybackManager::StopPlayback()
 {
 	UIKRetargetAnimInstance* AnimInstance = EditorController.Pin()->SourceAnimInstance.Get();
 	if (!AnimInstance)
 	{
 		return;
 	}
-	
+
+	AnimThatWasPlaying = AnimInstance->GetAnimationAsset();
 	AnimInstance->SetPlaying(false);
 	AnimInstance->SetAnimationAsset(nullptr);
 }
@@ -101,15 +102,13 @@ void FRetargetPlaybackManager::PausePlayback()
 		return;
 	}
 	
-	AnimThatWasPlaying = AnimInstance->GetAnimationAsset();
-	
 	if (AnimThatWasPlaying)
 	{
 		TimeWhenPaused = AnimInstance->GetCurrentTime();
 	}
-	
+
+	AnimThatWasPlaying = AnimInstance->GetAnimationAsset();
 	AnimInstance->SetPlaying(false);
-	AnimInstance->SetAnimationAsset(nullptr);
 }
 
 void FRetargetPlaybackManager::ResumePlayback() const
diff --git a/Engine/Plugins/Animation/IKRig/Source/IKRigEditor/Public/RetargetEditor/IKRetargetEditorController.h b/Engine/Plugins/Animation/IKRig/Source/IKRigEditor/Public/RetargetEditor/IKRetargetEditorController.h
--- a/Engine/Plugins/Animation/IKRig/Source/IKRigEditor/Public/RetargetEditor/IKRetargetEditorController.h
+++ b/Engine/Plugins/Animation/IKRig/Source/IKRigEditor/Public/RetargetEditor/IKRetargetEditorController.h
@@ -65,7 +65,7 @@ struct FRetargetPlaybackManager : public TSharedFromThis<FRetargetPlaybackManage
 {
 	FRetargetPlaybackManager(TWeakPtr<FIKRetargetEditorController> InEditorController);
 	void PlayAnimationAsset(UAnimationAsset* AssetToPlay);
-	void StopPlayback() const;
+	void StopPlayback();
 	void PausePlayback();
 	void ResumePlayback() const;
 	bool IsStopped() const;
