Fixed effector influence multiplier not affecting body mover solver.

diff --git a/Engine/Plugins/Animation/IKRig/Source/IKRig/Private/Rig/Solvers/IKRig_BodyMover.cpp b/Engine/Plugins/Animation/IKRig/Source/IKRig/Private/Rig/Solvers/IKRig_BodyMover.cpp
--- a/Engine/Plugins/Animation/IKRig/Source/IKRig/Private/Rig/Solvers/IKRig_BodyMover.cpp
+++ b/Engine/Plugins/Animation/IKRig/Source/IKRig/Private/Rig/Solvers/IKRig_BodyMover.cpp
@@ -34,7 +34,7 @@ void UIKRig_BodyMover::Solve(FIKRigSkeleton& IKRigSkeleton, const FIKRigGoalCont
 	// calculate a "best fit" transform from deformed goal locations
 	TArray<FVector> InitialPoints;
 	TArray<FVector> CurrentPoints;
-	for (UIKRig_BodyMoverEffector* Effector : Effectors)
+	for (const UIKRig_BodyMoverEffector* Effector : Effectors)
 	{
 		const FIKRigGoal* Goal = Goals.FindGoalByName(Effector->GoalName);
 		if (!Goal)
@@ -43,10 +43,12 @@ void UIKRig_BodyMover::Solve(FIKRigSkeleton& IKRigSkeleton, const FIKRigGoalCont
 		}
 
 		const int32 BoneIndex = IKRigSkeleton.GetBoneIndexFromName(Effector->BoneName);
-		const FTransform InitialEffector = IKRigSkeleton.CurrentPoseGlobal[BoneIndex];
-
-		InitialPoints.Add(InitialEffector.GetTranslation());
-		CurrentPoints.Add(Goal->FinalBlendedPosition);
+		const FVector InitialPosition = IKRigSkeleton.CurrentPoseGlobal[BoneIndex].GetTranslation();
+		const FVector GoalPosition = Goal->FinalBlendedPosition;
+		const FVector FinalPosition = FMath::Lerp(InitialPosition, GoalPosition, Effector->InfluenceMultiplier);
+		
+		InitialPoints.Add(InitialPosition);
+		CurrentPoints.Add(FinalPosition);
 	}
 	
 	FVector InitialCentroid;
