Fixes undo rename goal removing effector in IK Rig asset.

diff --git a/Engine/Plugins/Animation/IKRig/Source/IKRig/Private/Rig/Solvers/IKRig_BodyMover.cpp b/Engine/Plugins/Animation/IKRig/Source/IKRig/Private/Rig/Solvers/IKRig_BodyMover.cpp
--- a/Engine/Plugins/Animation/IKRig/Source/IKRig/Private/Rig/Solvers/IKRig_BodyMover.cpp
+++ b/Engine/Plugins/Animation/IKRig/Source/IKRig/Private/Rig/Solvers/IKRig_BodyMover.cpp
@@ -174,6 +174,7 @@ void UIKRig_BodyMover::RenameGoal(const FName& OldName, const FName& NewName)
 	}
 
 	// rename
+	Effectors[GoalIndex]->Modify();
 	Effectors[GoalIndex]->GoalName = NewName;
 }
 
