Added SetTransform blueprint exposure in IK Rig

diff --git a/Engine/Plugins/Animation/IKRig/Source/IKRig/Public/Rig/Solvers/IKRig_SetTransform.h b/Engine/Plugins/Animation/IKRig/Source/IKRig/Public/Rig/Solvers/IKRig_SetTransform.h
--- a/Engine/Plugins/Animation/IKRig/Source/IKRig/Public/Rig/Solvers/IKRig_SetTransform.h
+++ b/Engine/Plugins/Animation/IKRig/Source/IKRig/Public/Rig/Solvers/IKRig_SetTransform.h
@@ -7,7 +7,7 @@
 
 #include "IKRig_SetTransform.generated.h"
 
-UCLASS()
+UCLASS(BlueprintType)
 class IKRIG_API UIKRig_SetTransformEffector : public UObject
 {
 	GENERATED_BODY()
@@ -15,15 +15,15 @@ class IKRIG_API UIKRig_SetTransformEffector : public UObject
 public:
 
 	/** If true, Goal will drive the translation of the target bone. Default is true. */
-	UPROPERTY(EditAnywhere, Category = "Set Transform Effector")
+	UPROPERTY(BlueprintReadWrite, EditAnywhere, Category = "Set Transform Effector")
 	bool bEnablePosition = true;
 
 	/** If true, Goal will drive the rotation of the target bone. Default is true. */
-	UPROPERTY(EditAnywhere, Category = "Set Transform Effector")
+	UPROPERTY(BlueprintReadWrite, EditAnywhere, Category = "Set Transform Effector")
 	bool bEnableRotation = true;
 
 	/** Blend the effector on/off. Range is 0-1. Default is 1.0. */
-	UPROPERTY(EditAnywhere, Category = "Set Transform Effector", meta = (ClampMin = "0.0", ClampMax = "1.0", UIMin = "0.0", UIMax = "1.0"))
+	UPROPERTY(BlueprintReadWrite, EditAnywhere, Category = "Set Transform Effector", meta = (ClampMin = "0.0", ClampMax = "1.0", UIMin = "0.0", UIMax = "1.0"))
 	float Alpha = 1.0f;
 };
 
