Fixed "reset to default" behavior of IK Rig anim node property.

diff --git a/Engine/Plugins/Animation/IKRig/Source/IKRigDeveloper/Private/AnimGraphNode_IKRig.cpp b/Engine/Plugins/Animation/IKRig/Source/IKRigDeveloper/Private/AnimGraphNode_IKRig.cpp
--- a/Engine/Plugins/Animation/IKRig/Source/IKRigDeveloper/Private/AnimGraphNode_IKRig.cpp
+++ b/Engine/Plugins/Animation/IKRig/Source/IKRigDeveloper/Private/AnimGraphNode_IKRig.cpp
@@ -100,13 +100,21 @@ void FIKRigGoalLayout::GenerateChildContent(IDetailChildrenBuilder& InOutChildre
 		if (bExposePosition)
 		{
 			const TSharedPtr<IPropertyHandle> PosSpaceHandle = GoalPropHandle->GetChildHandle(GET_MEMBER_NAME_CHECKED(FIKRigGoal, PositionSpace));
-			InOutChildrenBuilder.AddProperty(PosSpaceHandle.ToSharedRef());
+			IDetailPropertyRow& PropertyRow = InOutChildrenBuilder.AddProperty(PosSpaceHandle.ToSharedRef());
+			
+			// Hide the reset to default button since it provides little value
+			const FResetToDefaultOverride	ResetDefaultOverride = FResetToDefaultOverride::Create(TAttribute<bool>(false));
+			PropertyRow.OverrideResetToDefault(ResetDefaultOverride);
 		}
 
 		if (bExposeRotation)
 		{
 			const TSharedPtr<IPropertyHandle> RotSpaceHandle = GoalPropHandle->GetChildHandle(GET_MEMBER_NAME_CHECKED(FIKRigGoal, RotationSpace));
-			InOutChildrenBuilder.AddProperty(RotSpaceHandle.ToSharedRef());
+			IDetailPropertyRow& PropertyRow = InOutChildrenBuilder.AddProperty(RotSpaceHandle.ToSharedRef());
+
+			// Hide the reset to default button since it provides little value
+			const FResetToDefaultOverride	ResetDefaultOverride = FResetToDefaultOverride::Create(TAttribute<bool>(false));
+			PropertyRow.OverrideResetToDefault(ResetDefaultOverride);
 		}
 	}
 }
