Fixing extra transactions in undo queue when editing retarget pose

diff --git a/Engine/Plugins/Animation/IKRig/Source/IKRigEditor/Private/RetargetEditor/IKRetargetDetails.cpp b/Engine/Plugins/Animation/IKRig/Source/IKRigEditor/Private/RetargetEditor/IKRetargetDetails.cpp
--- a/Engine/Plugins/Animation/IKRig/Source/IKRigEditor/Private/RetargetEditor/IKRetargetDetails.cpp
+++ b/Engine/Plugins/Animation/IKRig/Source/IKRigEditor/Private/RetargetEditor/IKRetargetDetails.cpp
@@ -468,7 +468,8 @@ void UIKRetargetBoneDetails::OnNumericValueCommitted(
 				Component,
 				Representation,
 				SubComponent,
-				Value);
+				Value,
+				bIsCommit);
 			break;
 		}
 	case EIKRetargetTransformType::Bone:
@@ -482,7 +483,8 @@ void UIKRetargetBoneDetails::OnNumericValueCommitted(
 				Component,
 				Representation,
 				SubComponent,
-				Value);
+				Value,
+				bIsCommit);
 			break;
 		}
 	default:
@@ -500,7 +502,8 @@ void UIKRetargetBoneDetails::CommitValueAsRelativeOffset(
 	ESlateTransformComponent::Type Component,
 	ESlateRotationRepresentation::Type Representation,
 	ESlateTransformSubComponent::Type SubComponent,
-	FVector::FReal Value)
+	FVector::FReal Value,
+	bool bShouldTransact)
 {
 	switch(Component)
 	{
@@ -535,7 +538,7 @@ void UIKRetargetBoneDetails::CommitValueAsRelativeOffset(
 			}
 			
 			// store the new transform in the retarget pose
-			FScopedTransaction Transaction(LOCTEXT("EditRootTranslation", "Edit Retarget Root Pose Translation"));
+			FScopedTransaction Transaction(LOCTEXT("EditRootTranslation", "Edit Retarget Root Pose Translation"), bShouldTransact);
 			EditorController->AssetController->GetAsset()->Modify();
 			EditorController->AssetController->GetCurrentRetargetPose(SourceOrTarget).SetRootTranslationDelta(CurrentGlobalOffset.GetTranslation());
 			
@@ -584,7 +587,7 @@ void UIKRetargetBoneDetails::CommitValueAsRelativeOffset(
 			}
 			
 			// store the new rotation in the retarget pose
-			FScopedTransaction Transaction(LOCTEXT("EditRootRotation", "Edit Retarget Pose Rotation"));
+			FScopedTransaction Transaction(LOCTEXT("EditRootRotation", "Edit Retarget Pose Rotation"), bShouldTransact);
 			EditorController->AssetController->GetAsset()->Modify();
 			EditorController->AssetController->SetRotationOffsetForRetargetPoseBone(SelectedBone, NewLocalRotationDelta, SourceOrTarget);
 			break;
@@ -603,7 +606,8 @@ void UIKRetargetBoneDetails::CommitValueAsBoneSpace(
 	ESlateTransformComponent::Type Component,
 	ESlateRotationRepresentation::Type Representation,
 	ESlateTransformSubComponent::Type SubComponent,
-	FVector::FReal Value)
+	FVector::FReal Value,
+	bool bShouldTransact)
 {
 	switch(Component)
 	{
@@ -638,7 +642,7 @@ void UIKRetargetBoneDetails::CommitValueAsBoneSpace(
 			}
 			
 			// store the new transform in the retarget pose
-			FScopedTransaction Transaction(LOCTEXT("EditRootTranslation", "Edit Retarget Root Pose Translation"));
+			FScopedTransaction Transaction(LOCTEXT("EditRootTranslation", "Edit Retarget Root Pose Translation"), bShouldTransact);
 			EditorController->AssetController->GetAsset()->Modify();
 			EditorController->AssetController->GetCurrentRetargetPose(SourceOrTarget).SetRootTranslationDelta(CurrentGlobalOffset.GetTranslation());
 			
@@ -663,7 +667,7 @@ void UIKRetargetBoneDetails::CommitValueAsBoneSpace(
 				FQuat NewLocalRotationDelta = LocalRefTransform.GetRotation().Inverse() * CombinedLocalDeltaTransform.GetRotation();
 			
 				// store the new rotation in the retarget pose
-				FScopedTransaction Transaction(LOCTEXT("EditRootRotation", "Edit Retarget Pose Rotation"));
+				FScopedTransaction Transaction(LOCTEXT("EditRootRotation", "Edit Retarget Pose Rotation"), bShouldTransact);
 				EditorController->AssetController->GetAsset()->Modify();
 				EditorController->AssetController->SetRotationOffsetForRetargetPoseBone(SelectedBone, NewLocalRotationDelta, SourceOrTarget);
 			}
@@ -802,6 +806,12 @@ void FIKRetargetBoneDetailCustomization::CustomizeDetails(IDetailLayoutBuilder&
 		// edit transform
 		if(bIsEditable)
 		{
+			TransformWidgetArgs.OnNumericValueCommitted_Static(
+				&UIKRetargetBoneDetails::OnMultiNumericValueCommitted,
+				TransformType,
+				BonesView,
+				true);
+			
 			TransformWidgetArgs.OnNumericValueChanged_Static(
 				&UIKRetargetBoneDetails::OnMultiNumericValueCommitted,
 				ETextCommit::Default,
diff --git a/Engine/Plugins/Animation/IKRig/Source/IKRigEditor/Public/RetargetEditor/IKRetargetDetails.h b/Engine/Plugins/Animation/IKRig/Source/IKRigEditor/Public/RetargetEditor/IKRetargetDetails.h
--- a/Engine/Plugins/Animation/IKRig/Source/IKRigEditor/Public/RetargetEditor/IKRetargetDetails.h
+++ b/Engine/Plugins/Animation/IKRig/Source/IKRigEditor/Public/RetargetEditor/IKRetargetDetails.h
@@ -64,6 +64,7 @@ public:
 		ESlateTransformComponent::Type Component,
 		ESlateRotationRepresentation::Type Representation,
 		ESlateTransformSubComponent::Type SubComponent);
+	
 	/** Method to react to changes of numeric values in the widget */
 	static void OnMultiNumericValueCommitted(
 		ESlateTransformComponent::Type Component,
@@ -91,7 +92,8 @@ private:
 		ESlateTransformComponent::Type Component,
 		ESlateRotationRepresentation::Type Representation,
 		ESlateTransformSubComponent::Type SubComponent,
-		FVector::FReal Value);
+		FVector::FReal Value,
+		bool bShouldTransact);
 
 	void CommitValueAsBoneSpace(
 		UIKRetargeterController* AssetController,
@@ -102,7 +104,8 @@ private:
 		ESlateTransformComponent::Type Component,
 		ESlateRotationRepresentation::Type Representation,
 		ESlateTransformSubComponent::Type SubComponent,
-		FVector::FReal Value);
+		FVector::FReal Value,
+		bool bShouldTransact);
 
 	static TOptional<FVector::FReal> CleanRealValue(TOptional<FVector::FReal> InValue);
 	
