Fix crash when updating retarget pose at runtime with a missing IK Rig asset.

diff --git a/Engine/Plugins/Animation/IKRig/Source/IKRig/Private/Retargeter/IKRetargetProcessor.cpp b/Engine/Plugins/Animation/IKRig/Source/IKRig/Private/Retargeter/IKRetargetProcessor.cpp
--- a/Engine/Plugins/Animation/IKRig/Source/IKRig/Private/Retargeter/IKRetargetProcessor.cpp
+++ b/Engine/Plugins/Animation/IKRig/Source/IKRig/Private/Retargeter/IKRetargetProcessor.cpp
@@ -2300,6 +2300,10 @@ void UIKRetargetProcessor::ApplyNewRetargetPose(
 	// re-initialize the bone chains using the newly generated retarget pose
 	const UIKRigDefinition* SourceIKRig = RetargeterAsset->GetSourceIKRig();
 	const UIKRigDefinition* TargetIKRig = RetargeterAsset->GetTargetIKRig();
+	if (!(SourceIKRig && TargetIKRig))
+	{
+		return;
+	}
 	for (FRetargetChainPairFK& FKChainPair : ChainPairsFK)
 	{
 		const FBoneChain* SourceBoneChain = SourceIKRig->GetRetargetChainByName(FKChainPair.SourceBoneChainName);
