Fixes editor hanging when immediately using exported retargeted animation.

#swarm

diff --git a/Engine/Plugins/Animation/IKRig/Source/IKRigEditor/Private/RetargetEditor/IKRetargetBatchOperation.cpp b/Engine/Plugins/Animation/IKRig/Source/IKRigEditor/Private/RetargetEditor/IKRetargetBatchOperation.cpp
--- a/Engine/Plugins/Animation/IKRig/Source/IKRigEditor/Private/RetargetEditor/IKRetargetBatchOperation.cpp
+++ b/Engine/Plugins/Animation/IKRig/Source/IKRigEditor/Private/RetargetEditor/IKRetargetBatchOperation.cpp
@@ -171,16 +171,15 @@ void UIKRetargetBatchOperation::RetargetAssets(
 	
 	for (UAnimationAsset* AssetToRetarget : AnimationAssetsToRetarget)
 	{
-		// synchronize curves between old/new asset
-		UAnimSequence* AnimSequenceToRetarget = Cast<UAnimSequence>(AssetToRetarget);
-		if (AnimSequenceToRetarget)
+		// prepare animation sequence asset to receive retargeted animation
+		if (UAnimSequence* AnimSequenceToRetarget = Cast<UAnimSequence>(AssetToRetarget))
 		{
 			// copy curve data from source asset, preserving data in the target if present.
 			UAnimationBlueprintLibrary::CopyAnimationCurveNamesToSkeleton(OldSkeleton, NewSkeleton, AnimSequenceToRetarget, ERawCurveTrackTypes::RCT_Float);	
 			// clear transform curves since those curves won't work in new skeleton
 			IAnimationDataController& Controller = AnimSequenceToRetarget->GetController();
 			constexpr bool bShouldTransact = false;
-			Controller.OpenBracket(FText::FromString("Generating Retargeted Animation Data"), bShouldTransact);
+			Controller.OpenBracket(FText::FromString("Preparing for retargeted animation."), bShouldTransact);
 			Controller.RemoveAllCurvesOfType(ERawCurveTrackTypes::RCT_Transform, bShouldTransact);
 			// clear bone tracks to prevent recompression
 			Controller.RemoveAllBoneTracks(bShouldTransact);
@@ -384,6 +383,8 @@ void UIKRetargetBatchOperation::ConvertAnimation(
 			TargetSeqController.AddBoneCurve(TargetBoneName, bShouldTransact);
 			TargetSeqController.SetBoneTrackKeys(TargetBoneName, RawTrack.PosKeys, RawTrack.RotKeys, RawTrack.ScaleKeys, bShouldTransact);
 		}
+
+		TargetSeqController.CloseBracket(bShouldTransact);
 	}
 }
 
