IKRig: Fix format string.

#rnx

diff --git a/Engine/Plugins/Animation/IKRig/Source/IKRigEditor/Private/RetargetEditor/IKRetargetBatchOperation.cpp b/Engine/Plugins/Animation/IKRig/Source/IKRigEditor/Private/RetargetEditor/IKRetargetBatchOperation.cpp
--- a/Engine/Plugins/Animation/IKRig/Source/IKRigEditor/Private/RetargetEditor/IKRetargetBatchOperation.cpp
+++ b/Engine/Plugins/Animation/IKRig/Source/IKRigEditor/Private/RetargetEditor/IKRetargetBatchOperation.cpp
@@ -294,7 +294,7 @@ void UIKRetargetBatchOperation::ConvertAnimation(
 
 		// increment progress bar
 		FString AssetName = DestinationSequence->GetName();
-		Progress.EnterProgressFrame(1.f, FText::Format(LOCTEXT("RunningBatchRetarget", "Retargeting animation asset: {Asset}"), FText::FromString(AssetName)));
+		Progress.EnterProgressFrame(1.f, FText::Format(LOCTEXT("RunningBatchRetarget", "Retargeting animation asset: {0}"), FText::FromString(AssetName)));
 
 		// remove all keys from the destination animation sequence
 		IAnimationDataController& TargetSeqController = DestinationSequence->GetController();
