Fixes bug, IK Retargeter asset holds reference to skeletal meshes even after they have been cleared from the UI.

diff --git a/Engine/Plugins/Animation/IKRig/Source/IKRig/Public/Retargeter/IKRetargeter.h b/Engine/Plugins/Animation/IKRig/Source/IKRig/Public/Retargeter/IKRetargeter.h
--- a/Engine/Plugins/Animation/IKRig/Source/IKRig/Public/Retargeter/IKRetargeter.h
+++ b/Engine/Plugins/Animation/IKRig/Source/IKRig/Public/Retargeter/IKRetargeter.h
@@ -435,7 +435,6 @@ public:
 private:
 
 	// only ask to fix the root height once, then warn thereafter (don't nag) 
-	UPROPERTY()
 	TSet<TObjectPtr<USkeletalMesh>> MeshesAskedToFixRootHeightFor;
 #endif
 	
