Fixes IKRig slow to add multiple chains and inconsistent undo/redo behavior.

diff --git a/Engine/Plugins/Animation/IKRig/Source/IKRigEditor/Private/RetargetEditor/IKRetargetEditorController.cpp b/Engine/Plugins/Animation/IKRig/Source/IKRigEditor/Private/RetargetEditor/IKRetargetEditorController.cpp
--- a/Engine/Plugins/Animation/IKRig/Source/IKRigEditor/Private/RetargetEditor/IKRetargetEditorController.cpp
+++ b/Engine/Plugins/Animation/IKRig/Source/IKRigEditor/Private/RetargetEditor/IKRetargetEditorController.cpp
@@ -180,64 +180,6 @@ void FIKRetargetEditorController::Close()
 	}
 }
 
-void FIKRetargetEditorController::PromptUserToAssignIKRig(const ERetargetSourceOrTarget SourceOrTarget)
-{
-	// early out if we already have an IK Rig assigned or is unattended
-	if (AssetController->GetIKRig(SourceOrTarget) || FApp::IsUnattended())
-	{
-		return;
-	}
-	
-	// Load the content browser module to display an asset picker
-	FContentBrowserModule& ContentBrowserModule = FModuleManager::LoadModuleChecked<FContentBrowserModule>("ContentBrowser");
-	FAssetPickerConfig AssetPickerConfig;
-	// must set the parent UObject so that the resulting list filters correctly in multi-project environments
-	AssetPickerConfig.AdditionalReferencingAssets.Add(AssetController->GetAsset());
-	/** The asset picker will only show skeletal meshes */
-	AssetPickerConfig.Filter.ClassPaths.Add(UIKRigDefinition::StaticClass()->GetClassPathName());
-	AssetPickerConfig.Filter.bRecursiveClasses = true;
-	/** The delegate that fires when an asset was selected */
-	AssetPickerConfig.OnAssetSelected =  FOnAssetSelected::CreateLambda([this, &SourceOrTarget](const FAssetData& AssetData)
-	{
-		IKRigPickerWindow->RequestDestroyWindow();
-
-		// set the IK Rig
-		if (const TObjectPtr<UIKRigDefinition> IKRigAsset = Cast<UIKRigDefinition>(AssetData.GetAsset()))
-		{
-			AssetController->SetIKRig(SourceOrTarget, IKRigAsset.Get());
-		}
-	});
-	/** The default view mode should be a list view */
-	AssetPickerConfig.InitialAssetViewType = EAssetViewType::List;
-
-	// change title depending on if we're picking a source or target IK Rig
-	FText TitleText;
-	if (SourceOrTarget == ERetargetSourceOrTarget::Source)
-	{
-		TitleText = LOCTEXT("SelectSourceIKRig", "Select Source IK Rig (To Copy Animation From)");
-	}
-	else
-	{
-		TitleText = LOCTEXT("SelectTargetIKRig", "Select Target IK Rig (To Copy Animation To)");
-	}
-
-	// show an asset browser in a pop-up modal dialog
-	IKRigPickerWindow = SNew(SWindow)
-	.Title(TitleText)
-	.ClientSize(FVector2D(500, 600))
-	.SupportsMinimize(false).SupportsMaximize(false)
-	[
-		SNew(SBorder)
-		.BorderImage( FAppStyle::GetBrush("Menu.Background") )
-		[
-			ContentBrowserModule.Get().CreateAssetPicker(AssetPickerConfig)
-		]
-	];
-
-	GEditor->EditorAddModalWindow(IKRigPickerWindow.ToSharedRef());
-	IKRigPickerWindow.Reset();
-}
-
 void FIKRetargetEditorController::BindToIKRigAssets()
 {
 	const UIKRetargeter* Asset = AssetController->GetAsset();
diff --git a/Engine/Plugins/Animation/IKRig/Source/IKRigEditor/Private/RetargetEditor/IKRetargeterController.cpp b/Engine/Plugins/Animation/IKRig/Source/IKRigEditor/Private/RetargetEditor/IKRetargeterController.cpp
--- a/Engine/Plugins/Animation/IKRig/Source/IKRigEditor/Private/RetargetEditor/IKRetargeterController.cpp
+++ b/Engine/Plugins/Animation/IKRig/Source/IKRigEditor/Private/RetargetEditor/IKRetargeterController.cpp
@@ -53,6 +53,8 @@ void UIKRetargeterController::SetIKRig(const ERetargetSourceOrTarget SourceOrTar
 {
 	FScopeLock Lock(&ControllerLock);
 	
+	FScopedReinitializeIKRetargeter Reinitialize(this);
+	
 	if (SourceOrTarget == ERetargetSourceOrTarget::Source)
 	{
 		Asset->SourceIKRigAsset = IKRig;
@@ -78,7 +80,6 @@ void UIKRetargeterController::SetIKRig(const ERetargetSourceOrTarget SourceOrTar
 	// update any editors attached to this asset
 	BroadcastIKRigReplaced(SourceOrTarget);
 	BroadcastPreviewMeshReplaced(SourceOrTarget);
-	BroadcastNeedsReinitialized();
 }
 
 const UIKRigDefinition* UIKRetargeterController::GetIKRig(const ERetargetSourceOrTarget SourceOrTarget) const
@@ -98,6 +99,9 @@ void UIKRetargeterController::SetPreviewMesh(
 	USkeletalMesh* PreviewMesh) const
 {
 	FScopeLock Lock(&ControllerLock);
+
+	FScopedTransaction Transaction(LOCTEXT("SetPreviewMesh_Transaction", "Set Preview Mesh"));
+	FScopedReinitializeIKRetargeter Reinitialize(this);
 	
 	if (SourceOrTarget == ERetargetSourceOrTarget::Source)
 	{
@@ -113,7 +117,6 @@ void UIKRetargeterController::SetPreviewMesh(
 	
 	// update any editors attached to this asset
 	BroadcastPreviewMeshReplaced(SourceOrTarget);
-	BroadcastNeedsReinitialized();
 }
 
 USkeletalMesh* UIKRetargeterController::GetPreviewMesh(const ERetargetSourceOrTarget SourceOrTarget) const
@@ -147,6 +150,9 @@ FTargetRootSettings UIKRetargeterController::GetRootSettings() const
 void UIKRetargeterController::SetRootSettings(const FTargetRootSettings& RootSettings) const
 {
 	FScopeLock Lock(&ControllerLock);
+	FScopedTransaction Transaction(LOCTEXT("SetRootSettings_Transaction", "Set Root Settings"));
+	FScopedReinitializeIKRetargeter Reinitialize(this);
+	GetAsset()->Modify();
 	GetAsset()->GetRootSettingsUObject()->Settings = RootSettings;
 }
 
@@ -159,6 +165,9 @@ FRetargetGlobalSettings UIKRetargeterController::GetGlobalSettings() const
 void UIKRetargeterController::SetGlobalSettings(const FRetargetGlobalSettings& GlobalSettings) const
 {
 	FScopeLock Lock(&ControllerLock);
+	FScopedTransaction Transaction(LOCTEXT("SetGlobalSettings_Transaction", "Set Global Settings"));
+	FScopedReinitializeIKRetargeter Reinitialize(this);
+	GetAsset()->Modify();
 	GetAsset()->GetGlobalSettingsUObject()->Settings = GlobalSettings;
 }
 
@@ -178,9 +187,13 @@ FTargetChainSettings UIKRetargeterController::GetRetargetChainSettings(const FNa
 bool UIKRetargeterController::SetRetargetChainSettings(const FName& TargetChainName, const FTargetChainSettings& Settings) const
 {
 	FScopeLock Lock(&ControllerLock);
+
+	FScopedTransaction Transaction(LOCTEXT("SetChainSettings_Transaction", "Set Chain Settings"));
+	FScopedReinitializeIKRetargeter Reinitialize(this);
 	
 	if (URetargetChainSettings* ChainSettings = GetChainSettings(TargetChainName))
 	{
+		ChainSettings->Modify();
 		ChainSettings->Settings = Settings;
 		return true;
 	}
@@ -323,6 +336,7 @@ void UIKRetargeterController::AutoMapChains(const EAutoMapChainType AutoMapType,
 {
 	FScopeLock Lock(&ControllerLock);
 	FScopedTransaction Transaction(LOCTEXT("AutoMapRetargetChains", "Auto-Map Retarget Chains"));
+	FScopedReinitializeIKRetargeter Reinitialize(this);
 	
 	CleanChainMapping();
 
@@ -412,7 +426,7 @@ void UIKRetargeterController::HandleRetargetChainAdded(UIKRigDefinition* IKRig)
 
 	// add the new chain to the mapping data
 	CleanChainMapping();
-	BroadcastNeedsReinitialized();
+	FScopedReinitializeIKRetargeter Reinitialize(this);
 }
 
 void UIKRetargeterController::HandleRetargetChainRenamed(UIKRigDefinition* IKRig, FName OldChainName, FName NewChainName) const
@@ -425,7 +439,7 @@ void UIKRetargeterController::HandleRetargetChainRenamed(UIKRigDefinition* IKRig
 		if (ChainNameToUpdate == OldChainName)
 		{
 			ChainNameToUpdate = NewChainName;
-			BroadcastNeedsReinitialized();
+			FScopedReinitializeIKRetargeter Reinitialize(this);
 			return;
 		}
 	}
@@ -444,7 +458,7 @@ void UIKRetargeterController::HandleRetargetChainRemoved(UIKRigDefinition* IKRig
 			if (ChainMap->SourceChain == InChainRemoved)
 			{
 				ChainMap->SourceChain = NAME_None;
-				BroadcastNeedsReinitialized();
+				FScopedReinitializeIKRetargeter Reinitialize(this);
 				return;
 			}
 		}
@@ -460,7 +474,7 @@ void UIKRetargeterController::HandleRetargetChainRemoved(UIKRigDefinition* IKRig
 	if (ChainIndex != INDEX_NONE)
 	{
 		Asset->ChainSettings.RemoveAt(ChainIndex);
-		BroadcastNeedsReinitialized();
+		FScopedReinitializeIKRetargeter Reinitialize(this);
 	}
 }
 
@@ -475,11 +489,9 @@ bool UIKRetargeterController::SetSourceChain(FName SourceChainName, FName Target
 	}
 	
 	FScopedTransaction Transaction(LOCTEXT("SetRetargetChainSource", "Set Retarget Chain Source"));
+	FScopedReinitializeIKRetargeter Reinitialize(this);
 	ChainSettings->Modify();
 	ChainSettings->SourceChain = SourceChainName;
-	
-	BroadcastNeedsReinitialized();
-
 	return true;
 }
 
@@ -533,6 +545,7 @@ FName UIKRetargeterController::CreateRetargetPose(const FName& NewPoseName, cons
 {
 	FScopeLock Lock(&ControllerLock);
 	FScopedTransaction Transaction(LOCTEXT("CreateRetargetPose", "Create Retarget Pose"));
+	FScopedReinitializeIKRetargeter Reinitialize(this);
 	Asset->Modify();
 
 	// create a new pose with a unique name 
@@ -543,8 +556,6 @@ FName UIKRetargeterController::CreateRetargetPose(const FName& NewPoseName, cons
 	FName& CurrentRetargetPoseName = SourceOrTarget == ERetargetSourceOrTarget::Source ? Asset->CurrentSourceRetargetPose : Asset->CurrentTargetRetargetPose;
 	CurrentRetargetPoseName = UniqueNewPoseName;
 
-	BroadcastNeedsReinitialized();
-
 	return UniqueNewPoseName;
 }
 
@@ -563,6 +574,7 @@ bool UIKRetargeterController::RemoveRetargetPose(const FName& PoseToRemove, cons
 
 	FScopeLock Lock(&ControllerLock);
 	FScopedTransaction Transaction(LOCTEXT("RemoveRetargetPose", "Remove Retarget Pose"));
+	FScopedReinitializeIKRetargeter Reinitialize(this);
 	Asset->Modify();
 
 	Poses.Remove(PoseToRemove);
@@ -572,9 +584,7 @@ bool UIKRetargeterController::RemoveRetargetPose(const FName& PoseToRemove, cons
 	{
 		SetCurrentRetargetPose(UIKRetargeter::GetDefaultPoseName(), SourceOrTarget);
 	}
-
-	BroadcastNeedsReinitialized();
-
+	
 	return true;
 }
 
@@ -589,6 +599,7 @@ FName UIKRetargeterController::DuplicateRetargetPose( const FName PoseToDuplicat
 
 	FScopeLock Lock(&ControllerLock);
 	FScopedTransaction Transaction(LOCTEXT("DuplicateRetargetPose", "Duplicate Retarget Pose"));
+	FScopedReinitializeIKRetargeter Reinitialize(this);
 	Asset->Modify();
 
 	// create a new pose with a unique name
@@ -601,9 +612,6 @@ FName UIKRetargeterController::DuplicateRetargetPose( const FName PoseToDuplicat
 	// set duplicate to be the current pose
 	FName& CurrentRetargetPoseName = SourceOrTarget == ERetargetSourceOrTarget::Source ? Asset->CurrentSourceRetargetPose : Asset->CurrentTargetRetargetPose;
 	CurrentRetargetPoseName = UniqueNewPoseName;
-
-	BroadcastNeedsReinitialized();
-
 	return UniqueNewPoseName;
 }
 
@@ -629,6 +637,7 @@ bool UIKRetargeterController::RenameRetargetPose(const FName OldPoseName, const
 	const bool bWasCurrentPose = GetCurrentRetargetPoseName(SourceOrTarget) == OldPoseName;
 	
 	FScopedTransaction Transaction(LOCTEXT("RenameRetargetPose", "Rename Retarget Pose"));
+	FScopedReinitializeIKRetargeter Reinitialize(this);
 	Asset->Modify();
 
 	// make sure new name is unique
@@ -646,9 +655,6 @@ bool UIKRetargeterController::RenameRetargetPose(const FName OldPoseName, const
 	{
 		SetCurrentRetargetPose(UniqueNewPoseName, SourceOrTarget);
 	}
-
-	BroadcastNeedsReinitialized();
-
 	return true;
 }
 
@@ -666,6 +672,8 @@ void UIKRetargeterController::ResetRetargetPose(
 	}
 	
 	FIKRetargetPose& PoseToEdit = Poses[PoseToReset];
+
+	FScopedReinitializeIKRetargeter Reinitialize(this);
 	
 	if (BonesToReset.IsEmpty())
 	{
@@ -694,8 +702,6 @@ void UIKRetargeterController::ResetRetargetPose(
 			}
 		}
 	}
-
-	BroadcastNeedsReinitialized();
 }
 
 FName UIKRetargeterController::GetCurrentRetargetPoseName(const ERetargetSourceOrTarget SourceOrTarget) const
@@ -716,12 +722,10 @@ bool UIKRetargeterController::SetCurrentRetargetPose(FName NewCurrentPose, const
 	}
 	
 	FScopedTransaction Transaction(LOCTEXT("SetCurrentPose", "Set Current Pose"));
+	FScopedReinitializeIKRetargeter Reinitialize(this);
 	Asset->Modify();
 	FName& CurrentPose = SourceOrTarget == ERetargetSourceOrTarget::Source ? Asset->CurrentSourceRetargetPose : Asset->CurrentTargetRetargetPose;
 	CurrentPose = NewCurrentPose;
-	
-	BroadcastNeedsReinitialized();
-
 	return true;
 }
 
@@ -743,9 +747,9 @@ void UIKRetargeterController::SetRotationOffsetForRetargetPoseBone(
 	const ERetargetSourceOrTarget SourceOrTarget) const
 {
 	FScopeLock Lock(&ControllerLock);
-	const UIKRigDefinition* IKRig = SourceOrTarget == ERetargetSourceOrTarget::Source ? GetAsset()->GetSourceIKRig() : GetAsset()->GetTargetIKRig();
 	FIKRetargetPose& Pose = GetCurrentRetargetPose(SourceOrTarget);
 	Pose.SetDeltaRotationForBone(BoneName, RotationOffset);
+	const UIKRigDefinition* IKRig = SourceOrTarget == ERetargetSourceOrTarget::Source ? GetAsset()->GetSourceIKRig() : GetAsset()->GetTargetIKRig();
 	Pose.SortHierarchically(IKRig->GetSkeleton());
 }
 
diff --git a/Engine/Plugins/Animation/IKRig/Source/IKRigEditor/Private/RigEditor/IKRigController.cpp b/Engine/Plugins/Animation/IKRig/Source/IKRigEditor/Private/RigEditor/IKRigController.cpp
--- a/Engine/Plugins/Animation/IKRig/Source/IKRigEditor/Private/RigEditor/IKRigController.cpp
+++ b/Engine/Plugins/Animation/IKRig/Source/IKRigEditor/Private/RigEditor/IKRigController.cpp
@@ -52,11 +52,10 @@
 	}
 	
 	FScopedTransaction Transaction(LOCTEXT("AddBoneSetting_Label", "Add Bone Setting"));
+	FScopedReinitializeIKRig Reinitialize(this);
 	UIKRigSolver* Solver = Asset->Solvers[SolverIndex];
 	Solver->Modify();
 	Solver->AddBoneSetting(BoneName);
-	BroadcastNeedsReinitialized();
-
 	return true;
 }
 
@@ -106,9 +105,9 @@
 	}
 
 	FScopedTransaction Transaction(LOCTEXT("RemoveBoneSetting_Label", "Remove Bone Setting"));
+	FScopedReinitializeIKRig Reinitialize(this);
 	Solver->Modify();
 	Solver->RemoveBoneSetting(BoneName);
-	BroadcastNeedsReinitialized();
 
 	return true;
 }
@@ -175,42 +174,32 @@
 
 FName UIKRigController::AddRetargetChain(const FName ChainName, const FName StartBoneName, const FName EndBoneName, const FName GoalName) const
 {
-	return AddRetargetChain(FBoneChain(ChainName, StartBoneName, EndBoneName, GoalName));
+	return AddRetargetChainInternal(FBoneChain(ChainName, StartBoneName, EndBoneName, GoalName));
 }
 
-FName UIKRigController::AddRetargetChain(const FBoneChain& BoneChain) const
+FName UIKRigController::AddRetargetChainInternal(const FBoneChain& BoneChain) const
 {
 	if (BoneChain.StartBone.BoneName != NAME_None && Asset->Skeleton.GetBoneIndexFromName(BoneChain.StartBone.BoneName) == INDEX_NONE)
 	{
 		UE_LOG(LogIKRigEditor, Warning, TEXT("Could not create retarget chain. Start Bone does not exist, %s."), *BoneChain.StartBone.BoneName.ToString());
-		return NAME_None; // bone doesn't exist
+		return NAME_None; // start bone doesn't exist
 	}
 
 	if (BoneChain.EndBone.BoneName != NAME_None && Asset->Skeleton.GetBoneIndexFromName(BoneChain.EndBone.BoneName) == INDEX_NONE)
 	{
 		UE_LOG(LogIKRigEditor, Warning, TEXT("Could not create retarget chain. End Bone does not exist, %s."), *BoneChain.EndBone.BoneName.ToString());
-		return NAME_None; // bone doesn't exist
+		return NAME_None; // end bone doesn't exist
 	}
-	
+
+	// uniquify the chain name
 	FBoneChain ChainToAdd = BoneChain;
-
-	// if no name specified, use a default
-	if (ChainToAdd.ChainName == NAME_None)
-	{
-		const FName DefaultChainName = FName("DefaultChainName");
-		ChainToAdd.ChainName = DefaultChainName;
-	}
-	
 	ChainToAdd.ChainName = GetUniqueRetargetChainName(BoneChain.ChainName);
 	
 	FScopedTransaction Transaction(LOCTEXT("AddRetargetChain_Label", "Add Retarget Chain"));
+	FScopedReinitializeIKRig Reinitialize(this);
 	Asset->Modify();
-	
 	const int32 NewChainIndex = Asset->RetargetDefinition.BoneChains.Emplace(ChainToAdd);
-	
 	RetargetChainAdded.Broadcast(Asset);
-	BroadcastNeedsReinitialized();
-
 	return Asset->RetargetDefinition.BoneChains[NewChainIndex].ChainName;
 }
 
@@ -232,7 +221,7 @@
 	if (Asset->RetargetDefinition.BoneChains.RemoveAll(Pred) > 0)
 	{
 		RetargetChainRemoved.Broadcast(Asset, ChainName);
-		BroadcastNeedsReinitialized();
+		FScopedReinitializeIKRig Reinitialize(this);
 		return true;
 	}
 
@@ -253,10 +242,11 @@
 	const FName UniqueChainName = GetUniqueRetargetChainName(NewChainName);
 	
 	FScopedTransaction Transaction(LOCTEXT("RenameRetargetChain_Label", "Rename Retarget Chain"));
+	FScopedReinitializeIKRig Reinitialize(this);
 	Asset->Modify();
 	Chain->ChainName = UniqueChainName;
 	RetargetChainRenamed.Broadcast(GetAsset(), ChainName, UniqueChainName);
-	BroadcastNeedsReinitialized();
+	
 	return UniqueChainName;
 }
 
@@ -271,9 +261,9 @@
 	if (FBoneChain* BoneChain = Asset->RetargetDefinition.GetEditableBoneChainByName(ChainName))
 	{
 		FScopedTransaction Transaction(LOCTEXT("SetRetargetChainStartBone_Label", "Set Retarget Chain Start Bone"));
+		FScopedReinitializeIKRig Reinitialize(this);
 		Asset->Modify();
 		BoneChain->StartBone = StartBoneName;
-		BroadcastNeedsReinitialized();
 		return true;
 	}
 
@@ -292,9 +282,9 @@
 	if (FBoneChain* BoneChain = Asset->RetargetDefinition.GetEditableBoneChainByName(ChainName))
 	{
 		FScopedTransaction Transaction(LOCTEXT("SetRetargetChainEndBone_Label", "Set Retarget Chain End Bone"));
+		FScopedReinitializeIKRig Reinitialize(this);
 		Asset->Modify();
 		BoneChain->EndBone = EndBoneName;
-		BroadcastNeedsReinitialized();
 		return true;
 	}
 
@@ -318,9 +308,9 @@
 	}
 
 	FScopedTransaction Transaction(LOCTEXT("SetRetargetChainGoal_Label", "Set Retarget Chain Goal"));
+	FScopedReinitializeIKRig Reinitialize(this);
 	Asset->Modify();
 	BoneChain->IKGoalName = GoalNameToUse;
-	BroadcastNeedsReinitialized();
 	return true;
 }
 
@@ -394,12 +384,9 @@
 	}
 
 	FScopedTransaction Transaction(LOCTEXT("SetRetargetRootBone_Label", "Set Retarget Root Bone"));
+	FScopedReinitializeIKRig Reinitialize(this);
 	Asset->Modify();
-	
 	Asset->RetargetDefinition.RootBone = NewRootBone;
-
-	BroadcastNeedsReinitialized();
-
 	return true;
 }
 
@@ -425,7 +412,7 @@
 	});
 }
 
-FName UIKRigController::GetUniqueRetargetChainName(const FName NameToMakeUnique) const
+FName UIKRigController::GetUniqueRetargetChainName(FName NameToMakeUnique) const
 {
 	auto IsNameBeingUsed = [this](const FName NameToTry)->bool
 	{
@@ -439,6 +426,12 @@
 		return false;
 	};
 
+	// if no name specified, use a default
+	if (NameToMakeUnique == NAME_None)
+	{
+		NameToMakeUnique = FName("DefaultChainName");
+	}
+
 	// check if name is already unique
 	if (!IsNameBeingUsed(NameToMakeUnique))
 	{
@@ -547,7 +540,9 @@
 	if (GoalName == NAME_None)
 	{
 		return NAME_None;
-	}
+	}
+
+
 	
 	const TArray<FBoneChain>& Chains = GetRetargetChains();
 	for (const FBoneChain& Chain : Chains)
@@ -565,32 +560,24 @@
 // SKELETON
 //
 
-bool UIKRigController::SetSkeletalMesh(USkeletalMesh* SkeletalMesh, bool bTransact) const
+bool UIKRigController::SetSkeletalMesh(USkeletalMesh* SkeletalMesh) const
 {
 	// first determine runtime compatibility between the IK Rig asset and the skeleton we're trying to run it on
 	if (!IsSkeletalMeshCompatible(SkeletalMesh))
 	{
 		UE_LOG(LogIKRigEditor, Warning, TEXT("Trying to initialize IKRig with a Skeleton that is missing required bones. See output log. {0}"), *Asset->GetName());
-		BroadcastNeedsReinitialized();
+		FScopedReinitializeIKRig Reinitialize(this);
 		return false;
 	}
 	
 	FScopedTransaction Transaction(LOCTEXT("SetSkeletalMesh_Label", "Set Skeletal Mesh"));
-	if (bTransact)
-	{	
-		Asset->Modify();
-	}
-
+	FScopedReinitializeIKRig Reinitialize(this);
+	Asset->Modify();
 	// update stored skeletal mesh used for previewing results
 	Asset->PreviewSkeletalMesh = SkeletalMesh;
 	// copy skeleton data from the actual skeleton we want to run on
 	const FIKRigInputSkeleton InputSkeleton = FIKRigInputSkeleton(SkeletalMesh);
 	Asset->Skeleton.SetInputSkeleton(InputSkeleton, GetAsset()->Skeleton.ExcludedBones);
-	// update goal's initial transforms to reflect new skeleton
-	ResetInitialGoalTransforms();
-
-	BroadcastNeedsReinitialized();
-
 	return true;
 }
 
@@ -628,8 +615,8 @@
 	}
 
 	FScopedTransaction Transaction(LOCTEXT("SetBoneExcluded_Label", "Set Bone Excluded"));
+	FScopedReinitializeIKRig Reinitialize(this);
 	Asset->Modify();
-
 	if (bExclude)
 	{
 		Asset->Skeleton.ExcludedBones.Add(BoneName);
@@ -638,9 +625,6 @@
 	{
 		Asset->Skeleton.ExcludedBones.Remove(BoneName);
 	}
-
-	BroadcastNeedsReinitialized();
-
 	return true;
 }
 
@@ -676,16 +660,11 @@
 	}
 
 	FScopedTransaction Transaction(LOCTEXT("AddSolver_Label", "Add Solver"));
+	FScopedReinitializeIKRig Reinitialize(this);
 	Asset->Modify();
-
 	UIKRigSolver* NewSolver = NewObject<UIKRigSolver>(Asset, InIKRigSolverClass, NAME_None, RF_Transactional);
 	check(NewSolver);
-
-	const int32 SolverIndex = Asset->Solvers.Add(NewSolver);
-
-	BroadcastNeedsReinitialized();
-
-	return SolverIndex;
+	return Asset->Solvers.Add(NewSolver);
 }
 
 bool UIKRigController::RemoveSolver(const int32 SolverIndex) const
@@ -699,10 +678,9 @@
 	}
 
 	FScopedTransaction Transaction(LOCTEXT("RemoveSolver_Label", "Remove Solver"));
+	FScopedReinitializeIKRig Reinitialize(this);
 	Asset->Modify();
 	Asset->Solvers.RemoveAt(SolverIndex);
-
-	BroadcastNeedsReinitialized();
 	return true;
 }
 
@@ -727,15 +705,12 @@
 	}
 
 	FScopedTransaction Transaction(LOCTEXT("ReorderSolver_Label", "Reorder Solvers"));
+	FScopedReinitializeIKRig Reinitialize(this);
 	Asset->Modify();
-
 	UIKRigSolver* SolverToMove = Asset->Solvers[SolverToMoveIndex];
 	Asset->Solvers.Insert(SolverToMove, TargetSolverIndex + 1);
 	const int32 SolverToRemove = TargetSolverIndex > SolverToMoveIndex ? SolverToMoveIndex : SolverToMoveIndex + 1;
 	Asset->Solvers.RemoveAt(SolverToRemove);
-
-	BroadcastNeedsReinitialized();
-	
 	return true;
 }
 
@@ -748,14 +723,10 @@
 	}
 	
 	FScopedTransaction Transaction(LOCTEXT("SetSolverEndabled_Label", "Enable/Disable Solver"));
+	FScopedReinitializeIKRig Reinitialize(this);
 	UIKRigSolver* Solver = Asset->Solvers[SolverIndex];
 	Solver->Modify();
-	Asset->Modify();
-
 	Solver->SetEnabled(bIsEnabled);
-
-	BroadcastNeedsReinitialized();
-	
 	return true;
 }
 
@@ -785,13 +756,10 @@
 	}
 
 	FScopedTransaction Transaction(LOCTEXT("SetRootBone_Label", "Set Root Bone"));
+	FScopedReinitializeIKRig Reinitialize(this);
 	UIKRigSolver* Solver = Asset->Solvers[SolverIndex];
 	Solver->Modify();
-
 	Solver->SetRootBone(RootBoneName);
-
-	BroadcastNeedsReinitialized();
-
 	return true;
 }
 
@@ -827,13 +795,10 @@
 	}
 
 	FScopedTransaction Transaction(LOCTEXT("SetEndBone_Label", "Set End Bone"));
+	FScopedReinitializeIKRig Reinitialize(this);
 	UIKRigSolver* Solver = Asset->Solvers[SolverIndex];
 	Solver->Modify();
-
 	Solver->SetEndBone(EndBoneName);
-
-	BroadcastNeedsReinitialized();
-
 	return true;
 }
 
@@ -912,19 +877,17 @@
 	}
 	
 	FScopedTransaction Transaction(LOCTEXT("AddNewGoal_Label", "Add New Goal"));
-	Asset->Modify();
+	FScopedReinitializeIKRig Reinitialize(this, true /*bGoalsChanged*/);
 
 	UIKRigEffectorGoal* NewGoal = NewObject<UIKRigEffectorGoal>(Asset, UIKRigEffectorGoal::StaticClass(), NAME_None, RF_Transactional);
 	NewGoal->BoneName = BoneName;
 	NewGoal->GoalName = GoalName;
+	Asset->Modify();
 	Asset->Goals.Add(NewGoal);
 
 	// set initial transform
 	NewGoal->InitialTransform = GetRefPoseTransformOfBone(NewGoal->BoneName);
 	NewGoal->CurrentTransform = NewGoal->InitialTransform;
- 
-	BroadcastNeedsReinitialized();
-	BroadcastGoalsChange();
 	
 	return NewGoal->GoalName;
 }
@@ -938,6 +901,7 @@
 	}
 
 	FScopedTransaction Transaction(LOCTEXT("RemoveGoal_Label", "Remove Goal"));
+	FScopedReinitializeIKRig Reinitialize(this, true /*bGoalsChanged*/);
 
 	// remove from all the solvers
 	const FName GoalToRemove = Asset->Goals[GoalIndex]->GoalName;
@@ -960,9 +924,6 @@
 		}
 	}
 
-	BroadcastNeedsReinitialized();
-	BroadcastGoalsChange();
-	
 	return true;
 }
 
@@ -986,9 +947,10 @@
 	const FName NewName = GetUniqueGoalName(FName(CleanName));
 	
 	FScopedTransaction Transaction(LOCTEXT("RenameGoal_Label", "Rename Goal"));
-	Asset->Modify();
+	FScopedReinitializeIKRig Reinitialize(this, true /*bGoalsChanged*/);
 
 	// rename in core
+	Asset->Modify();
 	Asset->Goals[GoalIndex]->Modify();
 	Asset->Goals[GoalIndex]->GoalName = NewName;
 
@@ -1008,9 +970,6 @@
 		Solver->RenameGoal(OldName, NewName);
 	}
 
-	BroadcastNeedsReinitialized();
-	BroadcastGoalsChange();
-
 	return NewName;
 }
 
@@ -1070,6 +1029,7 @@
 	}
 	
 	FScopedTransaction Transaction(LOCTEXT("SetGoalBone_Label", "Set Goal Bone"));
+	FScopedReinitializeIKRig Reinitialize(this);
 
 	// update goal's bone name and it's initial transform
 	TObjectPtr<UIKRigEffectorGoal> Goal = Asset->Goals[GoalIndex]; 
@@ -1084,8 +1044,6 @@
 		Solver->SetGoalBone(GoalName, NewBoneName);
 	}
 
-	BroadcastNeedsReinitialized();
-	
 	return true;
 }
 
@@ -1136,14 +1094,11 @@
 	}
 
 	FScopedTransaction Transaction(LOCTEXT("ConnectGoalSolver_Label", "Connect Goal to Solver"));
-
+	FScopedReinitializeIKRig Reinitialize(this);
 	UIKRigSolver* Solver = Asset->Solvers[SolverIndex];
 	const UIKRigEffectorGoal& Goal = *GetGoal(GoalName);
 	Solver->Modify();
 	Solver->AddGoal(Goal);
-
-	BroadcastNeedsReinitialized();
-	
 	return true;
 }
 
@@ -1164,12 +1119,10 @@
 	}
 
     FScopedTransaction Transaction(LOCTEXT("DisconnectGoalSolver_Label", "Disconnect Goal from Solver"));
+	FScopedReinitializeIKRig Reinitialize(this);
 	UIKRigSolver* Solver = Asset->Solvers[SolverIndex];
 	Solver->Modify();
-	
 	Solver->RemoveGoal(GoalToRemove);
-
-	BroadcastNeedsReinitialized();
 	return true;
 }
 
@@ -1308,14 +1261,14 @@
 
 void UIKRigController::BroadcastGoalsChange() const
 {
-	if (Asset)
+	if(ensure(Asset))
 	{
 		static const FName GoalsPropName = GET_MEMBER_NAME_CHECKED(UIKRigDefinition, Goals);
 		if (FProperty* GoalProperty = UIKRigDefinition::StaticClass()->FindPropertyByName(GoalsPropName) )
 		{
 			FPropertyChangedEvent GoalPropertyChangedEvent(GoalProperty);
 			FCoreUObjectDelegates::OnObjectPropertyChanged.Broadcast(Asset, GoalPropertyChangedEvent);
-		}
+		}	
 	}
 }
 
diff --git a/Engine/Plugins/Animation/IKRig/Source/IKRigEditor/Private/RigEditor/IKRigEditMode.cpp b/Engine/Plugins/Animation/IKRig/Source/IKRigEditor/Private/RigEditor/IKRigEditMode.cpp
--- a/Engine/Plugins/Animation/IKRig/Source/IKRigEditor/Private/RigEditor/IKRigEditMode.cpp
+++ b/Engine/Plugins/Animation/IKRig/Source/IKRigEditor/Private/RigEditor/IKRigEditMode.cpp
@@ -516,26 +516,11 @@ bool FIKRigEditMode::InputKey(FEditorViewportClient* ViewportClient, FViewport*
 		
 	if (Key == EKeys::Delete || Key == EKeys::BackSpace)
 	{
-		const TSharedPtr<FIKRigEditorController> Controller = EditorController.Pin();
-		if (!Controller.IsValid())
+		if (EditorController.IsValid())
 		{
-			return false;
-		}
-	
-		TArray<FName> SelectedGoalNames;
-		Controller->GetSelectedGoalNames(SelectedGoalNames);
-		if (SelectedGoalNames.IsEmpty())
-		{
-			return false; // nothing selected to manipulate
-		}
-
-		for (const FName& GoalName : SelectedGoalNames)
-		{
-			Controller->AssetController->RemoveGoal(GoalName);
+			EditorController.Pin()->HandleDeleteSelectedElements();
+			return true;
 		}
-
-		Controller->RefreshAllViews();
-		return true;
 	}
 
 	return false;
diff --git a/Engine/Plugins/Animation/IKRig/Source/IKRigEditor/Private/RigEditor/IKRigEditorController.cpp b/Engine/Plugins/Animation/IKRig/Source/IKRigEditor/Private/RigEditor/IKRigEditorController.cpp
--- a/Engine/Plugins/Animation/IKRig/Source/IKRigEditor/Private/RigEditor/IKRigEditorController.cpp
+++ b/Engine/Plugins/Animation/IKRig/Source/IKRigEditor/Private/RigEditor/IKRigEditorController.cpp
@@ -312,60 +312,6 @@ void FIKRigEditorController::Close() const
 	AssetController->OnIKRigNeedsInitialized().Remove(ReinitializeDelegateHandle);
 }
 
-void FIKRigEditorController::PromptUserToAssignMesh()
-{
-	// do we already have a skeletal mesh assigned?
-	if (AssetController->GetSkeletalMesh())
-	{
-		return;
-	}
-
-	// is there already an imported hierarchy of bones?
-	if (!AssetController->GetIKRigSkeleton().BoneNames.IsEmpty())
-	{
-		return;
-	}
-
-	// no skeletal mesh imported yet... so let's prompt the user to pick one...
-	
-	// Load the content browser module to display an asset picker
-	FContentBrowserModule& ContentBrowserModule = FModuleManager::LoadModuleChecked<FContentBrowserModule>("ContentBrowser");
-	FAssetPickerConfig AssetPickerConfig;
-	// must set the parent UObject so that the resulting list filters correctly in multi-project environments
-	AssetPickerConfig.AdditionalReferencingAssets.Add(AssetController->GetAsset());
-	// the asset picker will only show skeletal meshes
-	AssetPickerConfig.Filter.ClassPaths.Add(USkeletalMesh::StaticClass()->GetClassPathName());
-	// the delegate that fires when an asset is selected
-	AssetPickerConfig.OnAssetSelected = FOnAssetSelected::CreateLambda([this](const FAssetData& AssetData)
-	{
-		MeshPickerWindow->RequestDestroyWindow();
-		
-		if (const TObjectPtr<USkeletalMesh> SkeletalMesh = Cast<USkeletalMesh>(AssetData.GetAsset()))
-		{
-			// import the skeleton data into the IK Rig
-			AssetController->SetSkeletalMesh(SkeletalMesh.Get());
-		}
-		
-	});
-	// the default view mode should be a list view
-	AssetPickerConfig.InitialAssetViewType = EAssetViewType::List;
-
-	MeshPickerWindow = SNew(SWindow)
-	.Title(LOCTEXT("CreateIKRigOptions", "Assign Skeletal Mesh to IK Rig"))
-	.ClientSize(FVector2D(500, 600))
-	.SupportsMinimize(false) .SupportsMaximize(false)
-	[
-		SNew(SBorder)
-		.BorderImage( FAppStyle::GetBrush("Menu.Background") )
-		[
-			ContentBrowserModule.Get().CreateAssetPicker(AssetPickerConfig)
-		]
-	];
-
-	GEditor->EditorAddModalWindow(MeshPickerWindow.ToSharedRef());
-	MeshPickerWindow.Reset();
-}
-
 UIKRigProcessor* FIKRigEditorController::GetIKRigProcessor() const
 {
 	if (AnimInstance)
@@ -419,49 +365,70 @@ void FIKRigEditorController::HandleIKRigNeedsInitialized(UIKRigDefinition* Modif
 		SkelMeshComponent->EnablePreview(true, nullptr);
 	}
 
-
 	// update the bone details so it can pull on the current data
 	BoneDetails->AnimInstancePtr = AnimInstance;
 	BoneDetails->AssetPtr = ModifiedIKRig;
 
-	RefreshAllViews();
-}
-
-void FIKRigEditorController::Reset() const
-{
-	SkelMeshComponent->ShowReferencePose(true);
-	AssetController->ResetGoalTransforms();
-}
-
-void FIKRigEditorController::RefreshAllViews() const
-{
-	if (SolverStackView.IsValid())
+	// refresh all views
 	{
-		SolverStackView->RefreshStackView();
-	}
+		if (SolverStackView.IsValid())
+		{
+			SolverStackView->RefreshStackView();
+		}
 
-	if (SkeletonView.IsValid())
-	{
-		SkeletonView->RefreshTreeView();
-	}
+		if (SkeletonView.IsValid())
+		{
+			SkeletonView->RefreshTreeView();
+		}
 
-	if (DetailsView.IsValid())
-	{
-		DetailsView->ForceRefresh();
-	}
+		if (DetailsView.IsValid())
+		{
+			DetailsView->ForceRefresh();
+		}
 
-	if (RetargetingView.IsValid())
-	{
-		RetargetingView->RefreshView();
+		if (RetargetingView.IsValid())
+		{
+			RetargetingView->RefreshView();
+		}
+		
+		if (AssetBrowserView.IsValid())
+		{
+			AssetBrowserView.Get()->RefreshView(); // refresh the asset browser to ensure it shows compatible sequences
+		}
 	}
+}
 
-	// refresh the asset browser to ensure it shows compatible sequences
-	if (AssetBrowserView.IsValid())
+void FIKRigEditorController::HandleDeleteSelectedElements()
+{
+	FScopedTransaction Transaction(LOCTEXT("DeleteSelectedIKRigElements_Label", "Delete Selected IK Rig Elements"));
+	FScopedReinitializeIKRig Reinitialize(AssetController);
+	
+	TArray<TSharedPtr<FIKRigTreeElement>> SelectedItems = SkeletonView->GetSelectedItems();
+	for (const TSharedPtr<FIKRigTreeElement>& SelectedItem : SelectedItems)
 	{
-		AssetBrowserView.Get()->RefreshView();
+		switch(SelectedItem->ElementType)
+		{
+		case IKRigTreeElementType::GOAL:
+			AssetController->RemoveGoal(SelectedItem->GoalName);
+			break;
+		case IKRigTreeElementType::SOLVERGOAL:
+			AssetController->DisconnectGoalFromSolver(SelectedItem->EffectorGoalName, SelectedItem->EffectorIndex);
+			break;
+		case IKRigTreeElementType::BONE_SETTINGS:
+			AssetController->RemoveBoneSetting(SelectedItem->BoneSettingBoneName, SelectedItem->BoneSettingsSolverIndex);
+			break;
+		default:
+			break; // can't delete anything else
+		}
 	}
 }
 
+void FIKRigEditorController::Reset() const
+{
+	SkelMeshComponent->ShowReferencePose(true);
+	AssetController->ResetGoalTransforms();
+}
+
 void FIKRigEditorController::RefreshTreeView() const
 {
 	if (SkeletonView.IsValid())
@@ -486,6 +453,10 @@ void FIKRigEditorController::AddNewGoals(const TArray<FName>& GoalNames, const T
 {
 	check(GoalNames.Num() == BoneNames.Num());
 
+	FScopedTransaction Transaction(LOCTEXT("AddNewGoals_Label", "Add New Goals"));
+	FScopedReinitializeIKRig Reinitialize(AssetController);
+	AssetController->GetAsset()->Modify();
+
 	// add a default solver if there isn't one already
 	PromptToAddDefaultSolver();
 
@@ -517,13 +488,10 @@ void FIKRigEditorController::AddNewGoals(const TArray<FName>& GoalNames, const T
 		LastCreatedGoalName = GoalName;
 	}
 	
-	// were any goals created?
+	// show last created goal in details view
 	if (LastCreatedGoalName != NAME_None)
 	{
-		// show last created goal in details view
 		ShowDetailsForGoal(LastCreatedGoalName);
-		// update all views
-		RefreshAllViews();
 	}
 }
 
@@ -673,15 +641,16 @@ void FIKRigEditorController::CreateNewRetargetChains()
 	{
 		SkeletonView->GetSelectedBoneChains(SelectedBoneChains);
 	}
-
-	const FIKRigSkeleton& IKRigSkeleton = AssetController->GetIKRigSkeleton();
+	
+	FScopedTransaction Transaction(LOCTEXT("AddMultipleRetargetChains_Label", "Add Retarget Chain(s)"));
+	FScopedReinitializeIKRig Reinitialize(AssetController);
 	
 	if (!SelectedBoneChains.IsEmpty())
 	{
 		// create a chain for each selected chain in hierarchy
 		for (FBoneChain& BoneChain : SelectedBoneChains)
 		{
-			ChainAnalyzer.AssignBestGuessName(BoneChain, IKRigSkeleton);
+			ChainAnalyzer.AssignBestGuessName(BoneChain, AssetController->GetIKRigSkeleton());
 			PromptToAddNewRetargetChain(BoneChain);
 		}
 	}
@@ -691,8 +660,6 @@ void FIKRigEditorController::CreateNewRetargetChains()
 		FBoneChain Chain(FRetargetChainAnalyzer::GetDefaultChainName(), NAME_None, NAME_None, NAME_None);
 		PromptToAddNewRetargetChain(Chain);
 	}
-	
-	RefreshAllViews();
 }
 
 bool FIKRigEditorController::PromptToAddDefaultSolver() const
@@ -739,7 +706,7 @@ bool FIKRigEditorController::PromptToAddDefaultSolver() const
 		}
 	}
 	
-	TSharedRef<SComboBox<TSharedPtr<FIKRigSolverTypeAndName>>> SolverOptionBox = SNew(SComboBox<TSharedPtr<FIKRigSolverTypeAndName>>)
+	const TSharedRef<SComboBox<TSharedPtr<FIKRigSolverTypeAndName>>> SolverOptionBox = SNew(SComboBox<TSharedPtr<FIKRigSolverTypeAndName>>)
 	.OptionsSource(&SolverTypes)
 	.OnGenerateWidget_Lambda([](TSharedPtr<FIKRigSolverTypeAndName> Item)
 	{
@@ -777,9 +744,12 @@ bool FIKRigEditorController::PromptToAddDefaultSolver() const
 
 	if (SelectedSolver->SolverType != nullptr && SolverStackView.IsValid())
 	{
-		SolverStackView->AddNewSolver(SelectedSolver->SolverType);
+		AssetController->AddSolver(SelectedSolver->SolverType);
 	}
 
+	// must refresh the view so that subsequent goal operations see a selected solver to connect to
+	SolverStackView->RefreshStackView();
+	
 	return true;
 }
 
@@ -933,7 +903,7 @@ void FIKRigEditorController::PromptToAssignGoalToChain(const FName NewGoalName)
 	{
 		return; // cancel button pressed, or window closed
 	}
-
+	
 	AssetController->SetRetargetChainGoal(ChainToAddGoalTo, NewGoal->GoalName);
 }
 
@@ -1058,12 +1028,12 @@ FName FIKRigEditorController::PromptToAddNewRetargetChain(FBoneChain& BoneChain)
 	}
 
 	// add the retarget chain
-	const FName NewChainName = AssetController->AddRetargetChain(BoneChain);
-
+	const FName NewChainName = AssetController->AddRetargetChainInternal(BoneChain);
+	
 	// did user choose to assign a goal
 	if (UserChoice == 1)
 	{
-		FName GoalName = NAME_None;
+		FName GoalName;
 		if (bHasExistingGoal)
 		{
 			// use the existing goal
@@ -1089,8 +1059,6 @@ FName FIKRigEditorController::PromptToAddNewRetargetChain(FBoneChain& BoneChain)
 		// assign the existing goal to it
 		AssetController->SetRetargetChainGoal(NewChainName, GoalName);
 	}
-
-	RefreshAllViews();
 	
 	return NewChainName;
 }
diff --git a/Engine/Plugins/Animation/IKRig/Source/IKRigEditor/Private/RigEditor/IKRigToolkit.cpp b/Engine/Plugins/Animation/IKRig/Source/IKRigEditor/Private/RigEditor/IKRigToolkit.cpp
--- a/Engine/Plugins/Animation/IKRig/Source/IKRigEditor/Private/RigEditor/IKRigToolkit.cpp
+++ b/Engine/Plugins/Animation/IKRig/Source/IKRigEditor/Private/RigEditor/IKRigToolkit.cpp
@@ -188,12 +188,12 @@ void FIKRigEditorToolkit::HandleOnPreviewSceneSettingsCustomized(IDetailLayoutBu
 
 void FIKRigEditorToolkit::PostUndo(bool bSuccess)
 {
-	EditorController->AssetController->BroadcastNeedsReinitialized();
+	FScopedReinitializeIKRig Reinitialize(EditorController->AssetController);
 }
 
 void FIKRigEditorToolkit::PostRedo(bool bSuccess)
 {
-	EditorController->AssetController->BroadcastNeedsReinitialized();
+	FScopedReinitializeIKRig Reinitialize(EditorController->AssetController);
 }
 
 void FIKRigEditorToolkit::HandlePreviewSceneCreated(const TSharedRef<IPersonaPreviewScene>& InPersonaPreviewScene)
diff --git a/Engine/Plugins/Animation/IKRig/Source/IKRigEditor/Private/RigEditor/SIKRigHierarchy.cpp b/Engine/Plugins/Animation/IKRig/Source/IKRigEditor/Private/RigEditor/SIKRigHierarchy.cpp
--- a/Engine/Plugins/Animation/IKRig/Source/IKRigEditor/Private/RigEditor/SIKRigHierarchy.cpp
+++ b/Engine/Plugins/Animation/IKRig/Source/IKRigEditor/Private/RigEditor/SIKRigHierarchy.cpp
@@ -229,7 +229,6 @@ void SIKRigHierarchyItem::OnNameCommitted(const FText& InText, ETextCommit::Type
 		WeakTreeElement.Pin()->GoalName = NewName;
 	}
 	
-	Controller->RefreshAllViews();
 	HierarchyView.Pin()->ReplaceItemInSelection(OldText, WeakTreeElement.Pin()->Key);
 }
 
@@ -373,11 +372,7 @@ FReply SIKRigSkeletonRow::HandleAcceptDrop(
 
 	// re-parent the goal to a different bone
 	const UIKRigController* AssetController = Controller->AssetController;
-	const bool bWasReparented = AssetController->SetGoalBone(DraggedElement.GoalName, TargetItem.Get()->BoneName);
-	if (bWasReparented)
-	{
-		Controller->RefreshAllViews();
-	}
+	AssetController->SetGoalBone(DraggedElement.GoalName, TargetItem.Get()->BoneName);
 	
 	return FReply::Handled();
 }
@@ -713,7 +708,7 @@ void SIKRigHierarchy::BindCommands()
         FCanExecuteAction::CreateSP(this, &SIKRigHierarchy::CanAddNewGoal));
 	
 	CommandList->MapAction(Commands.DeleteElement,
-        FExecuteAction::CreateSP(this, &SIKRigHierarchy::HandleDeleteElement),
+        FExecuteAction::CreateSP(this, &SIKRigHierarchy::HandleDeleteElements),
         FCanExecuteAction::CreateSP(this, &SIKRigHierarchy::CanDeleteElement));
 
 	CommandList->MapAction(Commands.ConnectGoalToSolver,
@@ -862,7 +857,7 @@ bool SIKRigHierarchy::CanAddNewGoal() const
 	return true;
 }
 
-void SIKRigHierarchy::HandleDeleteElement()
+void SIKRigHierarchy::HandleDeleteElements()
 {
 	const TSharedPtr<FIKRigEditorController> Controller = EditorController.Pin();
 	if (!Controller.IsValid())
@@ -870,30 +865,7 @@ void SIKRigHierarchy::HandleDeleteElement()
 		return; 
 	}
 
-	TArray<TSharedPtr<FIKRigTreeElement>> SelectedItems = TreeView->GetSelectedItems();
-	for (const TSharedPtr<FIKRigTreeElement>& SelectedItem : SelectedItems)
-	{
-		switch(SelectedItem->ElementType)
-		{
-			case IKRigTreeElementType::GOAL:
-				Controller->AssetController->RemoveGoal(SelectedItem->GoalName);
-				break;
-			case IKRigTreeElementType::SOLVERGOAL:
-				Controller->AssetController->DisconnectGoalFromSolver(SelectedItem->EffectorGoalName, SelectedItem->EffectorIndex);
-				break;
-			case IKRigTreeElementType::BONE_SETTINGS:
-				Controller->AssetController->RemoveBoneSetting(SelectedItem->BoneSettingBoneName, SelectedItem->BoneSettingsSolverIndex);
-				break;
-			default:
-				break; // can't delete anything else
-		}
-	}
-	
-	RefreshTreeView();
-
-	Controller->ShowEmptyDetails();
-	// update all views
-	Controller->RefreshAllViews();
+	Controller->HandleDeleteSelectedElements();
 }
 bool SIKRigHierarchy::CanDeleteElement() const
 {
@@ -956,27 +928,26 @@ void SIKRigHierarchy::ConnectSelectedGoalsToSelectedSolvers(bool bConnect)
 	TArray<TSharedPtr<FSolverStackElement>> SelectedSolvers;
 	Controller->GetSelectedSolvers(SelectedSolvers);
 
-	UIKRigController* AssetController = Controller->AssetController;
+	FScopedTransaction Transaction(LOCTEXT("ConnectGoalsToSolver_Label", "Connect/Disconnect Goal(s) to Solver"));
+	FScopedReinitializeIKRig Reinitialize(Controller->AssetController);
+	
 	for (const TSharedPtr<FIKRigTreeElement>& GoalElement : SelectedGoals)
 	{
 		const FName GoalName = GoalElement->GoalName;
-		const UIKRigEffectorGoal* Goal = AssetController->GetGoal(GoalName);
+		const UIKRigEffectorGoal* Goal = Controller->AssetController->GetGoal(GoalName);
 		check(Goal);
 		for (const TSharedPtr<FSolverStackElement>& SolverElement : SelectedSolvers)
 		{
 			if (bConnect)
 			{
-				AssetController->ConnectGoalToSolver(Goal->GoalName, SolverElement->IndexInStack);	
+				Controller->AssetController->ConnectGoalToSolver(Goal->GoalName, SolverElement->IndexInStack);	
 			}
 			else
 			{
-				AssetController->DisconnectGoalFromSolver(Goal->GoalName, SolverElement->IndexInStack);	
+				Controller->AssetController->DisconnectGoalFromSolver(Goal->GoalName, SolverElement->IndexInStack);	
 			}
 		}
 	}
-
-	// add/remove new effector under goal in skeleton view
-	RefreshTreeView();
 }
 
 int32 SIKRigHierarchy::GetNumSelectedGoalToSolverConnections(bool bCountOnlyConnected) const
@@ -1216,9 +1187,11 @@ void SIKRigHierarchy::HandleRemoveBoneSettings()
 	// get selected bones
 	TArray<TSharedPtr<FIKRigTreeElement>> SelectedBones;
 	GetSelectedBones(SelectedBones);
+
+	FScopedTransaction Transaction(LOCTEXT("RemoveBoneSettings_Label", "Remove Bone Settings"));
+	FScopedReinitializeIKRig Reinitialize(Controller->AssetController);
 	
-	// add settings for bone on all selected solvers (ignored if already present)
-	UIKRigController* AssetController = Controller->AssetController;
+	// remove settings for bone on all selected solvers
 	TArray<TSharedPtr<FSolverStackElement>> SelectedSolvers;
 	Controller->GetSelectedSolvers(SelectedSolvers);
 	FName BoneToShowInDetailsView;
@@ -1226,15 +1199,12 @@ void SIKRigHierarchy::HandleRemoveBoneSettings()
 	{
 		for (const TSharedPtr<FSolverStackElement>& Solver : SelectedSolvers)
 		{
-			AssetController->RemoveBoneSetting(BoneItem->BoneName, Solver->IndexInStack);
+			Controller->AssetController->RemoveBoneSetting(BoneItem->BoneName, Solver->IndexInStack);
 			BoneToShowInDetailsView = BoneItem->BoneName;
 		}
 	}
-
-	Controller->ShowDetailsForBone(BoneToShowInDetailsView);
 	
-	// show new icon when bone has settings applied
-	RefreshTreeView();
+	Controller->ShowDetailsForBone(BoneToShowInDetailsView);
 }
 
 bool SIKRigHierarchy::CanRemoveBoneSettings()
@@ -1393,9 +1363,6 @@ void SIKRigHierarchy::HandleSetRetargetRoot()
 
 	// set the first selected bone as the retarget root
 	Controller->AssetController->SetRetargetRoot(SelectedBones[0]->BoneName);
-
-	// show root bone after being set
-	Controller->RefreshAllViews();
 }
 
 bool SIKRigHierarchy::CanSetRetargetRoot()
@@ -1414,7 +1381,6 @@ void SIKRigHierarchy::HandleClearRetargetRoot()
 	}
 	
 	Controller->AssetController->SetRetargetRoot(NAME_None);
-	Controller->RefreshAllViews();
 }
 
 bool SIKRigHierarchy::CanClearRetargetRoot()
diff --git a/Engine/Plugins/Animation/IKRig/Source/IKRigEditor/Private/RigEditor/SIKRigRetargetChainList.cpp b/Engine/Plugins/Animation/IKRig/Source/IKRigEditor/Private/RigEditor/SIKRigRetargetChainList.cpp
--- a/Engine/Plugins/Animation/IKRig/Source/IKRigEditor/Private/RigEditor/SIKRigRetargetChainList.cpp
+++ b/Engine/Plugins/Animation/IKRig/Source/IKRigEditor/Private/RigEditor/SIKRigRetargetChainList.cpp
@@ -716,6 +716,10 @@ void SIKRigRetargetChainList::MirrorSelectedChains() const
 	const UIKRigController* AssetController = Controller.AssetController;
 	const FIKRigSkeleton& IKRigSkeleton = AssetController->GetIKRigSkeleton();
 
+	FScopedTransaction Transaction(LOCTEXT("MirrorRetargetChain_Label", "Mirror Retarget Chains"));
+	FScopedReinitializeIKRig Reinitialize(AssetController, true /*bGoalsChanged*/);
+	AssetController->GetAsset()->Modify();
+
 	for (const FName& SelectedChainName : SelectedChainNames)
 	{
 		const FBoneChain* Chain = AssetController->GetRetargetChainByName(SelectedChainName);
diff --git a/Engine/Plugins/Animation/IKRig/Source/IKRigEditor/Private/RigEditor/SIKRigSolverStack.cpp b/Engine/Plugins/Animation/IKRig/Source/IKRigEditor/Private/RigEditor/SIKRigSolverStack.cpp
--- a/Engine/Plugins/Animation/IKRig/Source/IKRigEditor/Private/RigEditor/SIKRigSolverStack.cpp
+++ b/Engine/Plugins/Animation/IKRig/Source/IKRigEditor/Private/RigEditor/SIKRigSolverStack.cpp
@@ -376,7 +376,7 @@ void SIKRigSolverStack::RefreshStackView()
 	}
 
 	// record/restore selection
-	int32 IndexToSelect = -1; // default to nothing selected
+	int32 IndexToSelect = 0; // default to first solver selected
 	const TArray<TSharedPtr<FSolverStackElement>> SelectedItems = ListView.Get()->GetSelectedItems();
 	if (!SelectedItems.IsEmpty())
 	{
@@ -443,7 +443,10 @@ FReply SIKRigSolverStack::OnDragDetected(
 
 void SIKRigSolverStack::OnSelectionChanged(TSharedPtr<FSolverStackElement> InItem, ESelectInfo::Type SelectInfo)
 {
-	ShowDetailsForItem(InItem);
+	if (SelectInfo != ESelectInfo::Direct)
+	{
+		ShowDetailsForItem(InItem);	
+	}
 }
 
 void SIKRigSolverStack::OnItemClicked(TSharedPtr<FSolverStackElement> InItem)
diff --git a/Engine/Plugins/Animation/IKRig/Source/IKRigEditor/Public/RetargetEditor/IKRetargetEditorController.h b/Engine/Plugins/Animation/IKRig/Source/IKRigEditor/Public/RetargetEditor/IKRetargetEditorController.h
--- a/Engine/Plugins/Animation/IKRig/Source/IKRigEditor/Public/RetargetEditor/IKRetargetEditorController.h
+++ b/Engine/Plugins/Animation/IKRig/Source/IKRigEditor/Public/RetargetEditor/IKRetargetEditorController.h
@@ -90,9 +90,6 @@ public:
 	/** Close the editor */
 	void Close();
 	
-	// prompt user to assign an IK Rig to this retargeter (either source or target), returns true if an IK Rig was assigned
-	void PromptUserToAssignIKRig(const ERetargetSourceOrTarget SourceOrTarget);
-	
 	/** Bind callbacks to this IK Rig */
 	void BindToIKRigAssets();
 	/** callback when IK Rig asset requires reinitialization */
diff --git a/Engine/Plugins/Animation/IKRig/Source/IKRigEditor/Public/RetargetEditor/IKRetargeterController.h b/Engine/Plugins/Animation/IKRig/Source/IKRigEditor/Public/RetargetEditor/IKRetargeterController.h
--- a/Engine/Plugins/Animation/IKRig/Source/IKRigEditor/Public/RetargetEditor/IKRetargeterController.h
+++ b/Engine/Plugins/Animation/IKRig/Source/IKRigEditor/Public/RetargetEditor/IKRetargeterController.h
@@ -286,6 +286,9 @@ private:
 
 	// only allow modifications to data model from one thread at a time
 	mutable FCriticalSection ControllerLock;
+
+	// prevent reinitializing from inner operations
+	mutable int32 ReinitializeScopeCounter = 0;
 	
 public:
 
@@ -299,4 +302,24 @@ public:
 	FOnRetargeterNeedsInitialized& OnRetargeterNeedsInitialized(){ return RetargeterNeedsInitialized; };
 	
 	friend class UIKRetargeter;
+	friend struct FScopedReinitializeIKRetargeter;
 };
+
+struct FScopedReinitializeIKRetargeter
+{
+	FScopedReinitializeIKRetargeter(const UIKRetargeterController *InController)
+	{
+		InController->ReinitializeScopeCounter++;
+		Controller = InController;
+	}
+	~FScopedReinitializeIKRetargeter()
+	{
+		if (--Controller->ReinitializeScopeCounter == 0)
+		{
+			Controller->BroadcastNeedsReinitialized();
+		}
+	};
+
+	const UIKRetargeterController* Controller;
+};
+
diff --git a/Engine/Plugins/Animation/IKRig/Source/IKRigEditor/Public/RigEditor/IKRigController.h b/Engine/Plugins/Animation/IKRig/Source/IKRigEditor/Public/RigEditor/IKRigController.h
--- a/Engine/Plugins/Animation/IKRig/Source/IKRigEditor/Public/RigEditor/IKRigController.h
+++ b/Engine/Plugins/Animation/IKRig/Source/IKRigEditor/Public/RigEditor/IKRigController.h
@@ -43,7 +43,7 @@ public:
 	// Sets the preview mesh to use. Loads the hierarchy into the asset's IKRigSkeleton.
 	// Returns true if the mesh was able to be set. False if it was incompatible for any reason. 
 	UFUNCTION(BlueprintCallable, BlueprintPure=false, Category=IKRig)
-	bool SetSkeletalMesh(USkeletalMesh* SkeletalMesh, bool bTransact=false) const;
+	bool SetSkeletalMesh(USkeletalMesh* SkeletalMesh) const;
 
 	// Get the skeletal mesh this asset is initialized with 
 	UFUNCTION(BlueprintCallable, Category=IKRig)
@@ -318,8 +318,8 @@ public:
 	// RETARGETING C++ ONLY API
 	//
 	
-	// Add a Chain with the given BoneChain settings. Returns newly created chain name. 
-	FName AddRetargetChain(const FBoneChain& BoneChain) const;
+	// Add a Chain with the given BoneChain settings. Returns newly created chain name.
+	FName AddRetargetChainInternal(const FBoneChain& BoneChain) const;
 	
 	// Get read-only access to a single retarget chain with the given name 
 	const FBoneChain* GetRetargetChainByName(const FName ChainName) const;
@@ -334,7 +334,7 @@ public:
 	void SortRetargetChains() const;
 
 	// Make unique name for a retargeting bone chain. Adds a numbered suffix to make it unique.
-	FName GetUniqueRetargetChainName(const FName NameToMakeUnique) const;
+	FName GetUniqueRetargetChainName(FName NameToMakeUnique) const;
 
 	// Returns true if this is a valid chain. Produces array of bone indices between start and end (inclusive).
 	// Optionally provide a runtime skeleton from an IKRigProcessor to get indices for a running instance (otherwise uses stored hierarchy in asset)
@@ -367,6 +367,10 @@ public:
 
 private:
 
+	// prevent reinitializing from inner operations
+	mutable int32 ReinitializeScopeCounter = 0;
+	mutable bool bGoalsChangedInScope = false;
+
 	// The actual IKRigDefinition asset that this Controller modifies. 
 	UPROPERTY(transient)
 	TObjectPtr<UIKRigDefinition> Asset = nullptr;
@@ -378,4 +382,30 @@ private:
 	void BroadcastGoalsChange() const;
 	
 	friend class UIKRigDefinition;
+	friend struct FScopedReinitializeIKRig;
+};
+
+struct FScopedReinitializeIKRig
+{
+	FScopedReinitializeIKRig(const UIKRigController *InController, const bool bGoalsChanged=false)
+	{
+		InController->ReinitializeScopeCounter++;
+		InController->bGoalsChangedInScope = bGoalsChanged ? true : InController->bGoalsChangedInScope;
+		Controller = InController;
+	}
+	~FScopedReinitializeIKRig()
+	{
+		if (--Controller->ReinitializeScopeCounter == 0)
+		{
+			Controller->BroadcastNeedsReinitialized();
+			
+			if (Controller->bGoalsChangedInScope)
+			{
+				Controller->bGoalsChangedInScope = false;
+				Controller->BroadcastGoalsChange();
+			}
+		}
+	};
+
+	const UIKRigController* Controller;
 };
diff --git a/Engine/Plugins/Animation/IKRig/Source/IKRigEditor/Public/RigEditor/IKRigEditorController.h b/Engine/Plugins/Animation/IKRig/Source/IKRigEditor/Public/RigEditor/IKRigEditorController.h
--- a/Engine/Plugins/Animation/IKRig/Source/IKRigEditor/Public/RigEditor/IKRigEditorController.h
+++ b/Engine/Plugins/Animation/IKRig/Source/IKRigEditor/Public/RigEditor/IKRigEditorController.h
@@ -113,9 +113,6 @@ public:
 	// cleanup when editor closed 
 	void Close() const;
 
-	// initial setup, used when asset is empty (no hierarchy, no mesh assigned)
-	void PromptUserToAssignMesh();
-
 	// get the currently active processor running the IK Rig in the editor 
 	UIKRigProcessor* GetIKRigProcessor() const;
 	// get the currently running IKRig skeleton (if there is a running processor) 
@@ -124,6 +121,9 @@ public:
 	// callback when IK Rig requires re-initialization 
 	void HandleIKRigNeedsInitialized(UIKRigDefinition* ModifiedIKRig) const;
 
+	// delete key pressed (in either viewport or tree view)
+	void HandleDeleteSelectedElements();
+
 	// create goals 
 	void AddNewGoals(const TArray<FName>& GoalNames, const TArray<FName>& BoneNames);
 	// clear all selected objects 
@@ -134,8 +134,6 @@ public:
 	void HandleBoneSelectedInViewport(const FName& BoneName, bool bReplace) const;
 	// reset all goals to initial transforms 
 	void Reset() const;
-	// refresh all views 
-	void RefreshAllViews() const;
 	// refresh just the skeleton tree view 
 	void RefreshTreeView() const;
 	// clear the output log 
diff --git a/Engine/Plugins/Animation/IKRig/Source/IKRigEditor/Public/RigEditor/SIKRigHierarchy.h b/Engine/Plugins/Animation/IKRig/Source/IKRigEditor/Public/RigEditor/SIKRigHierarchy.h
--- a/Engine/Plugins/Animation/IKRig/Source/IKRigEditor/Public/RigEditor/SIKRigHierarchy.h
+++ b/Engine/Plugins/Animation/IKRig/Source/IKRigEditor/Public/RigEditor/SIKRigHierarchy.h
@@ -220,7 +220,7 @@ private:
 	bool CanAddNewGoal() const;
 	void HandleRenameGoal() const;
 	bool CanRenameGoal() const;
-	void HandleDeleteElement();
+	void HandleDeleteElements();
 	bool CanDeleteElement() const;
 	/** END creating / renaming / deleting goals */
 
