Fixing unity build error from 27679172

diff --git a/Engine/Plugins/Animation/IKRig/Source/IKRigEditor/Private/RetargetEditor/SRetargetOpStack.cpp b/Engine/Plugins/Animation/IKRig/Source/IKRigEditor/Private/RetargetEditor/SRetargetOpStack.cpp
--- a/Engine/Plugins/Animation/IKRig/Source/IKRigEditor/Private/RetargetEditor/SRetargetOpStack.cpp
+++ b/Engine/Plugins/Animation/IKRig/Source/IKRigEditor/Private/RetargetEditor/SRetargetOpStack.cpp
@@ -4,6 +4,7 @@
 
 #include "RetargetEditor/SRetargetOpStack.h"
 
+#include "UObject/UObjectIterator.h"
 #include "SPositiveActionButton.h"
 #include "RetargetEditor/IKRetargetEditorController.h"
 #include "Retargeter/IKRetargetOps.h"
