Fixes static analysis warning in SRetargetOpStack.cpp

diff --git a/Engine/Plugins/Animation/IKRig/Source/IKRigEditor/Private/RetargetEditor/SRetargetOpStack.cpp b/Engine/Plugins/Animation/IKRig/Source/IKRigEditor/Private/RetargetEditor/SRetargetOpStack.cpp
--- a/Engine/Plugins/Animation/IKRig/Source/IKRigEditor/Private/RetargetEditor/SRetargetOpStack.cpp
+++ b/Engine/Plugins/Animation/IKRig/Source/IKRigEditor/Private/RetargetEditor/SRetargetOpStack.cpp
@@ -141,7 +141,7 @@ void SRetargetOpStackItem::Construct(
 bool SRetargetOpStackItem::GetWarningMessage(FText& Message) const
 {
 	const UIKRetargetProcessor* Processor = OpStackWidget.Pin()->EditorController.Pin()->GetRetargetProcessor();
-	if (!Processor && Processor->IsInitialized())
+	if (!(Processor && Processor->IsInitialized()))
 	{
 		return false;
 	}
