Restyling retarget phase toolbar in IK Retargeter.

diff --git a/Engine/Plugins/Animation/IKRig/Source/IKRigEditor/Private/RetargetEditor/IKRetargetEditor.cpp b/Engine/Plugins/Animation/IKRig/Source/IKRigEditor/Private/RetargetEditor/IKRetargetEditor.cpp
--- a/Engine/Plugins/Animation/IKRig/Source/IKRigEditor/Private/RetargetEditor/IKRetargetEditor.cpp
+++ b/Engine/Plugins/Animation/IKRig/Source/IKRigEditor/Private/RetargetEditor/IKRetargetEditor.cpp
@@ -259,7 +259,7 @@ void FIKRetargetEditor::ExtendToolbar()
 
 void FIKRetargetEditor::FillToolbar(FToolBarBuilder& ToolbarBuilder)
 {
-	ToolbarBuilder.BeginSection("Show Retarget Pose");
+	ToolbarBuilder.BeginSection("Retarget Modes");
 	{
 		ToolbarBuilder.AddToolBarButton(
 			FExecuteAction::CreateLambda([this]{ EditorController->SetRetargetModeToPreviousMode(); }),
@@ -276,6 +276,10 @@ void FIKRetargetEditor::FillToolbar(FToolBarBuilder& ToolbarBuilder)
 			LOCTEXT("RetargetMode_ToolTip", "Choose which mode to display in the viewport."),
 			FSlateIcon(FAppStyle::GetAppStyleSetName(), "LevelEditor.Recompile"),
 			true);
+
+		ToolbarBuilder.AddSeparator();
+		const TSharedPtr<SHorizontalBox> RetargetPhaseButtons = GenerateRetargetPhaseButtons();
+		ToolbarBuilder.AddWidget(RetargetPhaseButtons.ToSharedRef());
 	}
 	ToolbarBuilder.EndSection();
 
@@ -306,46 +310,59 @@ void FIKRetargetEditor::FillToolbar(FToolBarBuilder& ToolbarBuilder)
 		FSlateIcon(FIKRetargetEditorStyle::Get().GetStyleSetName(),"IKRetarget.PostSettings"));
 	}
 	ToolbarBuilder.EndSection();
+}
 
-	ToolbarBuilder.AddSeparator();
+TSharedRef<SWidget> FIKRetargetEditor::GenerateRetargetModesMenu()
+{
+	FMenuBuilder MenuBuilder(true, GetToolkitCommands());
+	
+	MenuBuilder.BeginSection(TEXT("Retarget Modes"));
+	MenuBuilder.AddMenuEntry(FIKRetargetCommands::Get().RunRetargeter, TEXT("Run Retargeter"), TAttribute<FText>(), TAttribute<FText>(),  EditorController->GetRetargeterModeIcon(ERetargeterOutputMode::RunRetarget));
+	MenuBuilder.AddMenuEntry(FIKRetargetCommands::Get().EditRetargetPose, TEXT("Edit Retarget Pose"), TAttribute<FText>(), TAttribute<FText>(), EditorController->GetRetargeterModeIcon(ERetargeterOutputMode::EditRetargetPose));
+	MenuBuilder.AddMenuEntry(FIKRetargetCommands::Get().ShowRetargetPose, TEXT("Show Retarget Pose"), TAttribute<FText>(), TAttribute<FText>(), EditorController->GetRetargeterModeIcon(ERetargeterOutputMode::ShowRetargetPose));
+	MenuBuilder.EndSection();
+	
+	return MenuBuilder.MakeWidget();
+}
 
-	FLinearColor OffColor = FLinearColor::White;
-	FLinearColor OnColor = FLinearColor(.32f, .66f, .32f, 1.f);
+TSharedPtr<SHorizontalBox> FIKRetargetEditor::GenerateRetargetPhaseButtons() const
+{
+	constexpr float PhaseButtonMargin = 2.f;
 	
-	TSharedPtr<SVerticalBox> Box = SNew(SVerticalBox)
-	+ SVerticalBox::Slot()
+	TSharedPtr<SHorizontalBox> Box = SNew(SHorizontalBox)
+	+ SHorizontalBox::Slot()
 	.VAlign(VAlign_Center)
 	.HAlign(HAlign_Center)
-	.AutoHeight()
+	.AutoWidth()
 	[
 		SNew(SHorizontalBox)
 		+ SHorizontalBox::Slot()
 		.FillWidth(1.f)
 		.HAlign(HAlign_Center)
 		[
-			SNew(STextBlock)
-			.Font(FCoreStyle::Get().GetFontStyle("SmallFont"))
-			.Text(FText::FromString("Toggle Retarget Phases"))
+			SNew(STextBlock).Text(FText::FromString("Retarget Phases: "))
 		]
 	]
-	+ SVerticalBox::Slot()
-	.AutoHeight()
+	+ SHorizontalBox::Slot()
+	.AutoWidth()
 	[
 		SNew(SHorizontalBox)
 		+ SHorizontalBox::Slot()
 		.AutoWidth()
-		.HAlign(HAlign_Center)
+		.VAlign(VAlign_Center)
+		.Padding(FMargin(PhaseButtonMargin))
 		[
-			SNew(SButton)
-			.OnClicked_Lambda([this]()
+			SNew(SCheckBox)
+			.Style(FAppStyle::Get(), "ToggleButtonCheckbox")
+			.IsChecked_Lambda([this]()
 			{
-				FRetargetGlobalSettings& GlobalSettings = EditorController->GetGlobalSettings();
-				GlobalSettings.bEnableRoot = !GlobalSettings.bEnableRoot;
-				return FReply::Handled();
+				const bool bIsOn =  EditorController->GetGlobalSettings().bEnableRoot;
+				return bIsOn ? ECheckBoxState::Checked : ECheckBoxState::Unchecked;
 			})
-			.ButtonColorAndOpacity_Lambda([this, OffColor, OnColor]() -> FLinearColor
+			.OnCheckStateChanged_Lambda([this](ECheckBoxState InCheckBoxState)
 			{
-				return EditorController->GetGlobalSettings().bEnableRoot ? OnColor : OffColor;
+				FRetargetGlobalSettings& GlobalSettings = EditorController->GetGlobalSettings();
+				GlobalSettings.bEnableRoot = !GlobalSettings.bEnableRoot;
 			})
 			[
 				SNew(STextBlock).Text(FText::FromString("Root"))
@@ -354,18 +371,20 @@ void FIKRetargetEditor::FillToolbar(FToolBarBuilder& ToolbarBuilder)
 
 		+ SHorizontalBox::Slot()
 		.AutoWidth()
-		.HAlign(HAlign_Center)
+		.VAlign(VAlign_Center)
+		.Padding(FMargin(PhaseButtonMargin))
 		[
-			SNew(SButton)
-			.OnClicked_Lambda([this]()
+			SNew(SCheckBox)
+			.Style(FAppStyle::Get(), "ToggleButtonCheckbox")
+			.IsChecked_Lambda([this]()
 			{
-				FRetargetGlobalSettings& GlobalSettings = EditorController->GetGlobalSettings();
-				GlobalSettings.bEnableFK = !GlobalSettings.bEnableFK;
-				return FReply::Handled();
+				const bool bIsOn =  EditorController->GetGlobalSettings().bEnableFK;
+				return bIsOn ? ECheckBoxState::Checked : ECheckBoxState::Unchecked;
 			})
-			.ButtonColorAndOpacity_Lambda([this, OffColor, OnColor]() -> FLinearColor
+			.OnCheckStateChanged_Lambda([this](ECheckBoxState InCheckBoxState)
 			{
-				return EditorController->GetGlobalSettings().bEnableFK ? OnColor : OffColor;
+				FRetargetGlobalSettings& GlobalSettings = EditorController->GetGlobalSettings();
+				GlobalSettings.bEnableFK = !GlobalSettings.bEnableFK;
 			})
 			[
 				SNew(STextBlock).Text(FText::FromString("FK"))
@@ -374,18 +393,21 @@ void FIKRetargetEditor::FillToolbar(FToolBarBuilder& ToolbarBuilder)
 		
 		+ SHorizontalBox::Slot()
 		.AutoWidth()
-		.HAlign(HAlign_Center)
+		.VAlign(VAlign_Center)
+		.Padding(FMargin(PhaseButtonMargin))
 		[
-			SNew(SButton)
-			.OnClicked_Lambda([this]()
+
+			SNew(SCheckBox)
+			.Style(FAppStyle::Get(), "ToggleButtonCheckbox")
+			.IsChecked_Lambda([this]()
 			{
-				FRetargetGlobalSettings& GlobalSettings = EditorController->GetGlobalSettings();
-				GlobalSettings.bEnableIK = !GlobalSettings.bEnableIK;
-				return FReply::Handled();
+				const bool bIsOn =  EditorController->GetGlobalSettings().bEnableIK;
+				return bIsOn ? ECheckBoxState::Checked : ECheckBoxState::Unchecked;
 			})
-			.ButtonColorAndOpacity_Lambda([this, OffColor, OnColor]() -> FLinearColor
+			.OnCheckStateChanged_Lambda([this](ECheckBoxState InCheckBoxState)
 			{
-				return EditorController->GetGlobalSettings().bEnableIK ? OnColor : OffColor;
+				FRetargetGlobalSettings& GlobalSettings = EditorController->GetGlobalSettings();
+				GlobalSettings.bEnableIK = !GlobalSettings.bEnableIK;
 			})
 			[
 				SNew(STextBlock).Text(FText::FromString("IK"))
@@ -394,43 +416,28 @@ void FIKRetargetEditor::FillToolbar(FToolBarBuilder& ToolbarBuilder)
 
 		+ SHorizontalBox::Slot()
 		.AutoWidth()
-		.HAlign(HAlign_Center)
+		.VAlign(VAlign_Center)
+		.Padding(FMargin(PhaseButtonMargin))
 		[
-			SNew(SButton)
-			.OnClicked_Lambda([this]()
+			SNew(SCheckBox)
+			.Style(FAppStyle::Get(), "ToggleButtonCheckbox")
+			.IsChecked_Lambda([this]()
 			{
-				FRetargetGlobalSettings& GlobalSettings = EditorController->GetGlobalSettings();
-				GlobalSettings.bEnablePost = !GlobalSettings.bEnablePost;
-				return FReply::Handled();
+				const bool bIsOn =  EditorController->GetGlobalSettings().bEnablePost;
+				return bIsOn ? ECheckBoxState::Checked : ECheckBoxState::Unchecked;
 			})
-			.ButtonColorAndOpacity_Lambda([this, OffColor, OnColor]() -> FLinearColor
+			.OnCheckStateChanged_Lambda([this](ECheckBoxState InCheckBoxState)
 			{
-				return EditorController->GetGlobalSettings().bEnablePost ? OnColor : OffColor;
+				FRetargetGlobalSettings& GlobalSettings = EditorController->GetGlobalSettings();
+				GlobalSettings.bEnablePost = !GlobalSettings.bEnablePost;
 			})
 			[
 				SNew(STextBlock).Text(FText::FromString("Post"))
 			]
 		]
 	];
-
-	ToolbarBuilder.BeginSection("Toggle Retarget Passes");
-	{
-		ToolbarBuilder.AddWidget(Box.ToSharedRef());
-	}
-	ToolbarBuilder.EndSection();
-}
-
-TSharedRef<SWidget> FIKRetargetEditor::GenerateRetargetModesMenu()
-{
-	FMenuBuilder MenuBuilder(true, GetToolkitCommands());
 	
-	MenuBuilder.BeginSection(TEXT("Retarget Modes"));
-	MenuBuilder.AddMenuEntry(FIKRetargetCommands::Get().RunRetargeter, TEXT("Run Retargeter"), TAttribute<FText>(), TAttribute<FText>(),  EditorController->GetRetargeterModeIcon(ERetargeterOutputMode::RunRetarget));
-	MenuBuilder.AddMenuEntry(FIKRetargetCommands::Get().EditRetargetPose, TEXT("Edit Retarget Pose"), TAttribute<FText>(), TAttribute<FText>(), EditorController->GetRetargeterModeIcon(ERetargeterOutputMode::EditRetargetPose));
-	MenuBuilder.AddMenuEntry(FIKRetargetCommands::Get().ShowRetargetPose, TEXT("Show Retarget Pose"), TAttribute<FText>(), TAttribute<FText>(), EditorController->GetRetargeterModeIcon(ERetargeterOutputMode::ShowRetargetPose));
-	MenuBuilder.EndSection();
-	
-	return MenuBuilder.MakeWidget();
+	return Box;
 }
 
 FName FIKRetargetEditor::GetToolkitFName() const
diff --git a/Engine/Plugins/Animation/IKRig/Source/IKRigEditor/Public/RetargetEditor/IKRetargetEditor.h b/Engine/Plugins/Animation/IKRig/Source/IKRigEditor/Public/RetargetEditor/IKRetargetEditor.h
--- a/Engine/Plugins/Animation/IKRig/Source/IKRigEditor/Public/RetargetEditor/IKRetargetEditor.h
+++ b/Engine/Plugins/Animation/IKRig/Source/IKRigEditor/Public/RetargetEditor/IKRetargetEditor.h
@@ -89,6 +89,7 @@ private:
 	void ExtendToolbar();
 	void FillToolbar(FToolBarBuilder& ToolbarBuilder);
 	TSharedRef<SWidget> GenerateRetargetModesMenu();
+	TSharedPtr<SHorizontalBox> GenerateRetargetPhaseButtons() const;
 	/** END toolbar */
 	
 	/** preview scene setup */
