Fixes issues in IK Rig anim node when used with Linked Animation Graph.

diff --git a/Engine/Plugins/Animation/IKRig/Source/IKRig/Private/AnimNodes/AnimNode_IKRig.cpp b/Engine/Plugins/Animation/IKRig/Source/IKRig/Private/AnimNodes/AnimNode_IKRig.cpp
--- a/Engine/Plugins/Animation/IKRig/Source/IKRig/Private/AnimNodes/AnimNode_IKRig.cpp
+++ b/Engine/Plugins/Animation/IKRig/Source/IKRig/Private/AnimNodes/AnimNode_IKRig.cpp
@@ -292,26 +292,15 @@ void FAnimNode_IKRig::CacheBones_AnyThread(const FAnimationCacheBonesContext& Co
 		return;
 	}
 	
-	if (!RigDefinitionAsset)
-	{
-		return;
-	}
-
-	if (!IKRigProcessor)
-	{
-		return;
-	}
-
-	if (!IKRigProcessor->IsInitialized())
+	// fill up node names, mapping the anim graph bone indices to the indices used by the IK Rig
+	CompactPoseToRigIndices.Reset();
+	const TArray<FBoneIndexType>& RequiredBonesArray = RequiredBones.GetBoneIndicesArray();
+	const USkeletalMesh* SkeletalMesh = RequiredBones.GetSkeletalMeshAsset();
+	if (!ensure(SkeletalMesh))
 	{
 		return;
 	}
-
-	// fill up node names, mapping the anim graph bone indices to the IK Rig bones
-	CompactPoseToRigIndices.Reset();
-	const TArray<FBoneIndexType>& RequiredBonesArray = RequiredBones.GetBoneIndicesArray();
-	const FReferenceSkeleton& RefSkeleton = RequiredBones.GetReferenceSkeleton();
-	const FIKRigSkeleton& IKRigSkeleton = IKRigProcessor->GetSkeleton();
+	const FReferenceSkeleton& MeshRefSkeleton = SkeletalMesh->GetRefSkeleton();
 	const int32 NumBones = RequiredBonesArray.Num();
 	for (uint16 Index = 0; Index < NumBones; ++Index)
 	{
@@ -322,8 +311,8 @@ void FAnimNode_IKRig::CacheBones_AnyThread(const FAnimationCacheBonesContext& Co
 		}
 		
 		FCompactPoseBoneIndex CPIndex = RequiredBones.MakeCompactPoseIndex(FMeshPoseBoneIndex(MeshBone));
-		const FName Name = RefSkeleton.GetBoneName(MeshBone);
-		CompactPoseToRigIndices.Add(CPIndex) = IKRigSkeleton.GetBoneIndexFromName(Name);
+		const FName Name = MeshRefSkeleton.GetBoneName(MeshBone);
+		CompactPoseToRigIndices.Add(CPIndex) = MeshRefSkeleton.FindBoneIndex(Name);
 	}
 }
 
