Adding root motion retarget op to IK Retargeter.
Various bug fixes in retarget op stack system.

diff --git a/Engine/Plugins/Animation/IKRig/Source/IKRig/Private/Retargeter/IKRetargetProcessor.cpp b/Engine/Plugins/Animation/IKRig/Source/IKRig/Private/Retargeter/IKRetargetProcessor.cpp
--- a/Engine/Plugins/Animation/IKRig/Source/IKRig/Private/Retargeter/IKRetargetProcessor.cpp
+++ b/Engine/Plugins/Animation/IKRig/Source/IKRig/Private/Retargeter/IKRetargetProcessor.cpp
@@ -1469,42 +1469,6 @@ FVector FPoleVectorMatcher::GetChainNormal(
 	return ChainAxis.GetSafeNormal();
 }
 
-bool FRetargetOps::Initialize(
-	const FRetargetSkeleton& SourceSkeleton,
-	const FTargetSkeleton& TargetSkeleton,
-	const TArray<TObjectPtr<URetargetOpBase>>& OpStackFromAsset,
-	UIKRetargetProcessor* Processor,
-	FIKRigLogger& Log)
-{
-	bool bAllOpsInitialized = true;
-	// create copies of all the operations in the retarget asset for runtime use
-	OpStack.Reset(OpStackFromAsset.Num());
-	int32 OpIndex = 0;
-	for (const URetargetOpBase* OpFromAsset : OpStackFromAsset)
-	{
-		if (!OpFromAsset)
-		{
-			// this can happen if asset references deleted op type which should only happen during development (if at all)
-			Log.LogWarning(FText(LOCTEXT("UnknownPostOperation", "Retargeter has null/unknown post operation in it. Please remove it.")));
-			continue;
-		}
-
-		// create duplicate solver instance with unique name
-		FString Name = "RetargetOpInstance";
-		Name.AppendInt(OpIndex++);
-		URetargetOpBase* NewOp  = DuplicateObject(OpFromAsset, Processor, FName(*Name));
-		
-		// initialize it and store in the processor
-		if (!NewOp->Initialize(SourceSkeleton, TargetSkeleton, Processor, Log))
-		{
-			bAllOpsInitialized = false;
-		}
-		OpStack.Add(NewOp);
-	}
-
-	return bAllOpsInitialized;
-}
-
 UIKRetargetProcessor::UIKRetargetProcessor()
 {
 	const FName LogName = FName("IKRetarget_",GetUniqueID());
@@ -1526,8 +1490,8 @@ void UIKRetargetProcessor::Initialize(
 	// record source asset
 	RetargeterAsset = InRetargeterAsset;
 
-	const UIKRigDefinition* SourceIKRig = RetargeterAsset->GetSourceIKRig();
-	const UIKRigDefinition* TargetIKRig = RetargeterAsset->GetTargetIKRig();
+	const UIKRigDefinition* SourceIKRig = RetargeterAsset->GetIKRig(ERetargetSourceOrTarget::Source);
+	const UIKRigDefinition* TargetIKRig = RetargeterAsset->GetIKRig(ERetargetSourceOrTarget::Target);
 
 	// check prerequisite assets
 	if (!SourceSkeletalMesh)
@@ -1590,7 +1554,7 @@ void UIKRetargetProcessor::Initialize(
 
 	// initialize the post operations
 	const TArray<TObjectPtr<URetargetOpBase>>& OpsFromAsset = RetargeterAsset->GetPostSettingsUObject()->RetargetOps;
-	const bool bAllOpsInitialized = RetargetOps.Initialize(SourceSkeleton, TargetSkeleton, OpsFromAsset, this, Log);
+	const bool bAllOpsInitialized = InitializeOpStack(OpsFromAsset);
 	if (!bAllOpsInitialized)
 	{
 		Log.LogWarning(FText(LOCTEXT("PostOperationInitError", "One of the retarget post operations has warnings. See output log for details.")));
@@ -1617,7 +1581,7 @@ bool UIKRetargetProcessor::InitializeRoots()
 	RootRetargeter.Reset();
 	
 	// initialize root encoder
-	const FName SourceRootBoneName = RetargeterAsset->GetSourceIKRig()->GetRetargetRoot();
+	const FName SourceRootBoneName = RetargeterAsset->GetIKRig(ERetargetSourceOrTarget::Source)->GetRetargetRoot();
 	const bool bRootEncoderInit = RootRetargeter.InitializeSource(SourceRootBoneName, SourceSkeleton, Log);
 	if (!bRootEncoderInit)
 	{
@@ -1627,7 +1591,7 @@ bool UIKRetargetProcessor::InitializeRoots()
 	}
 
 	// initialize root decoder
-	const FName TargetRootBoneName = RetargeterAsset->GetTargetIKRig()->GetRetargetRoot();
+	const FName TargetRootBoneName =RetargeterAsset->GetIKRig(ERetargetSourceOrTarget::Target)->GetRetargetRoot();
 	const bool bRootDecoderInit = RootRetargeter.InitializeTarget(TargetRootBoneName, TargetSkeleton, Log);
 	if (!bRootDecoderInit)
 	{
@@ -1643,18 +1607,17 @@ bool UIKRetargetProcessor::InitializeBoneChainPairs()
 {
 	ChainPairsFK.Reset();
 	ChainPairsIK.Reset();
-	
-	const UIKRigDefinition* TargetIKRig = RetargeterAsset->GetTargetIKRig();
-	const UIKRigDefinition* SourceIKRig = RetargeterAsset->GetSourceIKRig();
-	
+
+	const UIKRigDefinition* SourceIKRig = RetargeterAsset->GetIKRig(ERetargetSourceOrTarget::Source);
+	const UIKRigDefinition* TargetIKRig = RetargeterAsset->GetIKRig(ERetargetSourceOrTarget::Target);
 	check(SourceIKRig && TargetIKRig);
 
 	// check that chains are available in both IKRig assets before sorting them based on StartBone index
 	const TArray<TObjectPtr<URetargetChainSettings>>& ChainMapping = RetargeterAsset->GetAllChainSettings();	
-	for (URetargetChainSettings* ChainMap : ChainMapping)
+	for (const URetargetChainSettings* ChainMap : ChainMapping)
 	{
 		// get target bone chain
-		const FBoneChain* TargetBoneChain = RetargeterAsset->GetTargetIKRig()->GetRetargetChainByName(ChainMap->TargetChain);
+		const FBoneChain* TargetBoneChain = TargetIKRig->GetRetargetChainByName(ChainMap->TargetChain);
 		if (!TargetBoneChain)
 		{
 			Log.LogWarning( FText::Format(
@@ -1670,7 +1633,7 @@ bool UIKRetargetProcessor::InitializeBoneChainPairs()
 		}
 		
 		// get source bone chain
-		const FBoneChain* SourceBoneChain = RetargeterAsset->GetSourceIKRig()->GetRetargetChainByName(ChainMap->SourceChain);
+		const FBoneChain* SourceBoneChain = SourceIKRig->GetRetargetChainByName(ChainMap->SourceChain);
 		if (!SourceBoneChain)
 		{
 			Log.LogWarning( FText::Format(
@@ -1760,7 +1723,7 @@ bool UIKRetargetProcessor::InitializeBoneChainPairs()
 bool UIKRetargetProcessor::InitializeIKRig(UObject* Outer, const USkeletalMesh* InSkeletalMesh)
 {
 	// get the target IK Rig asset
-	UIKRigDefinition* IKRig = RetargeterAsset->GetTargetIKRigWriteable();
+	UIKRigDefinition* IKRig = RetargeterAsset->GetIKRigWriteable(ERetargetSourceOrTarget::Target);
 	
 	// gather list of excluded goals based on any chain mapping that has it's IK disabled
 	TArray<FName> GoalsToExclude;
@@ -1867,6 +1830,35 @@ bool UIKRetargetProcessor::InitializeIKRig(UObject* Outer, const USkeletalMesh*
 	return true;
 }
 
+bool UIKRetargetProcessor::InitializeOpStack(const TArray<TObjectPtr<URetargetOpBase>>& OpStackFromAsset)
+{
+	bool bAllOpsInitialized = true;
+	// create copies of all the operations in the retarget asset for runtime use
+	OpStack.Reset(OpStackFromAsset.Num());
+	int32 OpIndex = 0;
+	for (const URetargetOpBase* OpFromAsset : OpStackFromAsset)
+	{
+		if (!OpFromAsset)
+		{
+			// this can happen if asset references deleted op type which should only happen during development (if at all)
+			Log.LogWarning(FText(LOCTEXT("UnknownPostOperation", "Retargeter has null/unknown post operation in it. Please remove it.")));
+			continue;
+		}
+
+		// create duplicate op instance with unique name
+		const FName BaseName = FName(*(OpFromAsset->GetName() + "_RetargetOpInstance_"));
+		const FName UniqueName = MakeUniqueObjectName(this, OpFromAsset->GetClass(), BaseName, EUniqueObjectNameOptions::GloballyUnique);
+		URetargetOpBase* NewOp = DuplicateObject(OpFromAsset, this, UniqueName);
+		
+		// initialize it and store in the processor
+		NewOp->bIsInitialized = NewOp->Initialize(this, SourceSkeleton, TargetSkeleton, Log);
+		bAllOpsInitialized &= NewOp->bIsInitialized;
+		OpStack.Add(NewOp);
+	}
+
+	return bAllOpsInitialized;
+}
+
 TArray<FTransform>&  UIKRetargetProcessor::RunRetargeter(
 	const TArray<FTransform>& InSourceGlobalPose,
 	const TMap<FName, float>& SpeedValuesFromCurves,
@@ -2137,11 +2129,11 @@ void UIKRetargetProcessor::RunRetargetOps(
 	const TArray<FTransform>& InSourceGlobalPose,
 	TArray<FTransform>& OutTargetGlobalPose)
 {
-	for (URetargetOpBase* RetargetOp : RetargetOps.OpStack)
+	for (URetargetOpBase* RetargetOp : OpStack)
 	{
-		if (RetargetOp->bIsEnabled)
+		if (RetargetOp->bIsEnabled && RetargetOp->bIsInitialized)
 		{
-			RetargetOp->Run(InSourceGlobalPose, OutTargetGlobalPose);
+			RetargetOp->Run(this, InSourceGlobalPose, OutTargetGlobalPose);
 		}
 	}
 }
@@ -2248,7 +2240,7 @@ FName UIKRetargetProcessor::GetChainNameForBone(const int32& BoneIndex, const in
 void UIKRetargetProcessor::ApplySettingsFromAsset()
 {
 	// copy IK Rig settings
-	if (const UIKRigDefinition* TargetIKRig = RetargeterAsset->GetTargetIKRig())
+	if (const UIKRigDefinition* TargetIKRig = RetargeterAsset->GetIKRig(ERetargetSourceOrTarget::Target))
 	{
 		IKRigProcessor->CopyAllInputsFromSourceAssetAtRuntime(TargetIKRig);
 	}
@@ -2362,8 +2354,8 @@ void UIKRetargetProcessor::ApplyNewRetargetPose(
 	RetargetSkeleton.GenerateRetargetPose(NewRetargetPoseName, NewRetargetPose, RootBoneName);
 
 	// re-initialize the bone chains using the newly generated retarget pose
-	const UIKRigDefinition* SourceIKRig = RetargeterAsset->GetSourceIKRig();
-	const UIKRigDefinition* TargetIKRig = RetargeterAsset->GetTargetIKRig();
+	const UIKRigDefinition* SourceIKRig = RetargeterAsset->GetIKRig(ERetargetSourceOrTarget::Source);
+	const UIKRigDefinition* TargetIKRig = RetargeterAsset->GetIKRig(ERetargetSourceOrTarget::Target);
 	if (!(SourceIKRig && TargetIKRig))
 	{
 		return;
diff --git a/Engine/Plugins/Animation/IKRig/Source/IKRig/Private/Retargeter/IKRetargeter.cpp b/Engine/Plugins/Animation/IKRig/Source/IKRig/Private/Retargeter/IKRetargeter.cpp
--- a/Engine/Plugins/Animation/IKRig/Source/IKRig/Private/Retargeter/IKRetargeter.cpp
+++ b/Engine/Plugins/Animation/IKRig/Source/IKRig/Private/Retargeter/IKRetargeter.cpp
@@ -42,61 +42,58 @@ UIKRetargeter::UIKRetargeter(const FObjectInitializer& ObjectInitializer)
 	CleanAndInitialize();
 }
 
-const UIKRigDefinition* UIKRetargeter::GetSourceIKRig() const
+const UIKRigDefinition* UIKRetargeter::GetIKRig(ERetargetSourceOrTarget SourceOrTarget) const
 {
-	if (IsInGameThread())
+	if (!IsInGameThread())
 	{
-		return SourceIKRigAsset.LoadSynchronous();	
+		return nullptr;
 	}
-	
-	return nullptr;
-}
 
-const UIKRigDefinition* UIKRetargeter::GetTargetIKRig() const
-{
-	if (IsInGameThread())
-	{
-		return TargetIKRigAsset.LoadSynchronous();	
-	}
-	
-	return nullptr;
+	const TSoftObjectPtr<UIKRigDefinition> SoftIKRig = SourceOrTarget == ERetargetSourceOrTarget::Source ? SourceIKRigAsset : TargetIKRigAsset;
+	return SoftIKRig.LoadSynchronous();
 }
 
-UIKRigDefinition* UIKRetargeter::GetSourceIKRigWriteable() const
+UIKRigDefinition* UIKRetargeter::GetIKRigWriteable(ERetargetSourceOrTarget SourceOrTarget) const
 {
-	if (IsInGameThread())
+	if (!IsInGameThread())
 	{
-		return SourceIKRigAsset.LoadSynchronous();	
+		return nullptr;
 	}
-	
-	return nullptr;
-}
 
-UIKRigDefinition* UIKRetargeter::GetTargetIKRigWriteable() const
-{
-	if (IsInGameThread())
-	{
-		return TargetIKRigAsset.LoadSynchronous();	
-	}
-	
-	return nullptr;
+	const TSoftObjectPtr<UIKRigDefinition> SoftIKRig = SourceOrTarget == ERetargetSourceOrTarget::Source ? SourceIKRigAsset : TargetIKRigAsset;
+	return SoftIKRig.LoadSynchronous();
 }
 
 #if WITH_EDITORONLY_DATA
 const USkeletalMesh* UIKRetargeter::GetPreviewMesh(ERetargetSourceOrTarget SourceOrTarget) const
 {
-	if (IsInGameThread())
+	if (!IsInGameThread())
 	{
-		if (SourceOrTarget == ERetargetSourceOrTarget::Source)
+		return nullptr;
+	}
+
+	// the preview mesh override on the retarget takes precedence
+	if (SourceOrTarget == ERetargetSourceOrTarget::Source)
+	{
+		if (SourcePreviewMesh.IsValid())
 		{
-			return SourcePreviewMesh.LoadSynchronous();	
+			return SourcePreviewMesh.LoadSynchronous();
 		}
-		else
+	}
+	else
+	{
+		if (TargetPreviewMesh.IsValid())
 		{
-			return TargetPreviewMesh.LoadSynchronous();	
+			return TargetPreviewMesh.LoadSynchronous();
 		}
 	}
-	
+
+	// fallback to preview mesh from the IK Rig itself
+	if (const UIKRigDefinition* IKRig = GetIKRig(SourceOrTarget))
+	{
+		return IKRig->GetPreviewMesh();
+	}
+
 	return nullptr;
 }
 #endif
@@ -204,6 +201,9 @@ void UIKRetargeter::DeclareConstructClasses(TArray<FTopLevelAssetPath>& OutConst
 
 void UIKRetargeter::CleanAndInitialize()
 {
+	// remove null retarget ops
+	OpStack->RetargetOps.Remove(nullptr);
+	
 	// remove null settings
 	ChainSettings.Remove(nullptr);
 
@@ -412,7 +412,7 @@ FTargetChainSettings UIKRetargeter::GetChainUsingGoalFromRetargetAsset(
 		return EmptySettings;
 	}
 
-	const UIKRigDefinition* IKRig = RetargetAsset->GetTargetIKRig();
+	const UIKRigDefinition* IKRig = RetargetAsset->GetIKRig(ERetargetSourceOrTarget::Target);
 	if (!IKRig)
 	{
 		return EmptySettings;
diff --git a/Engine/Plugins/Animation/IKRig/Source/IKRig/Private/Retargeter/RetargetOps/PinBoneOp.cpp b/Engine/Plugins/Animation/IKRig/Source/IKRig/Private/Retargeter/RetargetOps/PinBoneOp.cpp
--- a/Engine/Plugins/Animation/IKRig/Source/IKRig/Private/Retargeter/RetargetOps/PinBoneOp.cpp
+++ b/Engine/Plugins/Animation/IKRig/Source/IKRig/Private/Retargeter/RetargetOps/PinBoneOp.cpp
@@ -1,14 +1,13 @@
 ﻿// Copyright Epic Games, Inc. All Rights Reserved.
 
-
 #include "Retargeter/RetargetOps/PinBoneOp.h"
 
 #define LOCTEXT_NAMESPACE "UPinBoneOp"
 
 bool UPinBoneOp::Initialize(
+	const UIKRetargetProcessor* Processor,
 	const FRetargetSkeleton& SourceSkeleton,
 	const FTargetSkeleton& TargetSkeleton,
-	const UIKRetargetProcessor* Processor,
 	FIKRigLogger& Log)
 {
 	bool bFullyInitialized = true;
@@ -47,11 +46,11 @@ bool UPinBoneOp::Initialize(
 			const FTransform& BoneToPinTransform = TargetSkeleton.RetargetGlobalPose[BoneToPin.BoneToPinIndex];
 			const TArray<FTransform>& PinToPose = PinTo == ERetargetSourceOrTarget::Source ? SourceSkeleton.RetargetGlobalPose : TargetSkeleton.RetargetGlobalPose;
 			const FTransform& BoneToPinToTransform = PinToPose[BoneToPin.BoneToPinToIndex];
-			BoneToPin.OffsetInRefPose =  BoneToPinToTransform.Inverse() * BoneToPinTransform;
+			BoneToPin.OffsetInRefPose =  BoneToPinToTransform.GetRelativeTransform(BoneToPinTransform);
 		}
 	}
 
-	#if WITH_EDITORONLY_DATA
+#if WITH_EDITOR
 	if (bFullyInitialized)
 	{
 		Message = FText::Format(LOCTEXT("ReadyToRun", "Pinning {0} bones."), FText::AsNumber(BonesToPin.Num()));
@@ -60,16 +59,20 @@ bool UPinBoneOp::Initialize(
 	{
 		Message = FText(LOCTEXT("MissingBonesWarning", "Bone(s) not found. See output log."));
 	}
-	#endif
-	
-	return bFullyInitialized;
+#endif
+
+	// always treat this op as "initialized", individual pins will only execute if their prerequisites are met
+	return true;
 }
 
-void UPinBoneOp::Run(const TArray<FTransform>& InSourceGlobalPose, TArray<FTransform>& OutTargetGlobalPose)
+void UPinBoneOp::Run(
+	const UIKRetargetProcessor* Processor,
+	const TArray<FTransform>& InSourceGlobalPose,
+	TArray<FTransform>& OutTargetGlobalPose)
 {
 	for (const FPinBoneData& BoneToPin : BonesToPin)
 	{
-		if (!bEnabled || BoneToPin.BoneToPinIndex == INDEX_NONE || BoneToPin.BoneToPinToIndex == INDEX_NONE)
+		if (BoneToPin.BoneToPinIndex == INDEX_NONE || BoneToPin.BoneToPinToIndex == INDEX_NONE)
 		{
 			continue; // disabled or not successfully initialized
 		}
diff --git a/Engine/Plugins/Animation/IKRig/Source/IKRig/Private/Retargeter/RetargetOps/RootMotionGeneratorOp.cpp b/Engine/Plugins/Animation/IKRig/Source/IKRig/Private/Retargeter/RetargetOps/RootMotionGeneratorOp.cpp
new file mode 100644
--- /dev/null
+++ b/Engine/Plugins/Animation/IKRig/Source/IKRig/Private/Retargeter/RetargetOps/RootMotionGeneratorOp.cpp
@@ -0,0 +1,206 @@
+// Copyright Epic Games, Inc. All Rights Reserved.
+#include "Retargeter/RetargetOps/RootMotionGeneratorOp.h"
+
+#include "Engine/SkeletalMesh.h"
+
+#define LOCTEXT_NAMESPACE "RootMotionRemapOp"
+
+#if WITH_EDITOR
+void URootMotionGeneratorOp::OnAddedToStack(const UIKRetargeter* Asset)
+{
+	// get the root of the SOURCE skeleton
+	if (const USkeletalMesh* Mesh = Asset->GetPreviewMesh(ERetargetSourceOrTarget::Source))
+	{
+		SourceRootBone = Mesh->GetRefSkeleton().GetBoneName(0);
+	}
+	
+	// get the root of the TARGET skeleton
+	if (const USkeletalMesh* Mesh = Asset->GetPreviewMesh(ERetargetSourceOrTarget::Target))
+	{
+		TargetRootBone = Mesh->GetRefSkeleton().GetBoneName(0);
+	}
+
+	// get the retarget root bone
+	if (const UIKRigDefinition* IKRig = Asset->GetIKRig(ERetargetSourceOrTarget::Target))
+	{
+		TargetPelvisBone = IKRig->GetRetargetRoot();
+	}
+}
+#endif
+
+bool URootMotionGeneratorOp::Initialize(
+	const UIKRetargetProcessor* Processor,
+	const FRetargetSkeleton& SourceSkeleton,
+	const FTargetSkeleton& TargetSkeleton,
+	FIKRigLogger& Log)
+{
+	SourceRootIndex = SourceSkeleton.BoneNames.Find(SourceRootBone);
+	bool bHasAllPrerequisites = true;
+	if (SourceRootIndex == INDEX_NONE)
+	{
+		Log.LogWarning(FText::Format(LOCTEXT("MissingRootBone", "Root Motion Remap Op, missing source root bone {0}."), FText::FromName(SourceRootBone)));
+		bHasAllPrerequisites = false;
+	}
+	
+	TargetRootIndex = TargetSkeleton.BoneNames.Find(TargetRootBone);
+	if (TargetRootIndex == INDEX_NONE)
+	{
+		Log.LogWarning(FText::Format(LOCTEXT("MissingRootBone", "Root Motion Remap Op, missing target root bone {0}."), FText::FromName(TargetRootBone)));
+		bHasAllPrerequisites = false;
+	}
+
+	TargetPelvisIndex = TargetSkeleton.BoneNames.Find(TargetPelvisBone);
+	if (TargetPelvisIndex == INDEX_NONE)
+	{
+		Log.LogWarning(FText::Format(LOCTEXT("MissingPelvisBone", "Root Motion Remap Op, missing target pelvis bone {0}."), FText::FromName(TargetPelvisBone)));
+		bHasAllPrerequisites = false;
+	}
+
+#if WITH_EDITOR
+	if (bHasAllPrerequisites)
+	{
+		Message = FText::Format(LOCTEXT("ReadyToRun", "Generating root motion on {0}."), FText::FromName(TargetRootBone));
+	}
+	else
+	{
+		Message = FText(LOCTEXT("MissingBonesWarning", "Bone(s) not found. See output log."));
+	}
+#endif
+
+	// can't cache everything unless all prerequisites are met
+	if (!bHasAllPrerequisites)
+	{
+		return false;
+	}
+
+	TargetPelvisInRefPose = TargetSkeleton.RetargetGlobalPose[TargetPelvisIndex];
+	SourceRootInRefPose = SourceSkeleton.RetargetGlobalPose[SourceRootIndex];
+	TargetRootInRefPose = TargetSkeleton.RetargetGlobalPose[TargetRootIndex];
+	TargetPelvisRelativeToTargetRootRefPose = TargetRootInRefPose.GetRelativeTransform(TargetPelvisInRefPose);
+
+	// generate list of non-retargeted child indices
+	const int32 NumTargetBones = TargetSkeleton.BoneNames.Num();
+	NonRetargetedChildrenOfRoot.Reset();
+	for (int32 BoneIndex=1; BoneIndex<NumTargetBones; ++BoneIndex)
+	{
+		bool bHasNoRetargetedParent = true;
+		
+		int32 ParentIndex = BoneIndex;
+		while(ParentIndex != INDEX_NONE)
+		{
+			if (TargetSkeleton.IsBoneRetargeted[ParentIndex])
+			{
+				bHasNoRetargetedParent = false;
+				break;
+			}
+			ParentIndex = TargetSkeleton.ParentIndices[ParentIndex];
+		}
+
+		if (bHasNoRetargetedParent)
+		{
+			NonRetargetedChildrenOfRoot.Add(BoneIndex);
+		}
+	}
+	
+	return true;
+}
+
+void URootMotionGeneratorOp::Run(
+	const UIKRetargetProcessor* Processor,
+	const TArray<FTransform>& InSourceGlobalPose,
+	TArray<FTransform>& OutTargetGlobalPose)
+{
+	// generate a new transform for the target root bone
+	FTransform NewRootTransform;
+	
+	// either generate motion from target pelvis, or copy from source root
+	if (RootMotionSource == ERootMotionSource::GenerateFromTargetPelvis)
+	{
+		GenerateRootMotionFromTargetPelvis(NewRootTransform, InSourceGlobalPose, OutTargetGlobalPose);
+	}
+	else
+	{
+		CopyRootMotionFromSourceRoot(NewRootTransform, Processor, InSourceGlobalPose);
+	}
+
+	// optionally propagate delta to all non-retargeted children
+	if (bPropagateToNonRetargetedChildren)
+	{
+		const FTransform Delta = OutTargetGlobalPose[TargetRootIndex].Inverse() * NewRootTransform;
+		for (const int32 BoneIndex : NonRetargetedChildrenOfRoot)
+		{
+			OutTargetGlobalPose[BoneIndex] = OutTargetGlobalPose[BoneIndex] * Delta;
+		}
+	}
+
+	// apply to the target root bone
+	OutTargetGlobalPose[TargetRootIndex] = NewRootTransform * GlobalOffset;
+}
+
+void URootMotionGeneratorOp::GenerateRootMotionFromTargetPelvis(
+	FTransform& OutRootTransform,
+	const TArray<FTransform>& InSourceGlobalPose,
+	const TArray<FTransform>& InTargetGlobalPose) const
+{
+	// In this case, we are generating root motion "from scratch"
+	// using the target Pelvis bone as the source of the motion
+		
+	// calculate new root transform relative to current hips transform
+	OutRootTransform = InTargetGlobalPose[TargetPelvisIndex];
+	// optionally remove all rotation when calculating the offset
+	if (!bRotateWithPelvis)
+	{
+		OutRootTransform.SetRotation(TargetPelvisInRefPose.GetRotation());
+	}
+	
+	// optionally maintain the offset from the ref pose
+	if (bMaintainOffsetFromPelvis)
+	{
+		OutRootTransform = TargetPelvisRelativeToTargetRootRefPose * OutRootTransform;
+	}
+
+	// optionally snap the root to the ground plane
+	if (RootHeightSource == ERootMotionHeightSource::SnapToGround)
+	{
+		FVector RootTranslation = OutRootTransform.GetTranslation();
+		RootTranslation.Z = 0.f;
+		OutRootTransform.SetTranslation(RootTranslation);
+	}
+
+	// optionally snap the root to the height from the source
+	if (RootHeightSource == ERootMotionHeightSource::CopyHeightFromSource)
+	{
+		FVector RootTranslation = OutRootTransform.GetTranslation();
+		RootTranslation.Z = InSourceGlobalPose[SourceRootIndex].GetTranslation().Z;
+		OutRootTransform.SetTranslation(RootTranslation);
+	}
+}
+
+void URootMotionGeneratorOp::CopyRootMotionFromSourceRoot(
+	FTransform& OutRootTransform,
+	const UIKRetargetProcessor* Processor,
+	const TArray<FTransform>& InSourceGlobalPose) const
+{
+	// In this case, we are copying root motion from the source root
+	// But we also scale it based on the relative height of the source/target Pelvis
+		
+	// copy root motion from source root bone
+	OutRootTransform = InSourceGlobalPose[SourceRootIndex];
+	FVector NewRootLocation = OutRootTransform.GetTranslation();
+
+	// scale root motion by same scale factor applied to Pelvis (and modified by settings)
+	FVector RootDelta = NewRootLocation - SourceRootInRefPose.GetTranslation();
+	RootDelta *= Processor->GetRootRetargeter().GetGlobalScaleVector();
+	NewRootLocation = SourceRootInRefPose.GetTranslation() + RootDelta;
+
+	// optionally snap the root to the ground plane
+	if (RootHeightSource == ERootMotionHeightSource::SnapToGround)
+	{
+		NewRootLocation.Z = 0.f;
+	}
+
+	// apply the modified translation
+	OutRootTransform.SetTranslation(NewRootLocation);
+}
+
+#undef LOCTEXT_NAMESPACE
diff --git a/Engine/Plugins/Animation/IKRig/Source/IKRig/Private/Rig/IKRigProcessor.cpp b/Engine/Plugins/Animation/IKRig/Source/IKRig/Private/Rig/IKRigProcessor.cpp
--- a/Engine/Plugins/Animation/IKRig/Source/IKRig/Private/Rig/IKRigProcessor.cpp
+++ b/Engine/Plugins/Animation/IKRig/Source/IKRig/Private/Rig/IKRigProcessor.cpp
@@ -145,9 +145,9 @@ void UIKRigProcessor::Initialize(
 		}
 
 		// create duplicate solver instance with unique name
-		FString Name = IKRigSolver->GetName() + "_SolverInstance_";
-		Name.AppendInt(SolverIndex++);
-		UIKRigSolver* Solver = DuplicateObject(IKRigSolver, this, FName(*Name));
+		const FName BaseName = FName(*(IKRigSolver->GetName() + "_SolverInstance_"));
+		const FName UniqueName = MakeUniqueObjectName(this, IKRigSolver->GetClass(), BaseName, EUniqueObjectNameOptions::GloballyUnique);
+		UIKRigSolver* Solver = DuplicateObject(IKRigSolver, this, UniqueName);
 
 		// remove excluded goals from the solver
 		for (const FName& ExcludedGoal : ExcludedGoals)
diff --git a/Engine/Plugins/Animation/IKRig/Source/IKRig/Public/Retargeter/IKRetargetOps.h b/Engine/Plugins/Animation/IKRig/Source/IKRig/Public/Retargeter/IKRetargetOps.h
--- a/Engine/Plugins/Animation/IKRig/Source/IKRig/Public/Retargeter/IKRetargetOps.h
+++ b/Engine/Plugins/Animation/IKRig/Source/IKRig/Public/Retargeter/IKRetargetOps.h
@@ -21,18 +21,21 @@ public:
 	
 	// override to cache internal data when initializing the processor
 	virtual bool Initialize(
+	const UIKRetargetProcessor* Processor,
 		const FRetargetSkeleton& SourceSkeleton,
 		const FTargetSkeleton& TargetSkeleton,
-		const UIKRetargetProcessor* Processor,
 		FIKRigLogger& Log) PURE_VIRTUAL(, return false;);
 
 	// override to evaluate this operation and modify the output pose
 	virtual void Run(
+		const UIKRetargetProcessor* Processor,
 		const TArray<FTransform>& InSourceGlobalPose,
 		TArray<FTransform>& OutTargetGlobalPose) PURE_VIRTUAL(,);
 
 	UPROPERTY()
 	bool bIsEnabled = true;
+
+	bool bIsInitialized = false;
 	
 #if WITH_EDITOR
 	// override to automate initial setup after being added to the stack
diff --git a/Engine/Plugins/Animation/IKRig/Source/IKRig/Public/Retargeter/IKRetargetProcessor.h b/Engine/Plugins/Animation/IKRig/Source/IKRig/Public/Retargeter/IKRetargetProcessor.h
--- a/Engine/Plugins/Animation/IKRig/Source/IKRig/Public/Retargeter/IKRetargetProcessor.h
+++ b/Engine/Plugins/Animation/IKRig/Source/IKRig/Public/Retargeter/IKRetargetProcessor.h
@@ -454,18 +454,6 @@ struct FRetargetDebugData
 	FTransform StrideWarpingFrame;
 };
 
-struct FRetargetOps
-{
-	TArray<TObjectPtr<URetargetOpBase>> OpStack;
-
-	bool Initialize(
-		const FRetargetSkeleton& SourceSkeleton,
-		const FTargetSkeleton& TargetSkeleton,
-		const TArray<TObjectPtr<URetargetOpBase>>& OpStackFromAsset,
-		UIKRetargetProcessor* Processor,
-		FIKRigLogger& Log);
-};
-
 /** The runtime processor that converts an input pose from a source skeleton into an output pose on a target skeleton.
  * To use:
  * 1. Initialize a processor with a Source/Target skeletal mesh and a UIKRetargeter asset.
@@ -546,7 +534,7 @@ public:
 	const FRootRetargeter& GetRootRetargeter() const { return RootRetargeter; };
 	
 	// Get read only access to the retarget ops currently running in processor
-	const FRetargetOps& GetRetargetOps() const {return RetargetOps; };
+	const TArray<TObjectPtr<URetargetOpBase>>& GetRetargetOps() const {return OpStack; };
 	
 	/** Reset the IK planting state. */
 	void ResetPlanting();
@@ -602,7 +590,8 @@ private:
 	FRetargetGlobalSettings GlobalSettings;
 
 	/** The collection of operations to run in the final phase of retargeting */
-	FRetargetOps RetargetOps;
+	UPROPERTY(Transient) // must be property to keep from being GC'd
+	TArray<TObjectPtr<URetargetOpBase>> OpStack;
 	
 	/** Initializes the FRootRetargeter */
 	bool InitializeRoots();
@@ -612,6 +601,9 @@ private:
 
 	/** Initializes the IK Rig that evaluates the IK solve for the target IK chains */
 	bool InitializeIKRig(UObject* Outer, const USkeletalMesh* InSkeletalMesh);
+
+	// Initialize the retarget op stack
+	bool InitializeOpStack(const TArray<TObjectPtr<URetargetOpBase>>& OpStackFromAsset);
 	
 	/** Internal retarget phase for the root. */
 	void RunRootRetarget(const TArray<FTransform>& InGlobalTransforms, TArray<FTransform>& OutGlobalTransforms);
diff --git a/Engine/Plugins/Animation/IKRig/Source/IKRig/Public/Retargeter/IKRetargeter.h b/Engine/Plugins/Animation/IKRig/Source/IKRig/Public/Retargeter/IKRetargeter.h
--- a/Engine/Plugins/Animation/IKRig/Source/IKRig/Public/Retargeter/IKRetargeter.h
+++ b/Engine/Plugins/Animation/IKRig/Source/IKRig/Public/Retargeter/IKRetargeter.h
@@ -208,16 +208,11 @@ public:
 	
 	UIKRetargeter(const FObjectInitializer& ObjectInitializer);
 
-	// Get read-only access to the source IK Rig asset 
-	const UIKRigDefinition* GetSourceIKRig() const;
-	// Get read-only access to the target IK Rig asset 
-	const UIKRigDefinition* GetTargetIKRig() const;
+	// Get read-only access to the source or target IK Rig asset 
+	const UIKRigDefinition* GetIKRig(ERetargetSourceOrTarget SourceOrTarget) const;
 	// Get read-write access to the source IK Rig asset.
 	// WARNING: do not use for editing the data model. Use Controller class instead. 
-	UIKRigDefinition* GetSourceIKRigWriteable() const;
-	// Get read-write access to the target IK Rig asset.
-	// WARNING: do not use for editing the data model. Use Controller class instead. 
-	UIKRigDefinition* GetTargetIKRigWriteable() const;
+	UIKRigDefinition* GetIKRigWriteable(ERetargetSourceOrTarget SourceOrTarget) const;
 	#if WITH_EDITORONLY_DATA
 	// Get read-only access to preview meshes
 	const USkeletalMesh* GetPreviewMesh(ERetargetSourceOrTarget SourceOrTarget) const;
diff --git a/Engine/Plugins/Animation/IKRig/Source/IKRig/Public/Retargeter/RetargetOps/PinBoneOp.h b/Engine/Plugins/Animation/IKRig/Source/IKRig/Public/Retargeter/RetargetOps/PinBoneOp.h
--- a/Engine/Plugins/Animation/IKRig/Source/IKRig/Public/Retargeter/RetargetOps/PinBoneOp.h
+++ b/Engine/Plugins/Animation/IKRig/Source/IKRig/Public/Retargeter/RetargetOps/PinBoneOp.h
@@ -44,12 +44,13 @@ class UPinBoneOp : public URetargetOpBase
 public:
 	
 	virtual bool Initialize(
+	const UIKRetargetProcessor* Processor,
 		const FRetargetSkeleton& SourceSkeleton,
 		const FTargetSkeleton& TargetSkeleton,
-		const UIKRetargetProcessor* Processor,
 		FIKRigLogger& Log) override;
 	
 	virtual void Run(
+		const UIKRetargetProcessor* Processor,
 		const TArray<FTransform>& InSourceGlobalPose,
 		TArray<FTransform>& OutTargetGlobalPose) override;
 
@@ -75,12 +76,8 @@ public:
 	// A manual offset to apply in local space
 	UPROPERTY(EditAnywhere, Category=Settings)
 	FTransform LocalOffset;
-
-	// Toggle this constraint on/off
-	UPROPERTY(EditAnywhere, Category=Settings)
-	bool bEnabled = true;
 	
-#if WITH_EDITORONLY_DATA
+#if WITH_EDITOR
 	virtual FText GetNiceName() const override { return FText(LOCTEXT("OpName", "Pin Bone")); };
 	virtual FText WarningMessage() const override { return Message; };
 	FText Message;
diff --git a/Engine/Plugins/Animation/IKRig/Source/IKRig/Public/Retargeter/RetargetOps/RootMotionGeneratorOp.h b/Engine/Plugins/Animation/IKRig/Source/IKRig/Public/Retargeter/RetargetOps/RootMotionGeneratorOp.h
new file mode 100644
--- /dev/null
+++ b/Engine/Plugins/Animation/IKRig/Source/IKRig/Public/Retargeter/RetargetOps/RootMotionGeneratorOp.h
@@ -0,0 +1,120 @@
+﻿// Copyright Epic Games, Inc. All Rights Reserved.
+#pragma once
+
+#include "Retargeter/IKRetargetOps.h"
+
+#include "RootMotionGeneratorOp.generated.h"
+
+#define LOCTEXT_NAMESPACE "URootMotionGeneratorOp"
+
+// where to copy the motion of the root from
+UENUM()
+enum class ERootMotionSource : uint8
+{
+	CopyFromSourceRoot,
+	GenerateFromTargetPelvis,
+};
+
+// where to copy the height of the root from
+UENUM()
+enum class ERootMotionHeightSource : uint8
+{
+	CopyHeightFromSource,
+	SnapToGround,
+};
+
+UCLASS(BlueprintType, EditInlineNew)
+class URootMotionGeneratorOp : public URetargetOpBase
+{
+	GENERATED_BODY()
+
+public:
+	
+	virtual bool Initialize(
+		const UIKRetargetProcessor* Processor,
+		const FRetargetSkeleton& SourceSkeleton,
+		const FTargetSkeleton& TargetSkeleton,
+		FIKRigLogger& Log) override;
+	
+	virtual void Run(
+		const UIKRetargetProcessor* Processor,
+		const TArray<FTransform>& InSourceGlobalPose,
+		TArray<FTransform>& OutTargetGlobalPose) override;
+
+	// The root of the source skeleton.
+	UPROPERTY(EditAnywhere, Category="OpSettings")
+	FName SourceRootBone;
+
+	// The root of the target skeleton.
+	UPROPERTY(EditAnywhere, Category="OpSettings")
+	FName TargetRootBone;
+
+	// The pelvis of the target skeleton.
+	UPROPERTY(EditAnywhere, Category="OpSettings")
+	FName TargetPelvisBone;
+
+	// Where to copy the root motion from.
+	// Copy From Source Root: copies the motion from the source root bone, but scales it according to relative height difference.
+	// Generate From Target Pelvis: uses the retargeted Pelvis motion to generate root motion.
+	// NOTE: Generating root motion from the Pelvis 
+	UPROPERTY(EditAnywhere, Category="OpSettings")
+	ERootMotionSource RootMotionSource;
+	
+	// How to set the height of the root bone.
+	// Copy Height From Source: copies the height of the root bone on the source skeleton's root bone.
+	// Snap To Ground: sets the root bone height to the ground plane (Z=0).
+	UPROPERTY(EditAnywhere, Category="OpSettings")
+	ERootMotionHeightSource RootHeightSource;
+
+	// Will transform all children of the target root that are not themselves part of a retarget chain.
+	UPROPERTY(EditAnywhere, Category="OpSettings")
+	bool bPropagateToNonRetargetedChildren = true;
+
+	// Applies only when generating root motion from the Pelvis.
+	// Maintains the offset between the root and pelvis as recorded in the target retarget pose.
+	// If false, the root bone is placed directly under the Pelvis bone.
+	UPROPERTY(EditAnywhere, Category="OpSettings")
+	bool bMaintainOffsetFromPelvis = true;
+
+	// Applies only when generating root motion from the Pelvis.
+	// When true, the applied offset will be rotated by the Pelvis.
+	// (NOTE: This may cause unwanted rotations, for example if Pelvis Yaw is animated.)
+	UPROPERTY(EditAnywhere, Category="OpSettings")
+	bool bRotateWithPelvis = false;
+
+	// A manual offset to apply in global space to the root bone.
+	UPROPERTY(EditAnywhere, Category=Settings)
+	FTransform GlobalOffset;
+
+#if WITH_EDITOR
+	virtual void OnAddedToStack(const UIKRetargeter* Asset) override;
+	virtual FText GetNiceName() const override { return FText(LOCTEXT("OpName", "Root Motion Generator")); };
+	virtual FText WarningMessage() const override { return Message; };
+	FText Message;
+#endif
+
+private:
+
+	void GenerateRootMotionFromTargetPelvis(
+		FTransform& OutRootTransform,
+		const TArray<FTransform>& InSourceGlobalPose,
+		const TArray<FTransform>& InTargetGlobalPose) const;
+	
+	void CopyRootMotionFromSourceRoot(
+		FTransform& OutRootTransform,
+		const UIKRetargetProcessor* Processor,
+		const TArray<FTransform>& InSourceGlobalPose) const;
+	
+	int32 SourceRootIndex;
+	int32 TargetRootIndex;
+	int32 TargetPelvisIndex;
+
+	FTransform TargetPelvisRelativeToTargetRootRefPose;
+	FTransform TargetPelvisInRefPose;
+	FTransform SourceRootInRefPose;
+	FTransform TargetRootInRefPose;
+	
+	TArray<int32> NonRetargetedChildrenOfRoot;
+};
+
+#undef LOCTEXT_NAMESPACE
diff --git a/Engine/Plugins/Animation/IKRig/Source/IKRigDeveloper/Private/AnimGraphNode_RetargetPoseFromMesh.cpp b/Engine/Plugins/Animation/IKRig/Source/IKRigDeveloper/Private/AnimGraphNode_RetargetPoseFromMesh.cpp
--- a/Engine/Plugins/Animation/IKRig/Source/IKRigDeveloper/Private/AnimGraphNode_RetargetPoseFromMesh.cpp
+++ b/Engine/Plugins/Animation/IKRig/Source/IKRigDeveloper/Private/AnimGraphNode_RetargetPoseFromMesh.cpp
@@ -76,18 +76,20 @@ void UAnimGraphNode_RetargetPoseFromMesh::ValidateAnimNodeDuringCompilation(USke
 	}
 
 	// validate SOURCE IK Rig asset has been assigned
-	if (!Node.IKRetargeterAsset->GetSourceIKRig())
+	const UIKRigDefinition* SourceIKRig = Node.IKRetargeterAsset->GetIKRig(ERetargetSourceOrTarget::Source);
+	if (!SourceIKRig)
 	{
 		MessageLog.Warning(TEXT("@@ has IK Retargeter that is missing a source IK Rig asset."), this);
 	}
 
 	// validate TARGET IK Rig asset has been assigned
-	if (!Node.IKRetargeterAsset->GetTargetIKRig())
+	const UIKRigDefinition* TargetIKRig = Node.IKRetargeterAsset->GetIKRig(ERetargetSourceOrTarget::Target);
+	if (!TargetIKRig)
 	{
 		MessageLog.Warning(TEXT("@@ has IK Retargeter that is missing a target IK Rig asset."), this);
 	}
 
-	if (!(Node.IKRetargeterAsset->GetSourceIKRig() && Node.IKRetargeterAsset->GetTargetIKRig()))
+	if (!(SourceIKRig && TargetIKRig))
 	{
 		return;
 	}
@@ -115,7 +117,7 @@ void UAnimGraphNode_RetargetPoseFromMesh::ValidateAnimNodeDuringCompilation(USke
 	{
 		// validate that target bone chains exist on this skeleton
 		const FReferenceSkeleton &RefSkel = ForSkeleton->GetReferenceSkeleton();
-		const TArray<FBoneChain> &TargetBoneChains = Node.IKRetargeterAsset->GetTargetIKRig()->GetRetargetChains();
+		const TArray<FBoneChain> &TargetBoneChains = TargetIKRig->GetRetargetChains();
 		for (const FBoneChain &Chain : TargetBoneChains)
 		{
 			if (RefSkel.FindBoneIndex(Chain.StartBone.BoneName) == INDEX_NONE)
@@ -138,8 +140,8 @@ void UAnimGraphNode_RetargetPoseFromMesh::PreloadRequiredAssets()
 	if (Node.IKRetargeterAsset)
 	{
 		PreloadObject(Node.IKRetargeterAsset);
-		PreloadObject(Node.IKRetargeterAsset->GetSourceIKRigWriteable());
-		PreloadObject(Node.IKRetargeterAsset->GetTargetIKRigWriteable());
+		PreloadObject(Node.IKRetargeterAsset->GetIKRigWriteable(ERetargetSourceOrTarget::Source));
+		PreloadObject(Node.IKRetargeterAsset->GetIKRigWriteable(ERetargetSourceOrTarget::Target));
 	}
 }
 
diff --git a/Engine/Plugins/Animation/IKRig/Source/IKRigEditor/Private/RetargetEditor/IKRetargetBatchOperation.cpp b/Engine/Plugins/Animation/IKRig/Source/IKRigEditor/Private/RetargetEditor/IKRetargetBatchOperation.cpp
--- a/Engine/Plugins/Animation/IKRig/Source/IKRigEditor/Private/RetargetEditor/IKRetargetBatchOperation.cpp
+++ b/Engine/Plugins/Animation/IKRig/Source/IKRigEditor/Private/RetargetEditor/IKRetargetBatchOperation.cpp
@@ -505,7 +505,7 @@ void UIKRetargetBatchOperation::RunRetarget(FIKRetargetBatchOperationContext& Co
 	}
 
 	// validate a source IK rig was provided
-	const UIKRigDefinition* SrcIKRig = Context.IKRetargetAsset->GetSourceIKRig();
+	const UIKRigDefinition* SrcIKRig = Context.IKRetargetAsset->GetIKRig(ERetargetSourceOrTarget::Source);
 	if (!SrcIKRig)
 	{
 		UE_LOG(LogTemp, Warning, TEXT("Batch retarget aborted. Specified IK Retargeter does not reference a source IK Rig."));
@@ -513,7 +513,7 @@ void UIKRetargetBatchOperation::RunRetarget(FIKRetargetBatchOperationContext& Co
 	}
 
 	// validate a target IK rig was provided
-	const UIKRigDefinition* TgtIKRig = Context.IKRetargetAsset->GetTargetIKRig();
+	const UIKRigDefinition* TgtIKRig = Context.IKRetargetAsset->GetIKRig(ERetargetSourceOrTarget::Target);
 	if (!TgtIKRig)
 	{
 		UE_LOG(LogTemp, Warning, TEXT("Batch retarget aborted. Specified IK Retargeter does not reference a target IK Rig."));
diff --git a/Engine/Plugins/Animation/IKRig/Source/IKRigEditor/Private/RetargetEditor/IKRetargetDetails.cpp b/Engine/Plugins/Animation/IKRig/Source/IKRigEditor/Private/RetargetEditor/IKRetargetDetails.cpp
--- a/Engine/Plugins/Animation/IKRig/Source/IKRigEditor/Private/RetargetEditor/IKRetargetDetails.cpp
+++ b/Engine/Plugins/Animation/IKRig/Source/IKRigEditor/Private/RetargetEditor/IKRetargetDetails.cpp
@@ -1461,6 +1461,11 @@ void FRetargetOpStackCustomization::CustomizeDetails(IDetailLayoutBuilder& Detai
 			const TSharedPtr<IPropertyHandle> PropertyHandle = DetailBuilder.AddObjectPropertyData({Op}, Prop->GetFName());
 			if (PropertyHandle && PropertyHandle->IsValidHandle())
 			{
+				PropertyHandle->SetOnPropertyValueChanged(FSimpleDelegate::CreateLambda([this]()
+				{
+					Controller->ReinitializeProcessor();
+				}));
+				
 				PropertyHandle->SetOnChildPropertyValueChanged(FSimpleDelegate::CreateLambda([this]()
 				{
 					Controller->ReinitializeProcessor();
diff --git a/Engine/Plugins/Animation/IKRig/Source/IKRigEditor/Private/RetargetEditor/IKRetargetEditorController.cpp b/Engine/Plugins/Animation/IKRig/Source/IKRigEditor/Private/RetargetEditor/IKRetargetEditorController.cpp
--- a/Engine/Plugins/Animation/IKRig/Source/IKRigEditor/Private/RetargetEditor/IKRetargetEditorController.cpp
+++ b/Engine/Plugins/Animation/IKRig/Source/IKRigEditor/Private/RetargetEditor/IKRetargetEditorController.cpp
@@ -198,20 +198,20 @@ void FIKRetargetEditorController::BindToIKRigAssets()
 
 	BoundIKRigs.Empty();
 	
-	if (Asset->GetSourceIKRigWriteable())
+	if (UIKRigDefinition* SourceIKRig = Asset->GetIKRigWriteable(ERetargetSourceOrTarget::Source))
 	{
-		BoundIKRigs.Emplace(FBoundIKRig(Asset->GetSourceIKRigWriteable(), *this));
+		BoundIKRigs.Emplace(FBoundIKRig(SourceIKRig, *this));
 	}
 
-	if (Asset->GetTargetIKRigWriteable())
+	if (UIKRigDefinition* TargetIKRig = Asset->GetIKRigWriteable(ERetargetSourceOrTarget::Target))
 	{
-		BoundIKRigs.Emplace(FBoundIKRig(Asset->GetTargetIKRigWriteable(), *this));
+		BoundIKRigs.Emplace(FBoundIKRig(TargetIKRig, *this));
 	}
 }
 
 void FIKRetargetEditorController::HandleIKRigNeedsInitialized(UIKRigDefinition* ModifiedIKRig) const
 {
-	UIKRetargeter* Retargeter = AssetController->GetAsset();
+	const UIKRetargeter* Retargeter = AssetController->GetAsset();
 	check(Retargeter)
 	HandleRetargeterNeedsInitialized();
 }
diff --git a/Engine/Plugins/Animation/IKRig/Source/IKRigEditor/Private/RetargetEditor/IKRetargeterController.cpp b/Engine/Plugins/Animation/IKRig/Source/IKRigEditor/Private/RetargetEditor/IKRetargeterController.cpp
--- a/Engine/Plugins/Animation/IKRig/Source/IKRigEditor/Private/RetargetEditor/IKRetargeterController.cpp
+++ b/Engine/Plugins/Animation/IKRig/Source/IKRigEditor/Private/RetargetEditor/IKRetargeterController.cpp
@@ -86,13 +86,13 @@ void UIKRetargeterController::SetIKRig(const ERetargetSourceOrTarget SourceOrTar
 const UIKRigDefinition* UIKRetargeterController::GetIKRig(const ERetargetSourceOrTarget SourceOrTarget) const
 {
 	FScopeLock Lock(&ControllerLock);
-	return SourceOrTarget == ERetargetSourceOrTarget::Source ? Asset->GetSourceIKRig() : Asset->GetTargetIKRig();
+	return Asset->GetIKRig(SourceOrTarget);
 }
 
 UIKRigDefinition* UIKRetargeterController::GetIKRigWriteable(const ERetargetSourceOrTarget SourceOrTarget) const
 {
 	FScopeLock Lock(&ControllerLock);
-	return SourceOrTarget == ERetargetSourceOrTarget::Source ? Asset->GetSourceIKRigWriteable() : Asset->GetTargetIKRigWriteable();
+	return Asset->GetIKRigWriteable(SourceOrTarget);
 }
 
 void UIKRetargeterController::SetPreviewMesh(
@@ -359,7 +359,7 @@ void UIKRetargeterController::GetChainNames(const ERetargetSourceOrTarget Source
 
 void UIKRetargeterController::CleanChainMapping() const
 {
-	if (IsValid(Asset->GetTargetIKRig()))
+	if (IsValid(Asset->GetIKRig(ERetargetSourceOrTarget::Target)))
 	{
 		TArray<FName> TargetChainNames;
 		GetChainNames(ERetargetSourceOrTarget::Target, TargetChainNames);
@@ -398,7 +398,7 @@ void UIKRetargeterController::CleanChainMapping() const
 		}
 	}
 
-	if (IsValid(Asset->GetSourceIKRig()))
+	if (IsValid(Asset->GetIKRig(ERetargetSourceOrTarget::Source)))
 	{
 		TArray<FName> SourceChainNames;
 		GetChainNames(ERetargetSourceOrTarget::Source,SourceChainNames);
@@ -536,7 +536,7 @@ void UIKRetargeterController::AutoMapChains(const EAutoMapChainType AutoMapType,
 
 void UIKRetargeterController::HandleRetargetChainAdded(UIKRigDefinition* IKRig) const
 {
-	const bool bIsTargetRig = IKRig == Asset->GetTargetIKRig();
+	const bool bIsTargetRig = IKRig == Asset->GetIKRig(ERetargetSourceOrTarget::Target);
 	if (!bIsTargetRig)
 	{
 		// if a source chain is added, it will simply be available as a new option, no need to reinitialize until it's used
@@ -550,8 +550,9 @@ void UIKRetargeterController::HandleRetargetChainAdded(UIKRigDefinition* IKRig)
 
 void UIKRetargeterController::HandleRetargetChainRenamed(UIKRigDefinition* IKRig, FName OldChainName, FName NewChainName) const
 {
-	const bool bIsSourceRig = IKRig == Asset->GetSourceIKRig();
-	check(bIsSourceRig || IKRig == Asset->GetTargetIKRig())
+	const bool bIsSourceRig = IKRig == Asset->GetIKRig(ERetargetSourceOrTarget::Source);
+	const bool bIsTargetRig = IKRig == Asset->GetIKRig(ERetargetSourceOrTarget::Target);
+	check(bIsSourceRig || bIsTargetRig)
 	for (URetargetChainSettings* ChainMap : Asset->ChainSettings)
 	{
 		FName& ChainNameToUpdate = bIsSourceRig ? ChainMap->SourceChain : ChainMap->TargetChain;
@@ -566,8 +567,9 @@ void UIKRetargeterController::HandleRetargetChainRenamed(UIKRigDefinition* IKRig
 
 void UIKRetargeterController::HandleRetargetChainRemoved(UIKRigDefinition* IKRig, const FName& InChainRemoved) const
 {
-	const bool bIsSourceRig = IKRig == Asset->GetSourceIKRig();
-	check(bIsSourceRig || IKRig == Asset->GetTargetIKRig())
+	const bool bIsSourceRig = IKRig == Asset->GetIKRig(ERetargetSourceOrTarget::Source);
+	const bool bIsTargetRig = IKRig == Asset->GetIKRig(ERetargetSourceOrTarget::Target);
+	check(bIsSourceRig || bIsTargetRig)
 
 	// set source chain name to NONE if it has been deleted 
 	if (bIsSourceRig)
@@ -868,7 +870,7 @@ void UIKRetargeterController::SetRotationOffsetForRetargetPoseBone(
 	FScopeLock Lock(&ControllerLock);
 	FIKRetargetPose& Pose = GetCurrentRetargetPose(SourceOrTarget);
 	Pose.SetDeltaRotationForBone(BoneName, RotationOffset);
-	const UIKRigDefinition* IKRig = SourceOrTarget == ERetargetSourceOrTarget::Source ? GetAsset()->GetSourceIKRig() : GetAsset()->GetTargetIKRig();
+	const UIKRigDefinition* IKRig = GetAsset()->GetIKRig(SourceOrTarget);
 	Pose.SortHierarchically(IKRig->GetSkeleton());
 }
 
@@ -936,7 +938,7 @@ URetargetChainSettings* UIKRetargeterController::GetChainSettings(const FName& T
 
 void UIKRetargeterController::SortChainMapping() const
 {
-	const UIKRigDefinition* TargetIKRig = Asset->GetTargetIKRig();
+	const UIKRigDefinition* TargetIKRig = Asset->GetIKRig(ERetargetSourceOrTarget::Target);
 	if (!IsValid(TargetIKRig))
 	{
 		return;
diff --git a/Engine/Plugins/Animation/IKRig/Source/IKRigEditor/Private/RetargetEditor/SIKRetargetChainMapList.cpp b/Engine/Plugins/Animation/IKRig/Source/IKRigEditor/Private/RetargetEditor/SIKRetargetChainMapList.cpp
--- a/Engine/Plugins/Animation/IKRig/Source/IKRigEditor/Private/RetargetEditor/SIKRetargetChainMapList.cpp
+++ b/Engine/Plugins/Animation/IKRig/Source/IKRigEditor/Private/RetargetEditor/SIKRetargetChainMapList.cpp
@@ -421,9 +421,9 @@ bool SIKRetargetChainMapList::IsChainMapEnabled() const
 		return false; 
 	}
 	
-	if (RetargeterController->GetAsset()->GetTargetIKRig())
+	if (const UIKRigDefinition* IKRig = RetargeterController->GetAsset()->GetIKRig(ERetargetSourceOrTarget::Target))
 	{
-		const TArray<FBoneChain>& Chains = RetargeterController->GetAsset()->GetTargetIKRig()->GetRetargetChains();
+		const TArray<FBoneChain>& Chains = IKRig->GetRetargetChains();
 		return Chains.Num() > 0;
 	}
 
diff --git a/Engine/Plugins/Animation/IKRig/Source/IKRigEditor/Private/RetargetEditor/SRetargetAnimAssetsWindow.cpp b/Engine/Plugins/Animation/IKRig/Source/IKRigEditor/Private/RetargetEditor/SRetargetAnimAssetsWindow.cpp
--- a/Engine/Plugins/Animation/IKRig/Source/IKRigEditor/Private/RetargetEditor/SRetargetAnimAssetsWindow.cpp
+++ b/Engine/Plugins/Animation/IKRig/Source/IKRigEditor/Private/RetargetEditor/SRetargetAnimAssetsWindow.cpp
@@ -258,7 +258,7 @@ void SRetargetAnimAssetsWindow::Construct(const FArguments& InArgs)
 								return false;
 							}
 							
-							return BatchContext.IKRetargetAsset->GetSourceIKRig() != nullptr;
+							return BatchContext.IKRetargetAsset->GetIKRig(ERetargetSourceOrTarget::Source) != nullptr;
 						})
 						.ObjectPath(this, &SRetargetAnimAssetsWindow::GetCurrentSourceMeshPath)
 						.OnObjectChanged(this, &SRetargetAnimAssetsWindow::SourceMeshAssigned)
@@ -312,7 +312,7 @@ void SRetargetAnimAssetsWindow::Construct(const FArguments& InArgs)
 								return false;
 							}
 							
-							return BatchContext.IKRetargetAsset->GetTargetIKRig() != nullptr;
+							return BatchContext.IKRetargetAsset->GetIKRig(ERetargetSourceOrTarget::Target) != nullptr;
 						})
 						.ObjectPath(this, &SRetargetAnimAssetsWindow::GetCurrentTargetMeshPath)
 						.OnObjectChanged(this, &SRetargetAnimAssetsWindow::TargetMeshAssigned)
diff --git a/Engine/Plugins/Animation/IKRig/Source/IKRigEditor/Private/RetargetEditor/SRetargetOpStack.cpp b/Engine/Plugins/Animation/IKRig/Source/IKRigEditor/Private/RetargetEditor/SRetargetOpStack.cpp
--- a/Engine/Plugins/Animation/IKRig/Source/IKRigEditor/Private/RetargetEditor/SRetargetOpStack.cpp
+++ b/Engine/Plugins/Animation/IKRig/Source/IKRigEditor/Private/RetargetEditor/SRetargetOpStack.cpp
@@ -146,11 +146,11 @@ bool SRetargetOpStackItem::GetWarningMessage(FText& Message) const
 		return false;
 	}
 
-	const FRetargetOps& RetargetOps = Processor->GetRetargetOps();
+	const TArray<TObjectPtr<URetargetOpBase>>& RetargetOps = Processor->GetRetargetOps();
 	const int32 OpIndex = StackElement.Pin()->IndexInStack;
-	if (RetargetOps.OpStack.IsValidIndex(OpIndex))
+	if (RetargetOps.IsValidIndex(OpIndex))
 	{
-		Message = RetargetOps.OpStack[OpIndex]->WarningMessage();
+		Message = RetargetOps[OpIndex]->WarningMessage();
 		return true;
 	}
 	
@@ -358,7 +358,6 @@ void SRetargetOpStack::RefreshStackView()
 	}
 
 	// generate all list items
-	static const FText UnknownOpTxt = FText::FromString("Unknown Retarget Op");
 	ListViewItems.Reset();
 	UIKRetargeterController* AssetController = Controller->AssetController;
 	const int32 NumOps = AssetController->GetNumRetargetOps();
@@ -366,8 +365,7 @@ void SRetargetOpStack::RefreshStackView()
 	{
 		const URetargetOpBase* Op = AssetController->GetRetargetOpAtIndex(i);
 		FString UniqueName =  FString::FromInt(i+1) + " - " + Op->GetNiceName().ToString();
-		const FText DisplayName = Op ? FText::FromString(UniqueName) : UnknownOpTxt;
-		TSharedPtr<FRetargetOpStackElement> StackElement = FRetargetOpStackElement::Make(DisplayName, i);
+		TSharedPtr<FRetargetOpStackElement> StackElement = FRetargetOpStackElement::Make(FText::FromString(UniqueName), i);
 		ListViewItems.Add(StackElement);
 	}
 
