Fixed duplicate localization tag in RootMotionGeneratorOp

diff --git a/Engine/Plugins/Animation/IKRig/Source/IKRig/Private/Retargeter/RetargetOps/RootMotionGeneratorOp.cpp b/Engine/Plugins/Animation/IKRig/Source/IKRig/Private/Retargeter/RetargetOps/RootMotionGeneratorOp.cpp
--- a/Engine/Plugins/Animation/IKRig/Source/IKRig/Private/Retargeter/RetargetOps/RootMotionGeneratorOp.cpp
+++ b/Engine/Plugins/Animation/IKRig/Source/IKRig/Private/Retargeter/RetargetOps/RootMotionGeneratorOp.cpp
@@ -3,7 +3,7 @@
 
 #include "Engine/SkeletalMesh.h"
 
-#define LOCTEXT_NAMESPACE "RootMotionRemapOp"
+#define LOCTEXT_NAMESPACE "RootMotionGeneratorOp"
 
 #if WITH_EDITOR
 void URootMotionGeneratorOp::OnAddedToStack(const UIKRetargeter* Asset)
@@ -38,14 +38,14 @@ bool URootMotionGeneratorOp::Initialize(
 	bool bHasAllPrerequisites = true;
 	if (SourceRootIndex == INDEX_NONE)
 	{
-		Log.LogWarning(FText::Format(LOCTEXT("MissingRootBone", "Root Motion Remap Op, missing source root bone {0}."), FText::FromName(SourceRootBone)));
+		Log.LogWarning(FText::Format(LOCTEXT("MissingSourceRootBone", "Root Motion Remap Op, missing source root bone {0}."), FText::FromName(SourceRootBone)));
 		bHasAllPrerequisites = false;
 	}
 	
 	TargetRootIndex = TargetSkeleton.BoneNames.Find(TargetRootBone);
 	if (TargetRootIndex == INDEX_NONE)
 	{
-		Log.LogWarning(FText::Format(LOCTEXT("MissingRootBone", "Root Motion Remap Op, missing target root bone {0}."), FText::FromName(TargetRootBone)));
+		Log.LogWarning(FText::Format(LOCTEXT("MissingTargetRootBone", "Root Motion Remap Op, missing target root bone {0}."), FText::FromName(TargetRootBone)));
 		bHasAllPrerequisites = false;
 	}
 
