Adding button to access retarget asset settings in editor toolbar.
Converted buttons to toggles.
Removed behavior that switches to asset details when clicking in empty viewport

diff --git a/Engine/Plugins/Animation/IKRig/Source/IKRigEditor/Private/RetargetEditor/IKRetargetCommands.cpp b/Engine/Plugins/Animation/IKRig/Source/IKRigEditor/Private/RetargetEditor/IKRetargetCommands.cpp
--- a/Engine/Plugins/Animation/IKRig/Source/IKRigEditor/Private/RetargetEditor/IKRetargetCommands.cpp
+++ b/Engine/Plugins/Animation/IKRig/Source/IKRigEditor/Private/RetargetEditor/IKRetargetCommands.cpp
@@ -10,9 +10,10 @@ void FIKRetargetCommands::RegisterCommands()
 	UI_COMMAND(EditRetargetPose, "Edit Retarget Pose", "Enter into mode allowing manual editing of the target skeleton pose in the viewport.", EUserInterfaceActionType::Button, FInputChord());
 	UI_COMMAND(ShowRetargetPose, "Show Retarget Pose", "Display the retarget pose and retarget to the target skeleton.", EUserInterfaceActionType::Button, FInputChord());
 
-	UI_COMMAND(ShowGlobalSettings, "Global Settings", "Open global settings in detail tab.", EUserInterfaceActionType::Button, FInputChord());
-	UI_COMMAND(ShowRootSettings, "Root Settings", "Open root settings in detail tab.", EUserInterfaceActionType::Button, FInputChord());
-	UI_COMMAND(ShowPostSettings, "Post Settings", "Open post phase settings in detail tab.", EUserInterfaceActionType::Button, FInputChord());
+	UI_COMMAND(ShowAssetSettings, "Asset Settings", "Set the source or target IK Rigs, preview meshes. Adjust debug settings.", EUserInterfaceActionType::ToggleButton, FInputChord());
+	UI_COMMAND(ShowGlobalSettings, "Global Settings", "Access Stride Warping and retarget phases.", EUserInterfaceActionType::ToggleButton, FInputChord());
+	UI_COMMAND(ShowRootSettings, "Root Settings", "Adjust the motion of the Pelvis.", EUserInterfaceActionType::ToggleButton, FInputChord());
+	UI_COMMAND(ShowPostSettings, "Post Settings", "Create a stack of operations that run after the retarget.", EUserInterfaceActionType::ToggleButton, FInputChord());
 	
 	UI_COMMAND(ResetAllBones, "Reset All", "Sets the selected bones to the mesh reference pose.", EUserInterfaceActionType::Button, FInputChord());
 	UI_COMMAND(ResetSelectedBones, "Reset Selected Bones", "Sets the selected bones to the mesh reference pose.", EUserInterfaceActionType::Button, FInputChord());
diff --git a/Engine/Plugins/Animation/IKRig/Source/IKRigEditor/Private/RetargetEditor/IKRetargetEditor.cpp b/Engine/Plugins/Animation/IKRig/Source/IKRigEditor/Private/RetargetEditor/IKRetargetEditor.cpp
--- a/Engine/Plugins/Animation/IKRig/Source/IKRigEditor/Private/RetargetEditor/IKRetargetEditor.cpp
+++ b/Engine/Plugins/Animation/IKRig/Source/IKRigEditor/Private/RetargetEditor/IKRetargetEditor.cpp
@@ -139,8 +139,21 @@ void FIKRetargetEditor::BindCommands()
 
 
 	//
-	// Show global / root settings in details panel
+	// Show various settings in details panel
 	//
+	ToolkitCommands->MapAction(
+		Commands.ShowAssetSettings,
+		FExecuteAction::CreateLambda([this]()
+		{
+			UIKRetargeter* Asset = EditorController->AssetController->GetAsset();
+			return EditorController->SetDetailsObject(Asset);
+		}),
+		FCanExecuteAction(),
+		FIsActionChecked::CreateLambda([this]() ->bool
+		{
+			const UIKRetargeter* Asset = EditorController->AssetController->GetAsset();
+			return EditorController->IsObjectInDetailsView(Asset);	
+		}));
 	ToolkitCommands->MapAction(
 		Commands.ShowGlobalSettings,
 		FExecuteAction::CreateSP(EditorController, &FIKRetargetEditorController::ShowGlobalSettings),
@@ -288,6 +301,13 @@ void FIKRetargetEditor::FillToolbar(FToolBarBuilder& ToolbarBuilder)
 
 	ToolbarBuilder.BeginSection("Show Settings");
 	{
+		ToolbarBuilder.AddToolBarButton(
+		FIKRetargetCommands::Get().ShowAssetSettings,
+		NAME_None,
+		TAttribute<FText>(),
+		TAttribute<FText>(),
+		FSlateIcon(FIKRetargetEditorStyle::Get().GetStyleSetName(),"IKRetarget.AssetSettings"));
+		
 		ToolbarBuilder.AddToolBarButton(
 		FIKRetargetCommands::Get().ShowGlobalSettings,
 		NAME_None,
diff --git a/Engine/Plugins/Animation/IKRig/Source/IKRigEditor/Private/RetargetEditor/IKRetargetEditorController.cpp b/Engine/Plugins/Animation/IKRig/Source/IKRigEditor/Private/RetargetEditor/IKRetargetEditorController.cpp
--- a/Engine/Plugins/Animation/IKRig/Source/IKRigEditor/Private/RetargetEditor/IKRetargetEditorController.cpp
+++ b/Engine/Plugins/Animation/IKRig/Source/IKRigEditor/Private/RetargetEditor/IKRetargetEditorController.cpp
@@ -996,9 +996,6 @@ void FIKRetargetEditorController::ClearSelection(const bool bKeepBoneSelection)
 		SetRootSelected(false);
 		SelectedBoneNames[ERetargetSourceOrTarget::Source].Reset();
 		SelectedBoneNames[ERetargetSourceOrTarget::Target].Reset();
-
-		// show global details
-		SetDetailsObject(AssetController->GetAsset());
 	}
 
 	LastSelectedItem = ERetargetSelectionType::NONE;
diff --git a/Engine/Plugins/Animation/IKRig/Source/IKRigEditor/Public/RetargetEditor/IKRetargetCommands.h b/Engine/Plugins/Animation/IKRig/Source/IKRigEditor/Public/RetargetEditor/IKRetargetCommands.h
--- a/Engine/Plugins/Animation/IKRig/Source/IKRigEditor/Public/RetargetEditor/IKRetargetCommands.h
+++ b/Engine/Plugins/Animation/IKRig/Source/IKRigEditor/Public/RetargetEditor/IKRetargetCommands.h
@@ -25,11 +25,13 @@ public:
 	// edit retarget pose
 	TSharedPtr< FUICommandInfo > EditRetargetPose;
 
+	// open asset settings
+	TSharedPtr< FUICommandInfo > ShowAssetSettings;
 	// open global settings
 	TSharedPtr< FUICommandInfo > ShowGlobalSettings;
-	// open global settings
+	// open root settings
 	TSharedPtr< FUICommandInfo > ShowRootSettings;
-	// open global settings
+	// open post settings
 	TSharedPtr< FUICommandInfo > ShowPostSettings;
 
 	// enable Root retarget pass
diff --git a/Engine/Plugins/Animation/IKRig/Source/IKRigEditor/Public/RetargetEditor/IKRetargetEditorStyle.h b/Engine/Plugins/Animation/IKRig/Source/IKRigEditor/Public/RetargetEditor/IKRetargetEditorStyle.h
--- a/Engine/Plugins/Animation/IKRig/Source/IKRigEditor/Public/RetargetEditor/IKRetargetEditorStyle.h
+++ b/Engine/Plugins/Animation/IKRig/Source/IKRigEditor/Public/RetargetEditor/IKRetargetEditorStyle.h
@@ -20,6 +20,7 @@ public:
 		Set("ClassIcon.IKRetargeter", new IMAGE_BRUSH_SVG("Slate/IKRigRetargeter", Icon16x16));
 		Set("ClassThumbnail.IKRetargeter", new IMAGE_BRUSH_SVG("Slate/IKRigRetargeter_64", Icon64x64));
 		
+		Set("IKRetarget.AssetSettings", new IMAGE_BRUSH_SVG("Slate/AssetSettings", Icon64x64));
 		Set("IKRetarget.GlobalSettings", new IMAGE_BRUSH_SVG("Slate/GlobalSettings", Icon64x64));
 		Set("IKRetarget.RootSettings", new IMAGE_BRUSH_SVG("Slate/RootSettings", Icon64x64));
 		Set("IKRetarget.PostSettings", new IMAGE_BRUSH_SVG("Slate/PostSettings", Icon64x64));
