Added new IK Retargeter CurveRemapOp

diff --git a/Engine/Plugins/Animation/IKRig/Source/IKRig/Private/AnimNodes/AnimNode_RetargetPoseFromMesh.cpp b/Engine/Plugins/Animation/IKRig/Source/IKRig/Private/AnimNodes/AnimNode_RetargetPoseFromMesh.cpp
--- a/Engine/Plugins/Animation/IKRig/Source/IKRig/Private/AnimNodes/AnimNode_RetargetPoseFromMesh.cpp
+++ b/Engine/Plugins/Animation/IKRig/Source/IKRig/Private/AnimNodes/AnimNode_RetargetPoseFromMesh.cpp
@@ -6,6 +6,8 @@
 #include "Animation/AnimInstanceProxy.h"
 #include "Components/SkeletalMeshComponent.h"
 #include "Engine/SkeletalMesh.h"
+#include "Retargeter/IKRetargetOps.h"
+#include "Retargeter/RetargetOps/CurveRemapOp.h"
 
 #include UE_INLINE_GENERATED_CPP_BY_NAME(AnimNode_RetargetPoseFromMesh)
 
@@ -133,11 +135,8 @@ void FAnimNode_RetargetPoseFromMesh::Evaluate_AnyThread(FPoseContext& Output)
 	// convert to local space
 	FCSPose<FCompactPose>::ConvertComponentPosesToLocalPoses(ComponentPose, Output.Pose);
 
-	// copy curves over
-	if (bCopyCurves)
-	{
-		Output.Curve.CopyFrom(SourceCurves);
-	}
+	// copy and/or remap curves from the source to the target skeletal mesh
+	CopyAndRemapCurvesFromSourceToTarget(Output.Curve);
 }
 
 void FAnimNode_RetargetPoseFromMesh::PreUpdate(const UAnimInstance* InAnimInstance)
@@ -178,6 +177,50 @@ void FAnimNode_RetargetPoseFromMesh::PreUpdate(const UAnimInstance* InAnimInstan
 	}
 }
 
+void FAnimNode_RetargetPoseFromMesh::CopyAndRemapCurvesFromSourceToTarget(FBlendedCurve& OutputCurves) const
+{
+	if (!bCopyCurves)
+	{
+		return;
+	}
+	
+	if (!(Processor && Processor->IsInitialized()))
+	{
+		return;
+	}
+
+	// copy curves over to same name on target (if it exists)
+	OutputCurves.CopyFrom(SourceCurves);
+	
+	// now remap curves that have different names using the CurveRemapOp data (if there are any)
+	FBlendedCurve RemapedCurves;
+	const TArray<TObjectPtr<URetargetOpBase>>& RetargetOps = Processor->GetRetargetOps();
+	for (URetargetOpBase* RetargetOp : RetargetOps)
+	{
+		UCurveRemapOp* CurveRemapOp = Cast<UCurveRemapOp>(RetargetOp);
+		if (!CurveRemapOp)
+		{
+			continue;
+		}
+
+		if (!CurveRemapOp->bIsEnabled)
+		{
+			continue;
+		}
+		
+		for (const FCurveRemapPair& CurveToRemap : CurveRemapOp->CurvesToRemap)
+		{
+			bool bOutIsValid = false;
+			const float SourceValue = SourceCurves.Get(CurveToRemap.SourceCurve, bOutIsValid);
+			if (bOutIsValid)
+			{
+				RemapedCurves.Add(CurveToRemap.TargetCurve, SourceValue);
+			}
+		}
+	}
+	OutputCurves.Combine(RemapedCurves);
+}
+
 void FAnimNode_RetargetPoseFromMesh::UpdateSpeedValuesFromCurves()
 {
 	SpeedValuesFromCurves.Reset();
diff --git a/Engine/Plugins/Animation/IKRig/Source/IKRig/Public/AnimNodes/AnimNode_RetargetPoseFromMesh.h b/Engine/Plugins/Animation/IKRig/Source/IKRig/Public/AnimNodes/AnimNode_RetargetPoseFromMesh.h
--- a/Engine/Plugins/Animation/IKRig/Source/IKRig/Public/AnimNodes/AnimNode_RetargetPoseFromMesh.h
+++ b/Engine/Plugins/Animation/IKRig/Source/IKRig/Public/AnimNodes/AnimNode_RetargetPoseFromMesh.h
@@ -72,10 +72,14 @@ private:
 
 	// mapping from required bones to actual bones within the target skeleton
 	TArray< TPair<int32, int32> > RequiredToTargetBoneMapping;
-	// cached curves, copied on the game thread
+
+	// cached curves, copied on the game thread in PreUpdate()
 	FBlendedHeapCurve SourceCurves;
 
-	/** update map of curve values containing speeds used for IK planting */
+	// remap curves for CurveRemapOp if one is present in the RetargetOp stack
+	void CopyAndRemapCurvesFromSourceToTarget(FBlendedCurve& OutputCurves) const;
+	
+	// update map of curve values containing speeds used for IK planting
 	void UpdateSpeedValuesFromCurves();
 	// map of curve names to values, passed to retargeter for IK planting (copied from source mesh)
 	TMap<FName, float> SpeedValuesFromCurves;
diff --git a/Engine/Plugins/Animation/IKRig/Source/IKRig/Public/Retargeter/RetargetOps/CurveRemapOp.h b/Engine/Plugins/Animation/IKRig/Source/IKRig/Public/Retargeter/RetargetOps/CurveRemapOp.h
new file mode 100644
--- /dev/null
+++ b/Engine/Plugins/Animation/IKRig/Source/IKRig/Public/Retargeter/RetargetOps/CurveRemapOp.h
@@ -0,0 +1,62 @@
+﻿// Copyright Epic Games, Inc. All Rights Reserved.
+#pragma once
+
+#include "Retargeter/IKRetargetOps.h"
+#include "CurveRemapOp.generated.h"
+
+#define LOCTEXT_NAMESPACE "CurveRemapOp"
+
+USTRUCT(BlueprintType)
+struct FCurveRemapPair
+{
+	GENERATED_BODY()
+	
+	// The curve name on the SOURCE skeletal mesh to copy animation data from.
+	UPROPERTY(EditAnywhere, Category=Settings)
+	FName SourceCurve;
+
+	// The curve name on the TARGET skeletal mesh to receive animation data.
+	UPROPERTY(EditAnywhere, Category=Settings)
+	FName TargetCurve;
+};
+
+UCLASS(BlueprintType, EditInlineNew)
+class IKRIG_API UCurveRemapOp : public URetargetOpBase
+{
+	GENERATED_BODY()
+
+public:
+
+	// this op does not do anything in Initialize() or Run(), but is used by the Retarget Pose from Mesh node and the Batch Exporter
+	virtual bool Initialize(
+	const UIKRetargetProcessor* Processor,
+		const FRetargetSkeleton& SourceSkeleton,
+		const FTargetSkeleton& TargetSkeleton,
+		FIKRigLogger& Log) override { return true; };
+	virtual void Run(
+		const UIKRetargetProcessor* Processor,
+		const TArray<FTransform>& InSourceGlobalPose,
+		TArray<FTransform>& OutTargetGlobalPose) override {};
+
+	// Add pairs of Source/Target curve names to remap. While retargeting, the animation from the source curves
+	// will be redirected to the curves on the target skeletal meshes. Can be used to drive, blendshapes or other downstream systems.
+	// NOTE: By default the IK Retargeter will automatically copy all equivalently named curves from the source to the
+	// target. Remapping with this op is only necessary when the target curve name(s) are different.
+	UPROPERTY(EditAnywhere, Category="OpSettings")
+	TArray<FCurveRemapPair> CurvesToRemap;
+
+	// This setting only applies to all curves when exporting retargeted animations.
+	// True: all source curves are copied to the target animation sequence asset.
+	// False: only remapped curves are left on the target animation sequence asset.
+	UPROPERTY(EditAnywhere, Category="OpSettings")
+	bool bCopyAllSourceCurves = true;
+	
+#if WITH_EDITOR
+	virtual FText GetNiceName() const override { return FText(LOCTEXT("OpName", "Remap Curves")); };
+	virtual FText WarningMessage() const override { return Message; };
+	FText Message;
+#endif
+};
+
+
+#undef LOCTEXT_NAMESPACE
diff --git a/Engine/Plugins/Animation/IKRig/Source/IKRigEditor/Private/RetargetEditor/IKRetargetBatchOperation.cpp b/Engine/Plugins/Animation/IKRig/Source/IKRigEditor/Private/RetargetEditor/IKRetargetBatchOperation.cpp
--- a/Engine/Plugins/Animation/IKRig/Source/IKRigEditor/Private/RetargetEditor/IKRetargetBatchOperation.cpp
+++ b/Engine/Plugins/Animation/IKRig/Source/IKRigEditor/Private/RetargetEditor/IKRetargetBatchOperation.cpp
@@ -25,6 +25,8 @@
 #include "Retargeter/IKRetargetProcessor.h"
 #include "Widgets/Notifications/SNotificationList.h"
 #include "Animation/AnimMontage.h"
+#include "Retargeter/IKRetargetOps.h"
+#include "Retargeter/RetargetOps/CurveRemapOp.h"
 
 #define LOCTEXT_NAMESPACE "RetargetBatchOperation"
 
@@ -245,6 +247,9 @@ void UIKRetargetBatchOperation::RetargetAssets(
 		AnimBlueprint->PostEditChange();
 		AnimBlueprint->MarkPackageDirty();
 	}
+
+	// copy/remap curves to duplicate sequences
+	RemapCurves(Context, Progress);
 }
 
 void UIKRetargetBatchOperation::ConvertAnimation(
@@ -286,18 +291,18 @@ void UIKRetargetBatchOperation::ConvertAnimation(
 	for (TPair<UAnimationAsset*, UAnimationAsset*>& Pair : DuplicatedAnimAssets)
 	{
 		UAnimSequence* SourceSequence = Cast<UAnimSequence>(Pair.Key);
-		UAnimSequence* DestinationSequence = Cast<UAnimSequence>(Pair.Value);
-		if (!(SourceSequence && DestinationSequence))
+		UAnimSequence* TargetSequence = Cast<UAnimSequence>(Pair.Value);
+		if (!(SourceSequence && TargetSequence))
 		{
 			continue;
 		}
 
 		// increment progress bar
-		FString AssetName = DestinationSequence->GetName();
+		FString AssetName = TargetSequence->GetName();
 		Progress.EnterProgressFrame(1.f, FText::Format(LOCTEXT("RunningBatchRetarget", "Retargeting animation asset: {0}"), FText::FromString(AssetName)));
 
 		// remove all keys from the destination animation sequence
-		IAnimationDataController& TargetSeqController = DestinationSequence->GetController();
+		IAnimationDataController& TargetSeqController = TargetSequence->GetController();
 		constexpr bool bShouldTransact = false;
 		TargetSeqController.OpenBracket(FText::FromString("Generating Retargeted Animation Data"), bShouldTransact);
 		TargetSeqController.RemoveAllBoneTracks(bShouldTransact);
@@ -372,7 +377,8 @@ void UIKRetargetBatchOperation::ConvertAnimation(
 				BoneTrack.RotKeys[FrameIndex] = FQuat4f(LocalPose.GetRotation());
 				BoneTrack.ScaleKeys[FrameIndex] = FVector3f(LocalPose.GetScale3D());
 			}
-		}
+			
+		} // END for each frame
 
 		// add keys to bone tracks
 		for (int32 TargetBoneIndex=0; TargetBoneIndex<NumTargetBones; ++TargetBoneIndex)
@@ -388,6 +394,114 @@ void UIKRetargetBatchOperation::ConvertAnimation(
 	}
 }
 
+void UIKRetargetBatchOperation::RemapCurves(const FIKRetargetBatchOperationContext& Context, FScopedSlowTask& Progress)
+{
+	USkeleton* SourceSkeleton = Context.SourceMesh->GetSkeleton();
+	USkeleton* TargetSkeleton = Context.TargetMesh->GetSkeleton();
+	
+	// get map of curves to remap (Source:Target)
+	bool bCopyAllSourceCurves = true;
+	TMap<FAnimationCurveIdentifier, FAnimationCurveIdentifier> CurvesToRemap;
+	URetargetOpStack* OpStack = Context.IKRetargetAsset->GetPostSettingsUObject();
+	for (const TObjectPtr<URetargetOpBase>& RetargetOp : OpStack->RetargetOps)
+	{
+		UCurveRemapOp* CurveRemapOp = Cast<UCurveRemapOp>(RetargetOp);
+		if (!CurveRemapOp)
+		{
+			continue;
+		}
+
+		if (!CurveRemapOp->bIsEnabled)
+		{
+			continue;
+		}
+
+		for (const FCurveRemapPair& CurveToRemap : CurveRemapOp->CurvesToRemap)
+		{
+			const FAnimationCurveIdentifier SourceCurveId = UAnimationCurveIdentifierExtensions::FindCurveIdentifier(SourceSkeleton, CurveToRemap.SourceCurve, ERawCurveTrackTypes::RCT_Float);
+			const FAnimationCurveIdentifier TargetCurveId = UAnimationCurveIdentifierExtensions::FindCurveIdentifier(TargetSkeleton, CurveToRemap.TargetCurve, ERawCurveTrackTypes::RCT_Float);
+			CurvesToRemap.Add(SourceCurveId, TargetCurveId);
+		}
+
+		bCopyAllSourceCurves &= CurveRemapOp->bCopyAllSourceCurves;
+	}
+
+	// for each exported animation, remap curves from source to target anim
+	for (TPair<UAnimationAsset*, UAnimationAsset*>& Pair : DuplicatedAnimAssets)
+	{
+		UAnimSequence* SourceSequence = Cast<UAnimSequence>(Pair.Key);
+		UAnimSequence* TargetSequence = Cast<UAnimSequence>(Pair.Value);
+		if (!(SourceSequence && TargetSequence))
+		{
+			continue;
+		}
+
+		// increment progress bar
+		FString AssetName = TargetSequence->GetName();
+		Progress.EnterProgressFrame(1.f, FText::Format(LOCTEXT("RunningBatchRetarget", "Remapping Curves on Asset: {0}"), FText::FromString(AssetName)));
+
+		// all curves were copied when we duplicated the animation sequence, so now we have to rename curves
+		// based on the remapping defined in the curve remap op(s)
+		IAnimationDataController& TargetSeqController = TargetSequence->GetController();
+		constexpr bool bShouldTransact = false;
+		TargetSeqController.OpenBracket(FText::FromString("Remapping Curve Data"), bShouldTransact);
+
+		const IAnimationDataModel* SourceDataModel = SourceSequence->GetDataModel();
+		
+		for (const TTuple<FAnimationCurveIdentifier, FAnimationCurveIdentifier>& CurveToRemap : CurvesToRemap)
+		{
+			// get the source curve to copy from
+			const FFloatCurve* SourceCurve =  SourceDataModel->FindFloatCurve(CurveToRemap.Key);
+			if (!SourceCurve)
+			{
+				continue; // missing source curve to remap
+			}
+
+			// add a curve to the target to house the keys
+			FAnimationCurveIdentifier TargetCurveID = CurveToRemap.Value;
+			if (!TargetCurveID.IsValid())
+			{
+				continue; // must provide a valid name for the target curve
+			}
+			TargetSeqController.AddCurve(TargetCurveID, SourceCurve->GetCurveTypeFlags(), bShouldTransact);
+			
+			// copy data into target curve
+			TargetSeqController.SetCurveKeys(TargetCurveID, SourceCurve->FloatCurve.GetConstRefOfKeys(), bShouldTransact);
+			TargetSeqController.SetCurveColor(TargetCurveID, SourceCurve->GetColor(), bShouldTransact);
+		}
+
+		// optionally remove all source curves from the target asset
+		// (remove all curves that were copied when the source sequence was duplicated UNLESS they are remapped)
+		if (!bCopyAllSourceCurves)
+		{
+			// get list of target curves to keep
+			TArray<FAnimationCurveIdentifier> TargetCurvesToKeep;
+			CurvesToRemap.GenerateValueArray(TargetCurvesToKeep);
+
+			// get list of target curves to remove
+			const TArray<FFloatCurve>& AllTargetCurves = TargetSeqController.GetModel()->GetFloatCurves();
+			TArray<FAnimationCurveIdentifier> CurvesToRemove;
+			for (const FFloatCurve& TargetCurve : AllTargetCurves)
+			{
+				const FAnimationCurveIdentifier TargetCurveId = UAnimationCurveIdentifierExtensions::FindCurveIdentifier(TargetSkeleton, TargetCurve.GetName(), ERawCurveTrackTypes::RCT_Float);
+				if (TargetCurvesToKeep.Contains(TargetCurveId))
+				{
+					continue;
+				}
+				CurvesToRemove.Add(TargetCurveId);
+			}
+
+			// remove the curves
+			for (const FAnimationCurveIdentifier& CurveToRemove : CurvesToRemove)
+			{
+				TargetSeqController.RemoveCurve(CurveToRemove, bShouldTransact);
+			}
+		}
+
+		TargetSeqController.CloseBracket(bShouldTransact);
+	}
+}
+
 void UIKRetargetBatchOperation::NotifyUserOfResults(
 	const FIKRetargetBatchOperationContext& Context,
 	FScopedSlowTask& Progress) const
@@ -545,7 +659,8 @@ void UIKRetargetBatchOperation::RunRetarget(FIKRetargetBatchOperationContext& Co
 	}
 	
 	// show progress bar
-	FScopedSlowTask Progress(NumAssets + 2, LOCTEXT("GatheringBatchRetarget", "Gathering animation assets..."));
+	constexpr int NumAdditionalProgressFrames = 3;
+	FScopedSlowTask Progress(NumAssets + NumAdditionalProgressFrames, LOCTEXT("GatheringBatchRetarget", "Gathering animation assets..."));
 	Progress.MakeDialog();
 	
 	DuplicateRetargetAssets(Context, Progress);
diff --git a/Engine/Plugins/Animation/IKRig/Source/IKRigEditor/Public/RetargetEditor/IKRetargetBatchOperation.h b/Engine/Plugins/Animation/IKRig/Source/IKRigEditor/Public/RetargetEditor/IKRetargetBatchOperation.h
--- a/Engine/Plugins/Animation/IKRig/Source/IKRigEditor/Public/RetargetEditor/IKRetargetBatchOperation.h
+++ b/Engine/Plugins/Animation/IKRig/Source/IKRigEditor/Public/RetargetEditor/IKRetargetBatchOperation.h
@@ -115,6 +115,9 @@ private:
 	/* Convert animation on all the duplicates */
 	void ConvertAnimation(const FIKRetargetBatchOperationContext& Context, FScopedSlowTask& Progress);
 
+	/* Copy/remap curves on all the duplicates */
+	void RemapCurves(const FIKRetargetBatchOperationContext& Context, FScopedSlowTask& Progress);
+
 	/* Output notifications of results */
 	void NotifyUserOfResults(const FIKRetargetBatchOperationContext& Context, FScopedSlowTask& Progress) const;
 
