Added UProperty value clamps on retargeter settings.

diff --git a/Engine/Plugins/Animation/IKRig/Source/IKRig/Public/Retargeter/IKRetargetSettings.h b/Engine/Plugins/Animation/IKRig/Source/IKRig/Public/Retargeter/IKRetargetSettings.h
--- a/Engine/Plugins/Animation/IKRig/Source/IKRig/Public/Retargeter/IKRetargetSettings.h
+++ b/Engine/Plugins/Animation/IKRig/Source/IKRig/Public/Retargeter/IKRetargetSettings.h
@@ -31,7 +31,7 @@ struct IKRIG_API FTargetChainSpeedPlantSettings
 	GENERATED_BODY()
 
 	/** The name of the curve on the source animation that contains the speed of the end effector bone.*/
-	UPROPERTY(EditAnywhere, BlueprintReadWrite, Category = "Plant IK by Speed", meta = (ClampMin = "0.0", ClampMax = "100.0", UIMin = "0.0", UIMax = "100.0"))
+	UPROPERTY(EditAnywhere, BlueprintReadWrite, Category = "Plant IK by Speed")
 	bool EnableSpeedPlanting = false;
 	
 	/** The name of the curve on the source animation that contains the speed of the end effector bone.*/
@@ -40,7 +40,7 @@ struct IKRIG_API FTargetChainSpeedPlantSettings
 
 	/** Range 0 to 1000. Default 15. The maximum speed a source bone can be moving while being considered 'planted'.
 	*  The target IK goal will not be allowed to move whenever the source bone speed drops below this threshold speed. */
-	UPROPERTY(EditAnywhere, BlueprintReadWrite, Category = "Plant IK by Speed", meta = (ClampMin = "0.0", UIMin = "0.0", UIMax = "100.0"))
+	UPROPERTY(EditAnywhere, BlueprintReadWrite, Category = "Plant IK by Speed", meta = (ClampMin = "0.0", ClampMax = "1000.0", UIMin = "0.0", UIMax = "100.0"))
 	float SpeedThreshold = 15.0f;
 
 	// How stiff the spring model is that smoothly pulls the IK position after unplanting (more stiffness means more oscillation around the target value)
@@ -48,7 +48,7 @@ struct IKRIG_API FTargetChainSpeedPlantSettings
 	float UnplantStiffness = 250.0f;
 
 	// How much damping to apply to the spring (0 means no damping, 1 means critically damped which means no oscillation)
-	UPROPERTY(EditAnywhere, BlueprintReadWrite, Category = "Plant IK by Speed", meta = (ClampMin = "0.0", ClampMax = "10.0", UIMin = "0.0", UIMax = "1.0"))
+	UPROPERTY(EditAnywhere, BlueprintReadWrite, Category = "Plant IK by Speed", meta = (ClampMin = "0.0", ClampMax = "1.0", UIMin = "0.0", UIMax = "1.0"))
 	float UnplantCriticalDamping = 1.0f;
 
 	bool operator==(const FTargetChainSpeedPlantSettings& Other) const;
@@ -94,11 +94,11 @@ struct IKRIG_API FTargetChainFKSettings
 	/** Range 0 to 1. Default 1. Matches the twist angle of this chain (along the Pole direction) to the source chain.
 	*  At 0, the chain's pole vector direction will be left alon
 	*  At 1, the root bone of the chain will be twist-rotated in the pole direction to match the orientation of the source chain.*/
-	UPROPERTY(EditAnywhere, BlueprintReadWrite, Category = "FK Adjustments", meta = (UIMin = "0.0", UIMax = "1.0"))
+	UPROPERTY(EditAnywhere, BlueprintReadWrite, Category = "FK Adjustments", meta = (UIMin = "0.0", UIMax = "1.0", ClampMin = "0.0", ClampMax = "1.0"))
 	float PoleVectorMatching = 0.0f;
 
 	/** Range +/- 180. Default 0. An angular offset, in degrees, for the pole direction of the chain. */
-	UPROPERTY(EditAnywhere, BlueprintReadWrite, Category = "FK Adjustments", meta = (UIMin = "-180.0", UIMax = "180.0"))
+	UPROPERTY(EditAnywhere, BlueprintReadWrite, Category = "FK Adjustments", meta = (UIMin = "-180.0", UIMax = "180.0", ClampMin = "-180.0", ClampMax = "180.0"))
 	float PoleVectorOffset = 0.0f;
 
 	bool operator==(const FTargetChainFKSettings& Other) const;
@@ -195,13 +195,13 @@ public:
 	/** Range 0 to 1. Default 1. Blends the amount of retargeted root rotation to apply.
 	*  At 0 the root is left at the rotation from the retarget pose.
 	*  At 1 the root is rotated fully to match the source root rotation. */
-	UPROPERTY(EditAnywhere, BlueprintReadWrite, Category = "Root Retarget Settings", meta = (ClampMin = "-5.0", ClampMax = "5.0", UIMin = "0.0", UIMax = "1.0"))
+	UPROPERTY(EditAnywhere, BlueprintReadWrite, Category = "Root Retarget Settings", meta = (ClampMin = "0.0", ClampMax = "1.0", UIMin = "0.0", UIMax = "1.0"))
 	float RotationAlpha = 1.0f;
 	
 	/** Range 0 to 1. Default 1. Blends the amount of retargeted root translation to apply.
 	*  At 0 the root is left at the position from the retarget pose.
 	*  At 1 the root will follow the source motion according to the behavior defined in the subsequent settings. */
-	UPROPERTY(EditAnywhere, BlueprintReadWrite, Category = "Root Retarget Settings", meta = (ClampMin = "-5.0", ClampMax = "5.0", UIMin = "0.0", UIMax = "1.0"))
+	UPROPERTY(EditAnywhere, BlueprintReadWrite, Category = "Root Retarget Settings", meta = (ClampMin = "0.0", ClampMax = "1.0", UIMin = "0.0", UIMax = "1.0"))
 	float TranslationAlpha = 1.0f;
 	
 	/** Range 0 to 1. Default 0. Blends the retarget root's translation to the exact source location.
@@ -324,7 +324,7 @@ struct IKRIG_API FRetargetGlobalSettings
 
 	/** Range 0 to Inf. Default 1. Warps IK goal positions in the forward direction. Useful for stride warping.
 	 * Values below 1 will create smaller, squashed strides. Values greater than 1 will create stretched, longer strides.*/
-	UPROPERTY(EditAnywhere, BlueprintReadWrite, Category = Warping, meta = (UIMin = "0.0", UIMax = "5.0"))
+	UPROPERTY(EditAnywhere, BlueprintReadWrite, Category = Warping, meta = (UIMin = "0.0", UIMax = "5.0", ClampMin = "0.0"))
 	float WarpForwards = 1.0f;
 
 	/** Range -+Inf. Default is 0. A static offset in world units to move the IK goals perpendicular to the forward direction.
@@ -335,7 +335,7 @@ struct IKRIG_API FRetargetGlobalSettings
 	/** Range 0 to +Inf. Default is 1.0f.
 	 * Values below 1 pull all the goals towards the average of all the goals (towards each other).
 	 * Values greater than 1 push the goals apart.*/
-	UPROPERTY(EditAnywhere, BlueprintReadWrite, Category = Warping, meta = (UIMin = "0.0", UIMax = "2.0"))
+	UPROPERTY(EditAnywhere, BlueprintReadWrite, Category = Warping, meta = (UIMin = "0.0", UIMax = "2.0", ClampMin = "0.0"))
 	float WarpSplay = 1.0f;
 
 	static FVector GetAxisVector(const EBasicAxis& Axis)
