Resets all additive animation settings on exported retargeted animations.

diff --git a/Engine/Plugins/Animation/IKRig/Source/IKRigEditor/Private/RetargetEditor/IKRetargetBatchOperation.cpp b/Engine/Plugins/Animation/IKRig/Source/IKRigEditor/Private/RetargetEditor/IKRetargetBatchOperation.cpp
--- a/Engine/Plugins/Animation/IKRig/Source/IKRigEditor/Private/RetargetEditor/IKRetargetBatchOperation.cpp
+++ b/Engine/Plugins/Animation/IKRig/Source/IKRigEditor/Private/RetargetEditor/IKRetargetBatchOperation.cpp
@@ -178,17 +178,27 @@ void UIKRetargetBatchOperation::RetargetAssets(
 		{
 			// copy curve data from source asset, preserving data in the target if present.
 			UAnimationBlueprintLibrary::CopyAnimationCurveNamesToSkeleton(OldSkeleton, NewSkeleton, AnimSequenceToRetarget, ERawCurveTrackTypes::RCT_Float);	
+
 			// clear transform curves since those curves won't work in new skeleton
 			IAnimationDataController& Controller = AnimSequenceToRetarget->GetController();
 			constexpr bool bShouldTransact = false;
 			Controller.OpenBracket(FText::FromString("Preparing for retargeted animation."), bShouldTransact);
 			Controller.RemoveAllCurvesOfType(ERawCurveTrackTypes::RCT_Transform, bShouldTransact);
+
 			// clear bone tracks to prevent recompression
 			Controller.RemoveAllBoneTracks(bShouldTransact);
+
+			// reset all additive animation properties to ensure WYSIWYG playback of additive anims between retargeter and sequence
+			AnimSequenceToRetarget->AdditiveAnimType = EAdditiveAnimationType::AAT_None;
+			AnimSequenceToRetarget->RefPoseType = EAdditiveBasePoseType::ABPT_None;
+			AnimSequenceToRetarget->RefFrameIndex = 0;
+			AnimSequenceToRetarget->RefPoseSeq = nullptr;
+			
 			// set the retarget source to the target skeletal mesh
 			AnimSequenceToRetarget->RetargetSource = NAME_None;
 			AnimSequenceToRetarget->RetargetSourceAsset = Context.TargetMesh;
 			Controller.UpdateWithSkeleton(NewSkeleton, bShouldTransact);
+
 			// done editing sequence data, close bracket
 			Controller.CloseBracket(bShouldTransact);
 		}
