Add ability to ignore "Force Root Lock" when retargeting.

diff --git a/Engine/Plugins/Animation/IKRig/Source/IKRig/Public/Retargeter/IKRetargeter.h b/Engine/Plugins/Animation/IKRig/Source/IKRig/Public/Retargeter/IKRetargeter.h
--- a/Engine/Plugins/Animation/IKRig/Source/IKRig/Public/Retargeter/IKRetargeter.h
+++ b/Engine/Plugins/Animation/IKRig/Source/IKRig/Public/Retargeter/IKRetargeter.h
@@ -404,6 +404,12 @@ public:
 	UPROPERTY(EditAnywhere, Category = PreviewSettings)
 	FVector SourceMeshOffset;
 
+	// When true, animation sequences with "Force Root Lock" turned On will act as though it is Off.
+	// This affects the preview in the retarget editor, as well as exported animation sequences.
+	// This setting has no effect on runtime retargeting where root motion is copied from the source component.
+	UPROPERTY(EditAnywhere, Category = ExportSettings)
+	bool bIgnoreRootLock = true;
+
 	// Toggle debug drawing for retargeting in the viewport. 
 	UPROPERTY(EditAnywhere, Category = DebugSettings)
 	bool bDebugDraw = true;
diff --git a/Engine/Plugins/Animation/IKRig/Source/IKRigEditor/Private/RetargetEditor/IKRetargetAnimInstanceProxy.cpp b/Engine/Plugins/Animation/IKRig/Source/IKRigEditor/Private/RetargetEditor/IKRetargetAnimInstanceProxy.cpp
--- a/Engine/Plugins/Animation/IKRig/Source/IKRigEditor/Private/RetargetEditor/IKRetargetAnimInstanceProxy.cpp
+++ b/Engine/Plugins/Animation/IKRig/Source/IKRigEditor/Private/RetargetEditor/IKRetargetAnimInstanceProxy.cpp
@@ -42,6 +42,11 @@ void FIKRetargetAnimInstanceProxy::CacheBones()
 
 bool FIKRetargetAnimInstanceProxy::Evaluate(FPoseContext& Output)
 {
+	if (PreviewPoseNode->IKRetargeterAsset)
+	{
+		bIgnoreRootLock = PreviewPoseNode->IKRetargeterAsset->bIgnoreRootLock;	
+	}
+	
 	switch (OutputMode)
 	{
 	case ERetargeterOutputMode::RunRetarget:
diff --git a/Engine/Plugins/Animation/IKRig/Source/IKRigEditor/Private/RetargetEditor/IKRetargetBatchOperation.cpp b/Engine/Plugins/Animation/IKRig/Source/IKRigEditor/Private/RetargetEditor/IKRetargetBatchOperation.cpp
--- a/Engine/Plugins/Animation/IKRig/Source/IKRigEditor/Private/RetargetEditor/IKRetargetBatchOperation.cpp
+++ b/Engine/Plugins/Animation/IKRig/Source/IKRigEditor/Private/RetargetEditor/IKRetargetBatchOperation.cpp
@@ -332,6 +332,10 @@ void UIKRetargetBatchOperation::ConvertAnimation(
 		FAnimPoseEvaluationOptions EvaluationOptions = FAnimPoseEvaluationOptions();
 		EvaluationOptions.OptionalSkeletalMesh = SourceSkeleton.SkeletalMesh;
 
+		// optionally ignore root lock
+		const bool bIgnoreRootLock = Context.IKRetargetAsset->bIgnoreRootLock;
+		TGuardValue<bool> RootLockGuard(SourceSequence->bForceRootLock, bIgnoreRootLock ? false : SourceSequence->bForceRootLock);
+
 		// reset the planting state
 		Processor->ResetPlanting();
 		
@@ -448,7 +452,7 @@ void UIKRetargetBatchOperation::RemapCurves(const FIKRetargetBatchOperationConte
 
 		// increment progress bar
 		FString AssetName = TargetSequence->GetName();
-		Progress.EnterProgressFrame(1.f, FText::Format(LOCTEXT("RunningBatchRetarget", "Remapping Curves on Asset: {0}"), FText::FromString(AssetName)));
+		Progress.EnterProgressFrame(1.f, FText::Format(LOCTEXT("RemappingCurves", "Remapping Curves on Asset: {0}"), FText::FromString(AssetName)));
 
 		// all curves were copied when we duplicated the animation sequence, so now we have to rename curves
 		// based on the remapping defined in the curve remap op(s)
diff --git a/Engine/Source/Runtime/Engine/Private/Animation/AnimSingleNodeInstanceProxy.cpp b/Engine/Source/Runtime/Engine/Private/Animation/AnimSingleNodeInstanceProxy.cpp
--- a/Engine/Source/Runtime/Engine/Private/Animation/AnimSingleNodeInstanceProxy.cpp
+++ b/Engine/Source/Runtime/Engine/Private/Animation/AnimSingleNodeInstanceProxy.cpp
@@ -316,6 +316,9 @@ void FAnimNode_SingleNode::Evaluate_AnyThread(FPoseContext& Output)
 		else if (UAnimSequence* Sequence = Cast<UAnimSequence>(Proxy->CurrentAsset))
 		{
 			FAnimExtractContext ExtractionContext(static_cast<double>(Proxy->CurrentTime), Sequence->bEnableRootMotion, Proxy->DeltaTimeRecord, Proxy->bLooping);
+#if WITH_EDITOR
+			ExtractionContext.bIgnoreRootLock = Proxy->bIgnoreRootLock;
+#endif
 
 			if (Sequence->IsValidAdditive())
 			{
diff --git a/Engine/Source/Runtime/Engine/Public/Animation/AnimSingleNodeInstanceProxy.h b/Engine/Source/Runtime/Engine/Public/Animation/AnimSingleNodeInstanceProxy.h
--- a/Engine/Source/Runtime/Engine/Public/Animation/AnimSingleNodeInstanceProxy.h
+++ b/Engine/Source/Runtime/Engine/Public/Animation/AnimSingleNodeInstanceProxy.h
@@ -53,6 +53,7 @@ public:
 	{
 #if WITH_EDITOR
 		bCanProcessAdditiveAnimations = false;
+		bIgnoreRootLock = false;
 #endif
 	}
 
@@ -75,6 +76,7 @@ public:
 
 #if WITH_EDITOR
 		bCanProcessAdditiveAnimations = false;
+		bIgnoreRootLock = false;
 #endif
 	}
 
@@ -200,6 +202,9 @@ protected:
 	/** If this is being used by a user (for instance on a skeletal mesh placed in a level) we don't want to allow
 	additives. But we need to be able to override this for editor preview windows */
 	bool bCanProcessAdditiveAnimations;
+
+	/** Allows editor previews to ignore root lock and show root motion in the editor.*/
+	bool bIgnoreRootLock;
 #endif
 
 	/** Pose Weight value that can override curve data. In the future, we'd like to have UCurveSet that can play by default**/
