Fixes bug "Static Offset" in Retargeter chain IK settings not deserializing properly.

diff --git a/Engine/Plugins/Animation/IKRig/Config/BaseIKRig.ini b/Engine/Plugins/Animation/IKRig/Config/BaseIKRig.ini
--- a/Engine/Plugins/Animation/IKRig/Config/BaseIKRig.ini
+++ b/Engine/Plugins/Animation/IKRig/Config/BaseIKRig.ini
@@ -2,5 +2,4 @@
 +ClassRedirects=(OldName="/Script/IKRig.IKRigPBIKSolver", NewName="/Script/IKRig.IKRigFBIKSolver")
 +ClassRedirects=(OldName="/Script/IKRig.IKRig_PBIKBoneSettings", NewName="/Script/IKRig.IKRig_FBIKBoneSettings")
 +ClassRedirects=(OldName="/Script/IKRig.AnimNode_IKRetargeter", NewName="/Script/IKRig.AnimNode_RetargetPoseFromMesh")
-+PropertyRedirects=(OldName="/Script/IKRig.TargetChainIKSettings.StaticOffset",NewName="TranslationGlobalOffset")
 
