Adding ability to dynamically adjust retarget pose and see results at runtime.

diff --git a/Engine/Plugins/Animation/IKRig/Source/IKRig/Private/Retargeter/IKRetargetProcessor.cpp b/Engine/Plugins/Animation/IKRig/Source/IKRig/Private/Retargeter/IKRetargetProcessor.cpp
--- a/Engine/Plugins/Animation/IKRig/Source/IKRig/Private/Retargeter/IKRetargetProcessor.cpp
+++ b/Engine/Plugins/Animation/IKRig/Source/IKRig/Private/Retargeter/IKRetargetProcessor.cpp
@@ -82,6 +82,7 @@ void FRetargetSkeleton::GenerateRetargetPose(
 {
 	// record the name of the retarget pose (prevents re-initialization if profile swaps it)
 	RetargetPoseName = InRetargetPoseName;
+	RetargetPoseVersion = InRetargetPose->GetVersion();
 	
 	// initialize retarget pose to the skeletal mesh reference pose
 	RetargetLocalPose = SkeletalMesh->GetRefSkeleton().GetRefBonePose();
@@ -2272,7 +2273,14 @@ void UIKRetargetProcessor::ApplySettingsFromAsset()
 	// copy global settings
 	GlobalSettings = RetargeterAsset->GetGlobalSettings();
 
+	// apply current retarget poses (only applied if the pose has been switched to a different one OR if the current pose was modified)
+	const FName SourcePose = RetargeterAsset->GetCurrentRetargetPoseName(ERetargetSourceOrTarget::Source);
+	const FName TargetPose = RetargeterAsset->GetCurrentRetargetPoseName(ERetargetSourceOrTarget::Target);
+	ApplyNewRetargetPose(SourcePose, ERetargetSourceOrTarget::Source);
+	ApplyNewRetargetPose(TargetPose, ERetargetSourceOrTarget::Target);
+
 	// apply the current profile
+	// (this is always applied last so that profile overrides take precedence over asset settings)
 	if (const FRetargetProfile* CurrentProfile = RetargeterAsset->GetCurrentProfile())
 	{
 		ApplySettingsFromProfile(*CurrentProfile);
@@ -2331,23 +2339,25 @@ void UIKRetargetProcessor::ApplyNewRetargetPose(
 	const FName NewRetargetPoseName,
 	ERetargetSourceOrTarget SourceOrTarget)
 {
-	const bool bIsSource = SourceOrTarget == ERetargetSourceOrTarget::Source;
-	FRetargetSkeleton& RetargetSkeleton = bIsSource ? SourceSkeleton : TargetSkeleton;
-	if (NewRetargetPoseName == NAME_None || RetargetSkeleton.RetargetPoseName == NewRetargetPoseName)
-	{
-		return; // retarget pose not specified, or already in use
-	}
-
 	const FIKRetargetPose* NewRetargetPose = RetargeterAsset->GetRetargetPoseByName(SourceOrTarget, NewRetargetPoseName);
 	if (!NewRetargetPose)
 	{
-		return;
+		return; // retarget pose not found
 	}
 
+	const bool bIsSource = SourceOrTarget == ERetargetSourceOrTarget::Source;
+	FRetargetSkeleton& RetargetSkeleton = bIsSource ? SourceSkeleton : TargetSkeleton;
+	const bool bSameRetargetPose = RetargetSkeleton.RetargetPoseName == NewRetargetPoseName;
+	const bool bSameVersion = RetargetSkeleton.RetargetPoseVersion == NewRetargetPose->GetVersion();
+	if (bSameRetargetPose && bSameVersion)
+	{
+		return; // retarget pose has not changed since it was initialized
+	}
+	
 	const FName RootBoneName = bIsSource ? RootRetargeter.Source.BoneName : RootRetargeter.Target.BoneName;
 	if (RootBoneName == NAME_None)
 	{
-		return;
+		return; // cannot regenerate pose without a root bone
 	}
 
 	// re-generate the retarget pose
diff --git a/Engine/Plugins/Animation/IKRig/Source/IKRig/Private/Retargeter/IKRetargeter.cpp b/Engine/Plugins/Animation/IKRig/Source/IKRig/Private/Retargeter/IKRetargeter.cpp
--- a/Engine/Plugins/Animation/IKRig/Source/IKRig/Private/Retargeter/IKRetargeter.cpp
+++ b/Engine/Plugins/Animation/IKRig/Source/IKRig/Private/Retargeter/IKRetargeter.cpp
@@ -44,24 +44,24 @@ UIKRetargeter::UIKRetargeter(const FObjectInitializer& ObjectInitializer)
 
 const UIKRigDefinition* UIKRetargeter::GetIKRig(ERetargetSourceOrTarget SourceOrTarget) const
 {
-	if (!IsInGameThread())
+	const TSoftObjectPtr<UIKRigDefinition> SoftIKRig = SourceOrTarget == ERetargetSourceOrTarget::Source ? SourceIKRigAsset : TargetIKRigAsset;
+	if (SoftIKRig.IsValid())
 	{
-		return nullptr;
+		return SoftIKRig.Get();
 	}
-
-	const TSoftObjectPtr<UIKRigDefinition> SoftIKRig = SourceOrTarget == ERetargetSourceOrTarget::Source ? SourceIKRigAsset : TargetIKRigAsset;
-	return SoftIKRig.LoadSynchronous();
+	
+	return IsInGameThread() ? SoftIKRig.LoadSynchronous() : nullptr;
 }
 
 UIKRigDefinition* UIKRetargeter::GetIKRigWriteable(ERetargetSourceOrTarget SourceOrTarget) const
 {
-	if (!IsInGameThread())
+	const TSoftObjectPtr<UIKRigDefinition> SoftIKRig = SourceOrTarget == ERetargetSourceOrTarget::Source ? SourceIKRigAsset : TargetIKRigAsset;
+	if (SoftIKRig.IsValid())
 	{
-		return nullptr;
+		return SoftIKRig.Get();
 	}
-
-	const TSoftObjectPtr<UIKRigDefinition> SoftIKRig = SourceOrTarget == ERetargetSourceOrTarget::Source ? SourceIKRigAsset : TargetIKRigAsset;
-	return SoftIKRig.LoadSynchronous();
+	
+	return IsInGameThread() ? SoftIKRig.LoadSynchronous() : nullptr;
 }
 
 #if WITH_EDITORONLY_DATA
@@ -304,6 +304,8 @@ FQuat FIKRetargetPose::GetDeltaRotationForBone(const FName BoneName) const
 
 void FIKRetargetPose::SetDeltaRotationForBone(FName BoneName, const FQuat& RotationDelta)
 {
+	IncrementVersion();
+
 	FQuat* RotOffset = BoneRotationOffsets.Find(BoneName);
 	if (RotOffset == nullptr)
 	{
@@ -322,6 +324,8 @@ FVector FIKRetargetPose::GetRootTranslationDelta() const
 
 void FIKRetargetPose::SetRootTranslationDelta(const FVector& TranslationDelta)
 {
+	IncrementVersion();
+	
 	RootTranslationOffset = TranslationDelta;
 	// only allow vertical offset of root in retarget pose
 	RootTranslationOffset.X = 0.0f;
@@ -330,6 +334,8 @@ void FIKRetargetPose::SetRootTranslationDelta(const FVector& TranslationDelta)
 
 void FIKRetargetPose::AddToRootTranslationDelta(const FVector& TranslateDelta)
 {
+	IncrementVersion();
+	
 	RootTranslationOffset += TranslateDelta;
 	// only allow vertical offset of root in retarget pose
 	RootTranslationOffset.X = 0.0f;
diff --git a/Engine/Plugins/Animation/IKRig/Source/IKRig/Public/Retargeter/IKRetargetProcessor.h b/Engine/Plugins/Animation/IKRig/Source/IKRig/Public/Retargeter/IKRetargetProcessor.h
--- a/Engine/Plugins/Animation/IKRig/Source/IKRig/Public/Retargeter/IKRetargetProcessor.h
+++ b/Engine/Plugins/Animation/IKRig/Source/IKRig/Public/Retargeter/IKRetargetProcessor.h
@@ -31,6 +31,7 @@ struct IKRIG_API FRetargetSkeleton
 	TArray<FTransform> RetargetLocalPose;	// local space retarget pose
 	TArray<FTransform> RetargetGlobalPose;	// global space retarget pose
 	FName RetargetPoseName;					// the name of the retarget pose this was initialized with
+	int32 RetargetPoseVersion;				// the version of the retarget pose this was initialized with (transient)
 	USkeletalMesh* SkeletalMesh;			// the skeletal mesh this was initialized with
 	TArray<FName> ChainThatContainsBone;	// record which chain is actually controlling each bone
 
@@ -561,7 +562,7 @@ private:
 	bool bRootsInitialized = false;
 	/** true when at least one pair of bone chains is able to be retargeted */
 	bool bAtLeastOneValidBoneChainPair = false;
-	/** true when roots are able to be retargeted */
+	/** true when IK Rig has been initialized and is ready to run*/
 	bool bIKRigInitialized = false;
 
 	/** The source asset this processor was initialized with. */
@@ -627,6 +628,7 @@ private:
 	/** Run all post process operations on the retargeted result. */
 	void RunRetargetOps(const TArray<FTransform>& InSourceGlobalPose, TArray<FTransform>& OutTargetGlobalPose);
 
-	/** Does a partial reinitialization (at runtime) whenever the retarget pose is swapped to a different one. */
+	/** Does a partial reinitialization (at runtime) whenever the retarget pose is swapped to a different or if the
+	 * pose has been modified. Does nothing if the pose has not changed. */
 	void ApplyNewRetargetPose(const FName NewRetargetPoseName, ERetargetSourceOrTarget SourceOrTarget);
 };
diff --git a/Engine/Plugins/Animation/IKRig/Source/IKRig/Public/Retargeter/IKRetargeter.h b/Engine/Plugins/Animation/IKRig/Source/IKRig/Public/Retargeter/IKRetargeter.h
--- a/Engine/Plugins/Animation/IKRig/Source/IKRig/Public/Retargeter/IKRetargeter.h
+++ b/Engine/Plugins/Animation/IKRig/Source/IKRig/Public/Retargeter/IKRetargeter.h
@@ -178,6 +178,9 @@ public:
 	void AddToRootTranslationDelta(const FVector& TranslationDelta);
 	
 	void SortHierarchically(const FIKRigSkeleton& Skeleton);
+	
+	int32 GetVersion() const { return Version; };
+	void IncrementVersion() { ++Version; };
 
 private:
 	// a translational delta in GLOBAL space, applied only to the retarget root bone
@@ -187,6 +190,10 @@ private:
 	// these are LOCAL-space rotation deltas to be applied to a bone to modify it's retarget pose
 	UPROPERTY(EditAnywhere, Category = RetargetPose)
 	TMap<FName, FQuat> BoneRotationOffsets;
+	
+	// incremented by any edits to the retarget pose, indicating to any running instance that it should reinitialize
+	// this is not made "editor only" to leave open the possibility of programmatically modifying a retarget pose in cooked builds
+	int32 Version = INDEX_NONE;
 
 	friend class UIKRetargeterController;
 };
