Fixed bug, mirror chains in IK Rig not mirroring goals.

diff --git a/Engine/Plugins/Animation/IKRig/Source/IKRigEditor/Private/RigEditor/SIKRigRetargetChainList.cpp b/Engine/Plugins/Animation/IKRig/Source/IKRigEditor/Private/RigEditor/SIKRigRetargetChainList.cpp
--- a/Engine/Plugins/Animation/IKRig/Source/IKRigEditor/Private/RigEditor/SIKRigRetargetChainList.cpp
+++ b/Engine/Plugins/Animation/IKRig/Source/IKRigEditor/Private/RigEditor/SIKRigRetargetChainList.cpp
@@ -752,11 +752,7 @@ void SIKRigRetargetChainList::MirrorSelectedChains() const
 		MirroredChain.ChainName = UMirrorDataTable::GetSettingsMirrorName(MirroredChain.ChainName);
 		MirroredChain.StartBone = IKRigSkeleton.BoneNames[MirroredIndices[0]];
 		MirroredChain.EndBone = IKRigSkeleton.BoneNames[MirroredIndices.Last()];
-		const FName GoalOnMirroredBone = AssetController->GetGoalNameForBone(MirroredChain.EndBone.BoneName);
-		if (GoalOnMirroredBone != NAME_None)
-		{
-			MirroredChain.IKGoalName = GoalOnMirroredBone;
-		}
+		MirroredChain.IKGoalName = AssetController->GetGoalNameForBone(MirroredChain.EndBone.BoneName);
 
 		const FName NewChainName = Controller.PromptToAddNewRetargetChain(MirroredChain);
 		const FBoneChain* NewChain = AssetController->GetRetargetChainByName(NewChainName);
