PR #10854: Fixed IKRetargeter-menu from being added BlendSpace-based assets multiple times

diff --git a/Engine/Plugins/Animation/IKRig/Source/IKRigEditor/Private/RetargetEditor/AssetTypeActions_IKRetargeter.cpp b/Engine/Plugins/Animation/IKRig/Source/IKRigEditor/Private/RetargetEditor/AssetTypeActions_IKRetargeter.cpp
--- a/Engine/Plugins/Animation/IKRig/Source/IKRigEditor/Private/RetargetEditor/AssetTypeActions_IKRetargeter.cpp
+++ b/Engine/Plugins/Animation/IKRig/Source/IKRigEditor/Private/RetargetEditor/AssetTypeActions_IKRetargeter.cpp
@@ -129,8 +129,6 @@ void FAssetTypeActions_IKRetargeter::ExtendAnimAssetMenusForBatchRetargeting()
 	{
 		"ContentBrowser.AssetContextMenu.AnimSequence",
 		"ContentBrowser.AssetContextMenu.BlendSpace",
-		"ContentBrowser.AssetContextMenu.AimOffsetBlendSpace",
-		"ContentBrowser.AssetContextMenu.BlendSpace1D",
 		"ContentBrowser.AssetContextMenu.PoseAsset",
 		"ContentBrowser.AssetContextMenu.AnimBlueprint",
 		"ContentBrowser.AssetContextMenu.AnimMontage"
