Removing dead code in retargeter.

diff --git a/Engine/Plugins/Animation/IKRig/Source/IKRigEditor/Private/RetargetEditor/IKRetargetDefaultMode.cpp b/Engine/Plugins/Animation/IKRig/Source/IKRigEditor/Private/RetargetEditor/IKRetargetDefaultMode.cpp
--- a/Engine/Plugins/Animation/IKRig/Source/IKRigEditor/Private/RetargetEditor/IKRetargetDefaultMode.cpp
+++ b/Engine/Plugins/Animation/IKRig/Source/IKRigEditor/Private/RetargetEditor/IKRetargetDefaultMode.cpp
@@ -32,10 +32,6 @@ IPersonaPreviewScene& FIKRetargetDefaultMode::GetAnimPreviewScene() const
 	return *static_cast<IPersonaPreviewScene*>(static_cast<FAssetEditorModeManager*>(Owner)->GetPreviewScene());
 }
 
-void FIKRetargetDefaultMode::GetOnScreenDebugInfo(TArray<FText>& OutDebugInfo) const
-{
-}
-
 void FIKRetargetDefaultMode::Initialize()
 {
 	const TSharedPtr<FIKRetargetEditorController> Controller = EditorController.Pin();
@@ -474,9 +470,4 @@ void FIKRetargetDefaultMode::Tick(FEditorViewportClient* ViewportClient, float D
 	}
 }
 
-void FIKRetargetDefaultMode::DrawHUD(FEditorViewportClient* ViewportClient, FViewport* Viewport, const FSceneView* View, FCanvas* Canvas)
-{
-	FEdMode::DrawHUD(ViewportClient, Viewport, View, Canvas);
-}
-
 #undef LOCTEXT_NAMESPACE
diff --git a/Engine/Plugins/Animation/IKRig/Source/IKRigEditor/Public/RetargetEditor/IKRetargetDefaultMode.h b/Engine/Plugins/Animation/IKRig/Source/IKRigEditor/Public/RetargetEditor/IKRetargetDefaultMode.h
--- a/Engine/Plugins/Animation/IKRig/Source/IKRigEditor/Public/RetargetEditor/IKRetargetDefaultMode.h
+++ b/Engine/Plugins/Animation/IKRig/Source/IKRigEditor/Public/RetargetEditor/IKRetargetDefaultMode.h
@@ -27,14 +27,12 @@ public:
 	/** IPersonaEditMode interface */
 	virtual bool GetCameraTarget(FSphere& OutTarget) const override;
 	virtual class IPersonaPreviewScene& GetAnimPreviewScene() const override;
-	virtual void GetOnScreenDebugInfo(TArray<FText>& OutDebugInfo) const override;
 	/** END IPersonaEditMode interface */
 
 	/** FEdMode interface */
 	virtual void Initialize() override;
 	virtual void Tick(FEditorViewportClient* ViewportClient, float DeltaTime) override;
 	virtual void Render(const FSceneView* View, FViewport* Viewport, FPrimitiveDrawInterface* PDI) override;
-	virtual void DrawHUD(FEditorViewportClient* ViewportClient, FViewport* Viewport, const FSceneView* View, FCanvas* Canvas) override;
 	virtual bool IsCompatibleWith(FEditorModeID OtherModeID) const override { return true; }
 	virtual bool AllowWidgetMove() override;
 	virtual bool ShouldDrawWidget() const override;
