Added viewport overlay when IK Rig is not initialized (red) or has warnings (yellow).

diff --git a/Engine/Plugins/Animation/IKRig/Source/IKRigEditor/Private/RigEditor/IKRigToolkit.cpp b/Engine/Plugins/Animation/IKRig/Source/IKRigEditor/Private/RigEditor/IKRigToolkit.cpp
--- a/Engine/Plugins/Animation/IKRig/Source/IKRigEditor/Private/RigEditor/IKRigToolkit.cpp
+++ b/Engine/Plugins/Animation/IKRig/Source/IKRigEditor/Private/RigEditor/IKRigToolkit.cpp
@@ -262,6 +262,35 @@ void FIKRigEditorToolkit::HandleViewportCreated(const TSharedRef<IPersonaViewpor
 			return 1.0f;
 		});
 	}
+
+	// highlight viewport when processor disabled
+	auto GetBorderColorAndOpacity = [this]()
+	{
+		// no processor or processor not initialized
+		const UIKRigProcessor* Processor = EditorController.Get().GetIKRigProcessor();
+		if (!Processor || !Processor->IsInitialized() )
+		{
+			return FLinearColor::Red;
+		}
+
+		// highlight viewport if warnings
+		const TArray<FText>& Warnings = Processor->Log.GetWarnings();
+		if (!Warnings.IsEmpty())
+		{
+			return FLinearColor::Yellow;
+		}
+
+		return FLinearColor::Transparent;
+	};
+
+	InViewport->AddOverlayWidget(
+		SNew(SBorder)
+		.BorderImage(FIKRigEditorStyle::Get().GetBrush("IKRig.Viewport.Border"))
+		.BorderBackgroundColor_Lambda(GetBorderColorAndOpacity)
+		.Visibility(EVisibility::HitTestInvisible)
+		.Padding(0.0f)
+		.ShowEffectWhenDisabled(false)
+	);
 }
 
 #undef LOCTEXT_NAMESPACE
diff --git a/Engine/Plugins/Animation/IKRig/Source/IKRigEditor/Public/RigEditor/IKRigEditorStyle.h b/Engine/Plugins/Animation/IKRig/Source/IKRigEditor/Public/RigEditor/IKRigEditorStyle.h
--- a/Engine/Plugins/Animation/IKRig/Source/IKRigEditor/Public/RigEditor/IKRigEditorStyle.h
+++ b/Engine/Plugins/Animation/IKRig/Source/IKRigEditor/Public/RigEditor/IKRigEditorStyle.h
@@ -42,6 +42,9 @@ public:
 		FTextBlockStyle NormalText = FAppStyle::GetWidgetStyle<FTextBlockStyle>("SkeletonTree.NormalFont");
 		Set( "IKRig.Tree.NormalText", FTextBlockStyle(NormalText));
 		Set( "IKRig.Tree.ItalicText", FTextBlockStyle(NormalText).SetFont(DEFAULT_FONT("Italic", 10)));
+
+		SetContentRoot(FPaths::EngineContentDir() / TEXT("Editor/Slate"));
+		Set( "IKRig.Viewport.Border", new BOX_BRUSH( "Old/Window/ViewportDebugBorder", 0.8f, FLinearColor(1.0f,1.0f,1.0f,1.0f) ) );
 		
 		FSlateStyleRegistry::RegisterSlateStyle(*this);
 	}
