Allow canceling batch retargeting.

diff --git a/Engine/Plugins/Animation/IKRig/Source/IKRigEditor/Private/RetargetEditor/IKRetargetBatchOperation.cpp b/Engine/Plugins/Animation/IKRig/Source/IKRigEditor/Private/RetargetEditor/IKRetargetBatchOperation.cpp
--- a/Engine/Plugins/Animation/IKRig/Source/IKRigEditor/Private/RetargetEditor/IKRetargetBatchOperation.cpp
+++ b/Engine/Plugins/Animation/IKRig/Source/IKRigEditor/Private/RetargetEditor/IKRetargetBatchOperation.cpp
@@ -5,28 +5,25 @@
 #include "Animation/AnimSequence.h"
 #include "AnimationBlueprintLibrary.h"
 #include "AnimPose.h"
-#include "AnimPreviewInstance.h"
 #include "Animation/AnimSequence.h"
 #include "ContentBrowserModule.h"
 #include "EditorReimportHandler.h"
 #include "IContentBrowserSingleton.h"
 #include "RigEditor/IKRigController.h"
-#include "ObjectEditorUtils.h"
 #include "SSkeletonWidget.h"
-#include "PropertyCustomizationHelpers.h"
 #include "Animation/DebugSkelMeshComponent.h"
 #include "EditorFramework/AssetImportData.h"
 #include "Framework/Notifications/NotificationManager.h"
 #include "Kismet2/BlueprintEditorUtils.h"
 #include "Kismet2/KismetEditorUtilities.h"
 #include "Misc/ScopedSlowTask.h"
-#include "RetargetEditor/SRetargetAnimAssetsWindow.h"
 #include "Retargeter/IKRetargeter.h"
 #include "Retargeter/IKRetargetProcessor.h"
 #include "Widgets/Notifications/SNotificationList.h"
 #include "Animation/AnimMontage.h"
 #include "Retargeter/IKRetargetOps.h"
 #include "Retargeter/RetargetOps/CurveRemapOp.h"
+#include "ObjectTools.h"
 
 #define LOCTEXT_NAMESPACE "RetargetBatchOperation"
 
@@ -117,8 +114,33 @@ void UIKRetargetBatchOperation::DuplicateRetargetAssets(
 		AnimationAssetsToDuplicate.Remove(Pair.Key);
 	}
 
-	DuplicatedAnimAssets = DuplicateAssets<UAnimationAsset>(AnimationAssetsToDuplicate, DestinationPackage, &Context.NameRule);
-	DuplicatedBlueprints = DuplicateAssets<UAnimBlueprint>(AnimBlueprintsToDuplicate, DestinationPackage, &Context.NameRule);
+	// duplicate each asset individually (not done as a batch so user can cancel)
+	for (UAnimationAsset* Asset : AnimationAssetsToDuplicate)
+	{
+		if (Progress.ShouldCancel())
+		{
+			return;
+		}
+
+		FString AssetName = Asset->GetName();
+		Progress.EnterProgressFrame(1.f, FText::Format(LOCTEXT("DuplicatingAnimation", "Duplicating animation: {0}"), FText::FromString(AssetName)));
+		
+		TMap<UAnimationAsset*, UAnimationAsset*> DuplicateMap = DuplicateAssets<UAnimationAsset>({Asset}, DestinationPackage, &Context.NameRule);
+		DuplicatedAnimAssets.Append(DuplicateMap);
+	}
+	for (UAnimBlueprint* Asset : AnimBlueprintsToDuplicate)
+	{
+		if (Progress.ShouldCancel())
+		{
+			return;
+		}
+
+		FString AssetName = Asset->GetName();
+		Progress.EnterProgressFrame(1.f, FText::Format(LOCTEXT("DuplicatingBlueprint", "Duplicating blueprint: {0}"), FText::FromString(AssetName)));
+		
+		TMap<UAnimBlueprint*, UAnimBlueprint*> DuplicateMap = DuplicateAssets<UAnimBlueprint>({Asset}, DestinationPackage, &Context.NameRule);
+		DuplicatedBlueprints.Append(DuplicateMap);
+	}
 
 	// If we are moving the new asset to a different directory we need to fixup the reimport path.
 	// This should only effect source FBX paths within the project.
@@ -173,9 +195,17 @@ void UIKRetargetBatchOperation::RetargetAssets(
 	
 	for (UAnimationAsset* AssetToRetarget : AnimationAssetsToRetarget)
 	{
+		if (Progress.ShouldCancel())
+		{
+			return;
+		}
+		
 		// prepare animation sequence asset to receive retargeted animation
 		if (UAnimSequence* AnimSequenceToRetarget = Cast<UAnimSequence>(AssetToRetarget))
 		{
+			FString AssetName = AnimSequenceToRetarget->GetName();
+			Progress.EnterProgressFrame(1.f, FText::Format(LOCTEXT("PreparingAsset", "Preparing asset: {0}"), FText::FromString(AssetName)));
+			
 			// copy curve data from source asset, preserving data in the target if present.
 			UAnimationBlueprintLibrary::CopyAnimationCurveNamesToSkeleton(OldSkeleton, NewSkeleton, AnimSequenceToRetarget, ERawCurveTrackTypes::RCT_Float);	
 
@@ -216,6 +246,11 @@ void UIKRetargetBatchOperation::RetargetAssets(
 	static FProperty* RetargetAssetProperty = UAnimSequence::StaticClass()->FindPropertyByName(RetargetSourceAssetPropertyName);
 	for (UAnimationAsset* AssetToRetarget : AnimationAssetsToRetarget)
 	{
+		if (Progress.ShouldCancel())
+		{
+			return;
+		}
+		
 		// force updating of the retarget pose, this is normally done on PreSave() but is guarded against procedural saves
 		if (UAnimSequence* AnimSequenceToRetarget = Cast<UAnimSequence>(AssetToRetarget))
 		{
@@ -233,6 +268,11 @@ void UIKRetargetBatchOperation::RetargetAssets(
 	// convert all Animation Blueprints and compile 
 	for (UAnimBlueprint* AnimBlueprint : AnimBlueprintsToRetarget)
 	{
+		if (Progress.ShouldCancel())
+		{
+			return;
+		}
+		
 		// replace skeleton
 		AnimBlueprint->TargetSkeleton = NewSkeleton;
 
@@ -300,6 +340,11 @@ void UIKRetargetBatchOperation::ConvertAnimation(
 	// for each pair of source / target animation sequences
 	for (TPair<UAnimationAsset*, UAnimationAsset*>& Pair : DuplicatedAnimAssets)
 	{
+		if (Progress.ShouldCancel())
+		{
+			return;
+		}
+		
 		UAnimSequence* SourceSequence = Cast<UAnimSequence>(Pair.Key);
 		UAnimSequence* TargetSequence = Cast<UAnimSequence>(Pair.Value);
 		if (!(SourceSequence && TargetSequence))
@@ -342,6 +387,12 @@ void UIKRetargetBatchOperation::ConvertAnimation(
 		// retarget each frame's pose from source to target
 		for (int32 FrameIndex=0; FrameIndex<NumFrames; ++FrameIndex)
 		{
+			if (Progress.ShouldCancel())
+			{
+				TargetSeqController.CloseBracket(bShouldTransact);
+				return;
+			}
+			
 			// get the source global pose
 			FAnimPose SourcePoseAtFrame;
 			UAnimPoseExtensions::GetAnimPoseAtFrame(SourceSequence, FrameIndex, EvaluationOptions, SourcePoseAtFrame);
@@ -453,6 +504,10 @@ void UIKRetargetBatchOperation::RemapCurves(const FIKRetargetBatchOperationConte
 		// increment progress bar
 		FString AssetName = TargetSequence->GetName();
 		Progress.EnterProgressFrame(1.f, FText::Format(LOCTEXT("RemappingCurves", "Remapping Curves on Asset: {0}"), FText::FromString(AssetName)));
+		if (Progress.ShouldCancel())
+		{
+			return;
+		}
 
 		// all curves were copied when we duplicated the animation sequence, so now we have to rename curves
 		// based on the remapping defined in the curve remap op(s)
@@ -520,37 +575,50 @@ void UIKRetargetBatchOperation::NotifyUserOfResults(
 	const FIKRetargetBatchOperationContext& Context,
 	FScopedSlowTask& Progress) const
 {
-	Progress.EnterProgressFrame(1.f, FText(LOCTEXT("DoneBatchRetarget", "Duplicate and retarget complete!")));
-
 	// gather newly created objects
 	TArray<UObject*> NewAssets;
 	GetNewAssets(NewAssets);
-
-	// log details of what assets were created
-	for (UObject* NewAsset : NewAssets)
-	{
-		UE_LOG(LogTemp, Display, TEXT("Duplicate and Retarget - New Asset Created: %s"), *NewAsset->GetName());
-	}
 	
-	// notify user
-	FNotificationInfo Notification(FText::GetEmpty());
-	Notification.ExpireDuration = 5.f;
-	Notification.Text = FText::Format(
-		LOCTEXT("MultiNonDuplicatedAsset", "{0} assets were retargeted to new skeleton {1}. See Output for details."),
-		FText::AsNumber(NewAssets.Num()),
-		FText::FromString(Context.TargetMesh->GetName()));
-	FSlateNotificationManager::Get().AddNotification(Notification);
-	
-	// select all new assets
+	// select all new assets and show in the content browser
 	TArray<FAssetData> CurrentSelection;
 	for(UObject* NewObject : NewAssets)
 	{
 		CurrentSelection.Add(FAssetData(NewObject));
 	}
-
-	// show assets in browser
-	FContentBrowserModule& ContentBrowserModule = FModuleManager::Get().LoadModuleChecked<FContentBrowserModule>("ContentBrowser");
+	const FContentBrowserModule& ContentBrowserModule = FModuleManager::Get().LoadModuleChecked<FContentBrowserModule>("ContentBrowser");
 	ContentBrowserModule.Get().SyncBrowserToAssets(CurrentSelection);
+
+	// create pop-up notification in editor UI
+	constexpr float NotificationDuration = 5.0f;
+	if (Progress.ShouldCancel())
+	{
+		Progress.EnterProgressFrame(1.f, FText(LOCTEXT("CancelledBatchRetarget", "Cancelled.")));
+		
+		// notify user that retarget was cancelled
+		FNotificationInfo Notification(FText::GetEmpty());
+		Notification.ExpireDuration = NotificationDuration;
+		Notification.Text = FText(LOCTEXT("BatchRetargetCancelled", "Batch retarget cancelled."));
+		FSlateNotificationManager::Get().AddNotification(Notification);
+	}
+	else
+	{
+		Progress.EnterProgressFrame(1.f, FText(LOCTEXT("DoneBatchRetarget", "Duplicate and retarget complete!")));
+		
+		// log details of what assets were created
+		for (const UObject* NewAsset : NewAssets)
+		{
+			UE_LOG(LogTemp, Display, TEXT("Duplicate and Retarget - New Asset Created: %s"), *NewAsset->GetName());
+		}
+	
+		// notify user that retarget completed
+		FNotificationInfo Notification(FText::GetEmpty());
+		Notification.ExpireDuration = NotificationDuration;
+		Notification.Text = FText::Format(
+			LOCTEXT("MultiNonDuplicatedAsset", "{0} assets were retargeted to new skeleton {1}. See Output for details."),
+			FText::AsNumber(NewAssets.Num()),
+			FText::FromString(Context.TargetMesh->GetName()));
+		FSlateNotificationManager::Get().AddNotification(Notification);
+	}
 }
 
 void UIKRetargetBatchOperation::GetNewAssets(TArray<UObject*>& NewAssets) const
@@ -570,6 +638,22 @@ void UIKRetargetBatchOperation::GetNewAssets(TArray<UObject*>& NewAssets) const
 	}
 }
 
+void UIKRetargetBatchOperation::CleanupIfCancelled(const FScopedSlowTask& Progress) const
+{
+	if (!Progress.ShouldCancel())
+	{
+		return;
+	}
+
+	// get list of all the assets that were created
+	// (to be removed after being cancelled)
+	TArray<UObject*> NewAssets;
+	GetNewAssets(NewAssets);
+	
+	// delete any newly created assets
+	constexpr bool bShowConfirmation = true;
+	ObjectTools::DeleteObjects(NewAssets, bShowConfirmation);
+}
 
 TArray<FAssetData> UIKRetargetBatchOperation::DuplicateAndRetarget(
 	const TArray<FAssetData>& AssetsToRetarget,
@@ -674,12 +758,14 @@ void UIKRetargetBatchOperation::RunRetarget(FIKRetargetBatchOperationContext& Co
 	
 	// show progress bar
 	constexpr int NumAdditionalProgressFrames = 3;
-	FScopedSlowTask Progress(NumAssets + NumAdditionalProgressFrames, LOCTEXT("GatheringBatchRetarget", "Gathering animation assets..."));
-	Progress.MakeDialog();
+	FScopedSlowTask Progress((NumAssets*3) + NumAdditionalProgressFrames, LOCTEXT("GatheringBatchRetarget", "Gathering animation assets..."));
+	constexpr bool bShowCancelButton = true;
+	Progress.MakeDialog(bShowCancelButton);
 	
 	DuplicateRetargetAssets(Context, Progress);
 	RetargetAssets(Context, Progress);
 	NotifyUserOfResults(Context, Progress);
+	CleanupIfCancelled(Progress);
 }
 
 void UIKRetargetBatchOperation::Reset()
@@ -691,38 +777,4 @@ void UIKRetargetBatchOperation::Reset()
 	RemappedAnimAssets.Reset();
 }
 
-/**
-* Duplicates the supplied AssetsToDuplicate and returns a map of original asset to duplicate. Templated wrapper that calls DuplicateAssetInternal.
-*
-* @param	AssetsToDuplicate	The animations to duplicate
-* @param	DestinationPackage	The package that the duplicates should be placed in
-* @param	NameRule			The rules for how to rename the duplicated assets
-*
-* @return	TMap of original animation to duplicate
-*/
-template<class AssetType>
-TMap<AssetType*, AssetType*> UIKRetargetBatchOperation::DuplicateAssets(
-	const TArray<AssetType*>& AssetsToDuplicate,
-	UPackage* DestinationPackage,
-	const FNameDuplicationRule* NameRule)
-{
-	TArray<UObject*> Assets;
-	for (AssetType* Asset : AssetsToDuplicate)
-	{
-		Assets.Add(Asset);
-	}
-
-	// duplicate assets
-	TMap<UObject*, UObject*> DuplicateAssetsMap = DuplicateAssetsInternal(Assets, DestinationPackage, NameRule);
-
-	// cast to AssetType
-	TMap<AssetType*, AssetType*> ReturnMap;
-	for (const TTuple<UObject*, UObject*>& DuplicateAsset : DuplicateAssetsMap)
-	{
-		ReturnMap.Add(Cast<AssetType>(DuplicateAsset.Key), Cast<AssetType>(DuplicateAsset.Value));
-	}
-	
-	return ReturnMap;
-}
-
 #undef LOCTEXT_NAMESPACE
diff --git a/Engine/Plugins/Animation/IKRig/Source/IKRigEditor/Public/RetargetEditor/IKRetargetBatchOperation.h b/Engine/Plugins/Animation/IKRig/Source/IKRigEditor/Public/RetargetEditor/IKRetargetBatchOperation.h
--- a/Engine/Plugins/Animation/IKRig/Source/IKRigEditor/Public/RetargetEditor/IKRetargetBatchOperation.h
+++ b/Engine/Plugins/Animation/IKRig/Source/IKRigEditor/Public/RetargetEditor/IKRetargetBatchOperation.h
@@ -124,19 +124,8 @@ private:
 	/* Generate list of newly created assets to report to user */
 	void GetNewAssets(TArray<UObject*>& NewAssets) const;
 
-	/**
-	* Duplicates the supplied AssetsToDuplicate and returns a map of original asset to duplicate. Templated wrapper that calls DuplicateAssetInternal.
-	*
-	* @param	AssetsToDuplicate	The animations to duplicate
-	* @param	DestinationPackage	The package that the duplicates should be placed in
-	*
-	* @return	TMap of original animation to duplicate
-	*/
-	template<class AssetType>
-	static TMap<AssetType*, AssetType*> DuplicateAssets(
-		const TArray<AssetType*>& AssetsToDuplicate,
-		UPackage* DestinationPackage,
-		const EditorAnimUtils::FNameDuplicationRule* NameRule);
+	/* If user cancelled half way, cleanup all the duplicated assets */
+	void CleanupIfCancelled(const FScopedSlowTask& Progress) const;
 	
 	/** Lists of assets to retarget. Populated from selection during init */
 	TArray<UAnimationAsset*>	AnimationAssetsToRetarget;
