Fixes minor issue with batch retarget progress bar miscalculating the amount of work when remapping curves.

diff --git a/Engine/Plugins/Animation/IKRig/Source/IKRigEditor/Private/RetargetEditor/IKRetargetBatchOperation.cpp b/Engine/Plugins/Animation/IKRig/Source/IKRigEditor/Private/RetargetEditor/IKRetargetBatchOperation.cpp
--- a/Engine/Plugins/Animation/IKRig/Source/IKRigEditor/Private/RetargetEditor/IKRetargetBatchOperation.cpp
+++ b/Engine/Plugins/Animation/IKRig/Source/IKRigEditor/Private/RetargetEditor/IKRetargetBatchOperation.cpp
@@ -491,6 +491,9 @@ void UIKRetargetBatchOperation::RemapCurves(const FIKRetargetBatchOperationConte
 		bCopyAllSourceCurves &= CurveRemapOp->bCopyAllSourceCurves;
 	}
 
+	// update progress bar
+	Progress.EnterProgressFrame(1.f, FText::Format(LOCTEXT("RemappingCurves", "Remapping {0} curves on animation assets..."), FText::AsNumber(CurvesToRemap.Num())));
+
 	// for each exported animation, remap curves from source to target anim
 	for (TPair<UAnimationAsset*, UAnimationAsset*>& Pair : DuplicatedAnimAssets)
 	{
@@ -502,8 +505,6 @@ void UIKRetargetBatchOperation::RemapCurves(const FIKRetargetBatchOperationConte
 		}
 
 		// increment progress bar
-		FString AssetName = TargetSequence->GetName();
-		Progress.EnterProgressFrame(1.f, FText::Format(LOCTEXT("RemappingCurves", "Remapping Curves on Asset: {0}"), FText::FromString(AssetName)));
 		if (Progress.ShouldCancel())
 		{
 			return;
@@ -757,8 +758,10 @@ void UIKRetargetBatchOperation::RunRetarget(FIKRetargetBatchOperationContext& Co
 	}
 	
 	// show progress bar
-	constexpr int NumAdditionalProgressFrames = 3;
-	FScopedSlowTask Progress((NumAssets*3) + NumAdditionalProgressFrames, LOCTEXT("GatheringBatchRetarget", "Gathering animation assets..."));
+	constexpr int32 NumAdditionalProgressFrames = 4;
+	constexpr int32 NumPassesOverAssets = 3;
+	const int32 NumProgressSteps = (NumAssets * NumPassesOverAssets) + NumAdditionalProgressFrames;
+	FScopedSlowTask Progress(NumProgressSteps, LOCTEXT("GatheringBatchRetarget", "Gathering animation assets..."));
 	constexpr bool bShowCancelButton = true;
 	Progress.MakeDialog(bShowCancelButton);
 	
