Fixes issue with "None" rotation mode in IK Retargeter still applying minor rotation.

diff --git a/Engine/Plugins/Animation/IKRig/Source/IKRig/Private/Retargeter/IKRetargetProcessor.cpp b/Engine/Plugins/Animation/IKRig/Source/IKRig/Private/Retargeter/IKRetargetProcessor.cpp
--- a/Engine/Plugins/Animation/IKRig/Source/IKRig/Private/Retargeter/IKRetargetProcessor.cpp
+++ b/Engine/Plugins/Animation/IKRig/Source/IKRig/Private/Retargeter/IKRetargetProcessor.cpp
@@ -199,7 +199,7 @@ void FRetargetSkeleton::UpdateLocalTransformOfSingleBone(
 	OutLocalPose[BoneIndex] = ChildGlobalTransform.GetRelativeTransform(ParentGlobalTransform);
 }
 
-FTransform FRetargetSkeleton::GetGlobalRefPoseOfSingleBone(
+FTransform FRetargetSkeleton::GetGlobalRetargetPoseOfSingleBone(
 	const int32 BoneIndex,
 	const TArray<FTransform>& InGlobalPose) const
 {
@@ -522,7 +522,7 @@ void FChainFK::PutCurrentTransformsInRefPose(
 		if (ChainIndex == 0)
 		{
 			const int32 BoneIndex = InBoneIndices[ChainIndex];
-			CurrentGlobalTransforms[ChainIndex] = Skeleton.GetGlobalRefPoseOfSingleBone(BoneIndex, InCurrentGlobalPose);
+			CurrentGlobalTransforms[ChainIndex] = Skeleton.GetGlobalRetargetPoseOfSingleBone(BoneIndex, InCurrentGlobalPose);
 		}
 		else
 		{
@@ -551,11 +551,8 @@ void FChainEncoderFK::EncodePose(
 
 	CurrentLocalTransforms.SetNum(SourceBoneIndices.Num());
 	FillTransformsWithLocalSpaceOfChain(SourceSkeleton, InSourceGlobalPose, SourceBoneIndices, CurrentLocalTransforms);
-
-	if (ChainParentBoneIndex != INDEX_NONE)
-	{
-		ChainParentCurrentGlobalTransform = InSourceGlobalPose[ChainParentBoneIndex];
-	}
+	
+	ChainParentCurrentGlobalTransform = ChainParentBoneIndex != INDEX_NONE ? InSourceGlobalPose[ChainParentBoneIndex] : FTransform::Identity;
 }
 
 void FChainEncoderFK::TransformCurrentChainTransforms(const FTransform& NewParentTransform)
@@ -592,7 +589,7 @@ void FChainDecoderFK::DecodePose(
 	// retargeted in the prior step, then the neck bones will need updating first.
 	// Otherwise the neck bones will remain at their location prior to the spine update.
 	UpdateIntermediateParents(TargetSkeleton,InOutGlobalPose);
-
+	
 	// transform entire source chain from it's root to match target's current root orientation (maintaining offset from retarget pose)
 	// this ensures children are retargeted in a "local" manner free from skewing that will happen if source and target
 	// become misaligned as can happen if parent chains were not retargeted
@@ -681,8 +678,10 @@ void FChainDecoderFK::DecodePose(
 			break;
 			case ERetargetRotationMode::None:
 			{
-				SourceCurrentTransform = SourceChain.InitialGlobalTransforms.Last();
-				SourceInitialTransform = SourceChain.InitialGlobalTransforms.Last();
+					// in order to induce no rotation on the FK chain, we rotate the chain rigidly from the root of the chain
+					SourceInitialTransform = SourceChain.InitialGlobalTransforms[0];
+					// use the current global space retarget pose as the "current" transform, so chain rotates with parent
+					SourceCurrentTransform = SourceChain.InitialLocalTransforms[0] * SourceChain.ChainParentCurrentGlobalTransform;
 			}
 			break;
 			default:
@@ -739,13 +738,16 @@ void FChainDecoderFK::DecodePose(
 
 	// apply final blending between retarget pose of chain and newly retargeted pose
 	// blend must be done in local space, so we do it in a separate loop after full chain pose is generated
-	// (skipped if the alphas are not near 1.0)
-	if (!FMath::IsNearlyEqual(Settings.FK.RotationAlpha, 1.0f) || !FMath::IsNearlyEqual(Settings.FK.TranslationAlpha, 1.0f))
+	const bool bShouldBlendRotation = !FMath::IsNearlyEqual(Settings.FK.RotationAlpha, 1.0f);
+	const bool bShouldBlendTranslation = !FMath::IsNearlyEqual(Settings.FK.TranslationAlpha, 1.0f);
+	if (bShouldBlendRotation || bShouldBlendTranslation) // (skipped if the alphas are not near 1.0)
 	{
+		// generate local space pose of chain
 		TArray<FTransform> NewLocalTransforms;
 		NewLocalTransforms.SetNum(InitialLocalTransforms.Num());
 		FillTransformsWithLocalSpaceOfChain(TargetSkeleton, InOutGlobalPose, TargetBoneIndices, NewLocalTransforms);
 
+		// blend each bone in chain with the retarget pose
 		for (int32 ChainIndex=0; ChainIndex<InitialLocalTransforms.Num(); ++ChainIndex)
 		{
 			// blend between current local pose and initial local pose
diff --git a/Engine/Plugins/Animation/IKRig/Source/IKRig/Public/Retargeter/IKRetargetProcessor.h b/Engine/Plugins/Animation/IKRig/Source/IKRig/Public/Retargeter/IKRetargetProcessor.h
--- a/Engine/Plugins/Animation/IKRig/Source/IKRig/Public/Retargeter/IKRetargetProcessor.h
+++ b/Engine/Plugins/Animation/IKRig/Source/IKRig/Public/Retargeter/IKRetargetProcessor.h
@@ -73,7 +73,7 @@ struct IKRIG_API FRetargetSkeleton
 		TArray<FTransform>& OutLocalPose,
 		const TArray<FTransform>& InGlobalPose) const;
 
-	FTransform GetGlobalRefPoseOfSingleBone(
+	FTransform GetGlobalRetargetPoseOfSingleBone(
 		const int32 BoneIndex,
 		const TArray<FTransform>& InGlobalPose) const;
 
