Fixes issue with retargeter propagating scale from source to target bones.

diff --git a/Engine/Plugins/Animation/IKRig/Source/IKRig/Private/AnimNodes/AnimNode_RetargetPoseFromMesh.cpp b/Engine/Plugins/Animation/IKRig/Source/IKRig/Private/AnimNodes/AnimNode_RetargetPoseFromMesh.cpp
--- a/Engine/Plugins/Animation/IKRig/Source/IKRig/Private/AnimNodes/AnimNode_RetargetPoseFromMesh.cpp
+++ b/Engine/Plugins/Animation/IKRig/Source/IKRig/Private/AnimNodes/AnimNode_RetargetPoseFromMesh.cpp
@@ -350,6 +350,12 @@ void FAnimNode_RetargetPoseFromMesh::CopyBoneTransformsFromSource(USkeletalMeshC
 	{
 		SourceMeshComponentSpaceBoneTransforms.Append(ComponentToCopyFrom->GetComponentSpaceTransforms()); // copy directly
 	}
+
+	// strip all scale out of the pose values, the translation of a component-space pose has incorporated scale values
+	for (FTransform& Transform : SourceMeshComponentSpaceBoneTransforms)
+	{
+		Transform.SetScale3D(FVector::OneVector);
+	}
 }
 
 TObjectPtr<USkeletalMeshComponent> FAnimNode_RetargetPoseFromMesh::GetComponentToCopyPoseFrom() const
diff --git a/Engine/Plugins/Animation/IKRig/Source/IKRig/Private/Retargeter/IKRetargetProcessor.cpp b/Engine/Plugins/Animation/IKRig/Source/IKRig/Private/Retargeter/IKRetargetProcessor.cpp
--- a/Engine/Plugins/Animation/IKRig/Source/IKRig/Private/Retargeter/IKRetargetProcessor.cpp
+++ b/Engine/Plugins/Animation/IKRig/Source/IKRig/Private/Retargeter/IKRetargetProcessor.cpp
@@ -91,14 +91,18 @@ void FRetargetSkeleton::GenerateRetargetPose(
 	// convert to global space
 	UpdateGlobalTransformsBelowBone(-1, RetargetLocalPose, RetargetGlobalPose);
 
+	// strip scale (done AFTER generating global pose so that scales are baked into translation)
+	for (int32 BoneIndex=0; BoneIndex<BoneNames.Num(); ++BoneIndex)
+	{
+		RetargetLocalPose[BoneIndex].SetScale3D(FVector::OneVector);
+		RetargetGlobalPose[BoneIndex].SetScale3D(FVector::OneVector);
+	}
+
 	// no retarget pose specified (will use default pose from skeletal mesh with no offsets)
 	if (InRetargetPose==nullptr  || RetargetRootBone == NAME_None)
 	{
 		return;
 	}
-
-	// apply retarget pose offsets (retarget pose is stored as offset relative to reference pose)
-	const TArray<FTransform>& RefPoseLocal = SkeletalMesh->GetRefSkeleton().GetRefBonePose();
 	
 	// apply root translation offset
 	const int32 RootBoneIndex = FindBoneIndexByName(RetargetRootBone);
@@ -109,7 +113,8 @@ void FRetargetSkeleton::GenerateRetargetPose(
 		UpdateLocalTransformOfSingleBone(RootBoneIndex, RetargetLocalPose, RetargetGlobalPose);
 	}
 
-	// apply bone rotation offsets
+	// apply retarget pose offsets (retarget pose is stored as offset relative to reference pose)
+	const TArray<FTransform>& RefPoseLocal = SkeletalMesh->GetRefSkeleton().GetRefBonePose();
 	for (const TTuple<FName, FQuat>& BoneDelta : InRetargetPose->GetAllDeltaRotations())
 	{
 		const int32 BoneIndex = FindBoneIndexByName(BoneDelta.Key);
@@ -1868,6 +1873,15 @@ TArray<FTransform>&  UIKRetargetProcessor::RunRetargeter(
 	const float DeltaTime)
 {
 	check(bIsInitialized);
+
+#if WITH_EDITOR
+	// validate system running the retargeter has stripped all the scale out of the incoming pose
+	for (const FTransform& Transform : InSourceGlobalPose)
+	{
+		const bool bHasNoScale = Transform.GetScale3D().Equals(FVector::OneVector);
+		ensureMsgf(bHasNoScale, TEXT("Found scale values on incoming pose in retarget processor. Scale should be baked into translation and set to 1,1,1. "));
+	}
+#endif
 	
 	// start from retarget pose
 	TargetSkeleton.OutputGlobalPose = TargetSkeleton.RetargetGlobalPose;
