Added "Asset Settings" button to toolbar in IK Rig.

diff --git a/Engine/Plugins/Animation/IKRig/Source/IKRigEditor/Private/RigEditor/IKRigCommands.cpp b/Engine/Plugins/Animation/IKRig/Source/IKRigEditor/Private/RigEditor/IKRigCommands.cpp
--- a/Engine/Plugins/Animation/IKRig/Source/IKRigEditor/Private/RigEditor/IKRigCommands.cpp
+++ b/Engine/Plugins/Animation/IKRig/Source/IKRigEditor/Private/RigEditor/IKRigCommands.cpp
@@ -7,6 +7,7 @@
 void FIKRigCommands::RegisterCommands()
 {
 	UI_COMMAND(Reset, "Reset", "Reset state of the rig and goals to initial pose.", EUserInterfaceActionType::Button, FInputChord());
+	UI_COMMAND(ShowAssetSettings, "Asset Settings", "Show the settings for this IK Rig asset.", EUserInterfaceActionType::ToggleButton, FInputChord());
 }
 
 #undef LOCTEXT_NAMESPACE
diff --git a/Engine/Plugins/Animation/IKRig/Source/IKRigEditor/Private/RigEditor/IKRigEditorController.cpp b/Engine/Plugins/Animation/IKRig/Source/IKRigEditor/Private/RigEditor/IKRigEditorController.cpp
--- a/Engine/Plugins/Animation/IKRig/Source/IKRigEditor/Private/RigEditor/IKRigEditorController.cpp
+++ b/Engine/Plugins/Animation/IKRig/Source/IKRigEditor/Private/RigEditor/IKRigEditorController.cpp
@@ -503,7 +503,7 @@ void FIKRigEditorController::ClearSelection()
 		SkeletonView->TreeView->ClearSelection();	
 	}
 	
-	ShowEmptyDetails();
+	ShowAssetDetails();
 }
 
 void FIKRigEditorController::HandleGoalSelectedInViewport(const FName& GoalName, bool bReplace) const
@@ -790,7 +790,7 @@ void FIKRigEditorController::ShowDetailsForSolver(const int32 SolverIndex) const
 	DetailsView->SetObject(AssetController->GetSolverAtIndex(SolverIndex));
 }
 
-void FIKRigEditorController::ShowEmptyDetails() const
+void FIKRigEditorController::ShowAssetDetails() const
 {
 	DetailsView->SetObject(AssetController->GetAsset());
 }
@@ -799,7 +799,7 @@ void FIKRigEditorController::ShowDetailsForElements(const TArray<TSharedPtr<FIKR
 {
 	if (!InItems.Num())
 	{
-		ShowEmptyDetails();
+		ShowAssetDetails();
 		return;
 	}
 
@@ -862,11 +862,35 @@ void FIKRigEditorController::OnFinishedChangingDetails(const FPropertyChangedEve
 	}
 }
 
+bool FIKRigEditorController::IsObjectInDetailsView(const UObject* Object) const
+{
+	if (!DetailsView.IsValid())
+	{
+		return false;
+	}
+
+	if (!Object)
+	{
+		return false;
+	}
+	
+	TArray<TWeakObjectPtr<UObject>> SelectedObjects = DetailsView->GetSelectedObjects();
+	for (TWeakObjectPtr<UObject> SelectedObject : SelectedObjects)
+	{
+		if (SelectedObject.Get() == Object)
+		{
+			return true;
+		}
+	}
+
+	return false;
+}
+
 void FIKRigEditorController::SetDetailsView(const TSharedPtr<IDetailsView>& InDetailsView)
 {
 	DetailsView = InDetailsView;
 	DetailsView->OnFinishedChangingProperties().AddSP(this, &FIKRigEditorController::OnFinishedChangingDetails);
-	ShowEmptyDetails();
+	ShowAssetDetails();
 }
 
 void FIKRigEditorController::PromptToAssignGoalToChain(const FName NewGoalName) const
diff --git a/Engine/Plugins/Animation/IKRig/Source/IKRigEditor/Private/RigEditor/IKRigToolkit.cpp b/Engine/Plugins/Animation/IKRig/Source/IKRigEditor/Private/RigEditor/IKRigToolkit.cpp
--- a/Engine/Plugins/Animation/IKRig/Source/IKRigEditor/Private/RigEditor/IKRigToolkit.cpp
+++ b/Engine/Plugins/Animation/IKRig/Source/IKRigEditor/Private/RigEditor/IKRigToolkit.cpp
@@ -112,6 +112,19 @@ void FIKRigEditorToolkit::BindCommands()
         Commands.Reset,
         FExecuteAction::CreateSP(EditorController, &FIKRigEditorController::Reset),
 		EUIActionRepeatMode::RepeatDisabled);
+
+	ToolkitCommands->MapAction(
+		Commands.ShowAssetSettings,
+		FExecuteAction::CreateLambda([this]()
+		{
+			return EditorController->ShowAssetDetails();
+		}),
+		FCanExecuteAction(),
+		FIsActionChecked::CreateLambda([this]() ->bool
+		{
+			const UIKRigDefinition* Asset = EditorController->AssetController->GetAsset();
+			return EditorController->IsObjectInDetailsView(Asset);	
+		}));
 }
 
 void FIKRigEditorToolkit::ExtendToolbar()
@@ -140,6 +153,20 @@ void FIKRigEditorToolkit::FillToolbar(FToolBarBuilder& ToolbarBuilder)
 			FSlateIcon(FAppStyle::Get().GetStyleSetName(),"Icons.Refresh"));
 	}
 	ToolbarBuilder.EndSection();
+
+	ToolbarBuilder.AddSeparator();
+	ToolbarBuilder.AddWidget(SNew(SSpacer), NAME_None, true, HAlign_Right);
+
+	ToolbarBuilder.BeginSection("Show Settings");
+	{
+		ToolbarBuilder.AddToolBarButton(
+		FIKRigCommands::Get().ShowAssetSettings,
+		NAME_None,
+		TAttribute<FText>(),
+		TAttribute<FText>(),
+		FSlateIcon(FIKRigEditorStyle::Get().GetStyleSetName(),"IKRig.AssetSettings"));
+	}
+	ToolbarBuilder.EndSection();
 }
 
 FName FIKRigEditorToolkit::GetToolkitFName() const
diff --git a/Engine/Plugins/Animation/IKRig/Source/IKRigEditor/Private/RigEditor/SIKRigSolverStack.cpp b/Engine/Plugins/Animation/IKRig/Source/IKRigEditor/Private/RigEditor/SIKRigSolverStack.cpp
--- a/Engine/Plugins/Animation/IKRig/Source/IKRigEditor/Private/RigEditor/SIKRigSolverStack.cpp
+++ b/Engine/Plugins/Animation/IKRig/Source/IKRigEditor/Private/RigEditor/SIKRigSolverStack.cpp
@@ -475,7 +475,7 @@ void SIKRigSolverStack::ShowDetailsForItem(TSharedPtr<FSolverStackElement> InIte
 			return;
 		}
 		
-		Controller->ShowEmptyDetails();
+		Controller->ShowAssetDetails();
 	}
 	else
 	{
diff --git a/Engine/Plugins/Animation/IKRig/Source/IKRigEditor/Public/RigEditor/IKRigCommands.h b/Engine/Plugins/Animation/IKRig/Source/IKRigEditor/Public/RigEditor/IKRigCommands.h
--- a/Engine/Plugins/Animation/IKRig/Source/IKRigEditor/Public/RigEditor/IKRigCommands.h
+++ b/Engine/Plugins/Animation/IKRig/Source/IKRigEditor/Public/RigEditor/IKRigCommands.h
@@ -18,9 +18,12 @@ public:
 	)
 	{}
 	
-	/** reset whole system to initial state */
+	// reset whole system to initial state
 	TSharedPtr< FUICommandInfo > Reset;
 
-	/** initialize commands */
+	// show settings of the asset in the details panel
+	TSharedPtr< FUICommandInfo > ShowAssetSettings;
+
+	// initialize commands
 	virtual void RegisterCommands() override;
 };
diff --git a/Engine/Plugins/Animation/IKRig/Source/IKRigEditor/Public/RigEditor/IKRigEditorController.h b/Engine/Plugins/Animation/IKRig/Source/IKRigEditor/Public/RigEditor/IKRigEditorController.h
--- a/Engine/Plugins/Animation/IKRig/Source/IKRigEditor/Public/RigEditor/IKRigEditorController.h
+++ b/Engine/Plugins/Animation/IKRig/Source/IKRigEditor/Public/RigEditor/IKRigEditorController.h
@@ -174,11 +174,13 @@ public:
 	// show single SOLVER settings in details view 
 	void ShowDetailsForSolver(const int32 SolverIndex) const;
 	// show nothing in details view 
-	void ShowEmptyDetails() const;
+	void ShowAssetDetails() const;
 	// show selected items in details view 
 	void ShowDetailsForElements(const TArray<TSharedPtr<FIKRigTreeElement>>& InItems) const;
 	// callback when detail is edited 
 	void OnFinishedChangingDetails(const FPropertyChangedEvent& PropertyChangedEvent) const;
+	// returns true if the supplied UObject is being shown in the details panel
+	bool IsObjectInDetailsView(const UObject* Object) const;
 	
 	// set details tab view 
 	void SetDetailsView(const TSharedPtr<class IDetailsView>& InDetailsView);
diff --git a/Engine/Plugins/Animation/IKRig/Source/IKRigEditor/Public/RigEditor/IKRigEditorStyle.h b/Engine/Plugins/Animation/IKRig/Source/IKRigEditor/Public/RigEditor/IKRigEditorStyle.h
--- a/Engine/Plugins/Animation/IKRig/Source/IKRigEditor/Public/RigEditor/IKRigEditorStyle.h
+++ b/Engine/Plugins/Animation/IKRig/Source/IKRigEditor/Public/RigEditor/IKRigEditorStyle.h
@@ -38,6 +38,8 @@ public:
 		Set("IKRig.SolverStack", new IMAGE_BRUSH_SVG("Slate/SolverStack", Icon64x64));
 		Set("IKRig.IKRetargeting", new IMAGE_BRUSH_SVG("Slate/IKRetargeting", Icon64x64));
 		Set("IKRig.Hierarchy", new IMAGE_BRUSH_SVG("Slate/Hierarchy", Icon64x64));
+
+		Set("IKRig.AssetSettings", new IMAGE_BRUSH_SVG("Slate/AssetSettings", Icon64x64));
 		
 		FTextBlockStyle NormalText = FAppStyle::GetWidgetStyle<FTextBlockStyle>("SkeletonTree.NormalFont");
 		Set( "IKRig.Tree.NormalText", FTextBlockStyle(NormalText));
