Matching comment style with header

diff --git a/Engine/Plugins/Animation/IKRig/Source/IKRigEditor/Public/RetargetEditor/IKRetargetEditorController.h b/Engine/Plugins/Animation/IKRig/Source/IKRigEditor/Public/RetargetEditor/IKRetargetEditorController.h
--- a/Engine/Plugins/Animation/IKRig/Source/IKRigEditor/Public/RetargetEditor/IKRetargetEditorController.h
+++ b/Engine/Plugins/Animation/IKRig/Source/IKRigEditor/Public/RetargetEditor/IKRetargetEditorController.h
@@ -79,65 +79,65 @@ private:
 	bool bWasPlayingAnim = false;
 };
 
-/** a home for cross-widget communication to synchronize state across all tabs and viewport */
+// a home for cross-widget communication to synchronize state across all tabs and viewport
 class FIKRetargetEditorController : public TSharedFromThis<FIKRetargetEditorController>, FGCObject
 {
 public:
 
 	virtual ~FIKRetargetEditorController() override {};
 
-	/** Initialize the editor */
+	// Initialize the editor
 	void Initialize(TSharedPtr<FIKRetargetEditor> InEditor, UIKRetargeter* InAsset);
-	/** Close the editor */
+	// Close the editor
 	void Close();
 	
-	/** Bind callbacks to this IK Rig */
+	// Bind callbacks to this IK Rig
 	void BindToIKRigAssets();
-	/** callback when IK Rig asset requires reinitialization */
+	// callback when IK Rig asset requires reinitialization
 	void HandleIKRigNeedsInitialized(UIKRigDefinition* ModifiedIKRig) const;
-	/** callback when IK Rig asset's retarget chain's have been added or removed*/
+	// callback when IK Rig asset's retarget chain's have been added or removed
 	void HandleRetargetChainAdded(UIKRigDefinition* ModifiedIKRig) const;
-	/** callback when IK Rig asset's retarget chain has been renamed */
+	// callback when IK Rig asset's retarget chain has been renamed
 	void HandleRetargetChainRenamed(UIKRigDefinition* ModifiedIKRig, FName OldName, FName NewName) const;
-	/** callback when IK Rig asset's retarget chain has been removed */
+	// callback when IK Rig asset's retarget chain has been removed
 	void HandleRetargetChainRemoved(UIKRigDefinition* ModifiedIKRig, const FName InChainRemoved) const;
-	/** callback when IK Retargeter asset requires reinitialization */
+	// callback when IK Retargeter asset requires reinitialization
 	void HandleRetargeterNeedsInitialized() const;
 	void ReinitializeProcessor() const;
 	FDelegateHandle RetargeterReInitDelegateHandle;
-	/** callback when IK Rig asset has been swapped out */
+	// callback when IK Rig asset has been swapped out
 	void HandleIKRigReplaced(ERetargetSourceOrTarget SourceOrTarget);
 	FDelegateHandle IKRigReplacedDelegateHandle;
-	/** callback when Preview Mesh asset has been swapped out */
+	// callback when Preview Mesh asset has been swapped out
 	void HandlePreviewMeshReplaced(ERetargetSourceOrTarget SourceOrTarget);
 	FDelegateHandle PreviewMeshReplacedDelegateHandle;
 	
-	/** all modifications to the data model should go through this controller */
+	// all modifications to the data model should go through this controller
 	TObjectPtr<UIKRetargeterController> AssetController;
 
-	/** the persona toolkit */
+	// the persona toolkit
 	TWeakPtr<FIKRetargetEditor> Editor;
 
-	/** import / export retarget poses */
+	// import / export retarget poses
 	TSharedPtr<FIKRetargetPoseExporter> PoseExporter;
 
-	/** manage playback of animation in the editor */
+	// manage playback of animation in the editor
 	TUniquePtr<FRetargetPlaybackManager> PlaybackManager;
 
-	/** viewport skeletal mesh */
+	// viewport skeletal mesh
 	UDebugSkelMeshComponent* GetSkeletalMeshComponent(const ERetargetSourceOrTarget SourceOrTarget) const;
 	UDebugSkelMeshComponent* SourceSkelMeshComponent;
 	UDebugSkelMeshComponent* TargetSkelMeshComponent;
 
-	/** viewport anim instance */
+	// viewport anim instance
 	UIKRetargetAnimInstance* GetAnimInstance(const ERetargetSourceOrTarget SourceOrTarget) const;
 	UPROPERTY(transient, NonTransactional)
 	TObjectPtr<UIKRetargetAnimInstance> SourceAnimInstance;
 	UPROPERTY(transient, NonTransactional)
 	TObjectPtr<UIKRetargetAnimInstance> TargetAnimInstance;
 	
-	/** store pointers to various tabs of UI,
-	 * have to manage access to these because they can be null if the tabs are closed */
+	// store pointers to various tabs of UI,
+	// have to manage access to these because they can be null if the tabs are closed
 	void SetDetailsView(const TSharedPtr<IDetailsView>& InDetailsView) { DetailsView = InDetailsView; };
 	void SetChainsView(const TSharedPtr<SIKRetargetChainMapList>& InChainsView) { ChainsView = InChainsView; };
 	void SetAssetBrowserView(const TSharedPtr<SIKRetargetAssetBrowser>& InAssetBrowserView) { AssetBrowserView = InAssetBrowserView; };
@@ -146,7 +146,7 @@ public:
 	void SetOpStackView(const TSharedPtr<SRetargetOpStack>& InOpStackView) { OpStackView = InOpStackView; };
 	bool IsObjectInDetailsView(const UObject* Object);
 	
-	/** force refresh all views in the editor */
+	// force refresh all views in the editor
 	void RefreshAllViews() const;
 	void RefreshDetailsView() const;
 	void RefreshChainsView() const;
@@ -157,37 +157,37 @@ public:
 	void SetDetailsObject(UObject* DetailsObject) const;
 	void SetDetailsObjects(const TArray<UObject*>& DetailsObjects) const;
 
-	/** retargeter state */
+	// retargeter state
 	bool IsReadyToRetarget() const;
 	bool IsCurrentMeshLoaded() const;
 	bool IsEditingPose() const;
 
-	/** display settings in the details panel */
+	// display settings in the details panel
 	void ShowGlobalSettings();
 	void ShowRootSettings();
 	void ShowPostPhaseSettings();
 	FRetargetGlobalSettings& GetGlobalSettings() const;
 
-	/** clear the output log */
+	// clear the output log
 	void ClearOutputLog() const;
 
-	/** get the USkeletalMesh we are transferring animation between (either source or target)*/
+	// get the USkeletalMesh we are transferring animation between (either source or target)
 	USkeletalMesh* GetSkeletalMesh(const ERetargetSourceOrTarget SourceOrTarget) const;
-	/** get the USkeleton we are transferring animation between (either source or target)*/
+	// get the USkeleton we are transferring animation between (either source or target)
 	const USkeleton* GetSkeleton(const ERetargetSourceOrTarget SourceOrTarget) const;
-	/** get currently edited debug skeletal mesh */
+	// get currently edited debug skeletal mesh
 	UDebugSkelMeshComponent* GetEditedSkeletalMesh() const;
-	/** get the currently edited retarget skeleton */
+	// get the currently edited retarget skeleton
 	const FRetargetSkeleton& GetCurrentlyEditedSkeleton(const UIKRetargetProcessor& Processor) const;
 	
-	/** get world space pose of a bone (with component scale / offset applied) */
+	// get world space pose of a bone (with component scale / offset applied)
 	FTransform GetGlobalRetargetPoseOfBone(
 		const ERetargetSourceOrTarget SourceOrTarget,
 		const int32& BoneIndex,
 		const float& Scale,
 		const FVector& Offset) const;
 	
-	/** get world space positions of all immediate children of bone (with component scale / offset applied) */
+	// get world space positions of all immediate children of bone (with component scale / offset applied)
 	static void GetGlobalRetargetPoseOfImmediateChildren(
 		const FRetargetSkeleton& RetargetSkeleton,
 		const int32& BoneIndex,
@@ -196,12 +196,12 @@ public:
 		TArray<int32>& OutChildrenIndices,
 		TArray<FVector>& OutChildrenPositions);
 
-	/** get the retargeter that is running in the viewport (which is a duplicate of the source asset) */
+	// get the retargeter that is running in the viewport (which is a duplicate of the source asset)
 	UIKRetargetProcessor* GetRetargetProcessor() const;
-	/** Reset the planting state of the IK (when scrubbing or animation loops over) */
+	// Reset the planting state of the IK (when scrubbing or animation loops over)
 	void ResetIKPlantingState() const;
 
-	/** Set viewport / editor tool mode */
+	// Set viewport / editor tool mode
 	void SetRetargeterMode(ERetargeterOutputMode Mode);
 	void SetRetargetModeToPreviousMode() { SetRetargeterMode(PreviousMode); };
 	ERetargeterOutputMode GetRetargeterMode() const { return OutputMode; }
@@ -210,35 +210,35 @@ public:
 	FSlateIcon GetRetargeterModeIcon(ERetargeterOutputMode Mode) const;
 	float GetRetargetPoseAmount() const;
 	void SetRetargetPoseAmount(float InValue);
-	/** END viewport / editor tool mode */
+	// END viewport / editor tool mode
 	
-	/** general editor mode can be either viewing/editing source or target */
+	// general editor mode can be either viewing/editing source or target
 	ERetargetSourceOrTarget GetSourceOrTarget() const { return CurrentlyEditingSourceOrTarget; };
 	void SetSourceOrTargetMode(ERetargetSourceOrTarget SourceOrTarget);
 
-	/** ------------------------- SELECTION -----------------------------*/
+	// ------------------------- SELECTION -----------------------------
 
-	/** SELECTION - BONES (viewport or hierarchy view) */
+	// SELECTION - BONES (viewport or hierarchy view)
 	void EditBoneSelection(
 		const TArray<FName>& InBoneNames,
 		ESelectionEdit EditMode,
 		const bool bFromHierarchyView = false);
 	const TArray<FName>& GetSelectedBones() const {return SelectedBoneNames[CurrentlyEditingSourceOrTarget]; };
-	/* END bone selection */
+	// END bone selection
 	
-	/** SELECTION - MESHES (viewport view) */
+	// SELECTION - MESHES (viewport view)
 	void SetSelectedMesh(UPrimitiveComponent* InComponent);
 	UPrimitiveComponent* GetSelectedMesh() const;
 	void AddOffsetToMeshComponent(const FVector& Offset, USceneComponent* MeshComponent) const;
-	/* END mesh selection */
+	// END mesh selection
 
-	/** SELECTION - CHAINS (viewport or chains view) */
+	// SELECTION - CHAINS (viewport or chains view)
 	void EditChainSelection(
 		const TArray<FName>& InChainNames,
 		ESelectionEdit EditMode,
 		const bool bFromChainsView);
 	const TArray<FName>& GetSelectedChains() const {return SelectedChains; };
-	/* END chain selection */
+	// END chain selection
 
 	void SetRootSelected(const bool bIsSelected);
 	bool GetRootSelected() const { return bIsRootSelected; };
@@ -250,99 +250,99 @@ public:
 	// op stack selection
 	URetargetOpBase* GetSelectedOp() const;
 
-	/** to frame selection when pressing "f" in viewport */
+	// to frame selection when pressing "f" in viewport
 	bool GetCameraTargetForSelection(FSphere& OutTarget) const;
 
-	/** ------------------------- END SELECTION -----------------------------*/
+	// ------------------------- END SELECTION -----------------------------
 
-	/** determine if bone in the specified skeleton is part of the retarget (in a mapped chain) */
+	// determine if bone in the specified skeleton is part of the retarget (in a mapped chain)
 	bool IsBoneRetargeted(const FName& BoneName, ERetargetSourceOrTarget SourceOrTarget) const;
-	/** get the name of the chain that contains this bone */
+	// get the name of the chain that contains this bone
 	FName GetChainNameFromBone(const FName& BoneName, ERetargetSourceOrTarget SourceOrTarget) const;
 
-	/** factory to get/create bone details UObject */
+	// factory to get/create bone details UObject
 	TObjectPtr<UIKRetargetBoneDetails> GetOrCreateBoneDetailsObject(const FName& BoneName);
 
-	/** ------------------------- RETARGET POSES -----------------------------*/
+	// ------------------------- RETARGET POSES -----------------------------
 	
-	/** toggle current retarget pose */
+	// toggle current retarget pose
 	TArray<TSharedPtr<FName>> PoseNames;
 	FText GetCurrentPoseName() const;
 	void OnPoseSelected(TSharedPtr<FName> InPoseName, ESelectInfo::Type SelectInfo) const;
 
-	/** reset retarget pose */
+	// reset retarget pose
 	void HandleResetAllBones() const;
 	void HandleResetSelectedBones() const;
 	void HandleResetSelectedAndChildrenBones() const;
 	bool CanResetSelected() const;
 
-	/** create new retarget pose */
+	// create new retarget pose
 	void HandleNewPose();
 	bool CanCreatePose() const;
 	FReply CreateNewPose() const;
 	TSharedPtr<SWindow> NewPoseWindow;
 	TSharedPtr<SEditableTextBox> NewPoseEditableText;
 
-	/** duplicate current retarget pose */
+	// duplicate current retarget pose
 	void HandleDuplicatePose();
 	FReply CreateDuplicatePose() const;
 
-	/** delete retarget pose */
+	// delete retarget pose
 	void HandleDeletePose();
 	bool CanDeletePose() const;
 
-	/** rename retarget pose */
+	// rename retarget pose
 	void HandleRenamePose();
 	FReply RenamePose() const;
 	bool CanRenamePose() const;
 	TSharedPtr<SWindow> RenamePoseWindow;
 	TSharedPtr<SEditableTextBox> NewNameEditableText;
 
-	/* detect and auto-fix retarget pose that causes root height to be on the ground */
+	// detect and auto-fix retarget pose that causes root height to be on the ground
 	void FixZeroHeightRetargetRoot(ERetargetSourceOrTarget SourceOrTarget) const;
 	
-	/** ------------------------- END RETARGET POSES -----------------------------*/
+	// ------------------------- END RETARGET POSES -----------------------------
 
-	/** FGCObject interface */
+	// FGCObject interface
 	virtual void AddReferencedObjects(FReferenceCollector& Collector) override;
 	virtual FString GetReferencerName() const override { return TEXT("Retarget Editor"); };
-	/** END FGCObject interface */
+	// END FGCObject interface
 
-	/** render the skeleton's in the viewport (either source or target) */
+	// render the skeleton's in the viewport (either source or target)
 	void RenderSkeleton(FPrimitiveDrawInterface* PDI, ERetargetSourceOrTarget SourceOrTarget) const;
 
 private:
 
-	/** modal dialog to ask user if they want to fix root bones that are "on the ground" */
+	// modal dialog to ask user if they want to fix root bones that are "on the ground"
 	bool PromptToFixRootHeight(ERetargetSourceOrTarget SourceOrTarget) const;
 	bool bAskedToFixRoot = false;
 
-	/** asset properties tab */
+	// asset properties tab
 	TSharedPtr<IDetailsView> DetailsView;
-	/** chain list view */
+	// chain list view
 	TSharedPtr<SIKRetargetChainMapList> ChainsView;
-	/** asset browser view */
+	// asset browser view
 	TSharedPtr<SIKRetargetAssetBrowser> AssetBrowserView;
-	/** output log view */
+	// output log view
 	TSharedPtr<SIKRigOutputLog> OutputLogView;
-	/** hierarchy view */
+	// hierarchy view
 	TSharedPtr<SIKRetargetHierarchy> HierarchyView;
-	/** op stack widget */
+	// op stack widget
 	TSharedPtr<SRetargetOpStack> OpStackView;
 
-	/** when prompting user to assign an IK Rig */
+	// when prompting user to assign an IK Rig
 	TSharedPtr<SWindow> IKRigPickerWindow;
 
-	/** the current output mode of the retargeter */
+	// the current output mode of the retargeter
 	ERetargeterOutputMode OutputMode = ERetargeterOutputMode::RunRetarget;
 	ERetargeterOutputMode PreviousMode;
-	/** slider value to blend between reference pose and retarget pose */
+	// slider value to blend between reference pose and retarget pose
 	float RetargetPosePreviewBlend = 1.0f;
 	
-	/** which skeleton are we editing / viewing? */
+	// which skeleton are we editing / viewing?
 	ERetargetSourceOrTarget CurrentlyEditingSourceOrTarget = ERetargetSourceOrTarget::Target;
 
-	/** current selection set */
+	// current selection set
 	bool bIsRootSelected = false;
 	UPrimitiveComponent* SelectedMesh = nullptr;
 	TArray<FName> SelectedChains;
@@ -351,7 +351,7 @@ private:
 	UPROPERTY()
 	TMap<FName,TObjectPtr<UIKRetargetBoneDetails>> AllBoneDetails;
 
-	/** ik rigs bound to this editor (will receive callbacks when requiring reinitialization*/
+	// ik rigs bound to this editor (will receive callbacks when requiring reinitialization
 	TArray<FBoundIKRig> BoundIKRigs;
 };
 
