Fixed issue, IK Rig editor quit animation playback when reinitializing the processor.

diff --git a/Engine/Plugins/Animation/IKRig/Source/IKRigEditor/Private/RigEditor/IKRigEditorController.cpp b/Engine/Plugins/Animation/IKRig/Source/IKRigEditor/Private/RigEditor/IKRigEditorController.cpp
--- a/Engine/Plugins/Animation/IKRig/Source/IKRigEditor/Private/RigEditor/IKRigEditorController.cpp
+++ b/Engine/Plugins/Animation/IKRig/Source/IKRigEditor/Private/RigEditor/IKRigEditorController.cpp
@@ -350,9 +350,10 @@ void FIKRigEditorController::HandleIKRigNeedsInitialized(UIKRigDefinition* Modif
 	
 	// in case the skeletal mesh was swapped out, we need to ensure the preview scene is up-to-date
 	USkeletalMesh* NewMesh = AssetController->GetSkeletalMesh();
-	SkelMeshComponent->SetSkeletalMesh(NewMesh);
 	const TSharedRef<IPersonaPreviewScene> PreviewScene = EditorToolkit.Pin()->GetPersonaToolkit()->GetPreviewScene();
-	if (PreviewScene->GetPreviewMesh() != NewMesh)
+	const bool bMeshWasSwapped = PreviewScene->GetPreviewMesh() != NewMesh;
+	SkelMeshComponent->SetSkeletalMesh(NewMesh);
+	if (bMeshWasSwapped)
 	{
 		PreviewScene->SetPreviewMeshComponent(SkelMeshComponent);
 		PreviewScene->SetPreviewMesh(NewMesh);
@@ -361,9 +362,22 @@ void FIKRigEditorController::HandleIKRigNeedsInitialized(UIKRigDefinition* Modif
 	// re-initializes the anim instances running in the viewport
 	if (AnimInstance)
 	{
+		// record what anim was playing so we can restore it after reinit
+		UAnimationAsset* AnimThatWasPlaying = AnimInstance->GetCurrentAsset();
+		const float TimeWhenReset = AnimInstance->GetCurrentTime();
+		const bool bWasPlaying = AnimInstance->IsPlaying();
+		
 		SkelMeshComponent->PreviewInstance = AnimInstance;
 		AnimInstance->InitializeAnimation();
 		SkelMeshComponent->EnablePreview(true, nullptr);
+
+		// restore previously playing asset
+		if (AnimThatWasPlaying)
+		{
+			AnimInstance->SetAnimationAsset(AnimThatWasPlaying);
+			AnimInstance->SetPlaying(bWasPlaying);
+			AnimInstance->SetPosition(TimeWhenReset);
+		}
 	}
 
 	// update the bone details so it can pull on the current data
@@ -392,7 +406,7 @@ void FIKRigEditorController::HandleIKRigNeedsInitialized(UIKRigDefinition* Modif
 			RetargetingView->RefreshView();
 		}
 		
-		if (AssetBrowserView.IsValid())
+		if (bMeshWasSwapped && AssetBrowserView.IsValid())
 		{
 			AssetBrowserView.Get()->RefreshView(); // refresh the asset browser to ensure it shows compatible sequences
 		}
