Adding ability to explicitly control the root motion lock state when exporting animation from the retargeter.

diff --git a/Engine/Plugins/Animation/IKRig/Source/IKRig/Public/Retargeter/IKRetargeter.h b/Engine/Plugins/Animation/IKRig/Source/IKRig/Public/Retargeter/IKRetargeter.h
--- a/Engine/Plugins/Animation/IKRig/Source/IKRig/Public/Retargeter/IKRetargeter.h
+++ b/Engine/Plugins/Animation/IKRig/Source/IKRig/Public/Retargeter/IKRetargeter.h
@@ -206,6 +206,18 @@ enum class ERetargetSourceOrTarget : uint8
 	Target, // the TARGET skeleton (to copy TO)
 };
 
+// which skeleton are we referring to?
+UENUM()
+enum class ERetargetRootLockMode : uint8
+{
+	// Uses the "ForceRootLock" setting in the source animation
+	FromSourceAnimation,
+	// Force the root to be locked, regardless of whether the source animation has the root locked or not.
+	ForceRootLocked,
+	// Force the root to be unlocked, regardless of whether the source animation has the root locked or not.
+	ForceRootUnlocked,
+};
+
 UCLASS(BlueprintType)
 class IKRIG_API UIKRetargeter : public UObject
 {
@@ -412,10 +424,17 @@ public:
 	FVector SourceMeshOffset;
 
 	// When true, animation sequences with "Force Root Lock" turned On will act as though it is Off.
-	// This affects the preview in the retarget editor, as well as exported animation sequences.
+	// This affects only the preview in the retarget editor. Use ExportRootLockMode to control exported animation behavior.
+	// This setting has no effect on runtime retargeting where root motion is copied from the source component.
+	UPROPERTY(EditAnywhere, Category = RootLockSettings)
+	bool bIgnoreRootLockInPreview = true;
+	
+	// Either uses the "ForceRootLock" setting from the source animation or optionally forces the root to be locked or unlocked,
+	// regardless of whether the source animation has the root locked or not.
+	// This affects only the exported animation sequences, not the editor preview.
 	// This setting has no effect on runtime retargeting where root motion is copied from the source component.
-	UPROPERTY(EditAnywhere, Category = ExportSettings)
-	bool bIgnoreRootLock = true;
+	UPROPERTY(EditAnywhere, Category = RootLockSettings)
+	ERetargetRootLockMode ExportRootLockMode;
 
 	// Toggle debug drawing for retargeting in the viewport. 
 	UPROPERTY(EditAnywhere, Category = DebugSettings)
diff --git a/Engine/Plugins/Animation/IKRig/Source/IKRigEditor/Private/RetargetEditor/IKRetargetAnimInstanceProxy.cpp b/Engine/Plugins/Animation/IKRig/Source/IKRigEditor/Private/RetargetEditor/IKRetargetAnimInstanceProxy.cpp
--- a/Engine/Plugins/Animation/IKRig/Source/IKRigEditor/Private/RetargetEditor/IKRetargetAnimInstanceProxy.cpp
+++ b/Engine/Plugins/Animation/IKRig/Source/IKRigEditor/Private/RetargetEditor/IKRetargetAnimInstanceProxy.cpp
@@ -44,7 +44,7 @@ bool FIKRetargetAnimInstanceProxy::Evaluate(FPoseContext& Output)
 {
 	if (PreviewPoseNode->IKRetargeterAsset)
 	{
-		bIgnoreRootLock = PreviewPoseNode->IKRetargeterAsset->bIgnoreRootLock;	
+		bIgnoreRootLock = PreviewPoseNode->IKRetargeterAsset->bIgnoreRootLockInPreview;	
 	}
 	
 	switch (OutputMode)
diff --git a/Engine/Plugins/Animation/IKRig/Source/IKRigEditor/Private/RetargetEditor/IKRetargetBatchOperation.cpp b/Engine/Plugins/Animation/IKRig/Source/IKRigEditor/Private/RetargetEditor/IKRetargetBatchOperation.cpp
--- a/Engine/Plugins/Animation/IKRig/Source/IKRigEditor/Private/RetargetEditor/IKRetargetBatchOperation.cpp
+++ b/Engine/Plugins/Animation/IKRig/Source/IKRigEditor/Private/RetargetEditor/IKRetargetBatchOperation.cpp
@@ -127,6 +127,15 @@ void UIKRetargetBatchOperation::DuplicateRetargetAssets(
 		
 		TMap<UAnimationAsset*, UAnimationAsset*> DuplicateMap = DuplicateAssets<UAnimationAsset>({Asset}, DestinationPackage, &Context.NameRule);
 		DuplicatedAnimAssets.Append(DuplicateMap);
+
+		// optionally let user override root lock on exported animation sequences,
+		// (by default it will inherit the bForceRootLock state from the duplicated source animation)
+		const ERetargetRootLockMode RootLockMode = Context.IKRetargetAsset->ExportRootLockMode;
+		UAnimSequence* TargetSequence = Cast<UAnimSequence>(DuplicateMap[Asset]);
+		if (RootLockMode != ERetargetRootLockMode::FromSourceAnimation && TargetSequence)
+		{
+			TargetSequence->bForceRootLock = RootLockMode == ERetargetRootLockMode::ForceRootLocked ? true : false;
+		}
 	}
 	for (UAnimBlueprint* Asset : AnimBlueprintsToDuplicate)
 	{
@@ -376,10 +385,9 @@ void UIKRetargetBatchOperation::ConvertAnimation(
 		// ensure we evaluate the source animation using the skeletal mesh proportions that were evaluated in the viewport
 		FAnimPoseEvaluationOptions EvaluationOptions = FAnimPoseEvaluationOptions();
 		EvaluationOptions.OptionalSkeletalMesh = SourceSkeleton.SkeletalMesh;
-
-		// optionally ignore root lock
-		const bool bIgnoreRootLock = Context.IKRetargetAsset->bIgnoreRootLock;
-		TGuardValue<bool> RootLockGuard(SourceSequence->bForceRootLock, bIgnoreRootLock ? false : SourceSequence->bForceRootLock);
+		// ensure WYSIWYG with editor by ensuring the same root motion is evaluated during export as during the editor preview
+		EvaluationOptions.bExtractRootMotion = !TargetSequence->bForceRootLock;
+		EvaluationOptions.bIncorporateRootMotionIntoPose = !TargetSequence->bForceRootLock;
 
 		// reset the planting state
 		Processor->ResetPlanting();
