Fixes crash when reloading IK Rig asset while referening retarget editor is open.

diff --git a/Engine/Plugins/Animation/IKRig/Source/IKRigEditor/Private/RetargetEditor/IKRetargetEditorController.cpp b/Engine/Plugins/Animation/IKRig/Source/IKRigEditor/Private/RetargetEditor/IKRetargetEditorController.cpp
--- a/Engine/Plugins/Animation/IKRig/Source/IKRigEditor/Private/RetargetEditor/IKRetargetEditorController.cpp
+++ b/Engine/Plugins/Animation/IKRig/Source/IKRigEditor/Private/RetargetEditor/IKRetargetEditorController.cpp
@@ -44,12 +44,12 @@ FBoundIKRig::FBoundIKRig(UIKRigDefinition* InIKRig, const FIKRetargetEditorContr
 
 void FBoundIKRig::UnBind() const
 {
-	if (!IsValid(IKRig))
+	if (!IKRig.IsValid())
 	{
 		return;
 	}
 	
-	UIKRigController* IKRigController = UIKRigController::GetController(IKRig);
+	UIKRigController* IKRigController = UIKRigController::GetController(IKRig.Get());
 	IKRigController->OnIKRigNeedsInitialized().Remove(ReInitIKDelegateHandle);
 	IKRigController->OnRetargetChainAdded().Remove(AddedChainDelegateHandle);
 	IKRigController->OnRetargetChainRemoved().Remove(RemoveChainDelegateHandle);
diff --git a/Engine/Plugins/Animation/IKRig/Source/IKRigEditor/Public/RetargetEditor/IKRetargetEditorController.h b/Engine/Plugins/Animation/IKRig/Source/IKRigEditor/Public/RetargetEditor/IKRetargetEditorController.h
--- a/Engine/Plugins/Animation/IKRig/Source/IKRigEditor/Public/RetargetEditor/IKRetargetEditorController.h
+++ b/Engine/Plugins/Animation/IKRig/Source/IKRigEditor/Public/RetargetEditor/IKRetargetEditorController.h
@@ -55,7 +55,7 @@ struct FBoundIKRig
 	FBoundIKRig(UIKRigDefinition* InIKRig, const FIKRetargetEditorController& InController);
 	void UnBind() const;
 	
-	UIKRigDefinition* IKRig;
+	TWeakObjectPtr<UIKRigDefinition> IKRig;
 	FDelegateHandle ReInitIKDelegateHandle;
 	FDelegateHandle AddedChainDelegateHandle;
 	FDelegateHandle RenameChainDelegateHandle;
