Removed unnecessary ensure in bone editing transform widget.

diff --git a/Engine/Plugins/Animation/IKRig/Source/IKRigEditor/Private/RetargetEditor/IKRetargetDetails.cpp b/Engine/Plugins/Animation/IKRig/Source/IKRigEditor/Private/RetargetEditor/IKRetargetDetails.cpp
--- a/Engine/Plugins/Animation/IKRig/Source/IKRigEditor/Private/RetargetEditor/IKRetargetDetails.cpp
+++ b/Engine/Plugins/Animation/IKRig/Source/IKRigEditor/Private/RetargetEditor/IKRetargetDetails.cpp
@@ -127,8 +127,6 @@ FEulerTransform UIKRetargetBoneDetails::GetTransform(EIKRetargetTransformType Tr
 
 	case EIKRetargetTransformType::Bone:
 		{
-			ensure(bLocalSpace);
-			
 			// this is the only stored data we have for bone pose offsets
 			const ERetargetSourceOrTarget SourceOrTarget = EditorController->GetSourceOrTarget();
 			const FQuat LocalRotationOffset = EditorController->AssetController->GetRotationOffsetForRetargetPoseBone(SelectedBone, SourceOrTarget);
@@ -653,27 +651,22 @@ void UIKRetargetBoneDetails::CommitValueAsBoneSpace(
 		}
 	case ESlateTransformComponent::Rotation:
 		{
-			// only allow local bone space pose editing for now
-			const bool bIsRotationLocal = BoneRelative[1];
-			if (ensure(bIsRotationLocal))
-			{
-				// combine the local space offset from ref pose with the recorded offset in the retarget pose
-				FTransform LocalRefTransform = RefSkeleton.GetRefBonePose()[BoneIndex];
-				const FQuat LocalRotationOffset = AssetController->GetRotationOffsetForRetargetPoseBone(SelectedBone, SourceOrTarget);
-				FQuat CombinedLocalRotation = LocalRefTransform.GetRotation() * LocalRotationOffset;
-				FEulerTransform CombinedLocalDeltaTransform = FEulerTransform(FVector::ZeroVector, CombinedLocalRotation.Rotator(), FVector::OneVector);
-			
-				// rotations are stored in local space, so just apply the edit
-				SAdvancedTransformInputBox<FEulerTransform>::ApplyNumericValueChange(CombinedLocalDeltaTransform, Value, Component, Representation, SubComponent);
+			// combine the local space offset from ref pose with the recorded offset in the retarget pose
+			FTransform LocalRefTransform = RefSkeleton.GetRefBonePose()[BoneIndex];
+			const FQuat LocalRotationOffset = AssetController->GetRotationOffsetForRetargetPoseBone(SelectedBone, SourceOrTarget);
+			FQuat CombinedLocalRotation = LocalRefTransform.GetRotation() * LocalRotationOffset;
+			FEulerTransform CombinedLocalDeltaTransform = FEulerTransform(FVector::ZeroVector, CombinedLocalRotation.Rotator(), FVector::OneVector);
+		
+			// rotations are stored in local space, so just apply the edit
+			SAdvancedTransformInputBox<FEulerTransform>::ApplyNumericValueChange(CombinedLocalDeltaTransform, Value, Component, Representation, SubComponent);
 
-				// subtract the local space from the result to be left with JUST the retarget pose offset
-				FQuat NewLocalRotationDelta = LocalRefTransform.GetRotation().Inverse() * CombinedLocalDeltaTransform.GetRotation();
-			
-				// store the new rotation in the retarget pose
-				FScopedTransaction Transaction(LOCTEXT("EditRootRotation", "Edit Retarget Pose Rotation"), bShouldTransact);
-				EditorController->AssetController->GetAsset()->Modify();
-				EditorController->AssetController->SetRotationOffsetForRetargetPoseBone(SelectedBone, NewLocalRotationDelta, SourceOrTarget);
-			}
+			// subtract the local space from the result to be left with JUST the retarget pose offset
+			FQuat NewLocalRotationDelta = LocalRefTransform.GetRotation().Inverse() * CombinedLocalDeltaTransform.GetRotation();
+		
+			// store the new rotation in the retarget pose
+			FScopedTransaction Transaction(LOCTEXT("EditRootRotation", "Edit Retarget Pose Rotation"), bShouldTransact);
+			EditorController->AssetController->GetAsset()->Modify();
+			EditorController->AssetController->SetRotationOffsetForRetargetPoseBone(SelectedBone, NewLocalRotationDelta, SourceOrTarget);
 			
 			break;
 		}
