Changed URetargetOpBase to remove pure virtual functions based on feedback.

diff --git a/Engine/Plugins/Animation/IKRig/Source/IKRig/Public/Retargeter/IKRetargetOps.h b/Engine/Plugins/Animation/IKRig/Source/IKRig/Public/Retargeter/IKRetargetOps.h
--- a/Engine/Plugins/Animation/IKRig/Source/IKRig/Public/Retargeter/IKRetargetOps.h
+++ b/Engine/Plugins/Animation/IKRig/Source/IKRig/Public/Retargeter/IKRetargetOps.h
@@ -24,13 +24,13 @@ public:
 	const UIKRetargetProcessor* Processor,
 		const FRetargetSkeleton& SourceSkeleton,
 		const FTargetSkeleton& TargetSkeleton,
-		FIKRigLogger& Log) PURE_VIRTUAL(, return false;);
+		FIKRigLogger& Log) { return false; };
 
 	// override to evaluate this operation and modify the output pose
 	virtual void Run(
 		const UIKRetargetProcessor* Processor,
 		const TArray<FTransform>& InSourceGlobalPose,
-		TArray<FTransform>& OutTargetGlobalPose) PURE_VIRTUAL(,);
+		TArray<FTransform>& OutTargetGlobalPose){};
 
 	UPROPERTY()
 	bool bIsEnabled = true;
@@ -41,7 +41,7 @@ public:
 	// override to automate initial setup after being added to the stack
 	virtual void OnAddedToStack(const UIKRetargeter* Asset) {};
 	// override to give your operation a nice name to display in the UI
-	virtual FText GetNiceName() const { checkNoEntry() return FText::GetEmpty(); };
+	virtual FText GetNiceName() const { return FText::FromString(TEXT("Default Op Name")); };
 	// override to display a warning message in the op stack
 	virtual FText WarningMessage() const { return FText::GetEmpty(); };
 #endif
