Minor naming issues in new retarget op code.

diff --git a/Engine/Plugins/Animation/IKRig/Source/IKRigEditor/Private/RetargetEditor/IKRetargeterController.cpp b/Engine/Plugins/Animation/IKRig/Source/IKRigEditor/Private/RetargetEditor/IKRetargeterController.cpp
--- a/Engine/Plugins/Animation/IKRig/Source/IKRigEditor/Private/RetargetEditor/IKRetargeterController.cpp
+++ b/Engine/Plugins/Animation/IKRig/Source/IKRigEditor/Private/RetargetEditor/IKRetargeterController.cpp
@@ -208,7 +208,7 @@ int32 UIKRetargeterController::AddRetargetOp(TSubclassOf<URetargetOpBase> InOpCl
 
 	if (!InOpClass)
 	{
-		UE_LOG(LogIKRigEditor, Warning, TEXT("Could not add solver to IK Rig. Invalid solver class specified."));
+		UE_LOG(LogIKRigEditor, Warning, TEXT("Could not add Op to stack. Invalid Op class specified."));
 		return INDEX_NONE;
 	}
 
@@ -288,8 +288,8 @@ bool UIKRetargeterController::MoveRetargetOpInStack(int32 OpToMoveIndex, int32 T
 	Asset->OpStack->Modify();
 	URetargetOpBase* OpToMove = RetargetOps[OpToMoveIndex];
 	RetargetOps.Insert(OpToMove, TargetIndex + 1);
-	const int32 SolverToRemove = TargetIndex > OpToMoveIndex ? OpToMoveIndex : TargetIndex + 1;
-	RetargetOps.RemoveAt(SolverToRemove);
+	const int32 OpToRemove = TargetIndex > OpToMoveIndex ? OpToMoveIndex : TargetIndex + 1;
+	RetargetOps.RemoveAt(OpToRemove);
 	return true;
 }
 
@@ -301,7 +301,7 @@ bool UIKRetargeterController::SetRetargetOpEnabled(int32 RetargetOpIndex, bool b
 		return false;
 	}
 	
-	FScopedTransaction Transaction(LOCTEXT("SetRetargetOpEnabled_Label", "Enable/Disable Solver"));
+	FScopedTransaction Transaction(LOCTEXT("SetRetargetOpEnabled_Label", "Enable/Disable Op"));
 	FScopedReinitializeIKRetargeter Reinitialize(this);
 	URetargetOpBase* OpToMove = Asset->OpStack->RetargetOps[RetargetOpIndex];
 	OpToMove->Modify();
diff --git a/Engine/Plugins/Animation/IKRig/Source/IKRigEditor/Public/RetargetEditor/IKRetargeterController.h b/Engine/Plugins/Animation/IKRig/Source/IKRigEditor/Public/RetargetEditor/IKRetargeterController.h
--- a/Engine/Plugins/Animation/IKRig/Source/IKRigEditor/Public/RetargetEditor/IKRetargeterController.h
+++ b/Engine/Plugins/Animation/IKRig/Source/IKRigEditor/Public/RetargetEditor/IKRetargeterController.h
@@ -118,7 +118,7 @@ public:
 	UFUNCTION(BlueprintCallable, Category=RetargetOps)
 	int32 GetIndexOfRetargetOp(URetargetOpBase* RetargetOp) const;
 
-	// Get the number of solvers in the stack.
+	// Get the number of Ops in the stack.
 	UFUNCTION(BlueprintCallable, Category=RetargetOps)
 	int32 GetNumRetargetOps() const;
 
@@ -130,7 +130,7 @@ public:
 	UFUNCTION(BlueprintCallable, BlueprintPure=false, Category=RetargetOps)
 	bool SetRetargetOpEnabled(int32 RetargetOpIndex, bool bIsEnabled) const;
 
-	// Get enabled status of the given solver. 
+	// Get enabled status of the given Op. 
 	UFUNCTION(BlueprintCallable, Category=RetargetOps)
 	bool GetRetargetOpEnabled(int32 RetargetOpIndex) const;
 
