Fixed issue retarget ops could only be moved down, not up in stack.

diff --git a/Engine/Plugins/Animation/IKRig/Source/IKRigEditor/Private/RetargetEditor/IKRetargeterController.cpp b/Engine/Plugins/Animation/IKRig/Source/IKRigEditor/Private/RetargetEditor/IKRetargeterController.cpp
--- a/Engine/Plugins/Animation/IKRig/Source/IKRigEditor/Private/RetargetEditor/IKRetargeterController.cpp
+++ b/Engine/Plugins/Animation/IKRig/Source/IKRigEditor/Private/RetargetEditor/IKRetargeterController.cpp
@@ -288,7 +288,7 @@ bool UIKRetargeterController::MoveRetargetOpInStack(int32 OpToMoveIndex, int32 T
 	Asset->OpStack->Modify();
 	URetargetOpBase* OpToMove = RetargetOps[OpToMoveIndex];
 	RetargetOps.Insert(OpToMove, TargetIndex + 1);
-	const int32 OpToRemove = TargetIndex > OpToMoveIndex ? OpToMoveIndex : TargetIndex + 1;
+	const int32 OpToRemove = TargetIndex > OpToMoveIndex ? OpToMoveIndex : OpToMoveIndex + 1;
 	RetargetOps.RemoveAt(OpToRemove);
 	return true;
 }
