Fixed issue in IK Rig editor, goal locations do not update when swapping to a different skeletal mesh.

diff --git a/Engine/Plugins/Animation/IKRig/Source/IKRigEditor/Private/RigEditor/IKRigEditorController.cpp b/Engine/Plugins/Animation/IKRig/Source/IKRigEditor/Private/RigEditor/IKRigEditorController.cpp
--- a/Engine/Plugins/Animation/IKRig/Source/IKRigEditor/Private/RigEditor/IKRigEditorController.cpp
+++ b/Engine/Plugins/Animation/IKRig/Source/IKRigEditor/Private/RigEditor/IKRigEditorController.cpp
@@ -357,6 +357,7 @@ void FIKRigEditorController::HandleIKRigNeedsInitialized(UIKRigDefinition* Modif
 	{
 		PreviewScene->SetPreviewMeshComponent(SkelMeshComponent);
 		PreviewScene->SetPreviewMesh(NewMesh);
+		AssetController->ResetGoalTransforms();
 	}
 
 	// re-initializes the anim instances running in the viewport
