Fixes crash when adding virtual bone to a skeletal mesh while a retargeter using that mesh is currently running.

diff --git a/Engine/Plugins/Animation/IKRig/Source/IKRig/Private/AnimNodes/AnimNode_RetargetPoseFromMesh.cpp b/Engine/Plugins/Animation/IKRig/Source/IKRig/Private/AnimNodes/AnimNode_RetargetPoseFromMesh.cpp
--- a/Engine/Plugins/Animation/IKRig/Source/IKRig/Private/AnimNodes/AnimNode_RetargetPoseFromMesh.cpp
+++ b/Engine/Plugins/Animation/IKRig/Source/IKRig/Private/AnimNodes/AnimNode_RetargetPoseFromMesh.cpp
@@ -61,6 +61,12 @@ void FAnimNode_RetargetPoseFromMesh::CacheBones_AnyThread(const FAnimationCacheB
 			}
 		}
 	}
+
+	// any time required bones are modified, retargeter needs reinitialized
+	if (Processor)
+	{
+		Processor->SetNeedsInitialized();
+	}
 }
 
 void FAnimNode_RetargetPoseFromMesh::Update_AnyThread(const FAnimationUpdateContext& Context)
