Added ability to auto generate retarget pose in IK Retargeter.

diff --git a/Engine/Plugins/Animation/IKRig/Source/IKRig/Private/Retargeter/IKRetargetProcessor.cpp b/Engine/Plugins/Animation/IKRig/Source/IKRig/Private/Retargeter/IKRetargetProcessor.cpp
--- a/Engine/Plugins/Animation/IKRig/Source/IKRig/Private/Retargeter/IKRetargetProcessor.cpp
+++ b/Engine/Plugins/Animation/IKRig/Source/IKRig/Private/Retargeter/IKRetargetProcessor.cpp
@@ -453,6 +453,48 @@ bool FChainFK::Initialize(
 	return CalculateBoneParameters(Log);
 }
 
+FTransform FChainFK::GetTransformAtParam(const TArray<FTransform>& Transforms, const float& Param) const
+{
+	check(Transforms.Num() == Params.Num());
+	
+	if (Params.Num() == 1)
+	{
+		return Transforms[0];
+	}
+	
+	if (Param < KINDA_SMALL_NUMBER)
+	{
+		return Transforms[0];
+	}
+
+	if (Param > 1.0f - KINDA_SMALL_NUMBER)
+	{
+		return Transforms.Last();
+	}
+
+	for (int32 ChainIndex=1; ChainIndex<Params.Num(); ++ChainIndex)
+	{
+		const float CurrentParam = Params[ChainIndex];
+		if (CurrentParam <= Param)
+		{
+			continue;
+		}
+		
+		const float PrevParam = Params[ChainIndex-1];
+		const float PercentBetweenParams = (Param - PrevParam) / (CurrentParam - PrevParam);
+		const FTransform& Prev = Transforms[ChainIndex-1];
+		const FTransform& Next = Transforms[ChainIndex];
+		const FVector Position = FMath::Lerp(Prev.GetTranslation(), Next.GetTranslation(), PercentBetweenParams);
+		const FQuat Rotation = FQuat::FastLerp(Prev.GetRotation(), Next.GetRotation(), PercentBetweenParams).GetNormalized();
+		const FVector Scale = FMath::Lerp(Prev.GetScale3D(), Next.GetScale3D(), PercentBetweenParams);
+		
+		return FTransform(Rotation,Position, Scale);
+	}
+
+	checkNoEntry();
+	return FTransform::Identity;
+}
+
 bool FChainFK::CalculateBoneParameters(FIKRigLogger& Log)
 {
 	Params.Reset();
@@ -641,16 +683,8 @@ void FChainDecoderFK::DecodePose(
 				// get the initial and current transform of source chain at param
 				// this is the interpolated transform along the chain
 				const float Param = Params[ChainIndex];
-					
-				SourceCurrentTransform = GetTransformAtParam(
-					SourceChain.CurrentGlobalTransforms,
-					SourceChain.Params,
-					Param);
-
-				SourceInitialTransform = GetTransformAtParam(
-					SourceChain.InitialGlobalTransforms,
-					SourceChain.Params,
-					Param);
+				SourceCurrentTransform = SourceChain.GetTransformAtParam(SourceChain.CurrentGlobalTransforms, Param);
+				SourceInitialTransform = SourceChain.GetTransformAtParam(SourceChain.InitialGlobalTransforms, Param);
 			}
 			break;
 			case ERetargetRotationMode::OneToOne:
@@ -806,49 +840,6 @@ void FChainDecoderFK::UpdateIntermediateParents(
 	}
 }
 
-FTransform FChainDecoderFK::GetTransformAtParam(
-	const TArray<FTransform>& Transforms,
-	const TArray<float>& InParams,
-	const float& Param) const
-{
-	if (InParams.Num() == 1)
-	{
-		return Transforms[0];
-	}
-	
-	if (Param < KINDA_SMALL_NUMBER)
-	{
-		return Transforms[0];
-	}
-
-	if (Param > 1.0f - KINDA_SMALL_NUMBER)
-	{
-		return Transforms.Last();
-	}
-
-	for (int32 ChainIndex=1; ChainIndex<InParams.Num(); ++ChainIndex)
-	{
-		const float CurrentParam = InParams[ChainIndex];
-		if (CurrentParam <= Param)
-		{
-			continue;
-		}
-		
-		const float PrevParam = InParams[ChainIndex-1];
-		const float PercentBetweenParams = (Param - PrevParam) / (CurrentParam - PrevParam);
-		const FTransform& Prev = Transforms[ChainIndex-1];
-		const FTransform& Next = Transforms[ChainIndex];
-		const FVector Position = FMath::Lerp(Prev.GetTranslation(), Next.GetTranslation(), PercentBetweenParams);
-		const FQuat Rotation = FQuat::FastLerp(Prev.GetRotation(), Next.GetRotation(), PercentBetweenParams).GetNormalized();
-		const FVector Scale = FMath::Lerp(Prev.GetScale3D(), Next.GetScale3D(), PercentBetweenParams);
-		
-		return FTransform(Rotation,Position, Scale);
-	}
-
-	checkNoEntry();
-	return FTransform::Identity;
-}
-
 bool FChainRetargeterIK::InitializeSource(
 	const TArray<int32>& BoneIndices,
 	const TArray<FTransform>& SourceInitialGlobalPose,
@@ -1903,7 +1894,7 @@ TArray<FTransform>&  UIKRetargetProcessor::RunRetargeter(
 	}
 	
 	// IK CHAIN retargeting
-	if (GlobalSettings.bEnableIK && bAtLeastOneValidBoneChainPair && bIKRigInitialized)
+	if (bAtLeastOneValidBoneChainPair && bIKRigInitialized)
 	{
 		RunIKRetarget(InSourceGlobalPose, TargetSkeleton.OutputGlobalPose, SpeedValuesFromCurves, DeltaTime);
 	}
@@ -2002,12 +1993,15 @@ void UIKRetargetProcessor::RunIKRetarget(
 		ChainPair.IKChainRetargeter.SaveDebugInfo(OutTargetGlobalPose);
 	}
 
-	// copy input pose to start IK solve from
-	IKRigProcessor->SetInputPoseGlobal(OutTargetGlobalPose);
-	// run IK solve
-	IKRigProcessor->Solve();
-	// copy results of solve
-	IKRigProcessor->CopyOutputGlobalPoseToArray(OutTargetGlobalPose);
+	if (GlobalSettings.bEnableIK)
+	{
+		// copy input pose to start IK solve from
+		IKRigProcessor->SetInputPoseGlobal(OutTargetGlobalPose);
+		// run IK solve
+		IKRigProcessor->Solve();
+		// copy results of solve
+		IKRigProcessor->CopyOutputGlobalPoseToArray(OutTargetGlobalPose);
+	}
 }
 
 void UIKRetargetProcessor::RunPoleVectorMatching(
@@ -2173,8 +2167,8 @@ bool UIKRetargetProcessor::WasInitializedWithTheseAssets(
 		return false;
 	}
 
-	const bool bSourceMatches = InSourceMesh == GetSourceSkeleton().SkeletalMesh;
-	const bool bTargetMatches = InTargetMesh == GetTargetSkeleton().SkeletalMesh;
+	const bool bSourceMatches = InSourceMesh == GetSkeleton(ERetargetSourceOrTarget::Source).SkeletalMesh;
+	const bool bTargetMatches = InTargetMesh == GetSkeleton(ERetargetSourceOrTarget::Target).SkeletalMesh;
 	const bool bAssetMatches = InRetargetAsset == RetargeterAsset;
 	
 	return bSourceMatches && bTargetMatches && bAssetMatches;
@@ -2193,13 +2187,19 @@ void UIKRetargetProcessor::SetNeedsInitialized()
 #if WITH_EDITOR
 
 bool UIKRetargetProcessor::IsBoneRetargeted(
-	const int32& BoneIndex,
-	const int8& SkeletonToCheck) const
+	const FName BoneName,
+	const ERetargetSourceOrTarget SourceOrTarget) const
 {
-	const bool bUseSourceSkeleton = SkeletonToCheck == 0;
-	const FRetargetSkeleton& Skeleton = bUseSourceSkeleton ? SourceSkeleton : TargetSkeleton;
-	const int32 RootBoneIndex = bUseSourceSkeleton ? GetSourceRetargetRoot() : GetTargetRetargetRoot();
-	if (BoneIndex==RootBoneIndex)
+	const FRetargetSkeleton& Skeleton = GetSkeleton(SourceOrTarget);
+	const int32 BoneIndex = Skeleton.FindBoneIndexByName(BoneName);
+	if (BoneIndex == INDEX_NONE)
+	{
+		// bone doesn't exist, therefore not retargeted
+		return false;
+	}
+	
+	const FName RootBoneName = GetRetargetRoot(SourceOrTarget);
+	if (BoneName == RootBoneName)
 	{
 		// retarget root is always retargeted
 		return true;
@@ -2220,7 +2220,7 @@ bool UIKRetargetProcessor::IsBoneRetargeted(
 			continue;
 		}
 		
-		if (bUseSourceSkeleton)
+		if (SourceOrTarget == ERetargetSourceOrTarget::Source)
 		{
 			if (ChainMap->SourceChain == ChainThatContainsBone)
 			{
@@ -2239,12 +2239,20 @@ bool UIKRetargetProcessor::IsBoneRetargeted(
 	return false;
 }
 
-FName UIKRetargetProcessor::GetChainNameForBone(const int32& BoneIndex, const int8& SkeletonToCheck) const
+int32 UIKRetargetProcessor::GetBoneIndexFromName(
+	const FName BoneName,
+	const ERetargetSourceOrTarget SourceOrTarget) const
+{
+	return GetSkeleton(SourceOrTarget).FindBoneIndexByName(BoneName);
+}
+
+FName UIKRetargetProcessor::GetChainNameForBone(const FName BoneName, const ERetargetSourceOrTarget SourceOrTarget) const
 {
-	const bool bUseSourceSkeleton = SkeletonToCheck == 0;
-	const FRetargetSkeleton& Skeleton = bUseSourceSkeleton ? SourceSkeleton : TargetSkeleton;
-	const int32 RootBoneIndex = bUseSourceSkeleton ? GetSourceRetargetRoot() : GetTargetRetargetRoot();
-	if (RootBoneIndex == BoneIndex)
+	const bool bUseSource = SourceOrTarget == ERetargetSourceOrTarget::Source;
+	const FRetargetSkeleton& Skeleton = bUseSource ? SourceSkeleton : TargetSkeleton;
+	const int32 BoneIndex = GetBoneIndexFromName(BoneName, SourceOrTarget);
+	const FName RootBoneName = GetRetargetRoot(SourceOrTarget);
+	if (BoneName == RootBoneName)
 	{
 		return FName("Retarget Root");
 	}
@@ -2252,6 +2260,94 @@ FName UIKRetargetProcessor::GetChainNameForBone(const int32& BoneIndex, const in
 	return Skeleton.ChainThatContainsBone[BoneIndex];
 }
 
+FTransform UIKRetargetProcessor::GetGlobalRetargetPoseAtParam(
+	const FName InChainName,
+	const float Param,
+	const ERetargetSourceOrTarget SourceOrTarget) const
+{
+	if (const FChainFK* Chain = GetChain(InChainName, SourceOrTarget))
+	{
+		return Chain->GetTransformAtParam(Chain->InitialGlobalTransforms, Param);
+	}
+
+	checkNoEntry();
+	return FTransform::Identity;
+}
+
+const FChainFK* UIKRetargetProcessor::GetChain(const FName InChainName, const ERetargetSourceOrTarget SourceOrTarget) const
+{
+	const bool bIsSource = SourceOrTarget == ERetargetSourceOrTarget::Source;
+	for (const FRetargetChainPairFK& ChainPair : ChainPairsFK)
+	{
+		const FName ChainName = bIsSource ? ChainPair.SourceBoneChainName : ChainPair.TargetBoneChainName;
+		if (ChainName == InChainName)
+		{
+			const FChainFK* Encoder = &ChainPair.FKEncoder;
+			const FChainFK* Decoder = &ChainPair.FKDecoder;
+			return bIsSource ? Encoder : Decoder;
+		}
+	}
+
+	return nullptr;
+}
+
+float UIKRetargetProcessor::GetParamOfBoneInChain(
+	const FName InBoneName,
+	const ERetargetSourceOrTarget SourceOrTarget) const
+{
+	const FName ChainName = GetChainNameForBone(InBoneName, SourceOrTarget);
+	const FChainFK* Chain = GetChain(ChainName, SourceOrTarget);
+	if (!Chain)
+	{
+		return INDEX_NONE;
+	}
+
+	const int32 BoneIndexInSkeleton = GetBoneIndexFromName(InBoneName, SourceOrTarget);
+	const int32 BoneIndexInChain =  Chain->BoneIndices.Find(BoneIndexInSkeleton);
+	return Chain->Params[BoneIndexInChain];
+}
+
+FName UIKRetargetProcessor::GetBoneAtParam(
+	const FName InChainName, 
+	const float InParam,
+	const ERetargetSourceOrTarget SourceOrTarget) const
+{
+	const FChainFK* Chain = GetChain(InChainName, SourceOrTarget);
+	const FRetargetSkeleton& Skeleton = GetSkeleton(SourceOrTarget);
+	if (Chain->BoneIndices.Num() == 1 || InParam < 0.0f)
+	{
+		return Skeleton.BoneNames[Chain->BoneIndices[0]];
+	}
+	
+	for (int32 ChainIndex = Chain->Params.Num()-1; ChainIndex >= 0 ; --ChainIndex)
+	{
+		if (Chain->Params[ChainIndex] <= InParam)
+		{
+			return Skeleton.BoneNames[Chain->BoneIndices[ChainIndex]];
+		}
+	}
+
+	checkNoEntry();
+	return NAME_None;
+}
+
+FName UIKRetargetProcessor::GetMappedChainName(
+	const FName InChainName,
+	const ERetargetSourceOrTarget SourceOrTarget)
+{
+	const bool bIsSource = SourceOrTarget == ERetargetSourceOrTarget::Source;
+	for (const FRetargetChainPairFK& ChainPair : ChainPairsFK)
+	{
+		const FName ChainName = bIsSource ? ChainPair.SourceBoneChainName : ChainPair.TargetBoneChainName;
+		if (ChainName == InChainName)
+		{
+			return bIsSource ? ChainPair.TargetBoneChainName : ChainPair.SourceBoneChainName;
+		}
+	}
+
+	return NAME_None;
+}
+
 #endif
 
 void UIKRetargetProcessor::ApplySettingsFromAsset()
@@ -2292,8 +2388,8 @@ void UIKRetargetProcessor::ApplySettingsFromAsset()
 	// apply current retarget poses (only applied if the pose has been switched to a different one OR if the current pose was modified)
 	const FName SourcePose = RetargeterAsset->GetCurrentRetargetPoseName(ERetargetSourceOrTarget::Source);
 	const FName TargetPose = RetargeterAsset->GetCurrentRetargetPoseName(ERetargetSourceOrTarget::Target);
-	ApplyNewRetargetPose(SourcePose, ERetargetSourceOrTarget::Source);
-	ApplyNewRetargetPose(TargetPose, ERetargetSourceOrTarget::Target);
+	UpdateRetargetPoseAtRuntime(SourcePose, ERetargetSourceOrTarget::Source);
+	UpdateRetargetPoseAtRuntime(TargetPose, ERetargetSourceOrTarget::Target);
 
 	// apply the current profile
 	// (this is always applied last so that profile overrides take precedence over asset settings)
@@ -2303,16 +2399,26 @@ void UIKRetargetProcessor::ApplySettingsFromAsset()
 	}
 }
 
+const FRetargetSkeleton& UIKRetargetProcessor::GetSkeleton(ERetargetSourceOrTarget SourceOrTarget) const
+{
+	return SourceOrTarget == ERetargetSourceOrTarget::Source ? SourceSkeleton : TargetSkeleton;
+}
+
+FName UIKRetargetProcessor::GetRetargetRoot(ERetargetSourceOrTarget SourceOrTarget) const
+{
+	return SourceOrTarget == ERetargetSourceOrTarget::Source ? RootRetargeter.Source.BoneName : RootRetargeter.Target.BoneName;
+}
+
 void UIKRetargetProcessor::ApplySettingsFromProfile(const FRetargetProfile& Profile)
 {
 	// assign retarget poses specified in the profile
 	if (Profile.bApplySourceRetargetPose)
 	{
-		ApplyNewRetargetPose(Profile.SourceRetargetPoseName, ERetargetSourceOrTarget::Source);
+		UpdateRetargetPoseAtRuntime(Profile.SourceRetargetPoseName, ERetargetSourceOrTarget::Source);
 	}
 	if (Profile.bApplyTargetRetargetPose)
 	{
-		ApplyNewRetargetPose(Profile.TargetRetargetPoseName, ERetargetSourceOrTarget::Target);
+		UpdateRetargetPoseAtRuntime(Profile.TargetRetargetPoseName, ERetargetSourceOrTarget::Target);
 	}
 
 	// assign chain settings
@@ -2351,7 +2457,7 @@ void UIKRetargetProcessor::ApplySettingsFromProfile(const FRetargetProfile& Prof
 	}
 }
 
-void UIKRetargetProcessor::ApplyNewRetargetPose(
+void UIKRetargetProcessor::UpdateRetargetPoseAtRuntime(
 	const FName NewRetargetPoseName,
 	ERetargetSourceOrTarget SourceOrTarget)
 {
diff --git a/Engine/Plugins/Animation/IKRig/Source/IKRig/Public/Retargeter/IKRetargetProcessor.h b/Engine/Plugins/Animation/IKRig/Source/IKRig/Public/Retargeter/IKRetargetProcessor.h
--- a/Engine/Plugins/Animation/IKRig/Source/IKRig/Public/Retargeter/IKRetargetProcessor.h
+++ b/Engine/Plugins/Animation/IKRig/Source/IKRig/Public/Retargeter/IKRetargetProcessor.h
@@ -248,6 +248,8 @@ struct FChainFK
 		const TArray<FTransform> &InitialGlobalPose,
 		FIKRigLogger& Log);
 
+	FTransform GetTransformAtParam(const TArray<FTransform>& Transforms, const float& Param) const;
+
 private:
 	
 	bool CalculateBoneParameters(FIKRigLogger& Log);
@@ -301,13 +303,8 @@ struct FChainDecoderFK : public FChainFK
 		FChainEncoderFK& SourceChain,
 		const FTargetSkeleton& TargetSkeleton,
 		TArray<FTransform> &InOutGlobalPose);
-
-private:
 	
-	FTransform GetTransformAtParam(
-		const TArray<FTransform>& Transforms,
-		const TArray<float>& InParams,
-		const float& Param) const;
+private:
 	
 	void UpdateIntermediateParents(
 		const FTargetSkeleton& TargetSkeleton,
@@ -491,8 +488,7 @@ public:
 	*/
 	TArray<FTransform>& RunRetargeter(
 		const TArray<FTransform>& InSourceGlobalPose,
-		const TMap<FName,
-		float>& SpeedValuesFromCurves,
+		const TMap<FName, float>& SpeedValuesFromCurves,
 		const float DeltaTime);
 
 	/** Apply the settings stored in a retarget profile. Call this before RunRetargeter() to use the settings stored in a profile. */
@@ -500,18 +496,12 @@ public:
 	
 	/** Apply the settings stored in the retargeter asset. */
 	void ApplySettingsFromAsset();
-	
-	/** Get read-only access to the target skeleton. */
-	const FTargetSkeleton& GetTargetSkeleton() const { return TargetSkeleton; };
-
-	/** Get read-only access to the source skeleton. */
-	const FRetargetSkeleton& GetSourceSkeleton() const { return SourceSkeleton; };
 
-	/** Get index of the root bone of the source skeleton. */
-	const int32 GetSourceRetargetRoot() const { return RootRetargeter.Source.BoneIndex; };
+	/** Get read-only access to either source or target skeleton. */
+	const FRetargetSkeleton& GetSkeleton(ERetargetSourceOrTarget SourceOrTarget) const;
 
-	/** Get index of the root bone of the target skeleton. */
-	const int32 GetTargetRetargetRoot() const { return RootRetargeter.Target.BoneIndex; };
+	/** Get name of the root bone of the given skeleton. */
+	FName GetRetargetRoot(ERetargetSourceOrTarget SourceOrTarget) const;
 	
 	/** Get whether this processor is ready to call RunRetargeter() and generate new poses. */
 	bool IsInitialized() const { return bIsInitialized; };
@@ -540,6 +530,10 @@ public:
 	/** Reset the IK planting state. */
 	void ResetPlanting();
 
+	/** Does a partial reinitialization (at runtime) whenever the retarget pose is swapped to a different or if the
+	 * pose has been modified. Does nothing if the pose has not changed. */
+	void UpdateRetargetPoseAtRuntime(const FName NewRetargetPoseName, ERetargetSourceOrTarget SourceOrTarget);
+
 	/** logging system */
 	FIKRigLogger Log;
 
@@ -548,9 +542,21 @@ public:
 	
 #if WITH_EDITOR
 	/** Returns true if the bone is part of a retarget chain or root bone, false otherwise. */
-	bool IsBoneRetargeted(const int32& BoneIndex, const int8& SkeletonToCheck) const;
+	bool IsBoneRetargeted(const FName BoneName, const ERetargetSourceOrTarget SourceOrTarget) const;
+	/** Returns index of the bone with the given name in either Source or Target skeleton. */
+	int32 GetBoneIndexFromName(const FName BoneName, const ERetargetSourceOrTarget SourceOrTarget) const;
 	/** Returns name of the chain associated with this bone. Returns NAME_None if bone is not in a chain. */
-	FName GetChainNameForBone(const int32& BoneIndex, const int8& SkeletonToCheck) const;
+	FName GetChainNameForBone(const FName BoneName, const ERetargetSourceOrTarget SourceOrTarget) const;
+	/** Get a transform at a given param in a chain */
+	FTransform GetGlobalRetargetPoseAtParam(const FName InChainName, const float Param, const ERetargetSourceOrTarget SourceOrTarget) const;
+	/** Get access to the internal chain on the given skeleton */
+	const FChainFK* GetChain(const FName InChainName, const ERetargetSourceOrTarget SourceOrTarget) const;
+	/** Get the param of the bone in it's retarget chain. Ranges from 0 to NumBonesInChain. */
+	float GetParamOfBoneInChain(const FName InBoneName, const ERetargetSourceOrTarget SourceOrTarget) const;
+	/** Get the bone in the chain at the given param. */
+	FName GetBoneAtParam(const FName InChainName, const float InParam, const ERetargetSourceOrTarget SourceOrTarget) const;
+	/** Get the chain mapped to this one. */
+	FName GetMappedChainName(const FName InChainName, const ERetargetSourceOrTarget SourceOrTarget);
 	/** store data for debug drawing */
 	FRetargetDebugData DebugData;
 #endif
@@ -628,8 +634,4 @@ private:
 
 	/** Run all post process operations on the retargeted result. */
 	void RunRetargetOps(const TArray<FTransform>& InSourceGlobalPose, TArray<FTransform>& OutTargetGlobalPose);
-
-	/** Does a partial reinitialization (at runtime) whenever the retarget pose is swapped to a different or if the
-	 * pose has been modified. Does nothing if the pose has not changed. */
-	void ApplyNewRetargetPose(const FName NewRetargetPoseName, ERetargetSourceOrTarget SourceOrTarget);
 };
diff --git a/Engine/Plugins/Animation/IKRig/Source/IKRig/Public/Retargeter/IKRetargeter.h b/Engine/Plugins/Animation/IKRig/Source/IKRig/Public/Retargeter/IKRetargeter.h
--- a/Engine/Plugins/Animation/IKRig/Source/IKRig/Public/Retargeter/IKRetargeter.h
+++ b/Engine/Plugins/Animation/IKRig/Source/IKRig/Public/Retargeter/IKRetargeter.h
@@ -202,8 +202,10 @@ private:
 UENUM()
 enum class ERetargetSourceOrTarget : uint8
 {
-	Source,	// the SOURCE skeleton (to copy FROM)
-	Target, // the TARGET skeleton (to copy TO)
+	// the SOURCE skeleton (to copy FROM)
+	Source,
+	// the TARGET skeleton (to copy TO)
+	Target,
 };
 
 // which skeleton are we referring to?
diff --git a/Engine/Plugins/Animation/IKRig/Source/IKRigEditor/IKRigEditor.Build.cs b/Engine/Plugins/Animation/IKRig/Source/IKRigEditor/IKRigEditor.Build.cs
--- a/Engine/Plugins/Animation/IKRig/Source/IKRigEditor/IKRigEditor.Build.cs
+++ b/Engine/Plugins/Animation/IKRig/Source/IKRigEditor/IKRigEditor.Build.cs
@@ -54,6 +54,8 @@ namespace UnrealBuildTool.Rules
 					"AnimationCore",
                     "AnimationWidgets",
                     "ApplicationCore",
+					"MeshDescription",
+					"SkeletalMeshDescription"
                 }
             );
         }
diff --git a/Engine/Plugins/Animation/IKRig/Source/IKRigEditor/Private/RetargetEditor/IKRetargetBatchOperation.cpp b/Engine/Plugins/Animation/IKRig/Source/IKRigEditor/Private/RetargetEditor/IKRetargetBatchOperation.cpp
--- a/Engine/Plugins/Animation/IKRig/Source/IKRigEditor/Private/RetargetEditor/IKRetargetBatchOperation.cpp
+++ b/Engine/Plugins/Animation/IKRig/Source/IKRigEditor/Private/RetargetEditor/IKRetargetBatchOperation.cpp
@@ -326,7 +326,7 @@ void UIKRetargetBatchOperation::ConvertAnimation(
 	}
 
 	// target skeleton data
-	const FTargetSkeleton& TargetSkeleton = Processor->GetTargetSkeleton();
+	const FRetargetSkeleton& TargetSkeleton = Processor->GetSkeleton(ERetargetSourceOrTarget::Target);
 	const TArray<FName>& TargetBoneNames = TargetSkeleton.BoneNames;
 	const int32 NumTargetBones = TargetBoneNames.Num();
 
@@ -335,7 +335,7 @@ void UIKRetargetBatchOperation::ConvertAnimation(
 	BoneTracks.SetNumZeroed(NumTargetBones);
 
 	// source skeleton data
-	const FRetargetSkeleton& SourceSkeleton = Processor->GetSourceSkeleton();
+	const FRetargetSkeleton& SourceSkeleton = Processor->GetSkeleton(ERetargetSourceOrTarget::Source);
 	const TArray<FName>& SourceBoneNames = SourceSkeleton.BoneNames;
 	const int32 NumSourceBones = SourceBoneNames.Num();
 
diff --git a/Engine/Plugins/Animation/IKRig/Source/IKRigEditor/Private/RetargetEditor/IKRetargetCommands.cpp b/Engine/Plugins/Animation/IKRig/Source/IKRigEditor/Private/RetargetEditor/IKRetargetCommands.cpp
--- a/Engine/Plugins/Animation/IKRig/Source/IKRigEditor/Private/RetargetEditor/IKRetargetCommands.cpp
+++ b/Engine/Plugins/Animation/IKRig/Source/IKRigEditor/Private/RetargetEditor/IKRetargetCommands.cpp
@@ -19,6 +19,12 @@ void FIKRetargetCommands::RegisterCommands()
 	UI_COMMAND(ResetSelectedBones, "Reset Selected Bones", "Sets the selected bones to the mesh reference pose.", EUserInterfaceActionType::Button, FInputChord());
 	UI_COMMAND(ResetSelectedAndChildrenBones, "Reset Selected and Children Bones", "Sets the selected bones (and all children recursively) to the mesh reference pose.", EUserInterfaceActionType::Button, FInputChord());
 
+	UI_COMMAND(AutoAlignAllBones, "Align All Bones", "Goes through all bones contained in mapped retarget chains and rotates them to match the worldspace orientation of the equivalent bone on the other skeleton.", EUserInterfaceActionType::Button, FInputChord());
+	UI_COMMAND(AlignSelected, "Align Selected", "Goes through the selected bones contained in mapped retarget chains and rotates them to match the worldspace orientation of the equivalent bone on the other skeleton.", EUserInterfaceActionType::Button, FInputChord());
+	UI_COMMAND(AlignSelectedAndChildren, "Align Selected and Children", "Goes through the selected bones (and their children). If the bone is contained in a mapped retarget chains, it is rotated to match the worldspace orientation of the equivalent bone on the other skeleton.", EUserInterfaceActionType::Button, FInputChord());
+	UI_COMMAND(AlignSelectedUsingMesh, "Align Selected Using Mesh", "Align the selected bones using the mesh to determine the bone directions.", EUserInterfaceActionType::Button, FInputChord());
+	UI_COMMAND(SnapCharacterToGround, "Snap Character to Ground", "Move the character vertically to restore original height from the ground. Uses the selected bone as the reference point for a grounded limb. Otherwise searches the skeleton for the lowest retargeted bone.", EUserInterfaceActionType::Button, FInputChord());
+
 	UI_COMMAND(NewRetargetPose, "Create", "Create a new retarget pose.", EUserInterfaceActionType::Button, FInputChord());
 	UI_COMMAND(DuplicateRetargetPose, "Duplicate Current", "Duplicate the current retarget pose.", EUserInterfaceActionType::Button, FInputChord());
 	UI_COMMAND(DeleteRetargetPose, "Delete", "Delete current retarget pose.", EUserInterfaceActionType::Button, FInputChord());
diff --git a/Engine/Plugins/Animation/IKRig/Source/IKRigEditor/Private/RetargetEditor/IKRetargetEditor.cpp b/Engine/Plugins/Animation/IKRig/Source/IKRigEditor/Private/RetargetEditor/IKRetargetEditor.cpp
--- a/Engine/Plugins/Animation/IKRig/Source/IKRigEditor/Private/RetargetEditor/IKRetargetEditor.cpp
+++ b/Engine/Plugins/Animation/IKRig/Source/IKRigEditor/Private/RetargetEditor/IKRetargetEditor.cpp
@@ -195,13 +195,13 @@ void FIKRetargetEditor::BindCommands()
 	ToolkitCommands->MapAction(
 		Commands.ResetSelectedBones,
 		FExecuteAction::CreateSP(EditorController, &FIKRetargetEditorController::HandleResetSelectedBones),
-		FCanExecuteAction::CreateSP(EditorController, &FIKRetargetEditorController::CanResetSelected),
+		FCanExecuteAction::CreateSP(EditorController, &FIKRetargetEditorController::IsAnyBoneSelected),
 		EUIActionRepeatMode::RepeatDisabled);
 
 	ToolkitCommands->MapAction(
 		Commands.ResetSelectedAndChildrenBones,
 		FExecuteAction::CreateSP(EditorController, &FIKRetargetEditorController::HandleResetSelectedAndChildrenBones),
-		FCanExecuteAction::CreateSP(EditorController, &FIKRetargetEditorController::CanResetSelected),
+		FCanExecuteAction::CreateSP(EditorController, &FIKRetargetEditorController::IsAnyBoneSelected),
 		EUIActionRepeatMode::RepeatDisabled);
 
 	ToolkitCommands->MapAction(
@@ -232,6 +232,49 @@ void FIKRetargetEditor::BindCommands()
 		FCanExecuteAction(),
 		EUIActionRepeatMode::RepeatDisabled);
 
+	//
+	// Auto-gen retarget pose
+	//
+	
+	ToolkitCommands->MapAction(
+		Commands.AutoAlignAllBones,
+		FExecuteAction::CreateSP(EditorController, &FIKRetargetEditorController::HandleAlignAllBones),
+		FCanExecuteAction(),
+		FCanExecuteAction(),
+		EUIActionRepeatMode::RepeatDisabled);
+	
+	ToolkitCommands->MapAction(
+		Commands.AlignSelected,
+		FExecuteAction::CreateSP(EditorController, &FIKRetargetEditorController::HandleAlignSelectedBones, ERetargetAutoAlignMethod::ChainToChain, false /* no children*/),
+		FCanExecuteAction(),
+		FCanExecuteAction(),
+		EUIActionRepeatMode::RepeatDisabled);
+
+	ToolkitCommands->MapAction(
+		Commands.AlignSelectedAndChildren,
+		FExecuteAction::CreateSP(EditorController, &FIKRetargetEditorController::HandleAlignSelectedBones, ERetargetAutoAlignMethod::ChainToChain, true /* include children*/),
+		FCanExecuteAction(),
+		FCanExecuteAction(),
+		EUIActionRepeatMode::RepeatDisabled);
+
+	ToolkitCommands->MapAction(
+		Commands.AlignSelectedUsingMesh,
+		FExecuteAction::CreateSP(EditorController, &FIKRetargetEditorController::HandleAlignSelectedBones, ERetargetAutoAlignMethod::MeshToMesh, false /* no children*/),
+		FCanExecuteAction(),
+		FCanExecuteAction(),
+		EUIActionRepeatMode::RepeatDisabled);
+
+	ToolkitCommands->MapAction(
+		Commands.SnapCharacterToGround,
+		FExecuteAction::CreateSP(EditorController, &FIKRetargetEditorController::HandleSnapToGround),
+		FCanExecuteAction(),
+		FCanExecuteAction(),
+		EUIActionRepeatMode::RepeatDisabled);
+	
+
+	//
+	// Pose exporter
+	//
 	const TSharedRef<FIKRetargetPoseExporter> PoseExporterRef = EditorController->PoseExporter.ToSharedRef();
 	
 	ToolkitCommands->MapAction(
diff --git a/Engine/Plugins/Animation/IKRig/Source/IKRigEditor/Private/RetargetEditor/IKRetargetEditorController.cpp b/Engine/Plugins/Animation/IKRig/Source/IKRigEditor/Private/RetargetEditor/IKRetargetEditorController.cpp
--- a/Engine/Plugins/Animation/IKRig/Source/IKRigEditor/Private/RetargetEditor/IKRetargetEditorController.cpp
+++ b/Engine/Plugins/Animation/IKRig/Source/IKRigEditor/Private/RetargetEditor/IKRetargetEditorController.cpp
@@ -13,6 +13,10 @@
 #include "Animation/PoseAsset.h"
 #include "Dialog/SCustomDialog.h"
 #include "Preferences/PersonaOptions.h"
+#include "IKRigEditor.h"
+#include "MeshDescription.h"
+#include "SkeletalMeshAttributes.h"
+#include "Rendering/SkeletalMeshRenderData.h"
 #include "RetargetEditor/IKRetargetAnimInstance.h"
 #include "RetargetEditor/IKRetargetDefaultMode.h"
 #include "RetargetEditor/IKRetargetEditPoseMode.h"
@@ -26,6 +30,7 @@
 #include "RetargetEditor/SRetargetOpStack.h"
 #include "RigEditor/SIKRigOutputLog.h"
 #include "RigEditor/IKRigController.h"
+#include "ScopedTransaction.h"
 #include "Styling/AppStyle.h"
 
 #define LOCTEXT_NAMESPACE "IKRetargetEditorController"
@@ -56,7 +61,7 @@ void FBoundIKRig::UnBind() const
 	IKRigController->OnRetargetChainRenamed().Remove(RenameChainDelegateHandle);
 }
 
-FRetargetPlaybackManager::FRetargetPlaybackManager(TWeakPtr<FIKRetargetEditorController> InEditorController)
+FRetargetPlaybackManager::FRetargetPlaybackManager(const TWeakPtr<FIKRetargetEditorController>& InEditorController)
 {
 	check(InEditorController.Pin())
 	EditorController = InEditorController;
@@ -131,7 +136,7 @@ void FRetargetPlaybackManager::ResumePlayback() const
 
 bool FRetargetPlaybackManager::IsStopped() const
 {
-	UIKRetargetAnimInstance* AnimInstance = EditorController.Pin()->SourceAnimInstance.Get();
+	const UIKRetargetAnimInstance* AnimInstance = EditorController.Pin()->SourceAnimInstance.Get();
 	if (!AnimInstance)
 	{
 		return true;
@@ -152,6 +157,8 @@ void FIKRetargetEditorController::Initialize(TSharedPtr<FIKRetargetEditor> InEdi
 
 	PlaybackManager = MakeUnique<FRetargetPlaybackManager>(SharedThis(this));
 
+	AutoPoseGenerator = MakeUnique<FRetargetAutoPoseGenerator>(SharedThis(this));
+
 	SelectedBoneNames.Add(ERetargetSourceOrTarget::Source);
 	SelectedBoneNames.Add(ERetargetSourceOrTarget::Target);
 	LastSelectedItem = ERetargetSelectionType::NONE;
@@ -458,6 +465,11 @@ bool FIKRetargetEditorController::GetCameraTargetForSelection(FSphere& OutTarget
 	return false;
 }
 
+bool FIKRetargetEditorController::IsAnyBoneSelected() const
+{
+	return !GetSelectedBones().IsEmpty();
+}
+
 bool FIKRetargetEditorController::IsBoneRetargeted(const FName& BoneName, ERetargetSourceOrTarget SourceOrTarget) const
 {
 	// get an initialized processor
@@ -467,16 +479,8 @@ bool FIKRetargetEditorController::IsBoneRetargeted(const FName& BoneName, ERetar
 		return false;
 	}
 
-	// get the bone index
-	const FRetargetSkeleton& Skeleton = SourceOrTarget == ERetargetSourceOrTarget::Source ? Processor->GetSourceSkeleton() : Processor->GetTargetSkeleton();
-	const int32 BoneIndex = Skeleton.FindBoneIndexByName(BoneName);
-	if (BoneIndex == INDEX_NONE)
-	{
-		return false;
-	}
-
 	// return if it's a retargeted bone
-	return Processor->IsBoneRetargeted(BoneIndex, (int8)SourceOrTarget);
+	return Processor->IsBoneRetargeted(BoneName, SourceOrTarget);
 }
 
 FName FIKRetargetEditorController::GetChainNameFromBone(const FName& BoneName, ERetargetSourceOrTarget SourceOrTarget) const
@@ -487,16 +491,8 @@ FName FIKRetargetEditorController::GetChainNameFromBone(const FName& BoneName, E
 	{
 		return NAME_None;
 	}
-
-	// get the bone index
-	const FRetargetSkeleton& Skeleton = SourceOrTarget == ERetargetSourceOrTarget::Source ? Processor->GetSourceSkeleton() : Processor->GetTargetSkeleton();
-	const int32 BoneIndex = Skeleton.FindBoneIndexByName(BoneName);
-	if (BoneIndex == INDEX_NONE)
-	{
-		return NAME_None;
-	}
-
-	return Processor->GetChainNameForBone(BoneIndex, (int8)SourceOrTarget);
+	
+	return Processor->GetChainNameForBone(BoneName, SourceOrTarget);
 }
 
 TObjectPtr<UIKRetargetBoneDetails> FIKRetargetEditorController::GetOrCreateBoneDetailsObject(const FName& BoneName)
@@ -539,7 +535,7 @@ UDebugSkelMeshComponent* FIKRetargetEditorController::GetEditedSkeletalMesh() co
 
 const FRetargetSkeleton& FIKRetargetEditorController::GetCurrentlyEditedSkeleton(const UIKRetargetProcessor& Processor) const
 {
-	return CurrentlyEditingSourceOrTarget == ERetargetSourceOrTarget::Source ? Processor.GetSourceSkeleton() : Processor.GetTargetSkeleton();
+	return Processor.GetSkeleton(CurrentlyEditingSourceOrTarget);
 }
 
 FTransform FIKRetargetEditorController::GetGlobalRetargetPoseOfBone(
@@ -1317,46 +1313,69 @@ void FIKRetargetEditorController::HandleResetSelectedBones() const
 
 void FIKRetargetEditorController::HandleResetSelectedAndChildrenBones() const
 {
-	// get the reference skeleton we're operating on
-	const USkeletalMesh* SkeletalMesh = GetSkeletalMesh(GetSourceOrTarget());
-	if (!SkeletalMesh)
+	// get all selected bones and their children (recursive)
+	const TArray<FName> BonesToReset = GetSelectedBonesAndChildren();
+	
+	// reset the bones in the current pose
+	const FName CurrentPose = AssetController->GetCurrentRetargetPoseName(CurrentlyEditingSourceOrTarget);
+	AssetController->ResetRetargetPose(CurrentPose, BonesToReset, GetSourceOrTarget());
+}
+
+void FIKRetargetEditorController::HandleAlignAllBones() const
+{
+	// get all the bones in the current skeleton
+	const UIKRetargetProcessor* Processor = GetRetargetProcessor();
+	if (!(Processor && Processor->IsInitialized()))
 	{
 		return;
 	}
-	const FReferenceSkeleton RefSkeleton = SkeletalMesh->GetRefSkeleton();
+
+	// undo transaction
+	constexpr bool bShouldTransact = true;
+	FScopedTransaction Transaction(LOCTEXT("AutoAlignSelectedBones", "Auto Align All Bones"), bShouldTransact);
+	AssetController->GetAsset()->Modify();
 	
-	// get list of all children of selected bones
-	TArray<int32> AllChildrenIndices;
-	for (const FName& SelectedBone : SelectedBoneNames[CurrentlyEditingSourceOrTarget])
-	{
-		const int32 SelectedBoneIndex = RefSkeleton.FindBoneIndex(SelectedBone);
-		AllChildrenIndices.Add(SelectedBoneIndex);
-		
-		for (int32 ChildIndex = 0; ChildIndex < RefSkeleton.GetNum(); ++ChildIndex)
-		{
-			const int32 ParentIndex = RefSkeleton.GetParentIndex(ChildIndex);
-			if (ParentIndex != INDEX_NONE && AllChildrenIndices.Contains(ParentIndex))
-			{
-				AllChildrenIndices.Add(ChildIndex);
-			}
-		}
-	}
+	// suppress warnings about bones that cannot be aligned when aligning ALL bones
+	const TArray<FName>& AllBones = Processor->GetSkeleton(GetSourceOrTarget()).BoneNames;
+	constexpr bool bSuppressWarnings = true;
+	AutoPoseGenerator.Get()->AlignBones(
+		AllBones,
+		ERetargetAutoAlignMethod::ChainToChain,
+		GetSourceOrTarget(),
+		bSuppressWarnings);
 
-	// merge total list of all selected bones and their children
-	TArray<FName> BonesToReset = SelectedBoneNames[CurrentlyEditingSourceOrTarget];
-	for (const int32 ChildIndex : AllChildrenIndices)
-	{
-		BonesToReset.AddUnique(RefSkeleton.GetBoneName(ChildIndex));
-	}
+	// if the retarget root was aligned
+	AutoPoseGenerator.Get()->SnapToGround(NAME_None, GetSourceOrTarget());
+}
+
+void FIKRetargetEditorController::HandleAlignSelectedBones(const ERetargetAutoAlignMethod Method, const bool bIncludeChildren) const
+{
+	// undo transaction
+	constexpr bool bShouldTransact = true;
+	FScopedTransaction Transaction(LOCTEXT("AutoAlignSelectedBones", "Auto Align Selected Bones"), bShouldTransact);
+	AssetController->GetAsset()->Modify();
+
+	const TArray<FName> BonesToAlign = bIncludeChildren ? GetSelectedBonesAndChildren() : GetSelectedBones();
 	
-	// reset the bones in the current pose
-	const FName CurrentPose = AssetController->GetCurrentRetargetPoseName(CurrentlyEditingSourceOrTarget);
-	AssetController->ResetRetargetPose(CurrentPose, BonesToReset, GetSourceOrTarget());
+	// allow warnings about bones that cannot be aligned when bones are explicitly specified by user
+	constexpr bool bSuppressWarnings = false;
+	AutoPoseGenerator.Get()->AlignBones(
+		BonesToAlign,
+		Method,
+		GetSourceOrTarget(),
+		bSuppressWarnings);
 }
 
-bool FIKRetargetEditorController::CanResetSelected() const
+void FIKRetargetEditorController::HandleSnapToGround() const
 {
-	return !GetSelectedBones().IsEmpty();
+	// undo transaction
+	constexpr bool bShouldTransact = true;
+	FScopedTransaction Transaction(LOCTEXT("AutoSnapToGround", "Snap Retarget Pose to Ground"), bShouldTransact);
+	AssetController->GetAsset()->Modify();
+	
+	const TArray<FName> SelectedBones = GetSelectedBones();
+	const FName FirstSelectedBone = SelectedBones.IsEmpty() ? NAME_None : SelectedBones[0];
+	AutoPoseGenerator.Get()->SnapToGround(FirstSelectedBone, GetSourceOrTarget());
 }
 
 void FIKRetargetEditorController::HandleRenamePose()
@@ -1551,6 +1570,43 @@ void FIKRetargetEditorController::RenderSkeleton(FPrimitiveDrawInterface* PDI, E
 	);
 }
 
+TArray<FName> FIKRetargetEditorController::GetSelectedBonesAndChildren() const
+{
+	// get the reference skeleton we're operating on
+	const USkeletalMesh* SkeletalMesh = GetSkeletalMesh(GetSourceOrTarget());
+	if (!SkeletalMesh)
+	{
+		return {};
+	}
+	const FReferenceSkeleton RefSkeleton = SkeletalMesh->GetRefSkeleton();
+
+	// get list of all children of selected bones
+	TArray<int32> AllChildrenIndices;
+	for (const FName& SelectedBone : SelectedBoneNames[CurrentlyEditingSourceOrTarget])
+	{
+		const int32 SelectedBoneIndex = RefSkeleton.FindBoneIndex(SelectedBone);
+		AllChildrenIndices.Add(SelectedBoneIndex);
+		
+		for (int32 ChildIndex = 0; ChildIndex < RefSkeleton.GetNum(); ++ChildIndex)
+		{
+			const int32 ParentIndex = RefSkeleton.GetParentIndex(ChildIndex);
+			if (ParentIndex != INDEX_NONE && AllChildrenIndices.Contains(ParentIndex))
+			{
+				AllChildrenIndices.Add(ChildIndex);
+			}
+		}
+	}
+
+	// merge total list of all selected bones and their children
+	TArray<FName> BonesToReturn = SelectedBoneNames[CurrentlyEditingSourceOrTarget];
+	for (const int32 ChildIndex : AllChildrenIndices)
+	{
+		BonesToReturn.AddUnique(RefSkeleton.GetBoneName(ChildIndex));
+	}
+
+	return BonesToReturn;
+}
+
 void FIKRetargetEditorController::FixZeroHeightRetargetRoot(ERetargetSourceOrTarget SourceOrTarget) const
 {
 	// is there a mesh to check?
diff --git a/Engine/Plugins/Animation/IKRig/Source/IKRigEditor/Private/RetargetEditor/IKRetargeterPoseGenerator.cpp b/Engine/Plugins/Animation/IKRig/Source/IKRigEditor/Private/RetargetEditor/IKRetargeterPoseGenerator.cpp
new file mode 100644
--- /dev/null
+++ b/Engine/Plugins/Animation/IKRig/Source/IKRigEditor/Private/RetargetEditor/IKRetargeterPoseGenerator.cpp
@@ -0,0 +1,586 @@
+﻿// Copyright Epic Games, Inc. All Rights Reserved.
+
+#include "RetargetEditor/IKRetargeterPoseGenerator.h"
+
+#include "AnimationRuntime.h"
+#include "Engine/SkeletalMesh.h"
+#include "ScopedTransaction.h"
+#include "MeshDescription.h"
+#include "ReferenceSkeleton.h"
+#include "SkeletalMeshAttributes.h"
+#include "Algo/LevenshteinDistance.h"
+#include "RetargetEditor/IKRetargetEditorController.h"
+#include "Retargeter/IKRetargeter.h"
+#include "Async/ParallelFor.h"
+
+#include "Templates/SharedPointer.h"
+
+#define LOCTEXT_NAMESPACE "RetargetPoseGenerator"
+
+FRetargetAutoPoseGenerator::FRetargetAutoPoseGenerator(const TWeakPtr<FIKRetargetEditorController>& InEditorController)
+{
+	check(InEditorController.Pin())
+	EditorController = InEditorController;
+}
+
+void FRetargetAutoPoseGenerator::AlignBones(
+	const TArray<FName>& BonesToAlign,
+	const ERetargetAutoAlignMethod Method, 
+	const ERetargetSourceOrTarget SourceOrTarget,
+	const bool bSuppressWarnings) const
+{
+	// cannot align bones until we have an initialized processor
+	if (!CheckReadyToAlignBones())
+	{
+		return;
+	}
+
+	// sort hierarchically (we want to align from root to leaf order)
+	TArray<FName> SortedBones = BonesToAlign;
+	const FRetargetSkeleton& Skeleton = Processor->GetSkeleton(SourceOrTarget);
+	SortedBones.Sort([Skeleton](FName A, FName B) -> bool
+	{
+		return Skeleton.FindBoneIndexByName(A) < Skeleton.FindBoneIndexByName(B);
+	});
+
+	// actually align the bones (generating new entries in the retarget pose)
+	for (const FName BoneToAlign : SortedBones)
+	{
+		AlignBone(BoneToAlign, Method, SourceOrTarget);
+	}
+
+	// warn if any bones would not be aligned so user isn't left wondering
+	if (!bSuppressWarnings)
+	{
+		// warn user if they tried to align the retarget root (doesn't work)
+		if (BonesToAlign.Contains(Processor->GetRetargetRoot(SourceOrTarget)))
+		{
+			Processor->Log.LogWarning(LOCTEXT("AutoAlignRootNotSupported", "Skipped aligning the retarget root. This is not supported."));
+		}
+
+		// warn user if they tried to align a non-retargeted bone (doesn't work)
+		for (const FName Bone : BonesToAlign)
+		{
+			if (!Processor->IsBoneRetargeted(Bone, SourceOrTarget))
+			{
+				Processor->Log.LogWarning(FText::Format(LOCTEXT("SkipAutoAlignNonRetargeted", "Skipped aligning {0}. The bone is not retargeted and therefore has no 'matching' equivalent on other skeleton."), FText::FromName(Bone)));
+			}
+		}
+	}
+}
+
+void FRetargetAutoPoseGenerator::AlignBone(
+	const FName BoneToAlign,
+	const ERetargetAutoAlignMethod Method,
+	const ERetargetSourceOrTarget SourceOrTarget) const
+{
+	// can't auto-align if the bone is not retargeted  (warning is elsewhere)
+	if (!Processor->IsBoneRetargeted(BoneToAlign, SourceOrTarget))
+	{
+		return;
+	}
+
+	// force regenerate the retarget pose to guarantee Skeleton.RetargetGlobalPose reflects any modifications made in the viewport
+	const FName CurrentPoseName = EditorController.Pin()->AssetController->GetCurrentRetargetPoseName(SourceOrTarget);
+	Processor->UpdateRetargetPoseAtRuntime(CurrentPoseName, SourceOrTarget);
+
+	// use a custom procedure to auto-align retarget root
+	if (BoneToAlign == Processor->GetRetargetRoot(SourceOrTarget))
+	{
+		AlignRetargetRootBone(SourceOrTarget);
+		return;
+	}
+
+	// determine which side we're aligning and which side we're matching
+	const ERetargetSourceOrTarget SideToAlign = SourceOrTarget;
+	const ERetargetSourceOrTarget SideToMatch = SideToAlign == ERetargetSourceOrTarget::Source ? ERetargetSourceOrTarget::Target : ERetargetSourceOrTarget::Source;
+	
+	// find chains on both sides
+	const FName ChainToAlign = Processor->GetChainNameForBone(BoneToAlign, SideToAlign);
+	const FName ChainToMatch = Processor->GetMappedChainName(ChainToAlign, SideToAlign);
+	checkf(ChainToMatch != NAME_None != 0, TEXT("Bone should never be retargeted and not in a mapped chain."));
+	
+	// get the normalized chain parameter to match
+	const float Param = Processor->GetParamOfBoneInChain(BoneToAlign, SideToAlign);
+
+	// find tangents to align between source/target
+	FVector DirectionToAlign;
+	FVector DirectionToMatch;
+	bool bFoundAlignVector = false;
+	bool bFoundMatchVector = false;
+	switch (Method)
+	{
+	case ERetargetAutoAlignMethod::ChainToChain:
+		{
+			bFoundAlignVector = GetChainTangent(ChainToAlign, Param, SideToAlign, DirectionToAlign);
+			bFoundMatchVector = GetChainTangent(ChainToMatch, Param, SideToMatch, DirectionToMatch);
+			break;
+		}
+
+	case ERetargetAutoAlignMethod::MeshToMesh:
+		{
+			bFoundAlignVector = GetChainTangentFromMesh(ChainToAlign, Param, SideToAlign, DirectionToAlign);
+			bFoundMatchVector = GetChainTangentFromMesh(ChainToMatch, Param, SideToMatch, DirectionToAlign);
+			break;
+		}
+		
+	default:
+		{
+			checkNoEntry();
+			break;
+		}
+	}
+	
+	// warn if unable to generate a direction vector from the mesh
+	if (!(bFoundAlignVector && bFoundMatchVector))
+	{
+		// user asked to align the mesh, but direction vectors for the target could not be generated
+		if (!bFoundAlignVector && Method == ERetargetAutoAlignMethod::MeshToMesh)
+		{
+			Processor->Log.LogWarning(
+			FText::Format(LOCTEXT("AutoAlignMeshNoTarget", "Could not determine the direction of bone to align {0} from the mesh. Likely because no vertices are skinned to the target bone."),
+				FText::FromName(BoneToAlign)));
+		}
+		
+		// user asked to align the mesh, but direction vectors for the target could not be generated
+		if (!bFoundMatchVector && Method == ERetargetAutoAlignMethod::MeshToMesh)
+		{
+			Processor->Log.LogWarning(
+			FText::Format(LOCTEXT("AutoAlignMeshNoSource", "Could not determine the direction of bone to match from it's mesh for {0}. Likely because no vertices are skinned to the source bone."),
+				FText::FromName(BoneToAlign)));
+		}
+		
+		return;
+	}
+
+	// find quaternion to rotate from current to target
+	const FQuat WorldRotationOffset = FQuat::FindBetweenNormals(DirectionToAlign, DirectionToMatch);
+
+	// convert to local space offset and store in retarget pose
+	ApplyWorldRotationToRetargetPose(WorldRotationOffset, BoneToAlign, SideToAlign);
+}
+
+void FRetargetAutoPoseGenerator::AlignRetargetRootBone(const ERetargetSourceOrTarget SourceOrTarget) const
+{
+	// when aligning the retarget root bone, we have to use the Spine to determine it's direction
+	// any other assumptions are bound to fail in different ways, so finding a Spine is the only reliable way (that we've found so far)
+	FName SourceSpineChainName;
+	FName TargetSpineChainName;
+	const bool bFoundSourceSpine = GetSpineChain(ERetargetSourceOrTarget::Source, SourceSpineChainName);
+	const bool bFoundTargetSpine = GetSpineChain(ERetargetSourceOrTarget::Target, TargetSpineChainName);
+	
+	// user asked to align the the retarget root but could not find the spine to use for it's direction vector
+	if (!(bFoundSourceSpine && bFoundTargetSpine))
+	{
+		Processor->Log.LogWarning(
+		LOCTEXT("MissingSpine", "Could not auto align the retarget root because a Spine chain was not found. Expects a retarget chain named 'Spine' or similar."));
+		return;
+	}
+
+	// get the direction of the source and target spine chains
+	const FVector SourceSpineDirection = GetChainDirection(SourceSpineChainName, ERetargetSourceOrTarget::Source);
+	const FVector TargetSpineDirection = GetChainDirection(TargetSpineChainName, ERetargetSourceOrTarget::Target);
+
+	// determine which direction to align
+	const bool bAlignSource = SourceOrTarget == ERetargetSourceOrTarget::Source;
+	const FVector DirectionToAlign = bAlignSource ? SourceSpineDirection : TargetSpineDirection;
+	const FVector DirectionToMatch = bAlignSource ? TargetSpineDirection : SourceSpineDirection;
+
+	// find quaternion to rotate from current to target
+	const FQuat WorldRotationOffset = FQuat::FindBetweenNormals(DirectionToAlign, DirectionToMatch);
+
+	// convert to local space offset and store in retarget pose
+	const FName BoneToAlign = Processor->GetRetargetRoot(SourceOrTarget);
+	ApplyWorldRotationToRetargetPose(WorldRotationOffset, BoneToAlign, SourceOrTarget);
+}
+
+void FRetargetAutoPoseGenerator::SnapToGround(const FName BoneToPutOnGround, const ERetargetSourceOrTarget SourceOrTarget) const
+{
+	// cannot align bones until we have an initialized processor
+	if (!CheckReadyToAlignBones())
+	{
+		return;
+	}
+
+	const TObjectPtr<UIKRetargeterController> AssetController = EditorController.Pin()->AssetController;
+	const FRetargetSkeleton& Skeleton = Processor->GetSkeleton(SourceOrTarget);
+	
+	// force regenerate the retarget pose to guarantee Skeleton.RetargetGlobalPose reflects any modifications made thus far
+	const FName CurrentPoseName = AssetController->GetCurrentRetargetPoseName(SourceOrTarget);
+	Processor->UpdateRetargetPoseAtRuntime(CurrentPoseName, SourceOrTarget);
+
+	// if no bone specified, we search the current retarget pose to find the lowest one and use that
+	float HeightOfBoneInRetargetPose = std::numeric_limits<float>::max();
+	int32 BoneToSnapIndex = INDEX_NONE;
+	if (BoneToPutOnGround == NAME_None)
+	{
+		// find the lowest bone in the current retarget pose
+		for (int32 BoneIndex=0; BoneIndex<Skeleton.BoneNames.Num(); ++BoneIndex)
+		{
+			if (!Processor->IsBoneRetargeted(Skeleton.BoneNames[BoneIndex], SourceOrTarget))
+			{
+				continue;
+			}
+			const FTransform& BoneTransform = Skeleton.RetargetGlobalPose[BoneIndex];
+			const float BoneHeight = BoneTransform.GetLocation().Z;
+			if (BoneHeight < HeightOfBoneInRetargetPose)
+			{
+				HeightOfBoneInRetargetPose = BoneHeight;
+				BoneToSnapIndex = BoneIndex;
+			}
+		}
+	}
+	else
+	{
+		// use the user specified bone
+		BoneToSnapIndex = Skeleton.FindBoneIndexByName(BoneToPutOnGround);
+		HeightOfBoneInRetargetPose = Skeleton.RetargetGlobalPose[BoneToSnapIndex].GetLocation().Z;
+	}
+
+	if (BoneToSnapIndex == INDEX_NONE)
+	{
+		// shouldn't happen because of check at start of this function
+		return;
+	}
+	
+	// find the height of this bone in the ref pose
+	const FReferenceSkeleton& RefSkeleton = Skeleton.SkeletalMesh->GetRefSkeleton();
+	const FTransform& RefPoseOfBone = FAnimationRuntime::GetComponentSpaceTransformRefPose(RefSkeleton, BoneToSnapIndex);
+	const float HeightInRefPose = RefPoseOfBone.GetLocation().Z;
+
+	// snap it back to original height off the ground
+	const float DeltaHeight = HeightInRefPose - HeightOfBoneInRetargetPose;
+	AssetController->SetRootOffsetInRetargetPose(FVector(0.f,0.f,DeltaHeight), SourceOrTarget);
+}
+
+bool FRetargetAutoPoseGenerator::GetSpineChain(const ERetargetSourceOrTarget SourceOrTarget, FName& OutChainName) const
+{
+	OutChainName = NAME_None;
+	const FString ChainNameToFind = "spine";
+	float HighestScore = -1.f;
+	const TArray<FRetargetChainPairFK>& ChainPairs = Processor->GetFKChainPairs();
+	for (const FRetargetChainPairFK& ChainPair : ChainPairs)
+	{
+		FName ChainNameToTest = SourceOrTarget == ERetargetSourceOrTarget::Source ? ChainPair.SourceBoneChainName : ChainPair.TargetBoneChainName;
+		FString ChainNameStr = ChainNameToTest.ToString().ToLower();
+		float WorstCase = ChainNameToFind.Len() + ChainNameStr.Len();
+		WorstCase = WorstCase < 1.0f ? 1.0f : WorstCase;
+		const float Score = 1.0f - (Algo::LevenshteinDistance(ChainNameToFind, ChainNameStr) / WorstCase);
+		if (Score > HighestScore && Score > 0.5f)
+		{
+			HighestScore = Score;
+			OutChainName = ChainNameToTest;
+		}
+	}
+
+	return OutChainName != NAME_None;
+}
+
+FVector FRetargetAutoPoseGenerator::GetChainDirection(const FName ChainName, const ERetargetSourceOrTarget SourceOrTarget) const
+{
+	const FTransform ChainStart = Processor->GetGlobalRetargetPoseAtParam(ChainName, 0.f, SourceOrTarget);
+	const FTransform ChainEnd = Processor->GetGlobalRetargetPoseAtParam(ChainName, 1.f, SourceOrTarget);
+	return (ChainEnd.GetLocation() - ChainStart.GetLocation()).GetSafeNormal();
+}
+
+bool FRetargetAutoPoseGenerator::GetChainTangent(
+	const FName ChainName,
+	const float Param,
+	const ERetargetSourceOrTarget SourceOrTarget,
+	FVector& OutDirection) const
+{
+	check(Param >= 0.f);
+	
+	// cannot use chain tangent at the end of the chain
+	if (Param + ParamStepSize >= 1.0f)
+	{
+		return false;
+	}
+	
+	// use finite differencing to generate a tangent vector at the given param
+	const float ParamA = Param;
+	const float ParamB = Param + ParamStepSize;
+	const FTransform SampleA = Processor->GetGlobalRetargetPoseAtParam(ChainName, ParamA, SourceOrTarget);
+	const FTransform SampleB = Processor->GetGlobalRetargetPoseAtParam(ChainName, ParamB, SourceOrTarget);
+	OutDirection = (SampleB.GetLocation() - SampleA.GetLocation()).GetSafeNormal();
+	return true;
+}
+
+bool FRetargetAutoPoseGenerator::GetChainTangentFromMesh(
+	const FName ChainName,
+	const float Param,
+	const ERetargetSourceOrTarget SourceOrTarget,
+	FVector& OutDirection) const
+{
+	// get the bone at this location in the chain
+	const FName BoneName = Processor->GetBoneAtParam(ChainName, Param, SourceOrTarget);
+	
+	// get all children of bone
+	TArray<FName> BonesToConsider = {BoneName};
+	GetAllChildrenOfBone(BoneName, SourceOrTarget, BonesToConsider);
+	
+	// get all vertices primarily weighted to [children+bone]
+	TArray<FVector> Points;
+	GetPointsWeightedToBones(BonesToConsider, SourceOrTarget, Points);
+	
+	// cannot get bone direction from mesh since it's not connected to the mesh
+	if (Points.IsEmpty())
+	{
+		OutDirection = FVector::ZeroVector;
+		return false;
+	}
+
+	// get centroid of point cloud
+	FVector Centroid(0, 0, 0);
+	for (const FVector& Point : Points)
+	{
+		Centroid += Point;
+	}
+	Centroid /= static_cast<float>(Points.Num());
+
+	// return vector from bone to centroid of mesh
+	const FRetargetSkeleton& Skeleton = Processor->GetSkeleton(SourceOrTarget);
+	const int32 BoneIndex = Skeleton.FindBoneIndexByName(BoneName);
+	const FVector BonePosition = Skeleton.RetargetGlobalPose[BoneIndex].GetLocation();
+	OutDirection = (Centroid - BonePosition).GetSafeNormal();
+	return true;
+}
+
+void FRetargetAutoPoseGenerator::GetPointsWeightedToBones(
+	const TArray<FName>& BonesToConsider,
+	const ERetargetSourceOrTarget SourceOrTarget,
+	TArray<FVector>& OutPoints) const
+{
+	// get all the skin weights and ref pose vertex positions of the mesh
+	TArray<TMap<FName, float>> SkinWeights;
+	TArray<FVector> VertexRefPose;
+	GetRefPoseVerticesAndWeights(SourceOrTarget, SkinWeights, VertexRefPose);
+	
+	// find vertices weighted to the bones we are considering
+	TArray<int32> VerticesToConsider;
+	for (int32 VertexIndex = 0; VertexIndex < SkinWeights.Num(); ++VertexIndex)
+	{
+		const TMap<FName, float>& VertexWeights = SkinWeights[VertexIndex];
+
+		// find the highest weight bone for this vertex
+		FName MaxBone;
+		float MaxWeight = -1.0f;
+		for (const TTuple<FName, float>& Pair : VertexWeights)
+		{
+			if (Pair.Value > MaxWeight)
+			{
+				MaxBone = Pair.Key;
+				MaxWeight = Pair.Value;
+			}
+		}
+
+		// check if the highest weight bone is in BonesToConsider
+		if (BonesToConsider.Contains(MaxBone))
+		{
+			VerticesToConsider.Add(VertexIndex);
+		}
+	}
+	
+	// get the component space bone transform for the ref pose
+	const FRetargetSkeleton& Skeleton = Processor->GetSkeleton(SourceOrTarget);
+	const TArray<FTransform>& RefPoseLocal = Skeleton.SkeletalMesh->GetRefSkeleton().GetRefBonePose();
+	TArray<FTransform> ComponentSpaceRefPoseBoneTransforms;
+	FAnimationRuntime::FillUpComponentSpaceTransforms(
+		Skeleton.SkeletalMesh->GetRefSkeleton(),
+		RefPoseLocal,
+		ComponentSpaceRefPoseBoneTransforms);
+
+	// calculate the deformed vertex positions in the current retarget pose
+	GetDeformedVertexPositions(
+		VerticesToConsider,
+		VertexRefPose,
+		ComponentSpaceRefPoseBoneTransforms,
+		Skeleton.RetargetGlobalPose,
+		SkinWeights,
+		SourceOrTarget,
+		OutPoints);
+}
+
+void FRetargetAutoPoseGenerator::GetAllChildrenOfBone(
+	const FName BoneName,
+	const ERetargetSourceOrTarget SourceOrTarget,
+	TArray<FName>& OutAllChildren) const
+{
+	const FRetargetSkeleton& Skeleton = Processor->GetSkeleton(SourceOrTarget);
+	const int32 BoneIndex = Skeleton.FindBoneIndexByName(BoneName);
+	TArray<int32> ChildIndices;
+	Skeleton.GetChildrenIndicesRecursive(BoneIndex, ChildIndices);
+	for (const int32 ChildBoneIndex : ChildIndices)
+	{
+		OutAllChildren.Add(Skeleton.BoneNames[ChildBoneIndex]);	
+	}
+}
+
+void FRetargetAutoPoseGenerator::GetRefPoseVerticesAndWeights(
+	const ERetargetSourceOrTarget SourceOrTarget,
+	TArray<TMap<FName, float>>& OutWeights,
+	TArray<FVector>& OutPositions) const
+{
+	OutWeights.Reset();
+	OutPositions.Reset();
+	
+	const FRetargetSkeleton& Skeleton = Processor->GetSkeleton(SourceOrTarget);
+	USkeletalMesh* Mesh = Skeleton.SkeletalMesh;
+	constexpr int32 LODIndex = 0;
+	FMeshDescription MeshDescription;
+	Mesh->GetMeshDescription(LODIndex, MeshDescription);
+	if (MeshDescription.IsEmpty())
+	{
+		return;
+	}
+	
+	const FSkeletalMeshConstAttributes MeshAttribs(MeshDescription);
+	const FSkinWeightsVertexAttributesConstRef VertexSkinWeights = MeshAttribs.GetVertexSkinWeights();
+	const int32 NumVertices = MeshDescription.Vertices().Num();
+	OutWeights.SetNum(NumVertices);
+	OutPositions.SetNum(NumVertices);
+	const FReferenceSkeleton& RefSkeleton = Mesh->GetRefSkeleton();
+	ParallelFor(NumVertices,[RefSkeleton, &VertexSkinWeights, &OutWeights, &MeshDescription, &OutPositions](int32 VertexIndex)
+	{
+		const FVertexID VertexID(VertexIndex);
+
+		// get weights of vertex
+		int32 InfluenceIndex = 0;
+		for (UE::AnimationCore::FBoneWeight BoneWeight: VertexSkinWeights.Get(VertexID))
+		{
+			const int32 BoneIndex = BoneWeight.GetBoneIndex();
+			const FName BoneName = RefSkeleton.GetBoneName(BoneIndex);
+			const float Weight = BoneWeight.GetWeight();
+			OutWeights[VertexIndex].Add(BoneName, Weight);
+			++InfluenceIndex;
+		}
+
+		// get position of vertex
+		const FVector Position = static_cast<FVector>(MeshDescription.GetVertexPosition(VertexID));
+		OutPositions[VertexIndex] = Position;
+	});
+}
+
+void FRetargetAutoPoseGenerator::GetDeformedVertexPositions(
+	const TArray<int32>& VerticesToDeform,
+	const TArray<FVector>& VertexRefPose,
+	const TArray<FTransform>& RefPoseBoneTransforms,
+	const TArray<FTransform>& CurrentBoneTransforms,
+	const TArray<TMap<FName, float>>& SkinWeights,
+	const ERetargetSourceOrTarget SourceOrTarget,
+	TArray<FVector>& OutDeformedPositions) const
+{
+	OutDeformedPositions.Reset(VerticesToDeform.Num());
+	
+	// generate map of bone name to bone index for fast look-up
+	TMap<FName, int32> BoneNameToIndex;
+	const FRetargetSkeleton& Skeleton = Processor->GetSkeleton(SourceOrTarget);
+	for (int32 BoneIndex=0; BoneIndex<Skeleton.BoneNames.Num(); ++BoneIndex)
+	{
+		BoneNameToIndex.Add(Skeleton.BoneNames[BoneIndex], BoneIndex);
+	}
+
+	// generate deformed vertex position for each vertex we care about
+	for (const int32 VertexIndex : VerticesToDeform)
+	{
+		const FVector& RefVertexPosition = VertexRefPose[VertexIndex];
+		FVector DeformedPosition = FVector::ZeroVector;
+
+		const TMap<FName, float>& VertexWeights = SkinWeights[VertexIndex];
+		for (const TTuple<FName, float>& Pair : VertexWeights)
+		{
+			const int32 BoneIndex = BoneNameToIndex[Pair.Key];
+			const float Weight = Pair.Value;
+			FVector VertexInBoneSpace = RefPoseBoneTransforms[BoneIndex].Inverse().TransformPosition(RefVertexPosition);
+			FVector BoneDeformedVertex = CurrentBoneTransforms[BoneIndex].TransformPosition(VertexInBoneSpace);
+			DeformedPosition += BoneDeformedVertex * Weight;
+		}
+
+		OutDeformedPositions.Add(DeformedPosition);
+	}
+}
+
+bool FRetargetAutoPoseGenerator::IsARetargetLeaf(
+	const FName ChainName, 
+	const float Param,
+	ERetargetSourceOrTarget SourceOrTarget) const
+{
+	const FName BoneName = Processor->GetBoneAtParam(ChainName, Param, SourceOrTarget);
+	TArray<FName> AllChildren;
+	GetAllChildrenOfBone(BoneName, SourceOrTarget, AllChildren);
+	for (const FName Child : AllChildren)
+	{
+		if (Processor->IsBoneRetargeted(Child, SourceOrTarget))
+		{
+			return false;
+		}
+	}
+	return true;
+}
+
+bool FRetargetAutoPoseGenerator::CheckReadyToAlignBones() const
+{
+	// this should never happen
+	if (!EditorController.IsValid())
+	{
+		return false;
+	}
+	
+	// can't auto align until processor is initialized because we need a fully resolved chain mapping on source/target
+	// this could happen if the user attempts to edit the retarget pose on an incompatible IK Rig
+	const FIKRetargetEditorController* Controller = EditorController.Pin().Get();
+	Processor = Controller->GetRetargetProcessor();
+	if (!Processor && Processor->IsInitialized())
+	{
+		return false;
+	}
+
+	return true;
+}
+
+void FRetargetAutoPoseGenerator::ApplyWorldRotationToRetargetPose(
+	const FQuat& GlobalDeltaRotation,
+	const FName BoneToAffect,
+	const ERetargetSourceOrTarget SourceOrTarget) const
+{
+	const TObjectPtr<UIKRetargeterController> AssetController = EditorController.Pin()->AssetController;
+	const FRetargetSkeleton& Skeleton =  Processor->GetSkeleton(SourceOrTarget);
+	const FReferenceSkeleton& RefSkeleton = Skeleton.SkeletalMesh->GetRefSkeleton();
+	const int32 BoneIndex = Skeleton.FindBoneIndexByName(BoneToAffect);
+
+	// get the curren parent global transform
+	FTransform ParentGlobalTransform = FTransform::Identity;
+	const int32 ParentIndex = RefSkeleton.GetParentIndex(BoneIndex);
+	if (ParentIndex != INDEX_NONE)
+	{
+		ParentGlobalTransform = Skeleton.RetargetGlobalPose[ParentIndex];
+	}
+
+	// get the current global transform with no delta's at all
+	FTransform LocalRefTransform = RefSkeleton.GetRefBonePose()[BoneIndex];
+	FTransform CurrentGlobalTransformNoDelta = LocalRefTransform * ParentGlobalTransform;
+				
+	// get the stored local rotation offset from the retarget pose
+	const FQuat LocalRotationDeltaFromPose = AssetController->GetRotationOffsetForRetargetPoseBone(BoneToAffect, SourceOrTarget);
+	// apply the retarget pose to get the global rotation w/ retarget pose
+	// TODO get this from retarget pose??
+	FQuat GlobalRetargetPoseRotation = CurrentGlobalTransformNoDelta.GetRotation() * LocalRotationDeltaFromPose;
+	// get global delta from the retarget pose
+	FQuat GlobalDeltaRotationFromRetargetPose = GlobalRetargetPoseRotation * CurrentGlobalTransformNoDelta.GetRotation().Inverse();
+	// combine it with the new delta from the alignment
+	FQuat NewGlobalDeltaRotation = GlobalDeltaRotation * GlobalDeltaRotationFromRetargetPose;
+
+	// convert world space delta quaternion to bone-space
+	const FVector RotationAxis = NewGlobalDeltaRotation.GetRotationAxis();
+	const FVector UnRotatedAxis = CurrentGlobalTransformNoDelta.InverseTransformVector(RotationAxis);
+	FQuat FinalLocalDeltaRotation = FQuat(UnRotatedAxis, NewGlobalDeltaRotation.GetAngle());
+
+	// store the new rotation in the retarget pose
+	AssetController->SetRotationOffsetForRetargetPoseBone(BoneToAffect, FinalLocalDeltaRotation, SourceOrTarget);
+
+	// regenerate the retarget pose in the currently running processor so it reflects the latest change
+	const FName CurrentPoseName = AssetController->GetCurrentRetargetPoseName(SourceOrTarget);
+	Processor->UpdateRetargetPoseAtRuntime(CurrentPoseName, SourceOrTarget);
+}
+
+#undef LOCTEXT_NAMESPACE
\ No newline at end of file
diff --git a/Engine/Plugins/Animation/IKRig/Source/IKRigEditor/Private/RetargetEditor/SIKRetargetPoseEditor.cpp b/Engine/Plugins/Animation/IKRig/Source/IKRigEditor/Private/RetargetEditor/SIKRetargetPoseEditor.cpp
--- a/Engine/Plugins/Animation/IKRig/Source/IKRigEditor/Private/RetargetEditor/SIKRetargetPoseEditor.cpp
+++ b/Engine/Plugins/Animation/IKRig/Source/IKRigEditor/Private/RetargetEditor/SIKRetargetPoseEditor.cpp
@@ -119,6 +119,13 @@ TSharedRef<SWidget>  SIKRetargetPoseEditor::MakeToolbar(TSharedPtr<FUICommandLis
 		LOCTEXT("ResetPoseToolTip_Label", "Reset bones to reference pose."),
 		FSlateIcon(FAppStyle::GetAppStyleSetName(), "Icons.Refresh"));
 
+	ToolbarBuilder.AddComboButton(
+		FUIAction(),
+		FOnGetContent::CreateSP(this, &SIKRetargetPoseEditor::GenerateEditMenuContent, Commands),
+		LOCTEXT("ResetPose_Label", "Auto Align"),
+		LOCTEXT("ResetPoseToolTip_Label", "Automatically aligns bones on source skeleton to target (or vice versa)."),
+		FSlateIcon(FAppStyle::GetAppStyleSetName(), "Icons.Edit"));
+
 	ToolbarBuilder.EndSection();
 
 	ToolbarBuilder.BeginSection("Create Poses");
@@ -152,24 +159,27 @@ TSharedRef<SWidget>  SIKRetargetPoseEditor::MakeToolbar(TSharedPtr<FUICommandLis
 TSharedRef<SWidget> SIKRetargetPoseEditor::GenerateResetMenuContent(TSharedPtr<FUICommandList> Commands)
 {
 	FMenuBuilder MenuBuilder(true, Commands);
+	MenuBuilder.AddMenuEntry(FIKRetargetCommands::Get().ResetSelectedBones);
+	MenuBuilder.AddMenuEntry(FIKRetargetCommands::Get().ResetSelectedAndChildrenBones);
+	MenuBuilder.AddMenuEntry(FIKRetargetCommands::Get().ResetAllBones);
+	return MenuBuilder.MakeWidget();
+}
 
-	MenuBuilder.AddMenuEntry(
-		FIKRetargetCommands::Get().ResetSelectedBones,
-		TEXT("Reset Selected"),
-		TAttribute<FText>(),
-		TAttribute<FText>());
-
-	MenuBuilder.AddMenuEntry(
-		FIKRetargetCommands::Get().ResetSelectedAndChildrenBones,
-		TEXT("Reset Selected And Children"),
-		TAttribute<FText>(),
-		TAttribute<FText>());
+TSharedRef<SWidget> SIKRetargetPoseEditor::GenerateEditMenuContent(TSharedPtr<FUICommandList> Commands)
+{
+	FMenuBuilder MenuBuilder(true, Commands);
 
-	MenuBuilder.AddMenuEntry(
-		FIKRetargetCommands::Get().ResetAllBones,
-		TEXT("Reset All"),
-		TAttribute<FText>(),
-		TAttribute<FText>());
+	MenuBuilder.AddMenuEntry(FIKRetargetCommands::Get().AutoAlignAllBones);
+	MenuBuilder.AddSeparator();
+	
+	MenuBuilder.AddMenuEntry(FIKRetargetCommands::Get().AlignSelected);
+	MenuBuilder.AddMenuEntry(FIKRetargetCommands::Get().AlignSelectedAndChildren);
+	MenuBuilder.AddSeparator();
+	
+	MenuBuilder.AddMenuEntry(FIKRetargetCommands::Get().AlignSelectedUsingMesh);
+	MenuBuilder.AddSeparator();
+	
+	MenuBuilder.AddMenuEntry(FIKRetargetCommands::Get().SnapCharacterToGround);
 
 	return MenuBuilder.MakeWidget();
 }
@@ -181,30 +191,15 @@ TSharedRef<SWidget> SIKRetargetPoseEditor::GenerateNewMenuContent(TSharedPtr<FUI
 	UToolMenu* ToolMenu = UToolMenus::Get()->ExtendMenu(MenuName);
 
 	FToolMenuSection& CreateSection = ToolMenu->AddSection("Create", LOCTEXT("CreatePoseOperations", "Create New Retarget Pose"));
-	CreateSection.AddMenuEntry(
-		FIKRetargetCommands::Get().NewRetargetPose,
-		TAttribute<FText>(),
-        TAttribute<FText>());
-	CreateSection.AddMenuEntry(
-		FIKRetargetCommands::Get().DuplicateRetargetPose,
-		TAttribute<FText>(),
-		TAttribute<FText>());
+	CreateSection.AddMenuEntry(FIKRetargetCommands::Get().NewRetargetPose);
+	CreateSection.AddMenuEntry(FIKRetargetCommands::Get().DuplicateRetargetPose);
 	
 	FToolMenuSection& ImportSection = ToolMenu->AddSection("Import",LOCTEXT("ImportPoseOperations", "Import Retarget Pose"));
-	ImportSection.AddMenuEntry(
-			FIKRetargetCommands::Get().ImportRetargetPose,
-			TAttribute<FText>(),
-			TAttribute<FText>());
-	ImportSection.AddMenuEntry(
-		FIKRetargetCommands::Get().ImportRetargetPoseFromAnim,
-		TAttribute<FText>(),
-		TAttribute<FText>());
+	ImportSection.AddMenuEntry(FIKRetargetCommands::Get().ImportRetargetPose);
+	ImportSection.AddMenuEntry(FIKRetargetCommands::Get().ImportRetargetPoseFromAnim);
 
 	FToolMenuSection& ExportSection = ToolMenu->AddSection("Export",LOCTEXT("ExportPoseOperations", "Export Retarget Pose"));
-	ExportSection.AddMenuEntry(
-			FIKRetargetCommands::Get().ExportRetargetPose,
-			TAttribute<FText>(),
-			TAttribute<FText>());
+	ExportSection.AddMenuEntry(FIKRetargetCommands::Get().ExportRetargetPose);
 
 	return UToolMenus::Get()->GenerateWidget(MenuName, FToolMenuContext(Commands));
 }
diff --git a/Engine/Plugins/Animation/IKRig/Source/IKRigEditor/Public/RetargetEditor/IKRetargetCommands.h b/Engine/Plugins/Animation/IKRig/Source/IKRigEditor/Public/RetargetEditor/IKRetargetCommands.h
--- a/Engine/Plugins/Animation/IKRig/Source/IKRigEditor/Public/RetargetEditor/IKRetargetCommands.h
+++ b/Engine/Plugins/Animation/IKRig/Source/IKRigEditor/Public/RetargetEditor/IKRetargetCommands.h
@@ -50,6 +50,17 @@ public:
 	// reset retarget pose to ref pose
 	TSharedPtr< FUICommandInfo > ResetSelectedAndChildrenBones;
 
+	// auto generate retarget pose
+	TSharedPtr< FUICommandInfo > AutoAlignAllBones;
+	// auto align selected bones
+	TSharedPtr< FUICommandInfo > AlignSelected;
+	// auto align selected bones and children
+	TSharedPtr< FUICommandInfo > AlignSelectedAndChildren;
+	// auto align selected bones
+	TSharedPtr< FUICommandInfo > AlignSelectedUsingMesh;
+	// auto align selected bones
+	TSharedPtr< FUICommandInfo > SnapCharacterToGround;
+	
 	// delete retarget pose
 	TSharedPtr< FUICommandInfo > DeleteRetargetPose;
 	// rename retarget pose
diff --git a/Engine/Plugins/Animation/IKRig/Source/IKRigEditor/Public/RetargetEditor/IKRetargetEditorController.h b/Engine/Plugins/Animation/IKRig/Source/IKRigEditor/Public/RetargetEditor/IKRetargetEditorController.h
--- a/Engine/Plugins/Animation/IKRig/Source/IKRigEditor/Public/RetargetEditor/IKRetargetEditorController.h
+++ b/Engine/Plugins/Animation/IKRig/Source/IKRigEditor/Public/RetargetEditor/IKRetargetEditorController.h
@@ -2,6 +2,7 @@
 #pragma once
 
 #include "IKRetargetDetails.h"
+#include "IKRetargeterPoseGenerator.h"
 #include "IKRetargetPoseExporter.h"
 #include "IPersonaToolkit.h"
 #include "SIKRetargetAssetBrowser.h"
@@ -64,7 +65,7 @@ struct FBoundIKRig
 
 struct FRetargetPlaybackManager : public TSharedFromThis<FRetargetPlaybackManager>
 {
-	FRetargetPlaybackManager(TWeakPtr<FIKRetargetEditorController> InEditorController);
+	FRetargetPlaybackManager(const TWeakPtr<FIKRetargetEditorController>& InEditorController);
 	void PlayAnimationAsset(UAnimationAsset* AssetToPlay);
 	void StopPlayback();
 	void PausePlayback();
@@ -120,6 +121,8 @@ public:
 
 	// import / export retarget poses
 	TSharedPtr<FIKRetargetPoseExporter> PoseExporter;
+	// auto pose generator
+	TUniquePtr<FRetargetAutoPoseGenerator> AutoPoseGenerator;
 
 	// manage playback of animation in the editor
 	TUniquePtr<FRetargetPlaybackManager> PlaybackManager;
@@ -253,6 +256,9 @@ public:
 	// to frame selection when pressing "f" in viewport
 	bool GetCameraTargetForSelection(FSphere& OutTarget) const;
 
+	// check if any bone is selected
+	bool IsAnyBoneSelected() const;
+
 	// ------------------------- END SELECTION -----------------------------
 
 	// determine if bone in the specified skeleton is part of the retarget (in a mapped chain)
@@ -274,7 +280,11 @@ public:
 	void HandleResetAllBones() const;
 	void HandleResetSelectedBones() const;
 	void HandleResetSelectedAndChildrenBones() const;
-	bool CanResetSelected() const;
+	
+	// auto generate retarget pose
+	void HandleAlignAllBones() const;
+	void HandleAlignSelectedBones(const ERetargetAutoAlignMethod Method, const bool bIncludeChildren) const;
+	void HandleSnapToGround() const;
 
 	// create new retarget pose
 	void HandleNewPose();
@@ -313,6 +323,8 @@ public:
 
 private:
 
+	TArray<FName> GetSelectedBonesAndChildren() const;
+
 	// modal dialog to ask user if they want to fix root bones that are "on the ground"
 	bool PromptToFixRootHeight(ERetargetSourceOrTarget SourceOrTarget) const;
 	bool bAskedToFixRoot = false;
diff --git a/Engine/Plugins/Animation/IKRig/Source/IKRigEditor/Public/RetargetEditor/IKRetargeterPoseGenerator.h b/Engine/Plugins/Animation/IKRig/Source/IKRigEditor/Public/RetargetEditor/IKRetargeterPoseGenerator.h
new file mode 100644
--- /dev/null
+++ b/Engine/Plugins/Animation/IKRig/Source/IKRigEditor/Public/RetargetEditor/IKRetargeterPoseGenerator.h
@@ -0,0 +1,125 @@
+﻿// Copyright Epic Games, Inc. All Rights Reserved.
+#pragma once
+
+#include "CoreMinimal.h"
+
+class FIKRetargetEditorController;
+class UIKRetargetProcessor;
+enum class ERetargetSourceOrTarget : uint8;
+
+// A Note on Determining the "Direction" of a Bone...
+//
+// What are we using to determine the "direction" of the source and target bones to align?
+// If chain only has one bone OR this bone is located at the end of a chain, then we cannot compute a tangent!
+// So we fallback to an analysis of the geometry... this is necessary because we cannot rely solely on the shape of the skeleton.
+// Any assumptions we make could be invalidated in different circumstances. For example,
+// 1. we cannot assume the local axes are aligned in any meaningful way, and even if they are we cannot assume WHICH axis is the correct one
+// 2. we cannot assume that any child bones are located in a position that would create a meaningful vector (ie twist bones in arbitrary locations)
+// 3. we cannot assume that extrapolating the direction from the parent is meaningful. For example, feet are usually perpendicular to the knee.
+//
+// Since the skeleton cannot be relied upon to provide robust directional information for individual and/or leaf bones, we must instead fallback
+// on looking at the mesh that is skinned to it and try to discern a direction vector from that. If there is no skinning data, then we simply
+// do NOT auto-align that bone.
+enum class ERetargetAutoAlignMethod : uint8
+{
+	// use the chain to determine the source and target directions to align
+	ChainToChain,
+	// use mesh direction as source, and mesh direction as the target to align
+	MeshToMesh,
+};
+
+// a sub-system used by the IKRetargetEditorController for automatically generating a retarget pose
+struct FRetargetAutoPoseGenerator
+{
+	// created and owned by the retarget editor controller
+	FRetargetAutoPoseGenerator(const TWeakPtr<FIKRetargetEditorController>& InEditorController);
+	
+	// align the specified bones on the current skeleton to the OTHER skeleton, using the specified method
+	void AlignBones(
+		const TArray<FName>& BonesToAlign,
+		const ERetargetAutoAlignMethod Method,
+		const ERetargetSourceOrTarget SourceOrTarget,
+		const bool bSuppressWarnings) const;
+
+	// adjust the vertical height of the retarget root to place the character back on the ground
+	// if BoneToPutOnGround is not specified (NAME_None), then will find the lowest bone and use that
+	void SnapToGround(const FName BoneToPutOnGround, const ERetargetSourceOrTarget SourceOrTarget) const;
+
+private:
+	
+	// generates a local rotation offset to align BoneToAlign to the opposite skeleton, and stores it in the current retarget pose
+	void AlignBone(
+		const FName BoneToAlign,
+		const ERetargetAutoAlignMethod Method,
+		const ERetargetSourceOrTarget SourceOrTarget) const;
+
+	// generates a local rotation offset to align the retarget root based on the spine orientation
+	void AlignRetargetRootBone(const ERetargetSourceOrTarget SourceOrTarget) const;
+
+	// get the chain most likely to be the spine
+	bool GetSpineChain(const ERetargetSourceOrTarget SourceOrTarget, FName& OutChainName) const;
+
+	// get the chain most likely to be the spine
+	FVector GetChainDirection(const FName ChainName, const ERetargetSourceOrTarget SourceOrTarget) const;
+	
+	// takes a world space rotation offset, converts it to local space and stores it in the current retarget pose
+	void ApplyWorldRotationToRetargetPose(
+		const FQuat& WorldRotationOffset,
+		const FName BoneToAffect,
+		const ERetargetSourceOrTarget SourceOrTarget) const;
+
+	// get the tangent vector of a retarget chain at a given param (0-1)
+	bool GetChainTangent(
+		const FName ChainName,
+		const float Param,
+		const ERetargetSourceOrTarget SourceOrTarget,
+		FVector& OutDirection) const;
+
+	// discern the direction of a bone based on the mesh it is skinned to
+	bool GetChainTangentFromMesh(
+		const FName ChainName,
+		const float Param,
+		const ERetargetSourceOrTarget SourceOrTarget,
+		FVector& OutDirection) const;
+
+	// get the locations of vertices skinned to a sub-set of bones
+	void GetPointsWeightedToBones(
+		const TArray<FName>& Bones,
+		const ERetargetSourceOrTarget SourceOrTarget,
+		TArray<FVector>& OutVertexPositions) const;
+
+	// get all children (recursive) of the given bone
+	void GetAllChildrenOfBone(
+		const FName Bone,
+		const ERetargetSourceOrTarget SourceOrTarget,
+		TArray<FName>& OutAllChildren) const;
+
+	// get the vertex ref pose and skinning weights for a given mesh (source or target)
+	void GetRefPoseVerticesAndWeights(
+		const ERetargetSourceOrTarget SourceOrTarget,
+		TArray<TMap<FName, float>>& OutWeights,
+		TArray<FVector>& OutPositions) const;
+
+	// calculates the skinned vertex positions for a subset of the mesh using the supplied bone transforms
+	void GetDeformedVertexPositions(
+		const TArray<int32>& VerticesToDeform,
+		const TArray<FVector>& VertexRefPose,
+		const TArray<FTransform>& RefPoseBoneTransforms,
+		const TArray<FTransform>& CurrentBoneTransforms,
+		const TArray<TMap<FName, float>>& SkinWeights,
+		const ERetargetSourceOrTarget SourceOrTarget,
+		TArray<FVector>& OutDeformedPositions) const;
+
+	// gets the bone at this point in the chain, returns true if non of it's children are retargeted 
+	bool IsARetargetLeaf(
+		const FName ChainName,
+		const float Param,
+		ERetargetSourceOrTarget SourceOrTarget) const;
+
+	// get the processor from the editor and return true if it's initialized
+	bool CheckReadyToAlignBones() const;
+
+	TWeakPtr<FIKRetargetEditorController> EditorController;
+	mutable UIKRetargetProcessor* Processor;
+	static constexpr float ParamStepSize = 0.1f;
+};
diff --git a/Engine/Plugins/Animation/IKRig/Source/IKRigEditor/Public/RetargetEditor/SIKRetargetPoseEditor.h b/Engine/Plugins/Animation/IKRig/Source/IKRigEditor/Public/RetargetEditor/SIKRetargetPoseEditor.h
--- a/Engine/Plugins/Animation/IKRig/Source/IKRigEditor/Public/RetargetEditor/SIKRetargetPoseEditor.h
+++ b/Engine/Plugins/Animation/IKRig/Source/IKRigEditor/Public/RetargetEditor/SIKRetargetPoseEditor.h
@@ -27,6 +27,7 @@ private:
 	
 	TSharedRef<SWidget> MakeToolbar(TSharedPtr<FUICommandList> Commands);
 	TSharedRef<SWidget> GenerateResetMenuContent(TSharedPtr<FUICommandList> Commands);
+	TSharedRef<SWidget> GenerateEditMenuContent(TSharedPtr<FUICommandList> Commands);
 	TSharedRef<SWidget> GenerateNewMenuContent(TSharedPtr<FUICommandList> Commands);
 	TObjectPtr<UIKRetargeterController> GetAssetControllerFromSelectedObjects(IDetailLayoutBuilder& DetailBuilder) const;
 	
