Localization conflicts in retarget editor.

diff --git a/Engine/Plugins/Animation/IKRig/Source/IKRigEditor/Private/RetargetEditor/IKRetargetEditorController.cpp b/Engine/Plugins/Animation/IKRig/Source/IKRigEditor/Private/RetargetEditor/IKRetargetEditorController.cpp
--- a/Engine/Plugins/Animation/IKRig/Source/IKRigEditor/Private/RetargetEditor/IKRetargetEditorController.cpp
+++ b/Engine/Plugins/Animation/IKRig/Source/IKRigEditor/Private/RetargetEditor/IKRetargetEditorController.cpp
@@ -1332,7 +1332,7 @@ void FIKRetargetEditorController::HandleAlignAllBones() const
 
 	// undo transaction
 	constexpr bool bShouldTransact = true;
-	FScopedTransaction Transaction(LOCTEXT("AutoAlignSelectedBones", "Auto Align All Bones"), bShouldTransact);
+	FScopedTransaction Transaction(LOCTEXT("AutoAlignAllBones", "Auto Align All Bones"), bShouldTransact);
 	AssetController->GetAsset()->Modify();
 	
 	// suppress warnings about bones that cannot be aligned when aligning ALL bones
diff --git a/Engine/Plugins/Animation/IKRig/Source/IKRigEditor/Private/RetargetEditor/SIKRetargetPoseEditor.cpp b/Engine/Plugins/Animation/IKRig/Source/IKRigEditor/Private/RetargetEditor/SIKRetargetPoseEditor.cpp
--- a/Engine/Plugins/Animation/IKRig/Source/IKRigEditor/Private/RetargetEditor/SIKRetargetPoseEditor.cpp
+++ b/Engine/Plugins/Animation/IKRig/Source/IKRigEditor/Private/RetargetEditor/SIKRetargetPoseEditor.cpp
@@ -122,8 +122,8 @@ TSharedRef<SWidget>  SIKRetargetPoseEditor::MakeToolbar(TSharedPtr<FUICommandLis
 	ToolbarBuilder.AddComboButton(
 		FUIAction(),
 		FOnGetContent::CreateSP(this, &SIKRetargetPoseEditor::GenerateEditMenuContent, Commands),
-		LOCTEXT("ResetPose_Label", "Auto Align"),
-		LOCTEXT("ResetPoseToolTip_Label", "Automatically aligns bones on source skeleton to target (or vice versa)."),
+		LOCTEXT("AutoAlign_Label", "Auto Align"),
+		LOCTEXT("AutoAlignTip_Label", "Automatically aligns bones on source skeleton to target (or vice versa)."),
 		FSlateIcon(FAppStyle::GetAppStyleSetName(), "Icons.Edit"));
 
 	ToolbarBuilder.EndSection();
