Fixes static analysis issue in retarget pose generator.

diff --git a/Engine/Plugins/Animation/IKRig/Source/IKRigEditor/Private/RetargetEditor/IKRetargeterPoseGenerator.cpp b/Engine/Plugins/Animation/IKRig/Source/IKRigEditor/Private/RetargetEditor/IKRetargeterPoseGenerator.cpp
--- a/Engine/Plugins/Animation/IKRig/Source/IKRigEditor/Private/RetargetEditor/IKRetargeterPoseGenerator.cpp
+++ b/Engine/Plugins/Animation/IKRig/Source/IKRigEditor/Private/RetargetEditor/IKRetargeterPoseGenerator.cpp
@@ -530,7 +530,7 @@ bool FRetargetAutoPoseGenerator::CheckReadyToAlignBones() const
 	// this could happen if the user attempts to edit the retarget pose on an incompatible IK Rig
 	const FIKRetargetEditorController* Controller = EditorController.Pin().Get();
 	Processor = Controller->GetRetargetProcessor();
-	if (!Processor && Processor->IsInitialized())
+	if (!(Processor && Processor->IsInitialized()))
 	{
 		return false;
 	}
