Fixes crash if bones added or removed from skeleton while IK Rig is active.

diff --git a/Engine/Plugins/Animation/IKRig/Source/IKRig/Private/AnimNodes/AnimNode_IKRig.cpp b/Engine/Plugins/Animation/IKRig/Source/IKRig/Private/AnimNodes/AnimNode_IKRig.cpp
--- a/Engine/Plugins/Animation/IKRig/Source/IKRig/Private/AnimNodes/AnimNode_IKRig.cpp
+++ b/Engine/Plugins/Animation/IKRig/Source/IKRig/Private/AnimNodes/AnimNode_IKRig.cpp
@@ -338,6 +338,9 @@ void FAnimNode_IKRig::CacheBones_AnyThread(const FAnimationCacheBonesContext& Co
 		const FName Name = MeshRefSkeleton.GetBoneName(MeshBone);
 		CompactPoseToRigIndices.Add(CPIndex) = MeshRefSkeleton.FindBoneIndex(Name);
 	}
+
+	// must reinitialize if bone count changes
+	IKRigProcessor->SetNeedsInitialized();
 }
 
 void FAnimNode_IKRig::InitializeProperties(const UObject* InSourceInstance, UClass* InTargetClass)
