Fix to support scale on skeleton when retargeting.

diff --git a/Engine/Plugins/Animation/IKRig/Source/IKRig/Private/AnimNodes/AnimNode_RetargetPoseFromMesh.cpp b/Engine/Plugins/Animation/IKRig/Source/IKRig/Private/AnimNodes/AnimNode_RetargetPoseFromMesh.cpp
--- a/Engine/Plugins/Animation/IKRig/Source/IKRig/Private/AnimNodes/AnimNode_RetargetPoseFromMesh.cpp
+++ b/Engine/Plugins/Animation/IKRig/Source/IKRig/Private/AnimNodes/AnimNode_RetargetPoseFromMesh.cpp
@@ -141,6 +141,18 @@ void FAnimNode_RetargetPoseFromMesh::Evaluate_AnyThread(FPoseContext& Output)
 	// convert to local space
 	FCSPose<FCompactPose>::ConvertComponentPosesToLocalPoses(ComponentPose, Output.Pose);
 
+	// once converted back to local space, we copy scale values back
+	// (retargeter strips scale values and deals with translation only in component space)
+	const TArray<FTransform>& RefPose = TargetMesh->GetRefSkeleton().GetRefBonePose();
+	for (const TPair<int32, int32>& Pair : RequiredToTargetBoneMapping)
+	{
+		const FCompactPoseBoneIndex CompactBoneIndex(Pair.Key);
+		if (Output.Pose.IsValidIndex(CompactBoneIndex))
+		{
+			Output.Pose[CompactBoneIndex].SetScale3D(RefPose[Pair.Value].GetScale3D());
+		}
+	}
+
 	// copy and/or remap curves from the source to the target skeletal mesh
 	CopyAndRemapCurvesFromSourceToTarget(Output.Curve);
 }
