IKRig: Fix localization tags.

#rnx

diff --git a/Engine/Plugins/Animation/IKRig/Source/IKRigEditor/Private/RigEditor/IKRigEditorController.cpp b/Engine/Plugins/Animation/IKRig/Source/IKRigEditor/Private/RigEditor/IKRigEditorController.cpp
--- a/Engine/Plugins/Animation/IKRig/Source/IKRigEditor/Private/RigEditor/IKRigEditorController.cpp
+++ b/Engine/Plugins/Animation/IKRig/Source/IKRigEditor/Private/RigEditor/IKRigEditorController.cpp
@@ -509,7 +509,7 @@ void FIKRigEditorController::AutoGenerateRetargetChains() const
 			for (const FName& MissingBone : Results.MissingBones)
 			{
 				Log.LogWarning(FText::Format(
-				LOCTEXT("MissingTemplateBone", "{0} but was not found in this skeleton, but is used by {1}."),
+				LOCTEXT("MissingTemplateBone", "{0} was not found in this skeleton, but is used by {1}."),
 				FText::FromName(MissingBone),
 				FText::FromName(Results.BestTemplateName)));
 			}
@@ -518,7 +518,7 @@ void FIKRigEditorController::AutoGenerateRetargetChains() const
 			for (const FName& MissingParent : Results.BonesWithMissingParent)
 			{
 				Log.LogWarning(FText::Format(
-				LOCTEXT("MissingTemplateBone", "{0} has a different parent in the template: {1}."),
+				LOCTEXT("DifferentParentTemplateBone", "{0} has a different parent in the template: {1}."),
 				FText::FromName(MissingParent),
 				FText::FromName(Results.BestTemplateName)));
 			}
@@ -542,7 +542,7 @@ void FIKRigEditorController::AutoGenerateRetargetChains() const
 	{
 		// notify user that no skeleton template was used
 		// TODO change this message once procedurally generated retarget definitions are provided
-		FNotificationInfo Info(LOCTEXT("AutoCharacterizeResults", "No matching skeletal template found. Characterization skipped."));
+		FNotificationInfo Info(LOCTEXT("MissingTemplateCharacterizeSkipped", "No matching skeletal template found. Characterization skipped."));
 		Info.ExpireDuration = 5.0f;
 		FSlateNotificationManager::Get().AddNotification(Info);
 	}
