Adding BP/Python support for auto retarget pose generation.
Fixed bug, edit pose buttons not disabled while in "play" mode
Removed "Show Retarget Pose" mode as per discussion / design.
Fixed bug, auto-align all now resets the pose before running.

diff --git a/Engine/Plugins/Animation/IKRig/Source/IKRigEditor/Private/RetargetEditor/IKRetargetAnimInstanceProxy.cpp b/Engine/Plugins/Animation/IKRig/Source/IKRigEditor/Private/RetargetEditor/IKRetargetAnimInstanceProxy.cpp
--- a/Engine/Plugins/Animation/IKRig/Source/IKRigEditor/Private/RetargetEditor/IKRetargetAnimInstanceProxy.cpp
+++ b/Engine/Plugins/Animation/IKRig/Source/IKRigEditor/Private/RetargetEditor/IKRetargetAnimInstanceProxy.cpp
@@ -14,7 +14,7 @@ FIKRetargetAnimInstanceProxy::FIKRetargetAnimInstanceProxy(
 	: FAnimPreviewInstanceProxy(InAnimInstance),
 	PreviewPoseNode(InPreviewPoseNode),
 	RetargetNode(InRetargetNode),
-	OutputMode(ERetargeterOutputMode::ShowRetargetPose)
+	OutputMode(ERetargeterOutputMode::EditRetargetPose)
 {
 }
 
@@ -62,7 +62,6 @@ bool FIKRetargetAnimInstanceProxy::Evaluate(FPoseContext& Output)
 			break;
 		}
 	case ERetargeterOutputMode::EditRetargetPose:
-	case ERetargeterOutputMode::ShowRetargetPose:
 		{
 			PreviewPoseNode->Evaluate_AnyThread(Output);
 			break;
diff --git a/Engine/Plugins/Animation/IKRig/Source/IKRigEditor/Private/RetargetEditor/IKRetargetEditor.cpp b/Engine/Plugins/Animation/IKRig/Source/IKRigEditor/Private/RetargetEditor/IKRetargetEditor.cpp
--- a/Engine/Plugins/Animation/IKRig/Source/IKRigEditor/Private/RetargetEditor/IKRetargetEditor.cpp
+++ b/Engine/Plugins/Animation/IKRig/Source/IKRigEditor/Private/RetargetEditor/IKRetargetEditor.cpp
@@ -125,12 +125,6 @@ void FIKRetargetEditor::BindCommands()
 		FExecuteAction::CreateSP(EditorController, &FIKRetargetEditorController::SetRetargeterMode, ERetargeterOutputMode::RunRetarget),
 		FCanExecuteAction::CreateSP(EditorController, &FIKRetargetEditorController::IsReadyToRetarget));
 
-	ToolkitCommands->MapAction(
-		Commands.ShowRetargetPose,
-		FExecuteAction::CreateSP(EditorController, &FIKRetargetEditorController::SetRetargeterMode, ERetargeterOutputMode::ShowRetargetPose),
-		FCanExecuteAction::CreateSP(EditorController, &FIKRetargetEditorController::IsCurrentMeshLoaded),
-		FIsActionChecked());
-
 	ToolkitCommands->MapAction(
 		Commands.EditRetargetPose,
 		FExecuteAction::CreateSP(EditorController, &FIKRetargetEditorController::SetRetargeterMode, ERetargeterOutputMode::EditRetargetPose),
@@ -189,47 +183,41 @@ void FIKRetargetEditor::BindCommands()
 	ToolkitCommands->MapAction(
 		Commands.ResetAllBones,
 		FExecuteAction::CreateSP(EditorController, &FIKRetargetEditorController::HandleResetAllBones),
-		FCanExecuteAction(),
+		FCanExecuteAction::CreateSP(EditorController, &FIKRetargetEditorController::IsEditingPose),
 		EUIActionRepeatMode::RepeatDisabled);
 
 	ToolkitCommands->MapAction(
 		Commands.ResetSelectedBones,
 		FExecuteAction::CreateSP(EditorController, &FIKRetargetEditorController::HandleResetSelectedBones),
-		FCanExecuteAction::CreateSP(EditorController, &FIKRetargetEditorController::IsAnyBoneSelected),
+		FCanExecuteAction::CreateSP(EditorController, &FIKRetargetEditorController::IsEditingPoseWithAnyBoneSelected),
 		EUIActionRepeatMode::RepeatDisabled);
 
 	ToolkitCommands->MapAction(
 		Commands.ResetSelectedAndChildrenBones,
 		FExecuteAction::CreateSP(EditorController, &FIKRetargetEditorController::HandleResetSelectedAndChildrenBones),
-		FCanExecuteAction::CreateSP(EditorController, &FIKRetargetEditorController::IsAnyBoneSelected),
+		FCanExecuteAction::CreateSP(EditorController, &FIKRetargetEditorController::IsEditingPoseWithAnyBoneSelected),
 		EUIActionRepeatMode::RepeatDisabled);
 
 	ToolkitCommands->MapAction(
 		Commands.NewRetargetPose,
 		FExecuteAction::CreateSP(EditorController, &FIKRetargetEditorController::HandleNewPose),
-		FCanExecuteAction(),
-		FCanExecuteAction(),
 		EUIActionRepeatMode::RepeatDisabled);
 
 	ToolkitCommands->MapAction(
 		Commands.DuplicateRetargetPose,
 		FExecuteAction::CreateSP(EditorController, &FIKRetargetEditorController::HandleDuplicatePose),
-		FCanExecuteAction(),
-		FCanExecuteAction(),
 		EUIActionRepeatMode::RepeatDisabled);
 
 	ToolkitCommands->MapAction(
 		Commands.DeleteRetargetPose,
 		FExecuteAction::CreateSP(EditorController, &FIKRetargetEditorController::HandleDeletePose),
 		FCanExecuteAction::CreateSP(EditorController, &FIKRetargetEditorController::CanDeletePose),
-		FCanExecuteAction(),
 		EUIActionRepeatMode::RepeatDisabled);
 
 	ToolkitCommands->MapAction(
 		Commands.RenameRetargetPose,
 		FExecuteAction::CreateSP(EditorController, &FIKRetargetEditorController::HandleRenamePose),
 		FCanExecuteAction::CreateSP(EditorController, &FIKRetargetEditorController::CanRenamePose),
-		FCanExecuteAction(),
 		EUIActionRepeatMode::RepeatDisabled);
 
 	//
@@ -239,36 +227,31 @@ void FIKRetargetEditor::BindCommands()
 	ToolkitCommands->MapAction(
 		Commands.AutoAlignAllBones,
 		FExecuteAction::CreateSP(EditorController, &FIKRetargetEditorController::HandleAlignAllBones),
-		FCanExecuteAction(),
-		FCanExecuteAction(),
+		FCanExecuteAction::CreateSP(EditorController, &FIKRetargetEditorController::IsEditingPose),
 		EUIActionRepeatMode::RepeatDisabled);
 	
 	ToolkitCommands->MapAction(
 		Commands.AlignSelected,
 		FExecuteAction::CreateSP(EditorController, &FIKRetargetEditorController::HandleAlignSelectedBones, ERetargetAutoAlignMethod::ChainToChain, false /* no children*/),
-		FCanExecuteAction(),
-		FCanExecuteAction(),
+		FCanExecuteAction::CreateSP(EditorController, &FIKRetargetEditorController::IsEditingPoseWithAnyBoneSelected),
 		EUIActionRepeatMode::RepeatDisabled);
 
 	ToolkitCommands->MapAction(
 		Commands.AlignSelectedAndChildren,
 		FExecuteAction::CreateSP(EditorController, &FIKRetargetEditorController::HandleAlignSelectedBones, ERetargetAutoAlignMethod::ChainToChain, true /* include children*/),
-		FCanExecuteAction(),
-		FCanExecuteAction(),
+		FCanExecuteAction::CreateSP(EditorController, &FIKRetargetEditorController::IsEditingPoseWithAnyBoneSelected),
 		EUIActionRepeatMode::RepeatDisabled);
 
 	ToolkitCommands->MapAction(
 		Commands.AlignSelectedUsingMesh,
 		FExecuteAction::CreateSP(EditorController, &FIKRetargetEditorController::HandleAlignSelectedBones, ERetargetAutoAlignMethod::MeshToMesh, false /* no children*/),
-		FCanExecuteAction(),
-		FCanExecuteAction(),
+		FCanExecuteAction::CreateSP(EditorController, &FIKRetargetEditorController::IsEditingPoseWithAnyBoneSelected),
 		EUIActionRepeatMode::RepeatDisabled);
 
 	ToolkitCommands->MapAction(
 		Commands.SnapCharacterToGround,
 		FExecuteAction::CreateSP(EditorController, &FIKRetargetEditorController::HandleSnapToGround),
-		FCanExecuteAction(),
-		FCanExecuteAction(),
+		FCanExecuteAction::CreateSP(EditorController, &FIKRetargetEditorController::IsEditingPose),
 		EUIActionRepeatMode::RepeatDisabled);
 	
 
@@ -382,7 +365,6 @@ TSharedRef<SWidget> FIKRetargetEditor::GenerateRetargetModesMenu()
 	MenuBuilder.BeginSection(TEXT("Retarget Modes"));
 	MenuBuilder.AddMenuEntry(FIKRetargetCommands::Get().RunRetargeter, TEXT("Run Retargeter"), TAttribute<FText>(), TAttribute<FText>(),  EditorController->GetRetargeterModeIcon(ERetargeterOutputMode::RunRetarget));
 	MenuBuilder.AddMenuEntry(FIKRetargetCommands::Get().EditRetargetPose, TEXT("Edit Retarget Pose"), TAttribute<FText>(), TAttribute<FText>(), EditorController->GetRetargeterModeIcon(ERetargeterOutputMode::EditRetargetPose));
-	MenuBuilder.AddMenuEntry(FIKRetargetCommands::Get().ShowRetargetPose, TEXT("Show Retarget Pose"), TAttribute<FText>(), TAttribute<FText>(), EditorController->GetRetargeterModeIcon(ERetargeterOutputMode::ShowRetargetPose));
 	MenuBuilder.EndSection();
 	
 	return MenuBuilder.MakeWidget();
diff --git a/Engine/Plugins/Animation/IKRig/Source/IKRigEditor/Private/RetargetEditor/IKRetargetEditorController.cpp b/Engine/Plugins/Animation/IKRig/Source/IKRigEditor/Private/RetargetEditor/IKRetargetEditorController.cpp
--- a/Engine/Plugins/Animation/IKRig/Source/IKRigEditor/Private/RetargetEditor/IKRetargetEditorController.cpp
+++ b/Engine/Plugins/Animation/IKRig/Source/IKRigEditor/Private/RetargetEditor/IKRetargetEditorController.cpp
@@ -150,15 +150,13 @@ void FIKRetargetEditorController::Initialize(TSharedPtr<FIKRetargetEditor> InEdi
 	Editor = InEditor;
 	AssetController = UIKRetargeterController::GetController(InAsset);
 	CurrentlyEditingSourceOrTarget = ERetargetSourceOrTarget::Target;
-	OutputMode = ERetargeterOutputMode::ShowRetargetPose;
+	OutputMode = ERetargeterOutputMode::RunRetarget;
 	PreviousMode = OutputMode;
 	PoseExporter = MakeShared<FIKRetargetPoseExporter>();
 	PoseExporter->Initialize(SharedThis(this));
 
 	PlaybackManager = MakeUnique<FRetargetPlaybackManager>(SharedThis(this));
 
-	AutoPoseGenerator = MakeUnique<FRetargetAutoPoseGenerator>(SharedThis(this));
-
 	SelectedBoneNames.Add(ERetargetSourceOrTarget::Source);
 	SelectedBoneNames.Add(ERetargetSourceOrTarget::Target);
 	LastSelectedItem = ERetargetSelectionType::NONE;
@@ -465,9 +463,9 @@ bool FIKRetargetEditorController::GetCameraTargetForSelection(FSphere& OutTarget
 	return false;
 }
 
-bool FIKRetargetEditorController::IsAnyBoneSelected() const
+bool FIKRetargetEditorController::IsEditingPoseWithAnyBoneSelected() const
 {
-	return !GetSelectedBones().IsEmpty();
+	return IsEditingPose() && !GetSelectedBones().IsEmpty();
 }
 
 bool FIKRetargetEditorController::IsBoneRetargeted(const FName& BoneName, ERetargetSourceOrTarget SourceOrTarget) const
@@ -738,7 +736,7 @@ void FIKRetargetEditorController::SetRetargetPoseAmount(float InValue)
 {
 	if (OutputMode==ERetargeterOutputMode::RunRetarget)
 	{
-		SetRetargeterMode(ERetargeterOutputMode::ShowRetargetPose);
+		SetRetargeterMode(ERetargeterOutputMode::EditRetargetPose);
 	}
 	
 	RetargetPosePreviewBlend = InValue;
@@ -1042,18 +1040,7 @@ void FIKRetargetEditorController::SetRetargeterMode(ERetargeterOutputMode Mode)
 			TargetAnimInstance->SetRetargetMode(ERetargeterOutputMode::RunRetarget);
 			PlaybackManager->ResumePlayback();
 			break;
-
-		case ERetargeterOutputMode::ShowRetargetPose:
-			EditorModeManager.DeactivateMode(FIKRetargetEditPoseMode::ModeName);
-			EditorModeManager.ActivateMode(FIKRetargetDefaultMode::ModeName);
-			OutputMode = ERetargeterOutputMode::ShowRetargetPose;
-			// show retarget pose
-			SourceAnimInstance->SetRetargetMode(ERetargeterOutputMode::ShowRetargetPose);
-			TargetAnimInstance->SetRetargetMode(ERetargeterOutputMode::ShowRetargetPose);
-			PlaybackManager->PausePlayback();
-			SetRetargetPoseAmount(1.0f);
-			break;
-
+		
 		default:
 			checkNoEntry();
 	}
@@ -1070,8 +1057,6 @@ FText FIKRetargetEditorController::GetRetargeterModeLabel()
 		return FText::FromString("Running Retarget");
 	case ERetargeterOutputMode::EditRetargetPose:
 		return FText::FromString("Editing Retarget Pose");
-	case ERetargeterOutputMode::ShowRetargetPose:
-		return FText::FromString("Showing Retarget Pose");
 	default:
 		checkNoEntry();
 		return FText::FromString("Unknown Mode.");
@@ -1091,8 +1076,6 @@ FSlateIcon FIKRetargetEditorController::GetRetargeterModeIcon(ERetargeterOutputM
 		return FSlateIcon(FIKRetargetEditorStyle::Get().GetStyleSetName(), "IKRetarget.RunRetargeter");
 	case ERetargeterOutputMode::EditRetargetPose:
 		return FSlateIcon(FIKRetargetEditorStyle::Get().GetStyleSetName(), "IKRetarget.EditRetargetPose");
-	case ERetargeterOutputMode::ShowRetargetPose:
-		return FSlateIcon(FIKRetargetEditorStyle::Get().GetStyleSetName(), "IKRetarget.ShowRetargetPose");
 	default:
 		checkNoEntry();
 		return FSlateIcon(FIKRetargetEditorStyle::Get().GetStyleSetName(), "IKRetarget.ShowRetargetPose");
@@ -1121,7 +1104,7 @@ FRetargetGlobalSettings& FIKRetargetEditorController::GetGlobalSettings() const
 
 void FIKRetargetEditorController::HandleNewPose()
 {
-	SetRetargeterMode(ERetargeterOutputMode::ShowRetargetPose);
+	SetRetargeterMode(ERetargeterOutputMode::EditRetargetPose);
 	
 	// get a unique pose name to use as suggestion
 	const FString DefaultNewPoseName = LOCTEXT("NewRetargetPoseName", "CustomRetargetPose").ToString();
@@ -1203,7 +1186,7 @@ FReply FIKRetargetEditorController::CreateNewPose() const
 
 void FIKRetargetEditorController::HandleDuplicatePose()
 {
-	SetRetargeterMode(ERetargeterOutputMode::ShowRetargetPose);
+	SetRetargeterMode(ERetargeterOutputMode::EditRetargetPose);
 	
 	// get a unique pose name to use as suggestion for duplicate
 	const FString DuplicateSuffix = LOCTEXT("DuplicateSuffix", "_Copy").ToString();
@@ -1284,7 +1267,7 @@ FReply FIKRetargetEditorController::CreateDuplicatePose() const
 
 void FIKRetargetEditorController::HandleDeletePose()
 {
-	SetRetargeterMode(ERetargeterOutputMode::ShowRetargetPose);
+	SetRetargeterMode(ERetargeterOutputMode::EditRetargetPose);
 	
 	const ERetargetSourceOrTarget SourceOrTarget = GetSourceOrTarget();
 	const FName CurrentPose = AssetController->GetCurrentRetargetPoseName(SourceOrTarget);
@@ -1323,59 +1306,20 @@ void FIKRetargetEditorController::HandleResetSelectedAndChildrenBones() const
 
 void FIKRetargetEditorController::HandleAlignAllBones() const
 {
-	// get all the bones in the current skeleton
-	const UIKRetargetProcessor* Processor = GetRetargetProcessor();
-	if (!(Processor && Processor->IsInitialized()))
-	{
-		return;
-	}
-
-	// undo transaction
-	constexpr bool bShouldTransact = true;
-	FScopedTransaction Transaction(LOCTEXT("AutoAlignAllBones", "Auto Align All Bones"), bShouldTransact);
-	AssetController->GetAsset()->Modify();
-	
-	// suppress warnings about bones that cannot be aligned when aligning ALL bones
-	const TArray<FName>& AllBones = Processor->GetSkeleton(GetSourceOrTarget()).BoneNames;
-	constexpr bool bSuppressWarnings = true;
-	AutoPoseGenerator.Get()->AlignBones(
-		AllBones,
-		ERetargetAutoAlignMethod::ChainToChain,
-		GetSourceOrTarget(),
-		bSuppressWarnings);
-
-	// if the retarget root was aligned
-	AutoPoseGenerator.Get()->SnapToGround(NAME_None, GetSourceOrTarget());
+	AssetController->AutoAlignAllBones(GetSourceOrTarget());
 }
 
 void FIKRetargetEditorController::HandleAlignSelectedBones(const ERetargetAutoAlignMethod Method, const bool bIncludeChildren) const
 {
-	// undo transaction
-	constexpr bool bShouldTransact = true;
-	FScopedTransaction Transaction(LOCTEXT("AutoAlignSelectedBones", "Auto Align Selected Bones"), bShouldTransact);
-	AssetController->GetAsset()->Modify();
-
 	const TArray<FName> BonesToAlign = bIncludeChildren ? GetSelectedBonesAndChildren() : GetSelectedBones();
-	
-	// allow warnings about bones that cannot be aligned when bones are explicitly specified by user
-	constexpr bool bSuppressWarnings = false;
-	AutoPoseGenerator.Get()->AlignBones(
-		BonesToAlign,
-		Method,
-		GetSourceOrTarget(),
-		bSuppressWarnings);
+	AssetController->AutoAlignBones(BonesToAlign, Method, GetSourceOrTarget());
 }
 
 void FIKRetargetEditorController::HandleSnapToGround() const
 {
-	// undo transaction
-	constexpr bool bShouldTransact = true;
-	FScopedTransaction Transaction(LOCTEXT("AutoSnapToGround", "Snap Retarget Pose to Ground"), bShouldTransact);
-	AssetController->GetAsset()->Modify();
-	
 	const TArray<FName> SelectedBones = GetSelectedBones();
 	const FName FirstSelectedBone = SelectedBones.IsEmpty() ? NAME_None : SelectedBones[0];
-	AutoPoseGenerator.Get()->SnapToGround(FirstSelectedBone, GetSourceOrTarget());
+	AssetController->SnapBoneToGround(FirstSelectedBone, GetSourceOrTarget());
 }
 
 void FIKRetargetEditorController::HandleRenamePose()
diff --git a/Engine/Plugins/Animation/IKRig/Source/IKRigEditor/Private/RetargetEditor/IKRetargetPoseExporter.cpp b/Engine/Plugins/Animation/IKRig/Source/IKRigEditor/Private/RetargetEditor/IKRetargetPoseExporter.cpp
--- a/Engine/Plugins/Animation/IKRig/Source/IKRigEditor/Private/RetargetEditor/IKRetargetPoseExporter.cpp
+++ b/Engine/Plugins/Animation/IKRig/Source/IKRigEditor/Private/RetargetEditor/IKRetargetPoseExporter.cpp
@@ -29,7 +29,7 @@ void FIKRetargetPoseExporter::HandleImportFromPoseAsset()
 {
 	FIKRetargetEditorController* ControllerPtr = Controller.Pin().Get();
 	
-	ControllerPtr->SetRetargeterMode(ERetargeterOutputMode::ShowRetargetPose);
+	ControllerPtr->SetRetargeterMode(ERetargeterOutputMode::EditRetargetPose);
 	
 	RetargetPoseToImport = nullptr;
 	PosesInSelectedAsset.Empty();
@@ -286,7 +286,7 @@ void FIKRetargetPoseExporter::HandleImportFromSequenceAsset()
 {
 	FIKRetargetEditorController* ControllerPtr = Controller.Pin().Get();
 	
-	ControllerPtr->SetRetargeterMode(ERetargeterOutputMode::ShowRetargetPose);
+	ControllerPtr->SetRetargeterMode(ERetargeterOutputMode::EditRetargetPose);
 	
 	SequenceToImportAsPose = nullptr;
 
diff --git a/Engine/Plugins/Animation/IKRig/Source/IKRigEditor/Private/RetargetEditor/IKRetargeterController.cpp b/Engine/Plugins/Animation/IKRig/Source/IKRigEditor/Private/RetargetEditor/IKRetargeterController.cpp
--- a/Engine/Plugins/Animation/IKRig/Source/IKRigEditor/Private/RetargetEditor/IKRetargeterController.cpp
+++ b/Engine/Plugins/Animation/IKRig/Source/IKRigEditor/Private/RetargetEditor/IKRetargeterController.cpp
@@ -9,12 +9,14 @@
 #include "Retargeter/IKRetargeter.h"
 #include "RigEditor/IKRigController.h"
 #include "IKRigEditor.h"
+#include "RetargetEditor/IKRetargeterPoseGenerator.h"
 #include "Retargeter/IKRetargetOps.h"
 
 #include UE_INLINE_GENERATED_CPP_BY_NAME(IKRetargeterController)
 
 #define LOCTEXT_NAMESPACE "IKRetargeterController"
 
+
 UIKRetargeterController* UIKRetargeterController::GetController(const UIKRetargeter* InRetargeterAsset)
 {
 	if (!InRetargeterAsset)
@@ -32,6 +34,12 @@ UIKRetargeterController* UIKRetargeterController::GetController(const UIKRetarge
 	return Cast<UIKRetargeterController>(InRetargeterAsset->Controller);
 }
 
+void UIKRetargeterController::PostInitProperties()
+{
+	Super::PostInitProperties();
+	AutoPoseGenerator = MakeUnique<FRetargetAutoPoseGenerator>(this);
+}
+
 UIKRetargeter* UIKRetargeterController::GetAsset() const
 {
 	return Asset;
@@ -904,6 +912,54 @@ FVector UIKRetargeterController::GetRootOffsetInRetargetPose(
 	return GetCurrentRetargetPose(SourceOrTarget).GetRootTranslationDelta();
 }
 
+void UIKRetargeterController::AutoAlignAllBones(ERetargetSourceOrTarget SourceOrTarget) const
+{
+	// undo transaction
+	constexpr bool bShouldTransact = true;
+	FScopedTransaction Transaction(LOCTEXT("AutoAlignAllBones", "Auto Align All Bones"), bShouldTransact);
+	Asset->Modify();
+	
+	FScopedReinitializeIKRetargeter Reinitialize(this);
+
+	// first reset the entire retarget pose
+	ResetRetargetPose(GetCurrentRetargetPoseName(SourceOrTarget), TArray<FName>(), SourceOrTarget);
+	
+	// suppress warnings about bones that cannot be aligned when aligning ALL bones
+	constexpr bool bSuppressWarnings = true;
+	AutoPoseGenerator.Get()->AlignAllBones(SourceOrTarget, bSuppressWarnings);
+}
+
+void UIKRetargeterController::AutoAlignBones(
+	const TArray<FName>& BonesToAlign,
+	const ERetargetAutoAlignMethod Method,
+	ERetargetSourceOrTarget SourceOrTarget) const
+{
+	// undo transaction
+	constexpr bool bShouldTransact = true;
+	FScopedTransaction Transaction(LOCTEXT("AutoAlignSelectedBones", "Auto Align Selected Bones"), bShouldTransact);
+	Asset->Modify();
+
+	FScopedReinitializeIKRetargeter Reinitialize(this);
+	
+	// allow warnings about bones that cannot be aligned when bones are explicitly specified by user
+	constexpr bool bSuppressWarnings = false;
+	AutoPoseGenerator.Get()->AlignBones(
+		BonesToAlign,
+		Method,
+		SourceOrTarget,
+		bSuppressWarnings);
+}
+
+void UIKRetargeterController::SnapBoneToGround(FName ReferenceBone, ERetargetSourceOrTarget SourceOrTarget)
+{
+	// undo transaction
+	constexpr bool bShouldTransact = true;
+	FScopedTransaction Transaction(LOCTEXT("SnapBoneToGround", "Snap Bone to Ground"), bShouldTransact);
+	Asset->Modify();
+
+	AutoPoseGenerator.Get()->SnapToGround(ReferenceBone, SourceOrTarget);
+}
+
 FName UIKRetargeterController::MakePoseNameUnique(const FString& PoseName, const ERetargetSourceOrTarget SourceOrTarget) const
 {
 	FString UniqueName = PoseName;
diff --git a/Engine/Plugins/Animation/IKRig/Source/IKRigEditor/Private/RetargetEditor/IKRetargeterPoseGenerator.cpp b/Engine/Plugins/Animation/IKRig/Source/IKRigEditor/Private/RetargetEditor/IKRetargeterPoseGenerator.cpp
--- a/Engine/Plugins/Animation/IKRig/Source/IKRigEditor/Private/RetargetEditor/IKRetargeterPoseGenerator.cpp
+++ b/Engine/Plugins/Animation/IKRig/Source/IKRigEditor/Private/RetargetEditor/IKRetargeterPoseGenerator.cpp
@@ -17,10 +17,32 @@
 
 #define LOCTEXT_NAMESPACE "RetargetPoseGenerator"
 
-FRetargetAutoPoseGenerator::FRetargetAutoPoseGenerator(const TWeakPtr<FIKRetargetEditorController>& InEditorController)
+FRetargetAutoPoseGenerator::FRetargetAutoPoseGenerator(const TWeakObjectPtr<UIKRetargeterController> InController)
 {
-	check(InEditorController.Pin())
-	EditorController = InEditorController;
+	Controller = InController;
+
+	// we maintain our own processor separate from the editor, this allows the auto-pose generator
+	// to run "headless" and be accessed by BP/Python API in the asset controller
+	Processor = NewObject<UIKRetargetProcessor>();
+}
+
+void FRetargetAutoPoseGenerator::AddReferencedObjects(FReferenceCollector& Collector)
+{
+	Collector.AddReferencedObject(Processor);
+}
+
+void FRetargetAutoPoseGenerator::AlignAllBones(
+	const ERetargetSourceOrTarget SourceOrTarget,
+	const bool bSuppressWarnings) const
+{
+	// cannot get all bones until we have an initialized processor
+	if (!CheckReadyToAlignBones())
+	{
+		return;
+	}
+	
+	const TArray<FName>& AllBones = Processor->GetSkeleton(SourceOrTarget).BoneNames;
+	AlignBones(AllBones, ERetargetAutoAlignMethod::ChainToChain, SourceOrTarget, bSuppressWarnings);
 }
 
 void FRetargetAutoPoseGenerator::AlignBones(
@@ -74,6 +96,12 @@ void FRetargetAutoPoseGenerator::AlignBone(
 	const ERetargetAutoAlignMethod Method,
 	const ERetargetSourceOrTarget SourceOrTarget) const
 {
+	// cannot align bones until we have an initialized processor
+	if (!CheckReadyToAlignBones())
+	{
+		return;
+	}
+	
 	// can't auto-align if the bone is not retargeted  (warning is elsewhere)
 	if (!Processor->IsBoneRetargeted(BoneToAlign, SourceOrTarget))
 	{
@@ -81,7 +109,7 @@ void FRetargetAutoPoseGenerator::AlignBone(
 	}
 
 	// force regenerate the retarget pose to guarantee Skeleton.RetargetGlobalPose reflects any modifications made in the viewport
-	const FName CurrentPoseName = EditorController.Pin()->AssetController->GetCurrentRetargetPoseName(SourceOrTarget);
+	const FName CurrentPoseName = Controller->GetCurrentRetargetPoseName(SourceOrTarget);
 	Processor->UpdateRetargetPoseAtRuntime(CurrentPoseName, SourceOrTarget);
 
 	// use a custom procedure to auto-align retarget root
@@ -202,11 +230,10 @@ void FRetargetAutoPoseGenerator::SnapToGround(const FName BoneToPutOnGround, con
 		return;
 	}
 
-	const TObjectPtr<UIKRetargeterController> AssetController = EditorController.Pin()->AssetController;
 	const FRetargetSkeleton& Skeleton = Processor->GetSkeleton(SourceOrTarget);
 	
 	// force regenerate the retarget pose to guarantee Skeleton.RetargetGlobalPose reflects any modifications made thus far
-	const FName CurrentPoseName = AssetController->GetCurrentRetargetPoseName(SourceOrTarget);
+	const FName CurrentPoseName = Controller->GetCurrentRetargetPoseName(SourceOrTarget);
 	Processor->UpdateRetargetPoseAtRuntime(CurrentPoseName, SourceOrTarget);
 
 	// if no bone specified, we search the current retarget pose to find the lowest one and use that
@@ -250,7 +277,7 @@ void FRetargetAutoPoseGenerator::SnapToGround(const FName BoneToPutOnGround, con
 
 	// snap it back to original height off the ground
 	const float DeltaHeight = HeightInRefPose - HeightOfBoneInRetargetPose;
-	AssetController->SetRootOffsetInRetargetPose(FVector(0.f,0.f,DeltaHeight), SourceOrTarget);
+	Controller->SetRootOffsetInRetargetPose(FVector(0.f,0.f,DeltaHeight), SourceOrTarget);
 }
 
 bool FRetargetAutoPoseGenerator::GetSpineChain(const ERetargetSourceOrTarget SourceOrTarget, FName& OutChainName) const
@@ -521,16 +548,26 @@ bool FRetargetAutoPoseGenerator::IsARetargetLeaf(
 bool FRetargetAutoPoseGenerator::CheckReadyToAlignBones() const
 {
 	// this should never happen
-	if (!EditorController.IsValid())
+	if (!Controller.IsValid())
+	{
+		return false;
+	}
+
+	if (!Processor)
 	{
 		return false;
 	}
+
+	// we have to initialize a processor to "resolve" the retarget asset setup onto actual skeletons to do the alignment
+	USkeletalMesh* SourceSkeletalMesh = Controller->GetPreviewMesh(ERetargetSourceOrTarget::Source);
+	USkeletalMesh* TargetSkeletalMesh = Controller->GetPreviewMesh(ERetargetSourceOrTarget::Target);
+	UIKRetargeter* RetargeterAsset = Controller->GetAssetPtr();
+	constexpr bool bSuppressWarnings = true;
+	Processor->Initialize(SourceSkeletalMesh, TargetSkeletalMesh, RetargeterAsset, bSuppressWarnings);
 	
 	// can't auto align until processor is initialized because we need a fully resolved chain mapping on source/target
 	// this could happen if the user attempts to edit the retarget pose on an incompatible IK Rig
-	const FIKRetargetEditorController* Controller = EditorController.Pin().Get();
-	Processor = Controller->GetRetargetProcessor();
-	if (!(Processor && Processor->IsInitialized()))
+	if (!Processor->IsInitialized())
 	{
 		return false;
 	}
@@ -543,7 +580,6 @@ void FRetargetAutoPoseGenerator::ApplyWorldRotationToRetargetPose(
 	const FName BoneToAffect,
 	const ERetargetSourceOrTarget SourceOrTarget) const
 {
-	const TObjectPtr<UIKRetargeterController> AssetController = EditorController.Pin()->AssetController;
 	const FRetargetSkeleton& Skeleton =  Processor->GetSkeleton(SourceOrTarget);
 	const FReferenceSkeleton& RefSkeleton = Skeleton.SkeletalMesh->GetRefSkeleton();
 	const int32 BoneIndex = Skeleton.FindBoneIndexByName(BoneToAffect);
@@ -561,7 +597,7 @@ void FRetargetAutoPoseGenerator::ApplyWorldRotationToRetargetPose(
 	FTransform CurrentGlobalTransformNoDelta = LocalRefTransform * ParentGlobalTransform;
 				
 	// get the stored local rotation offset from the retarget pose
-	const FQuat LocalRotationDeltaFromPose = AssetController->GetRotationOffsetForRetargetPoseBone(BoneToAffect, SourceOrTarget);
+	const FQuat LocalRotationDeltaFromPose = Controller->GetRotationOffsetForRetargetPoseBone(BoneToAffect, SourceOrTarget);
 	// apply the retarget pose to get the global rotation w/ retarget pose
 	// TODO get this from retarget pose??
 	FQuat GlobalRetargetPoseRotation = CurrentGlobalTransformNoDelta.GetRotation() * LocalRotationDeltaFromPose;
@@ -576,10 +612,10 @@ void FRetargetAutoPoseGenerator::ApplyWorldRotationToRetargetPose(
 	FQuat FinalLocalDeltaRotation = FQuat(UnRotatedAxis, NewGlobalDeltaRotation.GetAngle());
 
 	// store the new rotation in the retarget pose
-	AssetController->SetRotationOffsetForRetargetPoseBone(BoneToAffect, FinalLocalDeltaRotation, SourceOrTarget);
+	Controller->SetRotationOffsetForRetargetPoseBone(BoneToAffect, FinalLocalDeltaRotation, SourceOrTarget);
 
 	// regenerate the retarget pose in the currently running processor so it reflects the latest change
-	const FName CurrentPoseName = AssetController->GetCurrentRetargetPoseName(SourceOrTarget);
+	const FName CurrentPoseName = Controller->GetCurrentRetargetPoseName(SourceOrTarget);
 	Processor->UpdateRetargetPoseAtRuntime(CurrentPoseName, SourceOrTarget);
 }
 
diff --git a/Engine/Plugins/Animation/IKRig/Source/IKRigEditor/Public/RetargetEditor/IKRetargetEditorController.h b/Engine/Plugins/Animation/IKRig/Source/IKRigEditor/Public/RetargetEditor/IKRetargetEditorController.h
--- a/Engine/Plugins/Animation/IKRig/Source/IKRigEditor/Public/RetargetEditor/IKRetargetEditorController.h
+++ b/Engine/Plugins/Animation/IKRig/Source/IKRigEditor/Public/RetargetEditor/IKRetargetEditorController.h
@@ -31,7 +31,6 @@ struct FRetargetSkeleton;
 enum class ERetargeterOutputMode : uint8
 {
 	RunRetarget,		// output the retargeted target pose
-	ShowRetargetPose,	// output the retarget pose for viewing purposes
 	EditRetargetPose,	// allow editing the retarget pose
 };
 
@@ -121,8 +120,6 @@ public:
 
 	// import / export retarget poses
 	TSharedPtr<FIKRetargetPoseExporter> PoseExporter;
-	// auto pose generator
-	TUniquePtr<FRetargetAutoPoseGenerator> AutoPoseGenerator;
 
 	// manage playback of animation in the editor
 	TUniquePtr<FRetargetPlaybackManager> PlaybackManager;
@@ -257,7 +254,7 @@ public:
 	bool GetCameraTargetForSelection(FSphere& OutTarget) const;
 
 	// check if any bone is selected
-	bool IsAnyBoneSelected() const;
+	bool IsEditingPoseWithAnyBoneSelected() const;
 
 	// ------------------------- END SELECTION -----------------------------
 
diff --git a/Engine/Plugins/Animation/IKRig/Source/IKRigEditor/Public/RetargetEditor/IKRetargeterController.h b/Engine/Plugins/Animation/IKRig/Source/IKRigEditor/Public/RetargetEditor/IKRetargeterController.h
--- a/Engine/Plugins/Animation/IKRig/Source/IKRigEditor/Public/RetargetEditor/IKRetargeterController.h
+++ b/Engine/Plugins/Animation/IKRig/Source/IKRigEditor/Public/RetargetEditor/IKRetargeterController.h
@@ -4,6 +4,8 @@
 #include "CoreMinimal.h"
 #include "UObject/Object.h"
 
+#include "RetargetEditor/IKRetargeterPoseGenerator.h"
+
 #include "IKRetargeterController.generated.h"
 
 struct FRetargetGlobalSettings;
@@ -11,6 +13,7 @@ struct FTargetRootSettings;
 struct FTargetChainSettings;
 struct FIKRetargetPose;
 enum class ERetargetSourceOrTarget : uint8;
+enum class ERetargetAutoAlignMethod : uint8;
 class FIKRetargetEditorController;
 class URetargetChainSettings;
 class UIKRigDefinition;
@@ -33,6 +36,9 @@ class IKRIGEDITOR_API UIKRetargeterController : public UObject
 	GENERATED_BODY()
 
 public:
+
+	// UObject
+	virtual void PostInitProperties() override;
 	
 	// Get access to the retargeter asset.
 	// Warning: Do not make modifications to the asset directly. Use this API instead. 
@@ -275,6 +281,18 @@ public:
 	UFUNCTION(BlueprintCallable, Category=IKRetargeter, meta = (BlueprintThreadSafe))
 	FVector GetRootOffsetInRetargetPose(const ERetargetSourceOrTarget SourceOrTarget) const;
 
+	// Automatically align all bones in mapped chains and store in the current retarget pose.
+	UFUNCTION(BlueprintCallable, Category=RetargetOps)
+	void AutoAlignAllBones(ERetargetSourceOrTarget SourceOrTarget) const;
+
+	// Automatically align an array of bones and store in the current retarget pose. Bones not in a mapped chain will be ignored.
+	UFUNCTION(BlueprintCallable, Category=RetargetOps)
+	void AutoAlignBones(const TArray<FName>& BonesToAlign, const ERetargetAutoAlignMethod Method, ERetargetSourceOrTarget SourceOrTarget) const;
+
+	// Moves the entire skeleton vertically until the specified bone is the same height off the ground as in the reference pose.
+	UFUNCTION(BlueprintCallable, Category=RetargetOps)
+	void SnapBoneToGround(FName ReferenceBone, ERetargetSourceOrTarget SourceOrTarget);
+
 	//
 	// RETARGET POSE C++ ONLY API
 	//
@@ -320,6 +338,9 @@ private:
 		RetargeterNeedsInitialized.Broadcast();
 	}
 
+	// auto pose generator
+	TUniquePtr<FRetargetAutoPoseGenerator> AutoPoseGenerator;
+
 	// only allow modifications to data model from one thread at a time
 	mutable FCriticalSection ControllerLock;
 
diff --git a/Engine/Plugins/Animation/IKRig/Source/IKRigEditor/Public/RetargetEditor/IKRetargeterPoseGenerator.h b/Engine/Plugins/Animation/IKRig/Source/IKRigEditor/Public/RetargetEditor/IKRetargeterPoseGenerator.h
--- a/Engine/Plugins/Animation/IKRig/Source/IKRigEditor/Public/RetargetEditor/IKRetargeterPoseGenerator.h
+++ b/Engine/Plugins/Animation/IKRig/Source/IKRigEditor/Public/RetargetEditor/IKRetargeterPoseGenerator.h
@@ -2,7 +2,12 @@
 #pragma once
 
 #include "CoreMinimal.h"
+#include "UObject/GCObject.h"
+#include "UObject/WeakObjectPtr.h"
 
+#include "IKRetargeterPoseGenerator.generated.h"
+
+class UIKRetargeterController;
 class FIKRetargetEditorController;
 class UIKRetargetProcessor;
 enum class ERetargetSourceOrTarget : uint8;
@@ -20,6 +25,7 @@ enum class ERetargetSourceOrTarget : uint8;
 // Since the skeleton cannot be relied upon to provide robust directional information for individual and/or leaf bones, we must instead fallback
 // on looking at the mesh that is skinned to it and try to discern a direction vector from that. If there is no skinning data, then we simply
 // do NOT auto-align that bone.
+UENUM()
 enum class ERetargetAutoAlignMethod : uint8
 {
 	// use the chain to determine the source and target directions to align
@@ -29,10 +35,18 @@ enum class ERetargetAutoAlignMethod : uint8
 };
 
 // a sub-system used by the IKRetargetEditorController for automatically generating a retarget pose
-struct FRetargetAutoPoseGenerator
+struct FRetargetAutoPoseGenerator : public FGCObject
 {
 	// created and owned by the retarget editor controller
-	FRetargetAutoPoseGenerator(const TWeakPtr<FIKRetargetEditorController>& InEditorController);
+	FRetargetAutoPoseGenerator(const TWeakObjectPtr<UIKRetargeterController> InController);
+
+	// FGCObject interface
+	virtual void AddReferencedObjects(FReferenceCollector& Collector) override;
+	virtual FString GetReferencerName() const override { return TEXT("FRetargetAutoPoseGenerator"); }
+	// End FGCObject interface
+
+	// align all bones bones on the specified skeleton to the OTHER skeleton
+	void AlignAllBones(const ERetargetSourceOrTarget SourceOrTarget, const bool bSuppressWarnings) const;
 	
 	// align the specified bones on the current skeleton to the OTHER skeleton, using the specified method
 	void AlignBones(
@@ -119,7 +133,7 @@ private:
 	// get the processor from the editor and return true if it's initialized
 	bool CheckReadyToAlignBones() const;
 
-	TWeakPtr<FIKRetargetEditorController> EditorController;
-	mutable UIKRetargetProcessor* Processor;
+	TWeakObjectPtr<UIKRetargeterController> Controller;
+	TObjectPtr<UIKRetargetProcessor> Processor;
 	static constexpr float ParamStepSize = 0.1f;
 };
