Fixes bug, disabled IK Rig visual logging until initialized.

diff --git a/Engine/Plugins/Animation/IKRig/Source/IKRig/Private/AnimNodes/AnimNode_IKRig.cpp b/Engine/Plugins/Animation/IKRig/Source/IKRig/Private/AnimNodes/AnimNode_IKRig.cpp
--- a/Engine/Plugins/Animation/IKRig/Source/IKRig/Private/AnimNodes/AnimNode_IKRig.cpp
+++ b/Engine/Plugins/Animation/IKRig/Source/IKRig/Private/AnimNodes/AnimNode_IKRig.cpp
@@ -225,7 +225,8 @@ void FAnimNode_IKRig::Update_AnyThread(const FAnimationUpdateContext& Context)
 
 #if ENABLE_VISUAL_LOG
 #if WITH_EDITORONLY_DATA
-	if( FVisualLogger::IsRecording() )
+	// is node setup?
+	if (IKRigProcessor && IKRigProcessor->IsInitialized() && FVisualLogger::IsRecording())
 	{
 		static const FBox UnitBox(FVector(-1, -1, -1), FVector(1, 1, 1));
 
