Fixed regression, IK Retargeter goes to ref pose for 1 frame during LOD change.

diff --git a/Engine/Plugins/Animation/IKRig/Source/IKRig/Private/AnimNodes/AnimNode_RetargetPoseFromMesh.cpp b/Engine/Plugins/Animation/IKRig/Source/IKRig/Private/AnimNodes/AnimNode_RetargetPoseFromMesh.cpp
--- a/Engine/Plugins/Animation/IKRig/Source/IKRig/Private/AnimNodes/AnimNode_RetargetPoseFromMesh.cpp
+++ b/Engine/Plugins/Animation/IKRig/Source/IKRig/Private/AnimNodes/AnimNode_RetargetPoseFromMesh.cpp
@@ -61,12 +61,6 @@ void FAnimNode_RetargetPoseFromMesh::CacheBones_AnyThread(const FAnimationCacheB
 			}
 		}
 	}
-
-	// any time required bones are modified, retargeter needs reinitialized
-	if (Processor)
-	{
-		Processor->SetNeedsInitialized();
-	}
 }
 
 void FAnimNode_RetargetPoseFromMesh::Update_AnyThread(const FAnimationUpdateContext& Context)
@@ -313,13 +307,6 @@ bool FAnimNode_RetargetPoseFromMesh::EnsureProcessorIsInitialized(const TObjectP
 	{
 		return false; // cannot initialize if components are missing skeletal mesh references
 	}
-	// check that both have skeleton assets (shouldn't get this far without a skeleton)
-	const TObjectPtr<USkeleton> SourceSkeleton = SourceMesh->GetSkeleton();
-	const TObjectPtr<USkeleton> TargetSkeleton = TargetMesh->GetSkeleton();
-	if (!SourceSkeleton || !TargetSkeleton)
-	{
-		return false;
-	}
 	
 	// try initializing the processor
 	if (!Processor->WasInitializedWithTheseAssets(SourceMesh, TargetMesh, IKRetargeterAsset))
@@ -368,7 +355,7 @@ void FAnimNode_RetargetPoseFromMesh::CopyBoneTransformsFromSource(USkeletalMeshC
 	{
 		SourceMeshComponentSpaceBoneTransforms.Append(ComponentToCopyFrom->GetComponentSpaceTransforms()); // copy directly
 	}
-
+	
 	// strip all scale out of the pose values, the translation of a component-space pose has incorporated scale values
 	for (FTransform& Transform : SourceMeshComponentSpaceBoneTransforms)
 	{
diff --git a/Engine/Plugins/Animation/IKRig/Source/IKRig/Private/Retargeter/IKRetargetProcessor.cpp b/Engine/Plugins/Animation/IKRig/Source/IKRig/Private/Retargeter/IKRetargetProcessor.cpp
--- a/Engine/Plugins/Animation/IKRig/Source/IKRig/Private/Retargeter/IKRetargetProcessor.cpp
+++ b/Engine/Plugins/Animation/IKRig/Source/IKRig/Private/Retargeter/IKRetargetProcessor.cpp
@@ -2167,11 +2167,20 @@ bool UIKRetargetProcessor::WasInitializedWithTheseAssets(
 		return false;
 	}
 
-	const bool bSourceMatches = InSourceMesh == GetSkeleton(ERetargetSourceOrTarget::Source).SkeletalMesh;
-	const bool bTargetMatches = InTargetMesh == GetSkeleton(ERetargetSourceOrTarget::Target).SkeletalMesh;
+	// check that both the source and target skeletal meshes are the same as what we initialized with
+	const FRetargetSkeleton& SourceRetargetSkeleton = GetSkeleton(ERetargetSourceOrTarget::Source);
+	const FRetargetSkeleton& TargetRetargetSkeleton = GetSkeleton(ERetargetSourceOrTarget::Target);
+	const bool bSourceMatches = InSourceMesh == SourceRetargetSkeleton.SkeletalMesh;
+	const bool bTargetMatches = InTargetMesh == TargetRetargetSkeleton.SkeletalMesh;
+
+	// check that the retarget asset is the same as what we initialized with
 	const bool bAssetMatches = InRetargetAsset == RetargeterAsset;
+
+	// check that the number of bones are the same as what we initialized with
+	const bool bSourceHasSameNumberOfBones = InSourceMesh->GetRefSkeleton().GetNum() == SourceRetargetSkeleton.BoneNames.Num();
+	const bool bTargetHasSameNumberOfBones = InTargetMesh->GetRefSkeleton().GetNum() == TargetRetargetSkeleton.BoneNames.Num();
 	
-	return bSourceMatches && bTargetMatches && bAssetMatches;
+	return bSourceMatches && bTargetMatches && bAssetMatches && bSourceHasSameNumberOfBones && bTargetHasSameNumberOfBones;
 }
 
 void UIKRetargetProcessor::SetNeedsInitialized()
