Fixed undo bug in retarget details panel.

diff --git a/Engine/Plugins/Animation/IKRig/Source/IKRigEditor/Private/RetargetEditor/IKRetargetDetails.cpp b/Engine/Plugins/Animation/IKRig/Source/IKRigEditor/Private/RetargetEditor/IKRetargetDetails.cpp
--- a/Engine/Plugins/Animation/IKRig/Source/IKRigEditor/Private/RetargetEditor/IKRetargetDetails.cpp
+++ b/Engine/Plugins/Animation/IKRig/Source/IKRigEditor/Private/RetargetEditor/IKRetargetDetails.cpp
@@ -814,6 +814,23 @@ void FIKRetargetBoneDetailCustomization::CustomizeDetails(IDetailLayoutBuilder&
 				TransformType,
 				BonesView,
 				false);
+
+			TransformWidgetArgs.OnBeginSliderMovement_Lambda([](
+				ESlateTransformComponent::Type Component,
+				ESlateRotationRepresentation::Type Representation,
+				ESlateTransformSubComponent::Type SubComponent)
+			{
+				GEditor->BeginTransaction(LOCTEXT("EditRetargetPoseSlider", "Edit Retarget Pose Transform Slider"));
+			});
+			
+			TransformWidgetArgs.OnEndSliderMovement_Lambda([](
+				ESlateTransformComponent::Type Component,
+				ESlateRotationRepresentation::Type Representation,
+				ESlateTransformSubComponent::Type SubComponent,
+				double)
+			{
+				GEditor->EndTransaction();
+			});
 		}
 
 		// get/set relative
