Changed display names of properties in retarget profile.

diff --git a/Engine/Plugins/Animation/IKRig/Source/IKRig/Public/Retargeter/IKRetargetProfile.h b/Engine/Plugins/Animation/IKRig/Source/IKRig/Public/Retargeter/IKRetargetProfile.h
--- a/Engine/Plugins/Animation/IKRig/Source/IKRig/Public/Retargeter/IKRetargetProfile.h
+++ b/Engine/Plugins/Animation/IKRig/Source/IKRig/Public/Retargeter/IKRetargetProfile.h
@@ -18,7 +18,7 @@ struct FRetargetProfile
 public:
 
 	// If true, the TARGET Retarget Pose specified in this profile will be applied to the Retargeter (when plugged into the Retargeter).
-	UPROPERTY(EditAnywhere, BlueprintReadWrite, Category=RetargetPoses)
+	UPROPERTY(EditAnywhere, BlueprintReadWrite, Category=RetargetPoses, meta = (DisplayName = "Override Target Retarget Pose"))
 	bool bApplyTargetRetargetPose = false;
 	
 	// Override the TARGET Retarget Pose to use when this profile is active.
@@ -27,7 +27,7 @@ public:
 	FName TargetRetargetPoseName;
 
 	// If true, the Source Retarget Pose specified in this profile will be applied to the Retargeter (when plugged into the Retargeter).
-	UPROPERTY(EditAnywhere, BlueprintReadWrite, Category=RetargetPoses)
+	UPROPERTY(EditAnywhere, BlueprintReadWrite, Category=RetargetPoses, meta = (DisplayName = "Override Source Retarget Pose"))
 	bool bApplySourceRetargetPose = false;
 
 	// Override the SOURCE Retarget Pose to use when this profile is active.
@@ -36,15 +36,15 @@ public:
 	FName SourceRetargetPoseName;
 
 	// If true, the Chain Settings stored in this profile will be applied to the Retargeter (when plugged into the Retargeter).
-	UPROPERTY(EditAnywhere, BlueprintReadWrite, Category=ChainSettings)
-	bool bApplyChainSettings = true;
+	UPROPERTY(EditAnywhere, BlueprintReadWrite, Category=ChainSettings, meta = (DisplayName = "Override Chain Settings"))
+	bool bApplyChainSettings = false;
 	
 	// A (potentially sparse) set of setting overrides for the target chains (only applied when bApplyChainSettings is true).
 	UPROPERTY(EditAnywhere, BlueprintReadWrite, Category=ChainSettings, meta=(EditCondition="bApplyChainSettings"))
 	TMap<FName, FTargetChainSettings> ChainSettings;
 
 	// If true, the root settings stored in this profile will be applied to the Retargeter (when plugged into the Retargeter).
-	UPROPERTY(EditAnywhere, BlueprintReadWrite, Category=RootSettings)
+	UPROPERTY(EditAnywhere, BlueprintReadWrite, Category=RootSettings, meta = (DisplayName = "Override Root Settings"))
 	bool bApplyRootSettings = false;
 
 	// Retarget settings to control behavior of the retarget root motion (not applied unless bApplyRootSettings is true)
@@ -52,7 +52,7 @@ public:
 	FTargetRootSettings RootSettings;
 
 	// If true, the global settings stored in this profile will be applied to the Retargeter (when plugged into the Retargeter).
-	UPROPERTY(EditAnywhere, BlueprintReadWrite, Category=GlobalSettings)
+	UPROPERTY(EditAnywhere, BlueprintReadWrite, Category=GlobalSettings, meta = (DisplayName = "Override Global Settings"))
 	bool bApplyGlobalSettings = false;
 
 	// Retarget settings to control global behavior, like Stride Warping (not applied unless bApplyGlobalSettings is true)
