Fixed regression, retargeter animation playback broken.

diff --git a/Engine/Plugins/Animation/IKRig/Source/IKRigEditor/Private/RetargetEditor/IKRetargetEditorController.cpp b/Engine/Plugins/Animation/IKRig/Source/IKRigEditor/Private/RetargetEditor/IKRetargetEditorController.cpp
--- a/Engine/Plugins/Animation/IKRig/Source/IKRigEditor/Private/RetargetEditor/IKRetargetEditorController.cpp
+++ b/Engine/Plugins/Animation/IKRig/Source/IKRigEditor/Private/RetargetEditor/IKRetargetEditorController.cpp
@@ -150,8 +150,8 @@ void FIKRetargetEditorController::Initialize(TSharedPtr<FIKRetargetEditor> InEdi
 	Editor = InEditor;
 	AssetController = UIKRetargeterController::GetController(InAsset);
 	CurrentlyEditingSourceOrTarget = ERetargetSourceOrTarget::Target;
-	OutputMode = ERetargeterOutputMode::RunRetarget;
-	PreviousMode = OutputMode;
+	OutputMode = ERetargeterOutputMode::EditRetargetPose;
+	PreviousMode = ERetargeterOutputMode::EditRetargetPose;
 	PoseExporter = MakeShared<FIKRetargetPoseExporter>();
 	PoseExporter->Initialize(SharedThis(this));
 
