Removing dependency on private persona module from IK Rig editor module.

diff --git a/Engine/Plugins/Animation/IKRig/Source/IKRigEditor/IKRigEditor.Build.cs b/Engine/Plugins/Animation/IKRig/Source/IKRigEditor/IKRigEditor.Build.cs
--- a/Engine/Plugins/Animation/IKRig/Source/IKRigEditor/IKRigEditor.Build.cs
+++ b/Engine/Plugins/Animation/IKRig/Source/IKRigEditor/IKRigEditor.Build.cs
@@ -8,12 +8,6 @@ namespace UnrealBuildTool.Rules
     {
         public IKRigEditor(ReadOnlyTargetRules Target) : base(Target)
 		{
-			PrivateIncludePaths.AddRange(
-				new string[] {
-					System.IO.Path.Combine(GetModuleDirectory("Persona"), "Private"),
-				}
-			);
-
 			PublicDependencyModuleNames.AddRange(
                 new string[]
                 {
diff --git a/Engine/Plugins/Animation/IKRig/Source/IKRigEditor/Private/RigEditor/IKRigToolkit.cpp b/Engine/Plugins/Animation/IKRig/Source/IKRigEditor/Private/RigEditor/IKRigToolkit.cpp
--- a/Engine/Plugins/Animation/IKRig/Source/IKRigEditor/Private/RigEditor/IKRigToolkit.cpp
+++ b/Engine/Plugins/Animation/IKRig/Source/IKRigEditor/Private/RigEditor/IKRigToolkit.cpp
@@ -18,7 +18,6 @@
 
 #include "Rig/IKRigDefinition.h"
 #include "IPersonaViewport.h"
-#include "PersonaPreviewSceneDescription.h"
 #include "RigEditor/IKRigAnimInstance.h"
 #include "RigEditor/IKRigCommands.h"
 #include "RigEditor/IKRigEditMode.h"
