Fixed issue with retarget exporter flipping non animated bones with orientations that are near a quaternion singularity.

diff --git a/Engine/Plugins/Animation/IKRig/Source/IKRigEditor/Private/RetargetEditor/IKRetargetBatchOperation.cpp b/Engine/Plugins/Animation/IKRig/Source/IKRigEditor/Private/RetargetEditor/IKRetargetBatchOperation.cpp
--- a/Engine/Plugins/Animation/IKRig/Source/IKRigEditor/Private/RetargetEditor/IKRetargetBatchOperation.cpp
+++ b/Engine/Plugins/Animation/IKRig/Source/IKRigEditor/Private/RetargetEditor/IKRetargetBatchOperation.cpp
@@ -447,7 +447,7 @@ void UIKRetargetBatchOperation::ConvertAnimation(
 				FRawAnimSequenceTrack& BoneTrack = BoneTracks[TargetBoneIndex];
 				
 				BoneTrack.PosKeys[FrameIndex] = FVector3f(LocalPose.GetLocation());
-				BoneTrack.RotKeys[FrameIndex] = FQuat4f(LocalPose.GetRotation());
+				BoneTrack.RotKeys[FrameIndex] = FQuat4f(LocalPose.GetRotation().GetNormalized());
 				BoneTrack.ScaleKeys[FrameIndex] = FVector3f(LocalPose.GetScale3D());
 			}
 			
