Fixed IK Rig resetting goal transforms incorrectly when reinitializing and swapping input mesh.

diff --git a/Engine/Plugins/Animation/IKRig/Source/IKRigEditor/Private/RigEditor/IKRigController.cpp b/Engine/Plugins/Animation/IKRig/Source/IKRigEditor/Private/RigEditor/IKRigController.cpp
--- a/Engine/Plugins/Animation/IKRig/Source/IKRigEditor/Private/RigEditor/IKRigController.cpp
+++ b/Engine/Plugins/Animation/IKRig/Source/IKRigEditor/Private/RigEditor/IKRigController.cpp
@@ -1264,13 +1264,14 @@ void UIKRigController::ResetInitialGoalTransforms() const
 {
 	for (UIKRigEffectorGoal* Goal : Asset->Goals)
 	{
-		// record the current delta rotation
-		const FQuat DeltaRotation = Goal->CurrentTransform.GetRotation() * Goal->InitialTransform.GetRotation().Inverse();
-		// update the initial transform based on the new ref pose
-		const FTransform InitialTransform = GetRefPoseTransformOfBone(Goal->BoneName);
-		Goal->InitialTransform = InitialTransform;
-		// restore the current transform
-		Goal->CurrentTransform.SetRotation(Goal->InitialTransform.GetRotation() * DeltaRotation);
+		// record the current delta relative to the current bone
+		FTransform Delta = Goal->CurrentTransform.GetRelativeTransform(Goal->InitialTransform);
+		// get the initial transform based on the ref pose of the bone it's attached to
+		const FTransform NewInitialTransform = GetRefPoseTransformOfBone(Goal->BoneName);
+		// update the initial transform
+		Goal->InitialTransform = NewInitialTransform;
+		// reapply the delta
+		Goal->CurrentTransform = Delta * NewInitialTransform;
 	}
 }
 
diff --git a/Engine/Plugins/Animation/IKRig/Source/IKRigEditor/Private/RigEditor/IKRigEditorController.cpp b/Engine/Plugins/Animation/IKRig/Source/IKRigEditor/Private/RigEditor/IKRigEditorController.cpp
--- a/Engine/Plugins/Animation/IKRig/Source/IKRigEditor/Private/RigEditor/IKRigEditorController.cpp
+++ b/Engine/Plugins/Animation/IKRig/Source/IKRigEditor/Private/RigEditor/IKRigEditorController.cpp
@@ -360,7 +360,7 @@ void FIKRigEditorController::HandleIKRigNeedsInitialized(UIKRigDefinition* Modif
 	{
 		PreviewScene->SetPreviewMeshComponent(SkelMeshComponent);
 		PreviewScene->SetPreviewMesh(NewMesh);
-		AssetController->ResetGoalTransforms();
+		AssetController->ResetInitialGoalTransforms();
 	}
 
 	// re-initializes the anim instances running in the viewport
