Fixed issue with retargeter allowing assignment of skeletal mesh without an IK Rig.

diff --git a/Engine/Plugins/Animation/IKRig/Source/IKRig/Public/Retargeter/IKRetargeter.h b/Engine/Plugins/Animation/IKRig/Source/IKRig/Public/Retargeter/IKRetargeter.h
--- a/Engine/Plugins/Animation/IKRig/Source/IKRig/Public/Retargeter/IKRetargeter.h
+++ b/Engine/Plugins/Animation/IKRig/Source/IKRig/Public/Retargeter/IKRetargeter.h
@@ -311,6 +311,14 @@ public:
 	UFUNCTION(BlueprintCallable, BlueprintPure, Category=RetargetProfile)
 	static FRetargetGlobalSettings GetGlobalSettingsFromRetargetProfile(UPARAM(ref) FRetargetProfile& RetargetProfile);
 
+	// Returns true if the source IK Rig has been assigned
+	UFUNCTION(BlueprintPure, Category=RetargetProfile)
+	bool HasSourceIKRig() const { return SourceIKRigAsset != nullptr; }
+
+	// Returns true if the target IK Rig has been assigned
+	UFUNCTION(BlueprintPure, Category=RetargetProfile)
+	bool HasTargetIKRig() const { return TargetIKRigAsset != nullptr; }
+	
 	// BLUEPRINT SETTERS 
 
 	// Set the global settings in a retarget profile (will set bApplyGlobalSettings to true). 
@@ -385,7 +393,7 @@ private:
 
 #if WITH_EDITORONLY_DATA
 	// Optional. Override the Skeletal Mesh to copy animation from. Uses the preview mesh from the Source IK Rig asset by default. 
-	UPROPERTY(EditAnywhere, Category = Source)
+	UPROPERTY(EditAnywhere, Category = Source, meta = (EditCondition = "HasSourceIKRig"))
 	TSoftObjectPtr<USkeletalMesh> SourcePreviewMesh = nullptr;
 #endif
 	
@@ -395,7 +403,7 @@ private:
 
 #if WITH_EDITORONLY_DATA
 	// Optional. Override the Skeletal Mesh to preview the retarget on. Uses the preview mesh from the Target IK Rig asset by default. 
-	UPROPERTY(EditAnywhere, Category = Target)
+	UPROPERTY(EditAnywhere, Category = Target, meta = (EditCondition = "HasTargetIKRig"))
 	TSoftObjectPtr<USkeletalMesh> TargetPreviewMesh = nullptr;
 #endif
 	
