Trivial fix for IK Retargeter's mesh alignment auto-pose method.

diff --git a/Engine/Plugins/Animation/IKRig/Source/IKRigEditor/Private/RetargetEditor/IKRetargeterPoseGenerator.cpp b/Engine/Plugins/Animation/IKRig/Source/IKRigEditor/Private/RetargetEditor/IKRetargeterPoseGenerator.cpp
--- a/Engine/Plugins/Animation/IKRig/Source/IKRigEditor/Private/RetargetEditor/IKRetargeterPoseGenerator.cpp
+++ b/Engine/Plugins/Animation/IKRig/Source/IKRigEditor/Private/RetargetEditor/IKRetargeterPoseGenerator.cpp
@@ -148,7 +148,7 @@ void FRetargetAutoPoseGenerator::AlignBone(
 	case ERetargetAutoAlignMethod::MeshToMesh:
 		{
 			bFoundAlignVector = GetChainTangentFromMesh(ChainToAlign, Param, SideToAlign, DirectionToAlign);
-			bFoundMatchVector = GetChainTangentFromMesh(ChainToMatch, Param, SideToMatch, DirectionToAlign);
+			bFoundMatchVector = GetChainTangentFromMesh(ChainToMatch, Param, SideToMatch, DirectionToMatch);
 			break;
 		}
 		
