Fixes crash when editing bone details in retargeter.

diff --git a/Engine/Plugins/Animation/IKRig/Source/IKRigEditor/Private/RetargetEditor/IKRetargetDetails.cpp b/Engine/Plugins/Animation/IKRig/Source/IKRigEditor/Private/RetargetEditor/IKRetargetDetails.cpp
--- a/Engine/Plugins/Animation/IKRig/Source/IKRigEditor/Private/RetargetEditor/IKRetargetDetails.cpp
+++ b/Engine/Plugins/Animation/IKRig/Source/IKRigEditor/Private/RetargetEditor/IKRetargetDetails.cpp
@@ -688,10 +688,10 @@ void UIKRetargetBoneDetails::OnMultiNumericValueCommitted(
 	FVector::FReal Value,
 	ETextCommit::Type CommitType,
 	EIKRetargetTransformType TransformType,
-	TArrayView<UIKRetargetBoneDetails*> Bones,
+	TArrayView<TObjectPtr<UIKRetargetBoneDetails>> Bones,
 	bool bIsCommit)
 {
-	for(UIKRetargetBoneDetails* Bone : Bones)
+	for(TObjectPtr<UIKRetargetBoneDetails> Bone : Bones)
 	{	
 		Bone->OnNumericValueCommitted(Component, Representation, SubComponent, Value, CommitType, TransformType, bIsCommit);
 	}
@@ -787,7 +787,7 @@ void FIKRetargetBoneDetailCustomization::CustomizeDetails(IDetailLayoutBuilder&
 	.Font(IDetailLayoutBuilder::GetDetailFont())
 	.UseQuaternionForRotation(true);
 
-	TArrayView<UIKRetargetBoneDetails*> BonesView = TArrayView<UIKRetargetBoneDetails*>(Bones);
+	TArrayView< TObjectPtr<UIKRetargetBoneDetails> > BonesView = TArrayView< TObjectPtr<UIKRetargetBoneDetails> >(Bones);
 	
 	for(int32 PropertyIndex=0;PropertyIndex<UIData.Properties.Num();PropertyIndex++)
 	{
@@ -843,7 +843,7 @@ void FIKRetargetBoneDetailCustomization::CustomizeDetails(IDetailLayoutBuilder&
 		})
 		.OnIsComponentRelativeChanged_Lambda( [bIsEditable, BonesView, TransformType](ESlateTransformComponent::Type InComponent, bool bIsRelative)
 		{
-			for (const TObjectPtr<UIKRetargetBoneDetails> Bone: BonesView)
+			for (const TObjectPtr<UIKRetargetBoneDetails>& Bone: BonesView)
 			{
 				Bone->OnComponentRelativeChanged(InComponent, bIsRelative, TransformType);
 			}
@@ -877,8 +877,8 @@ void FIKRetargetBoneDetailCustomization::CustomizeDetails(IDetailLayoutBuilder&
 		});
 
 		// copy/paste bones transforms
-		TransformWidgetArgs.OnCopyToClipboard_UObject(Bones[0], &UIKRetargetBoneDetails::OnCopyToClipboard, TransformType);
-		TransformWidgetArgs.OnPasteFromClipboard_UObject(Bones[0], &UIKRetargetBoneDetails::OnPasteFromClipboard, TransformType);
+		TransformWidgetArgs.OnCopyToClipboard_UObject(Bones[0].Get(), &UIKRetargetBoneDetails::OnCopyToClipboard, TransformType);
+		TransformWidgetArgs.OnPasteFromClipboard_UObject(Bones[0].Get(), &UIKRetargetBoneDetails::OnPasteFromClipboard, TransformType);
 
 		TransformWidgetArgs.Visibility_Lambda([TransformChoiceWidget, TransformType]() -> EVisibility
 		{
@@ -893,6 +893,16 @@ void FIKRetargetBoneDetailCustomization::CustomizeDetails(IDetailLayoutBuilder&
 	}
 }
 
+void FIKRetargetBoneDetailCustomization::AddReferencedObjects(FReferenceCollector& Collector)
+{
+	Collector.AddReferencedObjects(Bones);
+}
+
+FString FIKRetargetBoneDetailCustomization::GetReferencerName() const
+{
+	return TEXT("FIKRetargetBoneDetailCustomization");
+}
+
 void FIKRetargetBoneDetailCustomization::GetTransformUIData(
 	const bool bIsEditingPose,
 	const IDetailLayoutBuilder& DetailBuilder,
diff --git a/Engine/Plugins/Animation/IKRig/Source/IKRigEditor/Public/RetargetEditor/IKRetargetDetails.h b/Engine/Plugins/Animation/IKRig/Source/IKRigEditor/Public/RetargetEditor/IKRetargetDetails.h
--- a/Engine/Plugins/Animation/IKRig/Source/IKRigEditor/Public/RetargetEditor/IKRetargetDetails.h
+++ b/Engine/Plugins/Animation/IKRig/Source/IKRigEditor/Public/RetargetEditor/IKRetargetDetails.h
@@ -73,7 +73,7 @@ public:
 		FVector::FReal Value,
 		ETextCommit::Type CommitType,
 		EIKRetargetTransformType TransformType,
-		TArrayView<UIKRetargetBoneDetails*> Bones,
+		TArrayView<TObjectPtr<UIKRetargetBoneDetails>> Bones,
 		bool bIsCommit);
 
 	template<typename DataType>
@@ -140,7 +140,7 @@ struct FIKRetargetTransformUIData
 	TArray<TSharedRef<IPropertyHandle>> Properties;
 };
 
-class FIKRetargetBoneDetailCustomization : public IDetailCustomization
+class FIKRetargetBoneDetailCustomization : public IDetailCustomization, FGCObject
 {
 public:
 
@@ -152,6 +152,11 @@ public:
 	/** IDetailCustomization interface */
 	virtual void CustomizeDetails(IDetailLayoutBuilder& DetailBuilder) override;
 
+	// FGCObject interface
+	virtual void AddReferencedObjects(FReferenceCollector& Collector) override;
+	virtual FString GetReferencerName() const override;
+	// End FGCObject interface
+
 private:
 
 	void GetTransformUIData(
@@ -159,7 +164,7 @@ private:
 		const IDetailLayoutBuilder& DetailBuilder,
 		FIKRetargetTransformUIData& OutData) const;
 
-	TArray<UIKRetargetBoneDetails*> Bones;
+	TArray<TObjectPtr<UIKRetargetBoneDetails>> Bones;
 };
 
 /** ------------------------------------- BEGIN CHAIN DETAILS CUSTOMIZATION -------------*/
