Fixes crash in retargeter auto-pose when bone chain did not resolve onto the skeletal mesh (incompatible mesh).

diff --git a/Engine/Plugins/Animation/IKRig/Source/IKRig/Private/Retargeter/IKRetargetProcessor.cpp b/Engine/Plugins/Animation/IKRig/Source/IKRig/Private/Retargeter/IKRetargetProcessor.cpp
--- a/Engine/Plugins/Animation/IKRig/Source/IKRig/Private/Retargeter/IKRetargetProcessor.cpp
+++ b/Engine/Plugins/Animation/IKRig/Source/IKRig/Private/Retargeter/IKRetargetProcessor.cpp
@@ -2199,6 +2199,11 @@ bool UIKRetargetProcessor::IsBoneRetargeted(
 	const FName BoneName,
 	const ERetargetSourceOrTarget SourceOrTarget) const
 {
+	ensure(bIsInitialized);
+	
+	// NOTE: cannot use Skeleton.IsBoneRetargeted because that only exists
+	// on the target skeleton and this function needs to work for either skeleton
+	
 	const FRetargetSkeleton& Skeleton = GetSkeleton(SourceOrTarget);
 	const int32 BoneIndex = Skeleton.FindBoneIndexByName(BoneName);
 	if (BoneIndex == INDEX_NONE)
@@ -2219,29 +2224,14 @@ bool UIKRetargetProcessor::IsBoneRetargeted(
 		// bones not in a chain, therefore cannot be retargeted
 		return false;
 	}
-
-	// bone must be in a chain that is actually mapped to something
-	const TArray<TObjectPtr<URetargetChainSettings>>& ChainMaps = RetargeterAsset->GetAllChainSettings();
-	for (const TObjectPtr<URetargetChainSettings>& ChainMap : ChainMaps)
+	
+	// bone must be in a successfully resolved chain that is mapped and ready to retarget
+	for (const FRetargetChainPairFK& ChainPair : ChainPairsFK)
 	{
-		if (ChainMap->SourceChain == NAME_None || ChainMap->TargetChain == NAME_None)
-		{
-			continue;
-		}
-		
-		if (SourceOrTarget == ERetargetSourceOrTarget::Source)
+		const FName ChainName = SourceOrTarget == ERetargetSourceOrTarget::Source ? ChainPair.SourceBoneChainName : ChainPair.TargetBoneChainName;
+		if (ChainName == ChainThatContainsBone)
 		{
-			if (ChainMap->SourceChain == ChainThatContainsBone)
-			{
-				return true;
-			}
-		}
-		else
-		{
-			if (ChainMap->TargetChain == ChainThatContainsBone)
-			{
-				return true;
-			}
+			return true;
 		}
 	}
 
