Fixing static analysis issue in IKRigAutoFBIK.cpp

diff --git a/Engine/Plugins/Animation/IKRig/Source/IKRigEditor/Private/RigEditor/IKRigAutoFBIK.cpp b/Engine/Plugins/Animation/IKRig/Source/IKRigEditor/Private/RigEditor/IKRigAutoFBIK.cpp
--- a/Engine/Plugins/Animation/IKRig/Source/IKRigEditor/Private/RigEditor/IKRigAutoFBIK.cpp
+++ b/Engine/Plugins/Animation/IKRig/Source/IKRigEditor/Private/RigEditor/IKRigAutoFBIK.cpp
@@ -40,7 +40,7 @@ void FAutoFBIKCreator::CreateFBIKSetup(const UIKRigController& IKRigController,
 		}
 		
 		const FBoneChain* Chain = IKRigController.GetRetargetChainByName(ExpectedChain.ChainName);
-		FName GoalName = Chain && Chain->IKGoalName != NAME_None ? Chain->IKGoalName : ExpectedChain.IKGoalName;
+		FName GoalName = (Chain && Chain->IKGoalName != NAME_None) ? Chain->IKGoalName : ExpectedChain.IKGoalName;
 		const UIKRigEffectorGoal* ChainGoal = IKRigController.GetGoal(GoalName);
 		if (!ChainGoal)
 		{
