Exposed auto-FBIK setup to BP/python.

diff --git a/Engine/Plugins/Animation/IKRig/Source/IKRigEditor/Private/RigEditor/IKRigController.cpp b/Engine/Plugins/Animation/IKRig/Source/IKRigEditor/Private/RigEditor/IKRigController.cpp
--- a/Engine/Plugins/Animation/IKRig/Source/IKRigEditor/Private/RigEditor/IKRigController.cpp
+++ b/Engine/Plugins/Animation/IKRig/Source/IKRigEditor/Private/RigEditor/IKRigController.cpp
@@ -416,14 +416,14 @@ void UIKRigController::SetRetargetDefinition(const FRetargetDefinition& Retarget
 	Asset->RetargetDefinition = RetargetDefinition;
 }
 
-void UIKRigController::ApplyAutoGeneratedRetargetDefinition()
+bool UIKRigController::ApplyAutoGeneratedRetargetDefinition()
 {
 	check(Asset)
 	
 	USkeletalMesh* Mesh = GetSkeletalMesh();
 	if (!Mesh)
 	{
-		return;
+		return false;
 	}
 
 	FScopedTransaction Transaction(LOCTEXT("ApplyAutoCharacterization_Label", "Apply Auto Characterization"));
@@ -433,7 +433,33 @@ void UIKRigController::ApplyAutoGeneratedRetargetDefinition()
 	// apply an auto generated retarget definition
 	FAutoCharacterizeResults Results;
 	AutoCharacterizer.Get()->GenerateRetargetDefinitionFromMesh(Mesh, Results);
-	SetRetargetDefinition(Results.RetargetDefinition);	
+	if (!Results.bUsedTemplate)
+	{
+		return false;
+	}
+	
+	SetRetargetDefinition(Results.RetargetDefinition);
+	return true;
+}
+
+bool UIKRigController::ApplyAutoFBIK()
+{
+	check(Asset)
+	
+	USkeletalMesh* Mesh = GetSkeletalMesh();
+	if (!Mesh)
+	{
+		return false;
+	}
+
+	FScopedTransaction Transaction(LOCTEXT("ApplyAutoFBIK_Label", "Apply Auto FBIK"));
+	FScopedReinitializeIKRig Reinitialize(this);
+	Asset->Modify();
+	
+	// apply an auto generated FBIK setup
+	FAutoFBIKResults Results;
+	AutoFBIKCreator.Get()->CreateFBIKSetup(*this, Results);
+	return Results.Outcome == EAutoFBIKResult::AllOk;
 }
 
 void UIKRigController::AutoGenerateRetargetDefinition(FAutoCharacterizeResults& Results) const
diff --git a/Engine/Plugins/Animation/IKRig/Source/IKRigEditor/Public/RigEditor/IKRigController.h b/Engine/Plugins/Animation/IKRig/Source/IKRigEditor/Public/RigEditor/IKRigController.h
--- a/Engine/Plugins/Animation/IKRig/Source/IKRigEditor/Public/RigEditor/IKRigController.h
+++ b/Engine/Plugins/Animation/IKRig/Source/IKRigEditor/Public/RigEditor/IKRigController.h
@@ -325,9 +325,15 @@ public:
 	UFUNCTION(BlueprintCallable, Category=IKRig)
 	void SetRetargetDefinition(const FRetargetDefinition& RetargetDefinition) const;
 
-	// Analyse the skeleton and automatically generate all retarget chains and set the retarget root
+	// Analyse the skeleton to see if it matches a known template and automatically generates all retarget chains and sets the retarget root
+	// Returns true if a matching skeletal template was found and the retarget definition for it was applied.
 	UFUNCTION(BlueprintCallable, Category=IKRig)
-	void ApplyAutoGeneratedRetargetDefinition();
+	bool ApplyAutoGeneratedRetargetDefinition();
+
+	// Analyse the skeleton to see if it matches a known template and automatically generates a full body IK setup
+	// Returns true if a matching skeletal template was found and the FBIK setup for it was applied.
+	UFUNCTION(BlueprintCallable, Category=IKRig)
+	bool ApplyAutoFBIK();
 	
 	//
 	// RETARGETING C++ ONLY API
