Fixes regression, retargeter asset creation in wrong submenu

diff --git a/Engine/Plugins/Animation/IKRig/Source/IKRigEditor/Private/RetargetEditor/AssetTypeActions_IKRetargeter.cpp b/Engine/Plugins/Animation/IKRig/Source/IKRigEditor/Private/RetargetEditor/AssetTypeActions_IKRetargeter.cpp
--- a/Engine/Plugins/Animation/IKRig/Source/IKRigEditor/Private/RetargetEditor/AssetTypeActions_IKRetargeter.cpp
+++ b/Engine/Plugins/Animation/IKRig/Source/IKRigEditor/Private/RetargetEditor/AssetTypeActions_IKRetargeter.cpp
@@ -156,4 +156,13 @@ void FAssetTypeActions_IKRetargeter::ExtendAnimAssetMenusForBatchRetargeting()
 	}
 }
 
+const TArray<FText>& FAssetTypeActions_IKRetargeter::GetSubMenus() const
+{
+	static const TArray<FText> SubMenus
+	{
+		LOCTEXT("AnimRetargetingSubMenu", "Retargeting")
+	};
+	return SubMenus;
+}
+
 #undef LOCTEXT_NAMESPACE
diff --git a/Engine/Plugins/Animation/IKRig/Source/IKRigEditor/Private/RigEditor/IKRigDefinitionFactory.cpp b/Engine/Plugins/Animation/IKRig/Source/IKRigEditor/Private/RigEditor/IKRigDefinitionFactory.cpp
--- a/Engine/Plugins/Animation/IKRig/Source/IKRigEditor/Private/RigEditor/IKRigDefinitionFactory.cpp
+++ b/Engine/Plugins/Animation/IKRig/Source/IKRigEditor/Private/RigEditor/IKRigDefinitionFactory.cpp
@@ -52,7 +52,7 @@ FText UIKRigDefinitionFactory::GetToolTip() const
 
 FString UIKRigDefinitionFactory::GetDefaultNewAssetName() const
 {
-	return FString(TEXT("NewIKRig"));
+	return FString(TEXT("IK_NewIKRig"));
 }
 #undef LOCTEXT_NAMESPACE
 
diff --git a/Engine/Plugins/Animation/IKRig/Source/IKRigEditor/Public/RetargetEditor/AssetTypeActions_IKRetargeter.h b/Engine/Plugins/Animation/IKRig/Source/IKRigEditor/Public/RetargetEditor/AssetTypeActions_IKRetargeter.h
--- a/Engine/Plugins/Animation/IKRig/Source/IKRigEditor/Public/RetargetEditor/AssetTypeActions_IKRetargeter.h
+++ b/Engine/Plugins/Animation/IKRig/Source/IKRigEditor/Public/RetargetEditor/AssetTypeActions_IKRetargeter.h
@@ -19,6 +19,7 @@ public:
 	virtual void OpenAssetEditor(const TArray<UObject*>& InObjects, TSharedPtr<class IToolkitHost> EditWithinLevelEditor = TSharedPtr<IToolkitHost>()) override;
 	virtual uint32 GetCategories() override { return EAssetTypeCategories::Animation; }
 	virtual UThumbnailInfo* GetThumbnailInfo(UObject* Asset) const override;
+	virtual const TArray<FText>& GetSubMenus() const override;
 	// END IAssetTypeActions
 
 	static void ExtendIKRigMenuToMakeRetargeter();
