Fixes crash when opening new IK Rig or IK Retargeter.

diff --git a/Engine/Plugins/Animation/IKRig/Source/IKRig/Private/AnimNodes/AnimNode_IKRig.cpp b/Engine/Plugins/Animation/IKRig/Source/IKRig/Private/AnimNodes/AnimNode_IKRig.cpp
--- a/Engine/Plugins/Animation/IKRig/Source/IKRig/Private/AnimNodes/AnimNode_IKRig.cpp
+++ b/Engine/Plugins/Animation/IKRig/Source/IKRig/Private/AnimNodes/AnimNode_IKRig.cpp
@@ -182,12 +182,8 @@ void FAnimNode_IKRig::Initialize_AnyThread(const FAnimationInitializeContext& Co
 
 	// Initial update of the node, so we dont have a frame-delay on setup
 	GetEvaluateGraphExposedInputs().Execute(Context);
-
-	// ensure there is always a processor available
-	if (!IKRigProcessor && IsInGameThread())
-	{
-		IKRigProcessor = NewObject<UIKRigProcessor>(Context.AnimInstanceProxy->GetSkelMeshComponent());	
-	}
+	
+	CreateIKRigProcessorIfNeeded(Context.AnimInstanceProxy->GetSkelMeshComponent());
 
 	InitializeProperties(Context.GetAnimInstanceObject(), GetTargetClass());
 }
@@ -255,7 +251,7 @@ void FAnimNode_IKRig::PreUpdate(const UAnimInstance* InAnimInstance)
 	
 	if (!IsValid(IKRigProcessor))
 	{
-		IKRigProcessor = NewObject<UIKRigProcessor>(InAnimInstance->GetOwningComponent());	
+		CreateIKRigProcessorIfNeeded(InAnimInstance->GetOwningComponent());
 	}
 	
 	// initialize the IK Rig (will only try once on the current version of the rig asset)
@@ -306,6 +302,20 @@ void FAnimNode_IKRig::SetProcessorNeedsInitialized()
 	}
 }
 
+void FAnimNode_IKRig::CreateIKRigProcessorIfNeeded(UObject* Outer)
+{
+	// ensure there is always a processor available
+	if (!IKRigProcessor && IsInGameThread())
+	{
+		IKRigProcessor = NewObject<UIKRigProcessor>(Outer);	
+	}
+}
+
+UIKRigProcessor* FAnimNode_IKRig::GetIKRigProcessor()
+{
+	return IKRigProcessor;
+}
+
 void FAnimNode_IKRig::CacheBones_AnyThread(const FAnimationCacheBonesContext& Context)
 {
 	FAnimNode_Base::CacheBones_AnyThread(Context);
diff --git a/Engine/Plugins/Animation/IKRig/Source/IKRig/Private/AnimNodes/AnimNode_RetargetPoseFromMesh.cpp b/Engine/Plugins/Animation/IKRig/Source/IKRig/Private/AnimNodes/AnimNode_RetargetPoseFromMesh.cpp
--- a/Engine/Plugins/Animation/IKRig/Source/IKRig/Private/AnimNodes/AnimNode_RetargetPoseFromMesh.cpp
+++ b/Engine/Plugins/Animation/IKRig/Source/IKRig/Private/AnimNodes/AnimNode_RetargetPoseFromMesh.cpp
@@ -20,10 +20,8 @@ void FAnimNode_RetargetPoseFromMesh::Initialize_AnyThread(const FAnimationInitia
 	// Initial update of the node, so we dont have a frame-delay on setup
 	GetEvaluateGraphExposedInputs().Execute(Context);
 
-	if (!Processor && IsInGameThread())
-	{
-		Processor = NewObject<UIKRetargetProcessor>(Context.AnimInstanceProxy->GetSkelMeshComponent());	
-	}
+	// create the processor
+	CreateRetargetProcessorIfNeeded(Context.AnimInstanceProxy->GetSkelMeshComponent());
 }
 
 void FAnimNode_RetargetPoseFromMesh::CacheBones_AnyThread(const FAnimationCacheBonesContext& Context)
@@ -189,6 +187,14 @@ void FAnimNode_RetargetPoseFromMesh::PreUpdate(const UAnimInstance* InAnimInstan
 	}
 }
 
+void FAnimNode_RetargetPoseFromMesh::CreateRetargetProcessorIfNeeded(UObject* Outer)
+{
+	if (!Processor && IsInGameThread())
+	{
+		Processor = NewObject<UIKRetargetProcessor>(Outer);	
+	}
+}
+
 void FAnimNode_RetargetPoseFromMesh::CopyAndRemapCurvesFromSourceToTarget(FBlendedCurve& OutputCurves) const
 {
 	if (!bCopyCurves)
diff --git a/Engine/Plugins/Animation/IKRig/Source/IKRig/Public/AnimNodes/AnimNode_IKRig.h b/Engine/Plugins/Animation/IKRig/Source/IKRig/Public/AnimNodes/AnimNode_IKRig.h
--- a/Engine/Plugins/Animation/IKRig/Source/IKRig/Public/AnimNodes/AnimNode_IKRig.h
+++ b/Engine/Plugins/Animation/IKRig/Source/IKRig/Public/AnimNodes/AnimNode_IKRig.h
@@ -110,8 +110,11 @@ public:
 	virtual void OnInitializeAnimInstance(const FAnimInstanceProxy* InProxy, const UAnimInstance* InAnimInstance) override;
 	// End of FAnimNode_Base interface
 
-	/** force reinitialization */
+	// force reinitialization
 	void SetProcessorNeedsInitialized();
+	
+	void CreateIKRigProcessorIfNeeded(UObject* Outer);
+	UIKRigProcessor* GetIKRigProcessor();
 
 protected:
 
@@ -128,6 +131,4 @@ private:
 	virtual void ConditionalDebugDraw(FPrimitiveDrawInterface* PDI, USkeletalMeshComponent* PreviewSkelMeshComp) const;
 	
 	friend class UAnimGraphNode_IKRig;
-	friend struct FIKRigAnimInstanceProxy;
-	friend class UIKRigAnimInstance;
 };
diff --git a/Engine/Plugins/Animation/IKRig/Source/IKRig/Public/AnimNodes/AnimNode_RetargetPoseFromMesh.h b/Engine/Plugins/Animation/IKRig/Source/IKRig/Public/AnimNodes/AnimNode_RetargetPoseFromMesh.h
--- a/Engine/Plugins/Animation/IKRig/Source/IKRig/Public/AnimNodes/AnimNode_RetargetPoseFromMesh.h
+++ b/Engine/Plugins/Animation/IKRig/Source/IKRig/Public/AnimNodes/AnimNode_RetargetPoseFromMesh.h
@@ -52,6 +52,7 @@ struct IKRIG_API FAnimNode_RetargetPoseFromMesh : public FAnimNode_Base
 	// End of FAnimNode_Base interface
 
 	/** Access to the runtime processor */
+	void CreateRetargetProcessorIfNeeded(UObject* Outer);
 	UIKRetargetProcessor* GetRetargetProcessor() const;
 
 private:
diff --git a/Engine/Plugins/Animation/IKRig/Source/IKRigEditor/Private/RetargetEditor/IKRetargetAnimInstance.cpp b/Engine/Plugins/Animation/IKRig/Source/IKRigEditor/Private/RetargetEditor/IKRetargetAnimInstance.cpp
--- a/Engine/Plugins/Animation/IKRig/Source/IKRigEditor/Private/RetargetEditor/IKRetargetAnimInstance.cpp
+++ b/Engine/Plugins/Animation/IKRig/Source/IKRigEditor/Private/RetargetEditor/IKRetargetAnimInstance.cpp
@@ -151,8 +151,11 @@ void UIKRetargetAnimInstance::ConfigureAnimInstance(
 	Proxy.ConfigureAnimInstance(InSourceOrTarget, InAsset, InSourceMeshComponent);
 }
 
-UIKRetargetProcessor* UIKRetargetAnimInstance::GetRetargetProcessor() const
+UIKRetargetProcessor* UIKRetargetAnimInstance::GetRetargetProcessor()
 {
+	// depending on initialization order, the anim node may not have created the processor yet, so optionally force it to do so now
+	RetargetNode.CreateRetargetProcessorIfNeeded(this);
+	
 	return RetargetNode.GetRetargetProcessor();
 }
 
diff --git a/Engine/Plugins/Animation/IKRig/Source/IKRigEditor/Private/RigEditor/IKRigAnimInstance.cpp b/Engine/Plugins/Animation/IKRig/Source/IKRigEditor/Private/RigEditor/IKRigAnimInstance.cpp
--- a/Engine/Plugins/Animation/IKRig/Source/IKRigEditor/Private/RigEditor/IKRigAnimInstance.cpp
+++ b/Engine/Plugins/Animation/IKRig/Source/IKRigEditor/Private/RigEditor/IKRigAnimInstance.cpp
@@ -24,7 +24,8 @@ void UIKRigAnimInstance::SetProcessorNeedsInitialized()
 
 UIKRigProcessor* UIKRigAnimInstance::GetCurrentlyRunningProcessor()
 {
-	return IKRigNode.IKRigProcessor;
+	IKRigNode.CreateIKRigProcessorIfNeeded(this);
+	return IKRigNode.GetIKRigProcessor();
 }
 
 FAnimInstanceProxy* UIKRigAnimInstance::CreateAnimInstanceProxy()
diff --git a/Engine/Plugins/Animation/IKRig/Source/IKRigEditor/Public/RetargetEditor/IKRetargetAnimInstance.h b/Engine/Plugins/Animation/IKRig/Source/IKRigEditor/Public/RetargetEditor/IKRetargetAnimInstance.h
--- a/Engine/Plugins/Animation/IKRig/Source/IKRigEditor/Public/RetargetEditor/IKRetargetAnimInstance.h
+++ b/Engine/Plugins/Animation/IKRig/Source/IKRigEditor/Public/RetargetEditor/IKRetargetAnimInstance.h
@@ -60,7 +60,7 @@ public:
 		UIKRetargeter* InIKRetargetAsset,
 		TWeakObjectPtr<USkeletalMeshComponent> InSourceMeshComponent);
 
-	UIKRetargetProcessor* GetRetargetProcessor() const;
+	UIKRetargetProcessor* GetRetargetProcessor();
 
 	const TArray<FTransform>& GetLocalRetargetPose() const { return PreviewPoseNode.GetLocalRetargetPose(); };
 	const TArray<FTransform>& GetGlobalRetargetPose() const { return PreviewPoseNode.GetGlobalRetargetPose(); };
