Fixed bug, exported IK Rig assets did not reference each other.

diff --git a/Engine/Plugins/Animation/IKRig/Source/IKRigEditor/Private/RetargetEditor/SRetargetAnimAssetsWindow.cpp b/Engine/Plugins/Animation/IKRig/Source/IKRigEditor/Private/RetargetEditor/SRetargetAnimAssetsWindow.cpp
--- a/Engine/Plugins/Animation/IKRig/Source/IKRigEditor/Private/RetargetEditor/SRetargetAnimAssetsWindow.cpp
+++ b/Engine/Plugins/Animation/IKRig/Source/IKRigEditor/Private/RetargetEditor/SRetargetAnimAssetsWindow.cpp
@@ -1027,19 +1027,27 @@ FReply SRetargetAnimAssetsWindow::OnExportRetargetAssets()
 
 	// save the retargeter
 	FString RetargetAssetName = TEXT("RTG_AutoGenerated");
-	UObject* ExportedRetargeter = SaveAssetToDisk(BatchContext.IKRetargetAsset, RetargetAssetName, BatchContext.NameRule.FolderPath);
+	UObject* ExportedRetargeterObj = SaveAssetToDisk(BatchContext.IKRetargetAsset, RetargetAssetName, BatchContext.NameRule.FolderPath);
 	// save the SOURCE IK Rig
 	FString SourceIKRigAssetName = TEXT("IK_AutoGeneratedSource");
-	UObject* ExportedSourceIKRig = SaveAssetToDisk(SourceIKRig, SourceIKRigAssetName, BatchContext.NameRule.FolderPath);
+	UObject* ExportedSourceIKRigObj = SaveAssetToDisk(SourceIKRig, SourceIKRigAssetName, BatchContext.NameRule.FolderPath);
 	// save the TARGET IK Rig
 	FString TargetIKRigAssetName = TEXT("IK_AutoGeneratedTarget");
-	UObject* ExportedTargetIKRig = SaveAssetToDisk(TargetIKRig, TargetIKRigAssetName, BatchContext.NameRule.FolderPath);
+	UObject* ExportedTargetIKRigObj = SaveAssetToDisk(TargetIKRig, TargetIKRigAssetName, BatchContext.NameRule.FolderPath);
+
+	// modify exported retargeter to point to exported ik rigs
+	UIKRetargeter* ExportedRetargeter = CastChecked<UIKRetargeter>(ExportedRetargeterObj);
+	UIKRigDefinition* ExportedSourceIKRig = CastChecked<UIKRigDefinition>(ExportedSourceIKRigObj);
+	UIKRigDefinition* ExportedTargetIKRig = CastChecked<UIKRigDefinition>(ExportedTargetIKRigObj);
+	UIKRetargeterController* ExportedController = UIKRetargeterController::GetController(ExportedRetargeter);
+	ExportedController->SetIKRig(ERetargetSourceOrTarget::Source, ExportedSourceIKRig);
+	ExportedController->SetIKRig(ERetargetSourceOrTarget::Target, ExportedTargetIKRig);
 
 	// select all new assets and show in the content browser
 	TArray<FAssetData> NewAssets;
-	NewAssets.Add(FAssetData(ExportedRetargeter));
-	NewAssets.Add(FAssetData(ExportedSourceIKRig));
-	NewAssets.Add(FAssetData(ExportedTargetIKRig));
+	NewAssets.Add(FAssetData(ExportedRetargeterObj));
+	NewAssets.Add(FAssetData(ExportedSourceIKRigObj));
+	NewAssets.Add(FAssetData(ExportedTargetIKRigObj));
 	const FContentBrowserModule& ContentBrowserModule = FModuleManager::Get().LoadModuleChecked<FContentBrowserModule>("ContentBrowser");
 	ContentBrowserModule.Get().SyncBrowserToAssets(NewAssets);
 	
