Added ability to LOD the retarget pose from mesh node.

diff --git a/Engine/Plugins/Animation/IKRig/Source/IKRig/Private/AnimNodes/AnimNode_RetargetPoseFromMesh.cpp b/Engine/Plugins/Animation/IKRig/Source/IKRig/Private/AnimNodes/AnimNode_RetargetPoseFromMesh.cpp
--- a/Engine/Plugins/Animation/IKRig/Source/IKRig/Private/AnimNodes/AnimNode_RetargetPoseFromMesh.cpp
+++ b/Engine/Plugins/Animation/IKRig/Source/IKRig/Private/AnimNodes/AnimNode_RetargetPoseFromMesh.cpp
@@ -77,7 +77,7 @@ void FAnimNode_RetargetPoseFromMesh::Evaluate_AnyThread(FPoseContext& Output)
 
 	SCOPE_CYCLE_COUNTER(STAT_IKRetarget);
 
-	if (!(IKRetargeterAsset && Processor && SourceMeshComponent.IsValid()))
+	if (!(IKRetargeterAsset && Processor && SourceMeshComponent.IsValid() && IsLODEnabled(Output.AnimInstanceProxy)))
 	{
 		Output.ResetToRefPose();
 		return;
@@ -111,6 +111,17 @@ void FAnimNode_RetargetPoseFromMesh::Evaluate_AnyThread(FPoseContext& Output)
 
 	// apply custom profile settings to the processor
 	Processor->ApplySettingsFromProfile(CustomRetargetProfile);
+	
+	// LOD off the IK pass
+	if (Output.AnimInstanceProxy->GetLODLevel() > LODThresholdForIK)
+	{
+		// override the custom profile's global settings but with IK forcibly turned off
+		FRetargetProfile TurnIKOffProfile;
+		TurnIKOffProfile.bApplyGlobalSettings = true;
+		TurnIKOffProfile.GlobalSettings = CustomRetargetProfile.GlobalSettings;
+		TurnIKOffProfile.GlobalSettings.bEnableIK = false;
+		Processor->ApplySettingsFromProfile(TurnIKOffProfile);
+	}
 
 	// run the retargeter
 	const TArray<FTransform>& RetargetedPose = Processor->RunRetargeter(SourceMeshComponentSpaceBoneTransforms, SpeedValuesFromCurves, DeltaTime);
diff --git a/Engine/Plugins/Animation/IKRig/Source/IKRig/Public/AnimNodes/AnimNode_RetargetPoseFromMesh.h b/Engine/Plugins/Animation/IKRig/Source/IKRig/Public/AnimNodes/AnimNode_RetargetPoseFromMesh.h
--- a/Engine/Plugins/Animation/IKRig/Source/IKRig/Public/AnimNodes/AnimNode_RetargetPoseFromMesh.h
+++ b/Engine/Plugins/Animation/IKRig/Source/IKRig/Public/AnimNodes/AnimNode_RetargetPoseFromMesh.h
@@ -41,6 +41,21 @@ struct IKRIG_API FAnimNode_RetargetPoseFromMesh : public FAnimNode_Base
 	/* Copy curves from SouceMeshComponent. This will copy any curves the source/target Skeleton have in common. */
 	UPROPERTY(EditAnywhere, BlueprintReadWrite, Category = Settings, meta = (NeverAsPin))
 	bool bCopyCurves = true;
+
+	/*
+	* Max LOD that this node is allowed to run.
+	* For example if you have LODThreshold to be 2, it will run until LOD 2 (based on 0 index) when the component LOD becomes 3, it will stop update/evaluate
+	*/
+	UPROPERTY(EditAnywhere, BlueprintReadWrite, Category = Performance, meta = (DisplayName = "LOD Threshold"))
+	int32 LODThreshold = -1;
+
+	/*
+	* Max LOD that IK is allowed to run.
+	* For example if you have LODThresholdForIK to be 2, it will skip the IK pass on LODs 3 and greater.
+	* This only disables IK and does not affect the Root or FK passes.
+	*/
+	UPROPERTY(EditAnywhere, BlueprintReadWrite, Category = Performance, meta = (DisplayName = "IK LOD Threshold"))
+	int32 LODThresholdForIK = -1;
 	
 	// FAnimNode_Base interface
 	virtual void Initialize_AnyThread(const FAnimationInitializeContext& Context) override;
@@ -49,6 +64,7 @@ struct IKRIG_API FAnimNode_RetargetPoseFromMesh : public FAnimNode_Base
 	virtual void Evaluate_AnyThread(FPoseContext& Output) override;
 	virtual bool HasPreUpdate() const override { return true; }
 	virtual void PreUpdate(const UAnimInstance* InAnimInstance) override;
+	virtual int32 GetLODThreshold() const override { return LODThreshold; }
 	// End of FAnimNode_Base interface
 
 	/** Access to the runtime processor */
diff --git a/Engine/Source/Runtime/Engine/Classes/Animation/AnimNodeBase.h b/Engine/Source/Runtime/Engine/Classes/Animation/AnimNodeBase.h
--- a/Engine/Source/Runtime/Engine/Classes/Animation/AnimNodeBase.h
+++ b/Engine/Source/Runtime/Engine/Classes/Animation/AnimNodeBase.h
@@ -1007,7 +1007,7 @@ protected:
 	/** return true if enabled, otherwise, return false. This is utility function that can be used per node level */
 	ENGINE_API bool IsLODEnabled(FAnimInstanceProxy* AnimInstanceProxy);
 
-	/** Get the LOD threshold at which this node is enabled. Node is enabled if the current LOD >= threshold. */
+	/** Get the LOD level at which this node is enabled. Node is enabled if the current LOD is less than or equal to this threshold. */
 	virtual int32 GetLODThreshold() const { return INDEX_NONE; }
 
 	/** Called once, from game thread as the parent anim instance is created */
