Overrelaxation upper limit wrong in FBIK solver (trivial fix)

diff --git a/Engine/Plugins/Animation/IKRig/Source/IKRig/Public/Rig/Solvers/IKRig_FBIKSolver.h b/Engine/Plugins/Animation/IKRig/Source/IKRig/Public/Rig/Solvers/IKRig_FBIKSolver.h
--- a/Engine/Plugins/Animation/IKRig/Source/IKRig/Public/Rig/Solvers/IKRig_FBIKSolver.h
+++ b/Engine/Plugins/Animation/IKRig/Source/IKRig/Public/Rig/Solvers/IKRig_FBIKSolver.h
@@ -228,7 +228,7 @@ public:
 	float MaxAngle = 30.f;
 
 	/** Pushes constraints beyond their normal amount to speed up convergence. Increasing this may speed up convergence, but at the cost of stability. Range is 1.0 - 2.0. Default is 1.3. */
-	UPROPERTY(BlueprintReadWrite, EditAnywhere, Category = AdvancedSettings, meta = (ClampMin = "1", UIMin = "1.0", UIMax = "10.0"))
+	UPROPERTY(BlueprintReadWrite, EditAnywhere, Category = AdvancedSettings, meta = (ClampMin = "1",  ClampMax = "2", UIMin = "1.0", UIMax = "2.0"))
 	float OverRelaxation = 1.3f;
 	
 	/** DEPRECATED: When true, the solver is reset each tick to start from the current input pose. Default is true.
