Improved IK Rig ability to guess which "side" a bone belongs to even when whole skeleton rotated.

diff --git a/Engine/Plugins/Animation/IKRig/Source/IKRigEditor/Private/RigEditor/IKRigEditorController.cpp b/Engine/Plugins/Animation/IKRig/Source/IKRigEditor/Private/RigEditor/IKRigEditorController.cpp
--- a/Engine/Plugins/Animation/IKRig/Source/IKRigEditor/Private/RigEditor/IKRigEditorController.cpp
+++ b/Engine/Plugins/Animation/IKRig/Source/IKRigEditor/Private/RigEditor/IKRigEditorController.cpp
@@ -280,8 +280,39 @@ FName FRetargetChainAnalyzer::GetDefaultChainName()
 	return FName(*NewChainText.ToString());
 }
 
-EChainSide FRetargetChainAnalyzer::GetSideOfChain(const TArray<int32>& BoneIndices, const FIKRigSkeleton& IKRigSkeleton)
+EChainSide FRetargetChainAnalyzer::GetSideOfChain(
+	const TArray<int32>& BoneIndices,
+	const FIKRigSkeleton& IKRigSkeleton) const
 {
+	// check if bones are predominantly named "Left" or "Right"
+	auto DoMajorityBonesContainText = [this, &BoneIndices, &IKRigSkeleton](const FString& StrToTest) -> bool
+	{
+		int32 Score = 0;
+		for (const int32 BoneIndex : BoneIndices)
+		{
+			const FString BoneName = IKRigSkeleton.GetBoneNameFromIndex(BoneIndex).ToString().ToLower();
+			if (BoneName.Contains(StrToTest))
+			{
+				++Score;
+			}
+		}
+		
+		return Score > BoneIndices.Num() * 0.5f;
+	};
+	
+	if (DoMajorityBonesContainText("Right"))
+	{
+		return EChainSide::Right;
+	}
+	if (DoMajorityBonesContainText("Left"))
+	{
+		return EChainSide::Left;
+	}
+
+	//
+	// bones don't have left/right prefix/suffix, so lets fallback on using a spatial test...
+	//
+	
 	// determine "sidedness" of the chain based on the location of the bones (left, right or center of YZ plane)
 	float AverageXPositionOfChain = 0.f;
 	for (const int32 BoneIndex : BoneIndices)
diff --git a/Engine/Plugins/Animation/IKRig/Source/IKRigEditor/Public/RigEditor/IKRigEditorController.h b/Engine/Plugins/Animation/IKRig/Source/IKRigEditor/Public/RigEditor/IKRigEditorController.h
--- a/Engine/Plugins/Animation/IKRig/Source/IKRigEditor/Public/RigEditor/IKRigEditorController.h
+++ b/Engine/Plugins/Animation/IKRig/Source/IKRigEditor/Public/RigEditor/IKRigEditorController.h
@@ -100,7 +100,7 @@ struct FRetargetChainAnalyzer
 	
 	static FName GetDefaultChainName();
 
-	static EChainSide GetSideOfChain(const TArray<int32>& BoneIndices, const FIKRigSkeleton& IKRigSkeleton);
+	EChainSide GetSideOfChain(const TArray<int32>& BoneIndices, const FIKRigSkeleton& IKRigSkeleton) const;
 };
 
 // a home for cross-widget communication to synchronize state across all tabs and viewport 
