Fixed bug, selection not framed on new IK Retargeter.

diff --git a/Engine/Plugins/Animation/IKRig/Source/IKRigEditor/Private/RetargetEditor/IKRetargetEditorController.cpp b/Engine/Plugins/Animation/IKRig/Source/IKRigEditor/Private/RetargetEditor/IKRetargetEditorController.cpp
--- a/Engine/Plugins/Animation/IKRig/Source/IKRigEditor/Private/RetargetEditor/IKRetargetEditorController.cpp
+++ b/Engine/Plugins/Animation/IKRig/Source/IKRigEditor/Private/RetargetEditor/IKRetargetEditorController.cpp
@@ -452,15 +452,18 @@ bool FIKRetargetEditorController::GetCameraTargetForSelection(FSphere& OutTarget
 	case ERetargetSelectionType::MESH:
 	case ERetargetSelectionType::NONE:
 	default:
-		// target the current mesh
-		if (const UPrimitiveComponent* CurrentMesh = GetSkeletalMeshComponent(GetSourceOrTarget()))
+		// frame both meshes
+		OutTarget = FSphere();
+		if (const UPrimitiveComponent* SourceComponent = GetSkeletalMeshComponent(ERetargetSourceOrTarget::Source))
 		{
-			OutTarget = CurrentMesh->Bounds.GetSphere();
-			return true;
+			OutTarget += SourceComponent->Bounds.GetSphere();
+		}
+		if (const UPrimitiveComponent* TargetComponent = GetSkeletalMeshComponent(ERetargetSourceOrTarget::Target))
+		{
+			OutTarget += TargetComponent->Bounds.GetSphere();
 		}
+		return true;
 	}
-
-	return false;
 }
 
 bool FIKRetargetEditorController::IsEditingPoseWithAnyBoneSelected() const
