Fixes order of properties in Target Root Settings struct in retargeter.

diff --git a/Engine/Plugins/Animation/IKRig/Source/IKRigEditor/Private/RetargetEditor/IKRetargetDetails.cpp b/Engine/Plugins/Animation/IKRig/Source/IKRigEditor/Private/RetargetEditor/IKRetargetDetails.cpp
--- a/Engine/Plugins/Animation/IKRig/Source/IKRigEditor/Private/RetargetEditor/IKRetargetDetails.cpp
+++ b/Engine/Plugins/Animation/IKRig/Source/IKRigEditor/Private/RetargetEditor/IKRetargetDetails.cpp
@@ -1335,20 +1335,20 @@ void FRetargetRootSettingsCustomization::CustomizeDetails(IDetailLayoutBuilder&
 	AddPropertyRowToGroup(AlphaGroup,GET_MEMBER_NAME_STRING_CHECKED(FTargetRootSettings, RotationAlpha));
 	AddPropertyRowToGroup(AlphaGroup,GET_MEMBER_NAME_STRING_CHECKED(FTargetRootSettings, TranslationAlpha));
 
-	// offset group
-	IDetailGroup& OffsetGroup = SettingsCategory.AddGroup("Offsets", LOCTEXT("OffsetRoot_Label", "Offsets"), false, true);
-	AddPropertyRowToGroup(OffsetGroup,GET_MEMBER_NAME_STRING_CHECKED(FTargetRootSettings, TranslationOffset));
-	AddPropertyRowToGroup(OffsetGroup,GET_MEMBER_NAME_STRING_CHECKED(FTargetRootSettings, RotationOffset));
+	// blend to source group
+	IDetailGroup& BlendToSourceGroup = SettingsCategory.AddGroup("Blend To Source", LOCTEXT("BlendToSource_Label", "Blend to Source"), false, true);
+	AddPropertyRowToGroup(BlendToSourceGroup,GET_MEMBER_NAME_STRING_CHECKED(FTargetRootSettings, BlendToSource));
+	AddPropertyRowToGroup(BlendToSourceGroup,GET_MEMBER_NAME_STRING_CHECKED(FTargetRootSettings, BlendToSourceWeights));
 
 	// scale group
 	IDetailGroup& ScaleGroup = SettingsCategory.AddGroup("Scale Translation", LOCTEXT("ScaleRoot_Label", "Scale Translation"), false, true);
 	AddPropertyRowToGroup(ScaleGroup,GET_MEMBER_NAME_STRING_CHECKED(FTargetRootSettings, ScaleHorizontal));
 	AddPropertyRowToGroup(ScaleGroup,GET_MEMBER_NAME_STRING_CHECKED(FTargetRootSettings, ScaleVertical));
 
-	// blend to source group
-	IDetailGroup& BlendToSourceGroup = SettingsCategory.AddGroup("Blend To Source", LOCTEXT("BlendToSource_Label", "Blend to Source"), false, true);
-	AddPropertyRowToGroup(BlendToSourceGroup,GET_MEMBER_NAME_STRING_CHECKED(FTargetRootSettings, BlendToSource));
-	AddPropertyRowToGroup(BlendToSourceGroup,GET_MEMBER_NAME_STRING_CHECKED(FTargetRootSettings, BlendToSourceWeights));
+	// offset group
+	IDetailGroup& OffsetGroup = SettingsCategory.AddGroup("Offsets", LOCTEXT("OffsetRoot_Label", "Offsets"), false, true);
+	AddPropertyRowToGroup(OffsetGroup,GET_MEMBER_NAME_STRING_CHECKED(FTargetRootSettings, TranslationOffset));
+	AddPropertyRowToGroup(OffsetGroup,GET_MEMBER_NAME_STRING_CHECKED(FTargetRootSettings, RotationOffset));
 
 	// affect IK directionally
 	IDetailGroup& AffectIKGroup = SettingsCategory.AddGroup("Affect IK", LOCTEXT("AffectIK_Label", "Affect IK"), false, true);
