Retargeter BP API insconsistent return value names.

diff --git a/Engine/Plugins/Animation/IKRig/Source/IKRig/Public/Retargeter/IKRetargeter.h b/Engine/Plugins/Animation/IKRig/Source/IKRig/Public/Retargeter/IKRetargeter.h
--- a/Engine/Plugins/Animation/IKRig/Source/IKRig/Public/Retargeter/IKRetargeter.h
+++ b/Engine/Plugins/Animation/IKRig/Source/IKRig/Public/Retargeter/IKRetargeter.h
@@ -282,7 +282,7 @@ public:
 	static void GetRootSettingsFromRetargetAsset(
 		const UIKRetargeter* RetargetAsset,
 		const FName OptionalProfileName,
-		FTargetRootSettings& OutSettings);
+		UPARAM(DisplayName = "ReturnValue") FTargetRootSettings& OutSettings);
 
 	// Returns the root settings in the supplied Retarget Profile. 
 	UFUNCTION(BlueprintCallable, BlueprintPure, Category=RetargetProfile)
@@ -293,7 +293,7 @@ public:
 	static void GetGlobalSettingsFromRetargetAsset(
 		const UIKRetargeter* RetargetAsset,
 		const FName OptionalProfileName,
-		FRetargetGlobalSettings& OutSettings);
+		UPARAM(DisplayName = "ReturnValue") FRetargetGlobalSettings& OutSettings);
 
 	// Returns the global settings in the supplied Retarget Profile. 
 	UFUNCTION(BlueprintCallable, BlueprintPure, Category=RetargetProfile)
