Fixed bug, IK Rig creating chain with existing goal always added to goal chain.

diff --git a/Engine/Plugins/Animation/IKRig/Source/IKRigEditor/Private/RigEditor/IKRigEditorController.cpp b/Engine/Plugins/Animation/IKRig/Source/IKRigEditor/Private/RigEditor/IKRigEditorController.cpp
--- a/Engine/Plugins/Animation/IKRig/Source/IKRigEditor/Private/RigEditor/IKRigEditorController.cpp
+++ b/Engine/Plugins/Animation/IKRig/Source/IKRigEditor/Private/RigEditor/IKRigEditorController.cpp
@@ -1221,16 +1221,28 @@ FName FIKRigEditorController::PromptToAddNewRetargetChain(FBoneChain& BoneChain)
 
 	// show the dialog and handle user choice
 	const int32 UserChoice = AddNewRetargetChainDialog->ShowModal();
-	if (UserChoice == 2 || UserChoice < 0)
+	const bool bUserClosedWindow = UserChoice < 0;
+	const bool bUserAddedChainNoGoal = UserChoice == 0;
+	const bool bUserAddedChainWithGoal = UserChoice == 1;
+	const bool bUserPressedCancel = UserChoice == 2;
+
+	// cancel button pressed, or window closed ?
+	if (bUserClosedWindow || bUserPressedCancel)
 	{
-		return NAME_None;  // cancel button pressed, or window closed
+		return NAME_None;  
+	}
+
+	// user opted not to add a goal, so remove it in case one was in the bone chain
+	if (bUserAddedChainNoGoal)
+	{
+		BoneChain.IKGoalName = NAME_None;
 	}
 
 	// add the retarget chain
 	const FName NewChainName = AssetController->AddRetargetChainInternal(BoneChain);
 	
 	// did user choose to assign a goal
-	if (UserChoice == 1)
+	if (bUserAddedChainWithGoal)
 	{
 		FName GoalName;
 		if (bHasExistingGoal)
@@ -1242,6 +1254,7 @@ FName FIKRigEditorController::PromptToAddNewRetargetChain(FBoneChain& BoneChain)
 		{
 			// add a default solver if there isn't one already
 			PromptToAddDefaultSolver();
+
 			// create new goal
 			const FName NewGoalName = FName(FText::Format(LOCTEXT("GoalOnNewChainName", "{0}_Goal"), FText::FromName(BoneChain.ChainName)).ToString());
 			GoalName = AssetController->AddNewGoal(NewGoalName, BoneChain.EndBone.BoneName);
