Fixes IK Rig animgraph node missing pin after undo.

diff --git a/Engine/Plugins/Animation/IKRig/Source/IKRigDeveloper/Private/AnimGraphNode_IKRig.cpp b/Engine/Plugins/Animation/IKRig/Source/IKRigDeveloper/Private/AnimGraphNode_IKRig.cpp
--- a/Engine/Plugins/Animation/IKRig/Source/IKRigDeveloper/Private/AnimGraphNode_IKRig.cpp
+++ b/Engine/Plugins/Animation/IKRig/Source/IKRigDeveloper/Private/AnimGraphNode_IKRig.cpp
@@ -672,6 +672,13 @@ void UAnimGraphNode_IKRig::PostLoad()
 	}
 }
 
+void UAnimGraphNode_IKRig::PostEditUndo()
+{
+	Super::PostEditUndo();
+
+	UpdateGoalsFromAsset();
+}
+
 void UAnimGraphNode_IKRig::CustomizeDetails(IDetailLayoutBuilder& DetailBuilder)
 {
 	Super::CustomizeDetails(DetailBuilder);
diff --git a/Engine/Plugins/Animation/IKRig/Source/IKRigDeveloper/Public/AnimGraphNode_IKRig.h b/Engine/Plugins/Animation/IKRig/Source/IKRigDeveloper/Public/AnimGraphNode_IKRig.h
--- a/Engine/Plugins/Animation/IKRig/Source/IKRigDeveloper/Public/AnimGraphNode_IKRig.h
+++ b/Engine/Plugins/Animation/IKRig/Source/IKRigDeveloper/Public/AnimGraphNode_IKRig.h
@@ -116,6 +116,7 @@ public:
 
 	// Begin UObject Interface.
 	virtual void PostLoad() override;
+	virtual void PostEditUndo() override;
 	// End UObject Interface.
 
 	virtual bool NeedsToSpecifyValidTargetClass() const override { return false; }
