Exported retarget poses now reference correct skeletal mesh. Export option disabled when not in editing pose mode.

diff --git a/Engine/Plugins/Animation/IKRig/Source/IKRigEditor/Private/RetargetEditor/IKRetargetEditor.cpp b/Engine/Plugins/Animation/IKRig/Source/IKRigEditor/Private/RetargetEditor/IKRetargetEditor.cpp
--- a/Engine/Plugins/Animation/IKRig/Source/IKRigEditor/Private/RetargetEditor/IKRetargetEditor.cpp
+++ b/Engine/Plugins/Animation/IKRig/Source/IKRigEditor/Private/RetargetEditor/IKRetargetEditor.cpp
@@ -201,11 +201,13 @@ void FIKRetargetEditor::BindCommands()
 	ToolkitCommands->MapAction(
 		Commands.NewRetargetPose,
 		FExecuteAction::CreateSP(EditorController, &FIKRetargetEditorController::HandleNewPose),
+		FCanExecuteAction::CreateSP(EditorController, &FIKRetargetEditorController::IsEditingPose),
 		EUIActionRepeatMode::RepeatDisabled);
 
 	ToolkitCommands->MapAction(
 		Commands.DuplicateRetargetPose,
 		FExecuteAction::CreateSP(EditorController, &FIKRetargetEditorController::HandleDuplicatePose),
+		FCanExecuteAction::CreateSP(EditorController, &FIKRetargetEditorController::IsEditingPose),
 		EUIActionRepeatMode::RepeatDisabled);
 
 	ToolkitCommands->MapAction(
@@ -263,22 +265,19 @@ void FIKRetargetEditor::BindCommands()
 	ToolkitCommands->MapAction(
 		Commands.ImportRetargetPose,
 		FExecuteAction::CreateSP(PoseExporterRef, &FIKRetargetPoseExporter::HandleImportFromPoseAsset),
-		FCanExecuteAction(),
-		FCanExecuteAction(),
+		FCanExecuteAction::CreateSP(EditorController, &FIKRetargetEditorController::IsEditingPose),
 		EUIActionRepeatMode::RepeatDisabled);
 
 	ToolkitCommands->MapAction(
 		Commands.ImportRetargetPoseFromAnim,
 		FExecuteAction::CreateSP(PoseExporterRef, &FIKRetargetPoseExporter::HandleImportFromSequenceAsset),
-		FCanExecuteAction(),
-		FCanExecuteAction(),
+		FCanExecuteAction::CreateSP(EditorController, &FIKRetargetEditorController::IsEditingPose),
 		EUIActionRepeatMode::RepeatDisabled);
 
 	ToolkitCommands->MapAction(
 		Commands.ExportRetargetPose,
 		FExecuteAction::CreateSP(PoseExporterRef, &FIKRetargetPoseExporter::HandleExportPoseAsset),
-		FCanExecuteAction(),
-		FCanExecuteAction(),
+		FCanExecuteAction::CreateSP(EditorController, &FIKRetargetEditorController::IsEditingPose),
 		EUIActionRepeatMode::RepeatDisabled);
 }
 
diff --git a/Engine/Plugins/Animation/IKRig/Source/IKRigEditor/Private/RetargetEditor/IKRetargetPoseExporter.cpp b/Engine/Plugins/Animation/IKRig/Source/IKRigEditor/Private/RetargetEditor/IKRetargetPoseExporter.cpp
--- a/Engine/Plugins/Animation/IKRig/Source/IKRigEditor/Private/RetargetEditor/IKRetargetPoseExporter.cpp
+++ b/Engine/Plugins/Animation/IKRig/Source/IKRigEditor/Private/RetargetEditor/IKRetargetPoseExporter.cpp
@@ -591,7 +591,7 @@ void FIKRetargetPoseExporter::HandleExportPoseAsset()
 
 	// get mesh to export pose
 	const ERetargetSourceOrTarget SourceOrTarget = ControllerPtr->GetSourceOrTarget();
-	const USkeletalMesh* Mesh = ControllerPtr->GetSkeletalMesh(SourceOrTarget);
+	USkeletalMesh* Mesh = ControllerPtr->GetSkeletalMesh(SourceOrTarget);
 	UDebugSkelMeshComponent* MeshComponent = ControllerPtr->GetSkeletalMeshComponent(SourceOrTarget);
 	if (!(Mesh && MeshComponent))
 	{
@@ -627,6 +627,7 @@ void FIKRetargetPoseExporter::HandleExportPoseAsset()
 
 	// fill new pose asset with current pose
 	NewPoseAsset->SetSkeleton(const_cast<USkeleton*>(Mesh->GetSkeleton()));
+	NewPoseAsset->SetPreviewMesh(Mesh);
 	const FName NewPoseName = NewPoseAsset->AddPoseWithUniqueName(MeshComponent);
 
 	// mark asset dirty 
