Follow up fix to 31055857 (fixed bug, selection not framed on new IK Retargeter). Fixes issue with sphere not initializing to zero.

diff --git a/Engine/Plugins/Animation/IKRig/Source/IKRigEditor/Private/RetargetEditor/IKRetargetEditorController.cpp b/Engine/Plugins/Animation/IKRig/Source/IKRigEditor/Private/RetargetEditor/IKRetargetEditorController.cpp
--- a/Engine/Plugins/Animation/IKRig/Source/IKRigEditor/Private/RetargetEditor/IKRetargetEditorController.cpp
+++ b/Engine/Plugins/Animation/IKRig/Source/IKRigEditor/Private/RetargetEditor/IKRetargetEditorController.cpp
@@ -453,7 +453,7 @@ bool FIKRetargetEditorController::GetCameraTargetForSelection(FSphere& OutTarget
 	case ERetargetSelectionType::NONE:
 	default:
 		// frame both meshes
-		OutTarget = FSphere();
+		OutTarget = FSphere(0);
 		if (const UPrimitiveComponent* SourceComponent = GetSkeletalMeshComponent(ERetargetSourceOrTarget::Source))
 		{
 			OutTarget += SourceComponent->Bounds.GetSphere();
