Fixes regression from adding LOD support to retarget pose node.

diff --git a/Engine/Plugins/Animation/IKRig/Source/IKRig/Private/AnimNodes/AnimNode_RetargetPoseFromMesh.cpp b/Engine/Plugins/Animation/IKRig/Source/IKRig/Private/AnimNodes/AnimNode_RetargetPoseFromMesh.cpp
--- a/Engine/Plugins/Animation/IKRig/Source/IKRig/Private/AnimNodes/AnimNode_RetargetPoseFromMesh.cpp
+++ b/Engine/Plugins/Animation/IKRig/Source/IKRig/Private/AnimNodes/AnimNode_RetargetPoseFromMesh.cpp
@@ -113,12 +113,12 @@ void FAnimNode_RetargetPoseFromMesh::Evaluate_AnyThread(FPoseContext& Output)
 	Processor->ApplySettingsFromProfile(CustomRetargetProfile);
 	
 	// LOD off the IK pass
-	if (Output.AnimInstanceProxy->GetLODLevel() > LODThresholdForIK)
+	if (LODThresholdForIK != INDEX_NONE && Output.AnimInstanceProxy->GetLODLevel() > LODThresholdForIK)
 	{
 		// override the custom profile's global settings but with IK forcibly turned off
 		FRetargetProfile TurnIKOffProfile;
 		TurnIKOffProfile.bApplyGlobalSettings = true;
-		TurnIKOffProfile.GlobalSettings = CustomRetargetProfile.GlobalSettings;
+		TurnIKOffProfile.GlobalSettings = Processor->GetGlobalSettings();
 		TurnIKOffProfile.GlobalSettings.bEnableIK = false;
 		Processor->ApplySettingsFromProfile(TurnIKOffProfile);
 	}
diff --git a/Engine/Plugins/Animation/IKRig/Source/IKRig/Public/AnimNodes/AnimNode_RetargetPoseFromMesh.h b/Engine/Plugins/Animation/IKRig/Source/IKRig/Public/AnimNodes/AnimNode_RetargetPoseFromMesh.h
--- a/Engine/Plugins/Animation/IKRig/Source/IKRig/Public/AnimNodes/AnimNode_RetargetPoseFromMesh.h
+++ b/Engine/Plugins/Animation/IKRig/Source/IKRig/Public/AnimNodes/AnimNode_RetargetPoseFromMesh.h
@@ -44,15 +44,17 @@ struct IKRIG_API FAnimNode_RetargetPoseFromMesh : public FAnimNode_Base
 
 	/*
 	* Max LOD that this node is allowed to run.
-	* For example if you have LODThreshold to be 2, it will run until LOD 2 (based on 0 index) when the component LOD becomes 3, it will stop update/evaluate
+	* For example if you have LODThreshold at 2, it will run until LOD 2 (based on 0 index) when the component LOD becomes 3, it will stop update/evaluate
+	* A value of -1 forces the node to execute at all LOD levels.
 	*/
 	UPROPERTY(EditAnywhere, BlueprintReadWrite, Category = Performance, meta = (DisplayName = "LOD Threshold"))
 	int32 LODThreshold = -1;
 
 	/*
 	* Max LOD that IK is allowed to run.
-	* For example if you have LODThresholdForIK to be 2, it will skip the IK pass on LODs 3 and greater.
+	* For example if you have LODThresholdForIK at 2, it will skip the IK pass on LODs 3 and greater.
 	* This only disables IK and does not affect the Root or FK passes.
+	* A value of -1 forces the node to execute at all LOD levels.
 	*/
 	UPROPERTY(EditAnywhere, BlueprintReadWrite, Category = Performance, meta = (DisplayName = "IK LOD Threshold"))
 	int32 LODThresholdForIK = -1;
diff --git a/Engine/Plugins/Animation/IKRig/Source/IKRig/Public/Retargeter/IKRetargetProcessor.h b/Engine/Plugins/Animation/IKRig/Source/IKRig/Public/Retargeter/IKRetargetProcessor.h
--- a/Engine/Plugins/Animation/IKRig/Source/IKRig/Public/Retargeter/IKRetargetProcessor.h
+++ b/Engine/Plugins/Animation/IKRig/Source/IKRig/Public/Retargeter/IKRetargetProcessor.h
@@ -526,6 +526,9 @@ public:
 	
 	// Get read only access to the retarget ops currently running in processor
 	const TArray<TObjectPtr<URetargetOpBase>>& GetRetargetOps() const {return OpStack; };
+
+	// Get read only access to the current global settings
+	const FRetargetGlobalSettings& GetGlobalSettings() const {return GlobalSettings; };
 	
 	/** Reset the IK planting state. */
 	void ResetPlanting();
