Import retarget pose dialog no longer supports multi select.
Now properly filters assets when importing sequence or pose assets.
Fixes issue, pose list not updating when pose asset selected.
Improves pose asset picker window configuration.
Changed Import button text to be consistent.

diff --git a/Engine/Plugins/Animation/IKRig/Source/IKRigEditor/Private/RetargetEditor/IKRetargetPoseExporter.cpp b/Engine/Plugins/Animation/IKRig/Source/IKRigEditor/Private/RetargetEditor/IKRetargetPoseExporter.cpp
--- a/Engine/Plugins/Animation/IKRig/Source/IKRigEditor/Private/RetargetEditor/IKRetargetPoseExporter.cpp
+++ b/Engine/Plugins/Animation/IKRig/Source/IKRigEditor/Private/RetargetEditor/IKRetargetPoseExporter.cpp
@@ -41,10 +41,35 @@ void FIKRetargetPoseExporter::HandleImportFromPoseAsset()
 
 	// the asset picker will only show pose assets
 	FAssetPickerConfig AssetPickerConfig;
+	AssetPickerConfig.SelectionMode = ESelectionMode::Single;
 	AssetPickerConfig.Filter.ClassPaths.Add(UPoseAsset::StaticClass()->GetClassPathName());
 	AssetPickerConfig.Filter.bRecursiveClasses = true;
-	AssetPickerConfig.OnAssetSelected = FOnAssetSelected::CreateSP(this, &FIKRetargetPoseExporter::OnRetargetPoseSelected);
-	AssetPickerConfig.InitialAssetViewType = EAssetViewType::Tile;
+	AssetPickerConfig.OnShouldFilterAsset = FOnShouldFilterAsset::CreateSP(this, &FIKRetargetPoseExporter::OnShouldFilterPoseToImport);
+	AssetPickerConfig.OnAssetSelected = FOnAssetSelected::CreateSP(this, &FIKRetargetPoseExporter::OnPoseAssetSelected);
+	AssetPickerConfig.InitialAssetViewType = EAssetViewType::Column;
+	AssetPickerConfig.bAddFilterUI = true;
+	AssetPickerConfig.bShowPathInColumnView = true;
+	AssetPickerConfig.bShowTypeInColumnView = true;
+	AssetPickerConfig.DefaultFilterMenuExpansion = EAssetTypeCategories::Animation;
+	AssetPickerConfig.bAllowNullSelection = false;
+
+	// hide all asset registry columns by default (we only really want the name and path)
+	UObject* PoseAssetDefaultObject = UPoseAsset::StaticClass()->GetDefaultObject();
+	FAssetRegistryTagsContextData TagsContext(PoseAssetDefaultObject, EAssetRegistryTagsCaller::Uncategorized);
+	PoseAssetDefaultObject->GetAssetRegistryTags(TagsContext);
+	const FName ColumnToKeep = FName("Number of Frames");
+	for (const TPair<FName, UObject::FAssetRegistryTag>& TagPair : TagsContext.Tags)
+	{
+		if (TagPair.Key != ColumnToKeep)
+		{
+			AssetPickerConfig.HiddenColumnNames.Add(TagPair.Key.ToString());
+		}
+	}
+
+	// also hide the type column by default (but allow users to enable it, so don't use bShowTypeInColumnView)
+	AssetPickerConfig.HiddenColumnNames.Add(TEXT("Class"));
+	AssetPickerConfig.HiddenColumnNames.Add(TEXT("HasVirtualizedData"));
+	AssetPickerConfig.HiddenColumnNames.Add(TEXT("Disk Size"));
 
 	ImportPoseWindow = SNew(SWindow)
 	.Title(LOCTEXT("ImportRetargetPose_Label", "Import Retarget Pose"))
@@ -74,6 +99,7 @@ void FIKRetargetPoseExporter::HandleImportFromPoseAsset()
 					.OptionsSource(&PosesInSelectedAsset)
 					.OnGenerateWidget(this, &FIKRetargetPoseExporter::OnGeneratePoseComboWidget)
 					.OnSelectionChanged(this, &FIKRetargetPoseExporter::OnSelectPoseFromPoseAsset)
+					.OnComboBoxOpening(this, &FIKRetargetPoseExporter::RefreshPoseList)
 					[
 						SNew(STextBlock).Text_Lambda([this]()
 						{
@@ -131,23 +157,11 @@ void FIKRetargetPoseExporter::HandleImportFromPoseAsset()
 	ImportPoseWindow.Reset();
 }
 
-void FIKRetargetPoseExporter::OnRetargetPoseSelected(const FAssetData& SelectedAsset)
+void FIKRetargetPoseExporter::OnPoseAssetSelected(const FAssetData& SelectedAsset)
 {
 	RetargetPoseToImport = SelectedAsset.ToSoftObjectPath();
 
-	// get all poses in the selected asset
-	PosesInSelectedAsset.Empty();
-	const TObjectPtr<UPoseAsset> PoseAsset = Cast<UPoseAsset>(RetargetPoseToImport.TryLoad());
-	if (!PoseAsset)
-	{
-		return;
-	}
-
-	// store pose names in list used by combobox
-	for (const FName& PoseName : PoseAsset->GetPoseFNames())
-	{
-		PosesInSelectedAsset.Add(MakeShared<FName>(PoseName));
-	}
+	RefreshPoseList();
 
 	// set the selected pose to the first pose in the asset (or null if empty)
 	if (PosesInSelectedAsset.IsEmpty())
@@ -173,6 +187,23 @@ void FIKRetargetPoseExporter::OnSelectPoseFromPoseAsset(TSharedPtr<FName> Item,
 	}
 }
 
+void FIKRetargetPoseExporter::RefreshPoseList()
+{
+	// get all poses in the selected asset
+	PosesInSelectedAsset.Empty();
+	const TObjectPtr<UPoseAsset> PoseAsset = Cast<UPoseAsset>(RetargetPoseToImport.TryLoad());
+	if (!PoseAsset)
+	{
+		return;
+	}
+
+	// store pose names in list used by combobox
+	for (const FName& PoseName : PoseAsset->GetPoseFNames())
+	{
+		PosesInSelectedAsset.Add(MakeShared<FName>(PoseName));
+	}
+}
+
 FReply FIKRetargetPoseExporter::ImportPoseAsset() const
 {
 	// close the import window
@@ -283,6 +314,24 @@ FReply FIKRetargetPoseExporter::ImportPoseAsset() const
 	return FReply::Unhandled();
 }
 
+bool FIKRetargetPoseExporter::OnShouldFilterPoseToImport(const FAssetData& AssetData) const
+{
+	// is this a pose asset
+	if (!AssetData.IsInstanceOf(UPoseAsset::StaticClass()))
+	{
+		return true;
+	}
+
+	// get currently edited skeleton
+	const USkeleton* DesiredSkeleton = Controller.Pin()->GetSkeleton(Controller.Pin()->GetSourceOrTarget());
+	if (!DesiredSkeleton)
+	{
+		return true;
+	}
+
+	return !DesiredSkeleton->IsCompatibleForEditor(AssetData);
+}
+
 void FIKRetargetPoseExporter::HandleImportFromSequenceAsset()
 {
 	FIKRetargetEditorController* ControllerPtr = Controller.Pin().Get();
@@ -301,6 +350,7 @@ void FIKRetargetPoseExporter::HandleImportFromSequenceAsset()
 
 	// the asset picker will only show animation sequences compatible with the preview mesh
 	FAssetPickerConfig AssetPickerConfig;
+	AssetPickerConfig.SelectionMode = ESelectionMode::Single;
 	AssetPickerConfig.Filter.ClassPaths.Add(UAnimSequence::StaticClass()->GetClassPathName());
 	AssetPickerConfig.InitialAssetViewType = EAssetViewType::Column;
 	AssetPickerConfig.bAddFilterUI = true;
@@ -327,7 +377,7 @@ void FIKRetargetPoseExporter::HandleImportFromSequenceAsset()
 	// also hide the type column by default (but allow users to enable it, so don't use bShowTypeInColumnView)
 	AssetPickerConfig.HiddenColumnNames.Add(TEXT("Class"));
 	AssetPickerConfig.HiddenColumnNames.Add(TEXT("HasVirtualizedData"));
-	AssetPickerConfig.HiddenColumnNames.Add(TEXT("DiskSize"));
+	AssetPickerConfig.HiddenColumnNames.Add(TEXT("Disk Size"));
 
 	// create pop-up window for user to select animation sequence asset to import as a retarget pose]
 	ImportPoseFromSequenceWindow = SNew(SWindow)
@@ -423,7 +473,7 @@ void FIKRetargetPoseExporter::HandleImportFromSequenceAsset()
 					.TextStyle( FAppStyle::Get(), "DialogButtonText" )
 					.HAlign(HAlign_Center)
 					.VAlign(VAlign_Center)
-					.Text(LOCTEXT("ImportAsRetargetPoseButtonLabel", "Import As Retarget Pose"))
+					.Text(LOCTEXT("ImportAsRetargetPoseButtonLabel", "Import New Retarget Pose"))
 					.OnClicked(this, &FIKRetargetPoseExporter::OnImportPoseFromSequence)
 				]
 				
@@ -453,8 +503,8 @@ void FIKRetargetPoseExporter::HandleImportFromSequenceAsset()
 
 bool FIKRetargetPoseExporter::OnShouldFilterSequenceToImport(const FAssetData& AssetData) const
 {
-	// is this an animation asset?
-	if (!AssetData.IsInstanceOf(UAnimationAsset::StaticClass()))
+	// is this an animation sequence?
+	if (!AssetData.IsInstanceOf(UAnimSequence::StaticClass()))
 	{
 		return true;
 	}
diff --git a/Engine/Plugins/Animation/IKRig/Source/IKRigEditor/Public/RetargetEditor/IKRetargetPoseExporter.h b/Engine/Plugins/Animation/IKRig/Source/IKRigEditor/Public/RetargetEditor/IKRetargetPoseExporter.h
--- a/Engine/Plugins/Animation/IKRig/Source/IKRigEditor/Public/RetargetEditor/IKRetargetPoseExporter.h
+++ b/Engine/Plugins/Animation/IKRig/Source/IKRigEditor/Public/RetargetEditor/IKRetargetPoseExporter.h
@@ -17,9 +17,11 @@ public:
 	/** import retarget pose from asset*/
 	void HandleImportFromPoseAsset();
 	FReply ImportPoseAsset() const;
-	void OnRetargetPoseSelected(const FAssetData& SelectedAsset);
+	bool OnShouldFilterPoseToImport(const struct FAssetData& AssetData) const;
+	void OnPoseAssetSelected(const FAssetData& SelectedAsset);
 	TSharedRef<SWidget> OnGeneratePoseComboWidget(TSharedPtr<FName> Item) const;
 	void OnSelectPoseFromPoseAsset(TSharedPtr<FName> Item, ESelectInfo::Type SelectInfo);
+	void RefreshPoseList();
 	FSoftObjectPath RetargetPoseToImport;
 	TSharedPtr<SWindow> ImportPoseWindow;
 	TArray<TSharedPtr<FName>> PosesInSelectedAsset;
