Fixes drag/drop reordering in IK Rig solver stack.

diff --git a/Engine/Plugins/Animation/IKRig/Source/IKRigEditor/Private/RigEditor/IKRigController.cpp b/Engine/Plugins/Animation/IKRig/Source/IKRigEditor/Private/RigEditor/IKRigController.cpp
--- a/Engine/Plugins/Animation/IKRig/Source/IKRigEditor/Private/RigEditor/IKRigController.cpp
+++ b/Engine/Plugins/Animation/IKRig/Source/IKRigEditor/Private/RigEditor/IKRigController.cpp
@@ -780,7 +780,8 @@ bool UIKRigController::MoveSolverInStack(int32 SolverToMoveIndex, int32 TargetSo
 		return false;
 	}
 
-	if (!Asset->Solvers.IsValidIndex(TargetSolverIndex))
+	// allow a target 1 greater than the last element (for dragging below last element of list)
+	if (TargetSolverIndex > Asset->Solvers.Num() || TargetSolverIndex < 0)
 	{
 		UE_LOG(LogIKRigEditor, Warning, TEXT("Solver not moved. Invalid target index, %d."), TargetSolverIndex);
 		return false;
@@ -796,7 +797,7 @@ bool UIKRigController::MoveSolverInStack(int32 SolverToMoveIndex, int32 TargetSo
 	FScopedReinitializeIKRig Reinitialize(this);
 	Asset->Modify();
 	UIKRigSolver* SolverToMove = Asset->Solvers[SolverToMoveIndex];
-	Asset->Solvers.Insert(SolverToMove, TargetSolverIndex + 1);
+	Asset->Solvers.Insert(SolverToMove, TargetSolverIndex);
 	const int32 SolverToRemove = TargetSolverIndex > SolverToMoveIndex ? SolverToMoveIndex : SolverToMoveIndex + 1;
 	Asset->Solvers.RemoveAt(SolverToRemove);
 	return true;
diff --git a/Engine/Plugins/Animation/IKRig/Source/IKRigEditor/Private/RigEditor/SIKRigSolverStack.cpp b/Engine/Plugins/Animation/IKRig/Source/IKRigEditor/Private/RigEditor/SIKRigSolverStack.cpp
--- a/Engine/Plugins/Animation/IKRig/Source/IKRigEditor/Private/RigEditor/SIKRigSolverStack.cpp
+++ b/Engine/Plugins/Animation/IKRig/Source/IKRigEditor/Private/RigEditor/SIKRigSolverStack.cpp
@@ -493,7 +493,7 @@ TOptional<EItemDropZone> SIKRigSolverStack::OnCanAcceptDrop(
 	const TSharedPtr<FIKRigSolverStackDragDropOp> DragDropOp = DragDropEvent.GetOperationAs<FIKRigSolverStackDragDropOp>();
 	if (DragDropOp.IsValid())
 	{
-		ReturnedDropZone = EItemDropZone::BelowItem;	
+		ReturnedDropZone = DropZone == EItemDropZone::BelowItem ? EItemDropZone::BelowItem : EItemDropZone::AboveItem;
 	}
 	
 	return ReturnedDropZone;
@@ -518,7 +518,8 @@ FReply SIKRigSolverStack::OnAcceptDrop(
 
 	const FSolverStackElement& DraggedElement = *DragDropOp.Get()->Element.Pin().Get();
 	UIKRigController* AssetController = Controller->AssetController;
-	const bool bWasReparented = AssetController->MoveSolverInStack(DraggedElement.IndexInStack, TargetItem.Get()->IndexInStack);
+	const int32 IndexOffset = DropZone == EItemDropZone::BelowItem || DropZone == EItemDropZone::OntoItem ? 1 : 0;
+	const bool bWasReparented = AssetController->MoveSolverInStack(DraggedElement.IndexInStack, TargetItem.Get()->IndexInStack + IndexOffset);
 	if (bWasReparented)
 	{
 		RefreshStackView();
