Fixes retargeter and IK rig viewports not updating when clicking in other tabs.

diff --git a/Engine/Plugins/Animation/IKRig/Source/IKRigEditor/Private/RetargetEditor/IKRetargetEditor.cpp b/Engine/Plugins/Animation/IKRig/Source/IKRigEditor/Private/RetargetEditor/IKRetargetEditor.cpp
--- a/Engine/Plugins/Animation/IKRig/Source/IKRigEditor/Private/RetargetEditor/IKRetargetEditor.cpp
+++ b/Engine/Plugins/Animation/IKRig/Source/IKRigEditor/Private/RetargetEditor/IKRetargetEditor.cpp
@@ -523,6 +523,9 @@ void FIKRetargetEditor::Tick(float DeltaTime)
 		EditorController->ResetIKPlantingState();
 	}
 	PreviousTime = CurrentTime;
+	
+	// forces viewport to always update, even when mouse pressed down in other tabs
+	GetPersonaToolkit()->GetPreviewScene()->InvalidateViews();
 }
 
 TStatId FIKRetargetEditor::GetStatId() const
diff --git a/Engine/Plugins/Animation/IKRig/Source/IKRigEditor/Private/RigEditor/IKRigToolkit.cpp b/Engine/Plugins/Animation/IKRig/Source/IKRigEditor/Private/RigEditor/IKRigToolkit.cpp
--- a/Engine/Plugins/Animation/IKRig/Source/IKRigEditor/Private/RigEditor/IKRigToolkit.cpp
+++ b/Engine/Plugins/Animation/IKRig/Source/IKRigEditor/Private/RigEditor/IKRigToolkit.cpp
@@ -219,6 +219,12 @@ FString FIKRigEditorToolkit::GetWorldCentricTabPrefix() const
 	return TEXT("IKRigEditor");
 }
 
+void FIKRigEditorToolkit::Tick(float DeltaTime)
+{
+	// forces viewport to always update, even when mouse pressed down in other tabs
+	GetPersonaToolkit()->GetPreviewScene()->InvalidateViews();
+}
+
 TStatId FIKRigEditorToolkit::GetStatId() const
 {
 	RETURN_QUICK_DECLARE_CYCLE_STAT(FIKRigEditorToolkit, STATGROUP_Tickables);
diff --git a/Engine/Plugins/Animation/IKRig/Source/IKRigEditor/Public/RigEditor/IKRigToolkit.h b/Engine/Plugins/Animation/IKRig/Source/IKRigEditor/Public/RigEditor/IKRigToolkit.h
--- a/Engine/Plugins/Animation/IKRig/Source/IKRigEditor/Public/RigEditor/IKRigToolkit.h
+++ b/Engine/Plugins/Animation/IKRig/Source/IKRigEditor/Public/RigEditor/IKRigToolkit.h
@@ -45,7 +45,7 @@ public:
 	/** END FAssetEditorToolkit interface */
 
 	//** FTickableEditorObject Interface
-	virtual void Tick(float DeltaTime) override {};
+	virtual void Tick(float DeltaTime) override;
 	virtual ETickableTickType GetTickableTickType() const override { return ETickableTickType::Always; }
 	virtual TStatId GetStatId() const override;
 	//~ END FTickableEditorObject Interface
