Fixes extra undo transaction caused when using numerical input on goal transform.

diff --git a/Engine/Plugins/Animation/IKRig/Source/IKRigEditor/Private/RigEditor/IKRigDetailCustomizations.cpp b/Engine/Plugins/Animation/IKRig/Source/IKRigEditor/Private/RigEditor/IKRigDetailCustomizations.cpp
--- a/Engine/Plugins/Animation/IKRig/Source/IKRigEditor/Private/RigEditor/IKRigDetailCustomizations.cpp
+++ b/Engine/Plugins/Animation/IKRig/Source/IKRigEditor/Private/RigEditor/IKRigDetailCustomizations.cpp
@@ -360,8 +360,13 @@ void FIKRigGenericDetailCustomization::CustomizeDetailsForClass<UIKRigEffectorGo
 			FTransform CurrentTransform, UpdatedTransform;
 			for (UIKRigEffectorGoal* Goal: Goals)
 			{
-				Tie(CurrentTransform, UpdatedTransform) = Goal->PrepareNumericValueChanged( InComponent, InRepresentation, InSubComponent, InValue, TransformType);					
-				Goal->SetTransform(UpdatedTransform, TransformType);
+				Tie(CurrentTransform, UpdatedTransform) = Goal->PrepareNumericValueChanged( InComponent, InRepresentation, InSubComponent, InValue, TransformType);
+
+				// don't do anything if the transform has not been changed
+				if (!UpdatedTransform.Equals(CurrentTransform))
+				{
+					Goal->SetTransform(UpdatedTransform, TransformType);
+				}
 			}
 			
 			ValueChangedTransaction.Reset();
