Adding new icons to IK Retargeter

diff --git a/Engine/Plugins/Animation/IKRig/Source/IKRigEditor/Private/RetargetEditor/SIKRetargetPoseEditor.cpp b/Engine/Plugins/Animation/IKRig/Source/IKRigEditor/Private/RetargetEditor/SIKRetargetPoseEditor.cpp
--- a/Engine/Plugins/Animation/IKRig/Source/IKRigEditor/Private/RetargetEditor/SIKRetargetPoseEditor.cpp
+++ b/Engine/Plugins/Animation/IKRig/Source/IKRigEditor/Private/RetargetEditor/SIKRetargetPoseEditor.cpp
@@ -124,7 +124,7 @@ TSharedRef<SWidget>  SIKRetargetPoseEditor::MakeToolbar(TSharedPtr<FUICommandLis
 		FOnGetContent::CreateSP(this, &SIKRetargetPoseEditor::GenerateEditMenuContent, Commands),
 		LOCTEXT("AutoAlign_Label", "Auto Align"),
 		LOCTEXT("AutoAlignTip_Label", "Automatically aligns bones on source skeleton to target (or vice versa)."),
-		FSlateIcon(FAppStyle::GetAppStyleSetName(), "Icons.Edit"));
+		FSlateIcon(FIKRetargetEditorStyle::Get().GetStyleSetName(),"IKRetarget.AutoAlign"));
 
 	ToolbarBuilder.EndSection();
 
diff --git a/Engine/Plugins/Animation/IKRig/Source/IKRigEditor/Private/RigEditor/IKRigToolkit.cpp b/Engine/Plugins/Animation/IKRig/Source/IKRigEditor/Private/RigEditor/IKRigToolkit.cpp
--- a/Engine/Plugins/Animation/IKRig/Source/IKRigEditor/Private/RigEditor/IKRigToolkit.cpp
+++ b/Engine/Plugins/Animation/IKRig/Source/IKRigEditor/Private/RigEditor/IKRigToolkit.cpp
@@ -168,14 +168,14 @@ void FIKRigEditorToolkit::FillToolbar(FToolBarBuilder& ToolbarBuilder)
 			NAME_None,
 			TAttribute<FText>(),
 			TAttribute<FText>(),
-			FSlateIcon(FAppStyle::Get().GetStyleSetName(),"Icons.Edit"));
+			FSlateIcon(FIKRigEditorStyle::Get().GetStyleSetName(),"IKRig.AutoRetarget"));
 
 		ToolbarBuilder.AddToolBarButton(
 			FIKRigCommands::Get().AutoSetupFBIK,
 			NAME_None,
 			TAttribute<FText>(),
 			TAttribute<FText>(),
-			FSlateIcon(FAppStyle::Get().GetStyleSetName(),"Icons.Edit"));
+			FSlateIcon(FIKRigEditorStyle::Get().GetStyleSetName(),"IKRig.AutoIK"));
 	}
 	ToolbarBuilder.EndSection();
 
diff --git a/Engine/Plugins/Animation/IKRig/Source/IKRigEditor/Public/RetargetEditor/IKRetargetEditorStyle.h b/Engine/Plugins/Animation/IKRig/Source/IKRigEditor/Public/RetargetEditor/IKRetargetEditorStyle.h
--- a/Engine/Plugins/Animation/IKRig/Source/IKRigEditor/Public/RetargetEditor/IKRetargetEditorStyle.h
+++ b/Engine/Plugins/Animation/IKRig/Source/IKRigEditor/Public/RetargetEditor/IKRetargetEditorStyle.h
@@ -29,6 +29,7 @@ public:
 		Set("IKRetarget.RunRetargeter", new IMAGE_BRUSH_SVG("Slate/RunRetargeter", Icon64x64));
 		Set("IKRetarget.EditRetargetPose", new IMAGE_BRUSH_SVG("Slate/EditRetargetPose", Icon64x64));
 		Set("IKRetarget.ShowRetargetPose", new IMAGE_BRUSH_SVG("Slate/ShowRetargetPose", Icon64x64));
+		Set("IKRetarget.AutoAlign", new IMAGE_BRUSH_SVG("Slate/AutoRetargetPose", Icon16x16));
 
 		SetContentRoot(FPaths::EngineContentDir() / TEXT("Editor/Slate"));
 		Set( "IKRetarget.Viewport.Border", new BOX_BRUSH( "Old/Window/ViewportDebugBorder", 0.8f, FLinearColor(1.0f,1.0f,1.0f,1.0f) ) );
diff --git a/Engine/Plugins/Animation/IKRig/Source/IKRigEditor/Public/RigEditor/IKRigEditorStyle.h b/Engine/Plugins/Animation/IKRig/Source/IKRigEditor/Public/RigEditor/IKRigEditorStyle.h
--- a/Engine/Plugins/Animation/IKRig/Source/IKRigEditor/Public/RigEditor/IKRigEditorStyle.h
+++ b/Engine/Plugins/Animation/IKRig/Source/IKRigEditor/Public/RigEditor/IKRigEditorStyle.h
@@ -40,6 +40,8 @@ public:
 		Set("IKRig.Hierarchy", new IMAGE_BRUSH_SVG("Slate/Hierarchy", Icon64x64));
 
 		Set("IKRig.AssetSettings", new IMAGE_BRUSH_SVG("Slate/AssetSettings", Icon64x64));
+		Set("IKRig.AutoRetarget", new IMAGE_BRUSH_SVG("Slate/AutoChainRetarget", Icon64x64));
+		Set("IKRig.AutoIK", new IMAGE_BRUSH_SVG("Slate/AutoFullBodyIK", Icon64x64));
 		
 		FTextBlockStyle NormalText = FAppStyle::GetWidgetStyle<FTextBlockStyle>("SkeletonTree.NormalFont");
 		Set( "IKRig.Tree.NormalText", FTextBlockStyle(NormalText));
