Fixes broken IK Rig engine test (trivial)

#changelist validated
#virtualized

diff --git a/Engine/Plugins/Animation/IKRig/Source/IKRigEditor/Private/RigEditor/IKRigController.cpp b/Engine/Plugins/Animation/IKRig/Source/IKRigEditor/Private/RigEditor/IKRigController.cpp
--- a/Engine/Plugins/Animation/IKRig/Source/IKRigEditor/Private/RigEditor/IKRigController.cpp
+++ b/Engine/Plugins/Animation/IKRig/Source/IKRigEditor/Private/RigEditor/IKRigController.cpp
@@ -796,10 +796,11 @@ bool UIKRigController::MoveSolverInStack(int32 SolverToMoveIndex, int32 TargetSo
 	FScopedTransaction Transaction(LOCTEXT("ReorderSolver_Label", "Reorder Solvers"));
 	FScopedReinitializeIKRig Reinitialize(this);
 	Asset->Modify();
+	
 	UIKRigSolver* SolverToMove = Asset->Solvers[SolverToMoveIndex];
+	Asset->Solvers.RemoveAt(SolverToMoveIndex);
 	Asset->Solvers.Insert(SolverToMove, TargetSolverIndex);
-	const int32 SolverToRemove = TargetSolverIndex > SolverToMoveIndex ? SolverToMoveIndex : SolverToMoveIndex + 1;
-	Asset->Solvers.RemoveAt(SolverToRemove);
+
 	return true;
 }
 
diff --git a/Engine/Plugins/Animation/IKRig/Source/IKRigEditor/Private/RigEditor/SIKRigSolverStack.cpp b/Engine/Plugins/Animation/IKRig/Source/IKRigEditor/Private/RigEditor/SIKRigSolverStack.cpp
--- a/Engine/Plugins/Animation/IKRig/Source/IKRigEditor/Private/RigEditor/SIKRigSolverStack.cpp
+++ b/Engine/Plugins/Animation/IKRig/Source/IKRigEditor/Private/RigEditor/SIKRigSolverStack.cpp
@@ -517,10 +517,19 @@ FReply SIKRigSolverStack::OnAcceptDrop(
 	}
 
 	const FSolverStackElement& DraggedElement = *DragDropOp.Get()->Element.Pin().Get();
+	if (DraggedElement.IndexInStack == TargetItem->IndexInStack)
+	{
+		return FReply::Handled();
+	}
+	
 	UIKRigController* AssetController = Controller->AssetController;
-	const int32 IndexOffset = DropZone == EItemDropZone::BelowItem || DropZone == EItemDropZone::OntoItem ? 1 : 0;
-	const bool bWasReparented = AssetController->MoveSolverInStack(DraggedElement.IndexInStack, TargetItem.Get()->IndexInStack + IndexOffset);
-	if (bWasReparented)
+	int32 TargetIndex = TargetItem.Get()->IndexInStack;
+	if ( DropZone == EItemDropZone::AboveItem)
+	{
+		TargetIndex = FMath::Max(0, TargetIndex-1);
+	}
+	const bool bWasMoved = AssetController->MoveSolverInStack(DraggedElement.IndexInStack, TargetIndex);
+	if (bWasMoved)
 	{
 		RefreshStackView();
 		Controller->RefreshTreeView(); // update solver indices in effector items
