Closing "Add default solver" dialog no longer creates goals.

diff --git a/Engine/Plugins/Animation/IKRig/Source/IKRigEditor/Private/RigEditor/IKRigEditorController.cpp b/Engine/Plugins/Animation/IKRig/Source/IKRigEditor/Private/RigEditor/IKRigEditorController.cpp
--- a/Engine/Plugins/Animation/IKRig/Source/IKRigEditor/Private/RigEditor/IKRigEditorController.cpp
+++ b/Engine/Plugins/Animation/IKRig/Source/IKRigEditor/Private/RigEditor/IKRigEditorController.cpp
@@ -633,7 +633,11 @@ void FIKRigEditorController::AddNewGoals(const TArray<FName>& GoalNames, const T
 	AssetController->GetAsset()->Modify();
 
 	// add a default solver if there isn't one already
-	PromptToAddDefaultSolver();
+	const bool bCancelled = PromptToAddDefaultSolver();
+	if (bCancelled)
+	{
+		return;
+	}
 
 	// create goals
 	FName LastCreatedGoalName = NAME_None;
@@ -912,11 +916,22 @@ bool FIKRigEditorController::PromptToAddDefaultSolver() const
 			SCustomDialog::FButton(LOCTEXT("Skip", "Skip"))
 	});
 
-	if (AddSolverDialog->ShowModal() != 0)
+	// show window and get result
+	const int32 Result = AddSolverDialog->ShowModal();
+	const bool bWindowClosed = Result < 0;
+	const bool bSkipped = Result == 1;
+	
+	if (bWindowClosed)
+	{
+		return true; // window closed
+	}
+
+	if (bSkipped)
 	{
-		return false; // cancel button pressed, or window closed
+		return false; // user opted NOT to add a solver
 	}
 
+	// add a solver
 	if (SelectedSolver->SolverType != nullptr && SolverStackView.IsValid())
 	{
 		AssetController->AddSolver(SelectedSolver->SolverType);
@@ -925,7 +940,7 @@ bool FIKRigEditorController::PromptToAddDefaultSolver() const
 	// must refresh the view so that subsequent goal operations see a selected solver to connect to
 	SolverStackView->RefreshStackView();
 	
-	return true;
+	return false;
 }
 
 void FIKRigEditorController::ShowDetailsForBone(const FName BoneName) const
@@ -1253,8 +1268,13 @@ FName FIKRigEditorController::PromptToAddNewRetargetChain(FBoneChain& BoneChain)
 		else
 		{
 			// add a default solver if there isn't one already
-			PromptToAddDefaultSolver();
-
+			const bool bCancelled = PromptToAddDefaultSolver();
+			if (bCancelled)
+			{
+				// user cancelled creating a goal
+				return NewChainName;
+			}
+			
 			// create new goal
 			const FName NewGoalName = FName(FText::Format(LOCTEXT("GoalOnNewChainName", "{0}_Goal"), FText::FromName(BoneChain.ChainName)).ToString());
 			GoalName = AssetController->AddNewGoal(NewGoalName, BoneChain.EndBone.BoneName);
