Fixed batch retarget asset browser filtering out Aim Offset and Blendspace assets.
Fixes issue batch retargeter swaps target mesh when disabling retarget asset.
Filters illegal characters from exported file names.

diff --git a/Engine/Plugins/Animation/IKRig/Source/IKRigEditor/Private/RetargetEditor/SRetargetAnimAssetsWindow.cpp b/Engine/Plugins/Animation/IKRig/Source/IKRigEditor/Private/RetargetEditor/SRetargetAnimAssetsWindow.cpp
--- a/Engine/Plugins/Animation/IKRig/Source/IKRigEditor/Private/RetargetEditor/SRetargetAnimAssetsWindow.cpp
+++ b/Engine/Plugins/Animation/IKRig/Source/IKRigEditor/Private/RetargetEditor/SRetargetAnimAssetsWindow.cpp
@@ -19,6 +19,8 @@
 #include "RetargetEditor/IKRetargeterController.h"
 #include "Retargeter/IKRetargeter.h"
 #include "Viewports.h"
+#include "Animation/AimOffsetBlendSpace.h"
+#include "Animation/AimOffsetBlendSpace1D.h"
 #include "AssetRegistry/AssetRegistryModule.h"
 #include "Components/SkyLightComponent.h"
 #include "Editor/EditorPerProjectUserSettings.h"
@@ -290,7 +292,7 @@ FText SBatchExportPathDialog::GetPrefixName() const
 
 void SBatchExportPathDialog::SetPrefixName(const FText &InText)
 {
-	BatchContext->NameRule.Prefix = InText.ToString();
+	BatchContext->NameRule.Prefix = ConvertToCleanString(InText);
 	UpdateExampleText();
 }
 
@@ -301,7 +303,7 @@ FText SBatchExportPathDialog::GetSuffixName() const
 
 void SBatchExportPathDialog::SetSuffixName(const FText &InText)
 {
-	BatchContext->NameRule.Suffix = InText.ToString();
+	BatchContext->NameRule.Suffix = ConvertToCleanString(InText);
 	UpdateExampleText();
 }
 
@@ -312,7 +314,7 @@ FText SBatchExportPathDialog::GetReplaceFrom() const
 
 void SBatchExportPathDialog::SetReplaceFrom(const FText &InText)
 {
-	BatchContext->NameRule.ReplaceFrom = InText.ToString();
+	BatchContext->NameRule.ReplaceFrom = ConvertToCleanString(InText);
 	UpdateExampleText();
 }
 
@@ -323,7 +325,7 @@ FText SBatchExportPathDialog::GetReplaceTo() const
 
 void SBatchExportPathDialog::SetReplaceTo(const FText &InText)
 {
-	BatchContext->NameRule.ReplaceTo = InText.ToString();
+	BatchContext->NameRule.ReplaceTo = ConvertToCleanString(InText);
 	UpdateExampleText();
 }
 
@@ -344,6 +346,22 @@ FText SBatchExportPathDialog::GetFolderPath() const
 	return FText::FromString(BatchContext->NameRule.FolderPath);
 }
 
+FString SBatchExportPathDialog::ConvertToCleanString(const FText& ToClean)
+{
+	static TSet<TCHAR> IllegalChars = {' ','$', '&', '^', '/', '\\', '#', '@', '!', '*', '_', '(', ')'};
+	
+	FString StrToClean = ToClean.ToString();
+	for (TCHAR& Char : StrToClean)
+	{
+		if (IllegalChars.Contains(Char))
+		{
+			Char = TEXT('_'); // Replace illegal char with underscore
+		}
+	}
+
+	return MoveTemp(StrToClean);
+}
+
 UBatchExportOptions* UBatchExportOptions::GetInstance()
 {
 	if (!SingletonInstance)
@@ -680,6 +698,10 @@ void SRetargetExporterAssetBrowser::RefreshView()
 	AssetPickerConfig.Filter.ClassPaths.Add(UAnimSequence::StaticClass()->GetClassPathName());
 	AssetPickerConfig.Filter.ClassPaths.Add(UAnimMontage::StaticClass()->GetClassPathName());
 	AssetPickerConfig.Filter.ClassPaths.Add(UPoseAsset::StaticClass()->GetClassPathName());
+	AssetPickerConfig.Filter.ClassPaths.Add(UBlendSpace::StaticClass()->GetClassPathName());
+	AssetPickerConfig.Filter.ClassPaths.Add(UBlendSpace1D::StaticClass()->GetClassPathName());
+	AssetPickerConfig.Filter.ClassPaths.Add(UAimOffsetBlendSpace::StaticClass()->GetClassPathName());
+	AssetPickerConfig.Filter.ClassPaths.Add(UAimOffsetBlendSpace1D::StaticClass()->GetClassPathName());
 	AssetPickerConfig.InitialAssetViewType = EAssetViewType::Column;
 	AssetPickerConfig.bAddFilterUI = true;
 	AssetPickerConfig.DefaultFilterMenuExpansion = EAssetTypeCategories::Animation;
@@ -757,7 +779,7 @@ bool SRetargetExporterAssetBrowser::OnShouldFilterAsset(const FAssetData& AssetD
 	const TObjectPtr<UBatchRetargetSettings> BatchRetargetSettings = RetargetWindow.Get()->GetSettings();
 	if (!ensure(BatchRetargetSettings))
 	{
-		return false;
+		return true;
 	}
 	
 	if (!BatchRetargetSettings->SourceSkeletalMesh)
@@ -1225,11 +1247,18 @@ void SRetargetAnimAssetsWindow::SetRetargetAsset(UIKRetargeter* RetargeterToUse)
 	else if (Settings->RetargetAsset)
 	{
 		// if user assigned a custom retarget asset, set the skeletal meshes based on that asset
+		// BUT, only do this if the user hasn't already specified their own skeletal mesh
 		const UIKRetargeterController* Controller = UIKRetargeterController::GetController(Settings->RetargetAsset);
-		USkeletalMesh* SourceMesh = Controller->GetPreviewMesh(ERetargetSourceOrTarget::Source);
-		USkeletalMesh* TargetMesh = Controller->GetPreviewMesh(ERetargetSourceOrTarget::Target);
-		SetSkeletalMesh(SourceMesh, ERetargetSourceOrTarget::Source);
-		SetSkeletalMesh(TargetMesh, ERetargetSourceOrTarget::Target);
+		if (!Settings->SourceSkeletalMesh)
+		{
+			USkeletalMesh* SourceMesh = Controller->GetPreviewMesh(ERetargetSourceOrTarget::Source);
+			SetSkeletalMesh(SourceMesh, ERetargetSourceOrTarget::Source);
+		}
+		if (!Settings->TargetSkeletalMesh)
+		{
+			USkeletalMesh* TargetMesh = Controller->GetPreviewMesh(ERetargetSourceOrTarget::Target);
+			SetSkeletalMesh(TargetMesh, ERetargetSourceOrTarget::Target);
+		}
 	}
 	
 	// store the asset to use in the context
diff --git a/Engine/Plugins/Animation/IKRig/Source/IKRigEditor/Public/RetargetEditor/SRetargetAnimAssetsWindow.h b/Engine/Plugins/Animation/IKRig/Source/IKRigEditor/Public/RetargetEditor/SRetargetAnimAssetsWindow.h
--- a/Engine/Plugins/Animation/IKRig/Source/IKRigEditor/Public/RetargetEditor/SRetargetAnimAssetsWindow.h
+++ b/Engine/Plugins/Animation/IKRig/Source/IKRigEditor/Public/RetargetEditor/SRetargetAnimAssetsWindow.h
@@ -22,7 +22,9 @@ class UDebugSkelMeshComponent;
 // dialog to select path to export to
 class SBatchExportPathDialog: public SWindow
 {
+	
 public:
+	
 	SLATE_BEGIN_ARGS(SBatchExportPathDialog){}
 	SLATE_ARGUMENT(FIKRetargetBatchOperationContext*, BatchContext)
 	SLATE_ARGUMENT(bool, ExportRetargetAssets)
@@ -32,11 +34,11 @@ public:
 
 	void Construct(const FArguments& InArgs);
 
-public:
 	// displays the dialog in a blocking fashion
 	EAppReturnType::Type ShowModal();
 
-protected:
+private:
+	
 	void OnPathChange(const FString& NewPath);
 	FReply OnButtonClick(EAppReturnType::Type ButtonID);
 	
@@ -56,6 +58,9 @@ protected:
 	void UpdateExampleText();
 	// modify folder output path
 	FText GetFolderPath() const;
+
+	// remove characters not allowed in asset names
+	static FString ConvertToCleanString(const FText& ToClean);
 	
 	FText ExampleText; // The rename rule sample text
 	EAppReturnType::Type UserResponse;
