Fixes retargeted ABP's having wrong preview mesh.

diff --git a/Engine/Plugins/Animation/IKRig/Source/IKRigEditor/Private/RetargetEditor/IKRetargetBatchOperation.cpp b/Engine/Plugins/Animation/IKRig/Source/IKRigEditor/Private/RetargetEditor/IKRetargetBatchOperation.cpp
--- a/Engine/Plugins/Animation/IKRig/Source/IKRigEditor/Private/RetargetEditor/IKRetargetBatchOperation.cpp
+++ b/Engine/Plugins/Animation/IKRig/Source/IKRigEditor/Private/RetargetEditor/IKRetargetBatchOperation.cpp
@@ -284,6 +284,8 @@ void UIKRetargetBatchOperation::RetargetAssets(
 		
 		// replace skeleton
 		AnimBlueprint->TargetSkeleton = NewSkeleton;
+		// replace preview mesh (uses skeleton default otherwise)
+		AnimBlueprint->SetPreviewMesh(Context.TargetMesh);
 
 		// if they have parent blueprint, make sure to re-link to the new one also
 		UAnimBlueprint* CurrentParentBP = Cast<UAnimBlueprint>(AnimBlueprint->ParentClass->ClassGeneratedBy);
