Fixes bug, batch retarget doesn't work on ABP with no animations.

diff --git a/Engine/Plugins/Animation/IKRig/Source/IKRigEditor/Private/RetargetEditor/IKRetargetBatchOperation.cpp b/Engine/Plugins/Animation/IKRig/Source/IKRigEditor/Private/RetargetEditor/IKRetargetBatchOperation.cpp
--- a/Engine/Plugins/Animation/IKRig/Source/IKRigEditor/Private/RetargetEditor/IKRetargetBatchOperation.cpp
+++ b/Engine/Plugins/Animation/IKRig/Source/IKRigEditor/Private/RetargetEditor/IKRetargetBatchOperation.cpp
@@ -95,7 +95,7 @@ int32 UIKRetargetBatchOperation::GenerateAssetLists(const FIKRetargetBatchOperat
 		}
 	}
 
-	return AnimationAssetsToRetarget.Num();
+	return AnimationAssetsToRetarget.Num() + AnimBlueprintsToRetarget.Num();
 }
 
 void UIKRetargetBatchOperation::DuplicateRetargetAssets(
