Fixes bug where some retarget editor settings did not propagate to instances running in PIE.

diff --git a/Engine/Plugins/Animation/IKRig/Source/IKRig/Private/Retargeter/IKRetargetProcessor.cpp b/Engine/Plugins/Animation/IKRig/Source/IKRig/Private/Retargeter/IKRetargetProcessor.cpp
--- a/Engine/Plugins/Animation/IKRig/Source/IKRig/Private/Retargeter/IKRetargetProcessor.cpp
+++ b/Engine/Plugins/Animation/IKRig/Source/IKRig/Private/Retargeter/IKRetargetProcessor.cpp
@@ -619,7 +619,7 @@ void FChainEncoderFK::TransformCurrentChainTransforms(const FTransform& NewParen
 
 void FChainDecoderFK::DecodePose(
 	const FRootRetargeter& RootRetargeter,
-	const FTargetChainSettings& Settings,
+	const FTargetChainFKSettings& Settings,
 	const TArray<int32>& TargetBoneIndices,
     FChainEncoderFK& SourceChain,
     const FTargetSkeleton& TargetSkeleton,
@@ -648,7 +648,7 @@ void FChainDecoderFK::DecodePose(
 	SourceChain.TransformCurrentChainTransforms(SourceChainParentTransform);
 
 	// if FK retargeting has been disabled for this chain, then simply set it to the retarget pose
-	if (!Settings.FK.EnableFK)
+	if (!Settings.EnableFK)
 	{
 		// put the chain in the global ref pose (globally rotated by parent bone in it's currently retargeted state)
 		PutCurrentTransformsInRefPose(TargetBoneIndices, TargetSkeleton, InOutGlobalPose);
@@ -676,7 +676,7 @@ void FChainDecoderFK::DecodePose(
 		FTransform SourceInitialTransform;
 
 		// get source current / initial transforms for this bone
-		switch (Settings.FK.RotationMode)
+		switch (Settings.RotationMode)
 		{
 			case ERetargetRotationMode::Interpolated:
 			{
@@ -743,7 +743,7 @@ void FChainDecoderFK::DecodePose(
 			ParentGlobalTransform = InOutGlobalPose[ParentIndex];
 		}
 		FVector OutPosition;
-		switch (Settings.FK.TranslationMode)
+		switch (Settings.TranslationMode)
 		{
 			case ERetargetTranslationMode::None:
 				{
@@ -777,8 +777,8 @@ void FChainDecoderFK::DecodePose(
 
 	// apply final blending between retarget pose of chain and newly retargeted pose
 	// blend must be done in local space, so we do it in a separate loop after full chain pose is generated
-	const bool bShouldBlendRotation = !FMath::IsNearlyEqual(Settings.FK.RotationAlpha, 1.0f);
-	const bool bShouldBlendTranslation = !FMath::IsNearlyEqual(Settings.FK.TranslationAlpha, 1.0f);
+	const bool bShouldBlendRotation = !FMath::IsNearlyEqual(Settings.RotationAlpha, 1.0f);
+	const bool bShouldBlendTranslation = !FMath::IsNearlyEqual(Settings.TranslationAlpha, 1.0f);
 	if (bShouldBlendRotation || bShouldBlendTranslation) // (skipped if the alphas are not near 1.0)
 	{
 		// generate local space pose of chain
@@ -792,8 +792,8 @@ void FChainDecoderFK::DecodePose(
 			// blend between current local pose and initial local pose
 			FTransform& NewLocalTransform = NewLocalTransforms[ChainIndex];
 			const FTransform& RefPoseLocalTransform = InitialLocalTransforms[ChainIndex];
-			NewLocalTransform.SetTranslation(FMath::Lerp(RefPoseLocalTransform.GetTranslation(), NewLocalTransform.GetTranslation(), Settings.FK.TranslationAlpha));
-			NewLocalTransform.SetRotation(FQuat::FastLerp(RefPoseLocalTransform.GetRotation(), NewLocalTransform.GetRotation(), Settings.FK.RotationAlpha).GetNormalized());
+			NewLocalTransform.SetTranslation(FMath::Lerp(RefPoseLocalTransform.GetTranslation(), NewLocalTransform.GetTranslation(), Settings.TranslationAlpha));
+			NewLocalTransform.SetRotation(FQuat::FastLerp(RefPoseLocalTransform.GetRotation(), NewLocalTransform.GetRotation(), Settings.RotationAlpha).GetNormalized());
 
 			// put blended transforms back in global space and store in final output pose
 			const int32 BoneIndex = TargetBoneIndices[ChainIndex];
@@ -922,7 +922,8 @@ bool FChainRetargeterIK::InitializeTarget(
 }
 	
 void FChainRetargeterIK::DecodePose(
-	const FTargetChainSettings& Settings,
+	const FTargetChainIKSettings& Settings,
+	const FTargetChainSpeedPlantSettings& SpeedPlantSettings,
 	const FRootRetargeter& RootRetargeter,
 	const TMap<FName, float>& SpeedValuesFromCurves,
 	const float DeltaTime,
@@ -933,7 +934,7 @@ void FChainRetargeterIK::DecodePose(
 	FQuat GoalRotation = DeltaRotation * Target.InitialEndRotation;
 
 	// apply static rotation offset in the local space of the foot
-	GoalRotation = GoalRotation * Settings.IK.StaticRotationOffset.Quaternion();
+	GoalRotation = GoalRotation * Settings.StaticRotationOffset.Quaternion();
 
 	//
 	// calculate position of IK goal ...
@@ -946,10 +947,10 @@ void FChainRetargeterIK::DecodePose(
 	FVector GoalPosition = Start + (Source.CurrentEndDirectionNormalized * Target.InitialLength);
 
 	// blend to source location
-	if (Settings.IK.BlendToSource > KINDA_SMALL_NUMBER)
+	if (Settings.BlendToSource > KINDA_SMALL_NUMBER)
 	{
 		const FVector RootModification = RootRetargeter.Target.RootTranslationDelta * RootRetargeter.Settings.GetAffectIKWeightVector();
-		const FVector Weight = Settings.IK.BlendToSource * Settings.IK.BlendToSourceWeights;
+		const FVector Weight = Settings.BlendToSource * Settings.BlendToSourceWeights;
 		const FVector SourceLocation = Source.CurrentEndPosition + RootModification;
 		GoalPosition.X = FMath::Lerp(GoalPosition.X, SourceLocation.X, Weight.X);
 		GoalPosition.Y = FMath::Lerp(GoalPosition.Y, SourceLocation.Y, Weight.Y);
@@ -957,31 +958,30 @@ void FChainRetargeterIK::DecodePose(
 	}
 
 	// apply global static offset
-	GoalPosition += Settings.IK.StaticOffset;
+	GoalPosition += Settings.StaticOffset;
 
 	// apply local static offset
-	GoalPosition += GoalRotation.RotateVector(Settings.IK.StaticLocalOffset);
+	GoalPosition += GoalRotation.RotateVector(Settings.StaticLocalOffset);
 
 	// apply vertical scale
-	GoalPosition.Z *= Settings.IK.ScaleVertical;
+	GoalPosition.Z *= Settings.ScaleVertical;
 	
 	// apply extension
-	if (!FMath::IsNearlyEqual(Settings.IK.Extension, 1.0f))
+	if (!FMath::IsNearlyEqual(Settings.Extension, 1.0f))
 	{
-		GoalPosition = Start + (GoalPosition - Start) * Settings.IK.Extension;	
+		GoalPosition = Start + (GoalPosition - Start) * Settings.Extension;	
 	}
 	
 	// match velocity
-	const FTargetChainSpeedPlantSettings& PlantSettings = Settings.SpeedPlanting;
-	if (!ResetThisTick && PlantSettings.EnableSpeedPlanting && SpeedValuesFromCurves.Contains(PlantSettings.SpeedCurveName))
+	if (!ResetThisTick && SpeedPlantSettings.EnableSpeedPlanting && SpeedValuesFromCurves.Contains(SpeedPlantSettings.SpeedCurveName))
 	{
-		const float SourceSpeed = SpeedValuesFromCurves[PlantSettings.SpeedCurveName];
-		if (SourceSpeed < 0.0f || SourceSpeed > PlantSettings.SpeedThreshold)
+		const float SourceSpeed = SpeedValuesFromCurves[SpeedPlantSettings.SpeedCurveName];
+		if (SourceSpeed < 0.0f || SourceSpeed > SpeedPlantSettings.SpeedThreshold)
 		{
 			GoalPosition = UKismetMathLibrary::VectorSpringInterp(
 				Target.PrevEndPosition, GoalPosition, PlantingSpringState,
-				PlantSettings.UnplantStiffness,
-				PlantSettings.UnplantCriticalDamping,
+				SpeedPlantSettings.UnplantStiffness,
+				SpeedPlantSettings.UnplantCriticalDamping,
 				DeltaTime, 1.0f, 0.0f);
 		}
 		else
@@ -1344,15 +1344,14 @@ bool FPoleVectorMatcher::Initialize(
 }
 
 void FPoleVectorMatcher::MatchPoleVector(
-	const FTargetChainSettings& Settings,
+	const FTargetChainFKSettings& Settings,
 	const TArray<int32>& SourceIndices,
 	const TArray<int32>& TargetIndices,
 	const TArray<FTransform>& SourceGlobalPose,
 	TArray<FTransform>& OutTargetGlobalPose,
 	FRetargetSkeleton& TargetSkeleton)
 {
-	if (Settings.FK.PoleVectorMatching <= KINDA_SMALL_NUMBER &&
-		Settings.FK.PoleVectorOffset <= KINDA_SMALL_NUMBER)
+	if (Settings.PoleVectorMatching <= KINDA_SMALL_NUMBER && Settings.PoleVectorOffset <= KINDA_SMALL_NUMBER)
 	{
 		return;
 	}
@@ -1377,11 +1376,11 @@ void FPoleVectorMatcher::MatchPoleVector(
 	const FVector TargetPoleVector = CalculatePoleVector(TargetPoleAxis, TargetIndices, OutTargetGlobalPose);
 	
 	const float RotateTargetToSource = FMath::Acos(FVector::DotProduct(SourcePoleVector, TargetPoleVector));
-	const float MatchPoleAngle = RotateTargetToSource - TargetToSourceAngularOffsetAtRefPose + FMath::DegreesToRadians(Settings.FK.PoleVectorOffset);
+	const float MatchPoleAngle = RotateTargetToSource - TargetToSourceAngularOffsetAtRefPose + FMath::DegreesToRadians(Settings.PoleVectorOffset);
 
 	const FVector TargetChainNormal = GetChainNormal(TargetIndices, OutTargetGlobalPose);
 	const FQuat PoleVectorMatchingRotation = FQuat(TargetChainNormal, MatchPoleAngle);
-	const FQuat FinalRotation = FQuat::FastLerp(FQuat::Identity, PoleVectorMatchingRotation, Settings.FK.PoleVectorMatching).GetNormalized();
+	const FQuat FinalRotation = FQuat::FastLerp(FQuat::Identity, PoleVectorMatchingRotation, Settings.PoleVectorMatching).GetNormalized();
 
 	// rotate the base of the chain to match the pole vectors
 	FTransform& BaseOfChain = OutTargetGlobalPose[TargetIndices[0]];
@@ -1571,7 +1570,8 @@ void UIKRetargetProcessor::Initialize(
 	// copy the initial settings from the asset
 	ApplySettingsFromAsset();
 	
-	bIsInitialized = true;
+	bIsInitialized = bRootsInitialized || bAtLeastOneValidBoneChainPair;
+	AssetVersionInitializedWith = RetargeterAsset->GetVersion();
 }
 
 bool UIKRetargetProcessor::InitializeRoots()
@@ -1968,6 +1968,7 @@ void UIKRetargetProcessor::RunIKRetarget(
 		// decode the IK goal and apply to IKRig
 		ChainPair.IKChainRetargeter.DecodePose(
 			ChainPair.Settings,
+			ChainPair.SpeedPlantSettings,
 			RootRetargeter,
 			SpeedValuesFromCurves,
 			DeltaTime,
@@ -2036,7 +2037,7 @@ void UIKRetargetProcessor::RunStrideWarping(const TArray<FTransform>& InTargetGl
 		// use goals to determine the body's rotation
 		for (const FRetargetChainPairIK& ChainPair : ChainPairsIK)
 		{
-			if (!ChainPair.Settings.IK.bAffectedByIKWarping)
+			if (!ChainPair.Settings.bAffectedByIKWarping)
 			{
 				continue;
 			}
@@ -2101,7 +2102,7 @@ void UIKRetargetProcessor::RunStrideWarping(const TArray<FTransform>& InTargetGl
 	// warp goal positions...
 	for (FRetargetChainPairIK& ChainPair : ChainPairsIK)
 	{
-		if (!ChainPair.Settings.IK.bAffectedByIKWarping)
+		if (!ChainPair.Settings.bAffectedByIKWarping)
 		{
 			continue;
 		}
@@ -2162,10 +2163,17 @@ bool UIKRetargetProcessor::WasInitializedWithTheseAssets(
 	const TObjectPtr<USkeletalMesh> InTargetMesh,
 	const TObjectPtr<UIKRetargeter> InRetargetAsset)
 {
+	// not initialized at all
 	if (!IsInitialized())
 	{
 		return false;
 	}
+	
+	// asset has been modified in a way that requires reinitialization
+	if (AssetVersionInitializedWith != InRetargetAsset->GetVersion())
+	{
+		return false;
+	}
 
 	// check that both the source and target skeletal meshes are the same as what we initialized with
 	const FRetargetSkeleton& SourceRetargetSkeleton = GetSkeleton(ERetargetSourceOrTarget::Source);
@@ -2178,13 +2186,20 @@ bool UIKRetargetProcessor::WasInitializedWithTheseAssets(
 	}
 
 	// check that the retarget asset is the same as what we initialized with
-	const bool bAssetMatches = InRetargetAsset == RetargeterAsset;
+	if (InRetargetAsset != RetargeterAsset)
+	{
+		return false;
+	}
 
 	// check that the number of bones are the same as what we initialized with
 	const bool bSourceHasSameNumberOfBones = InSourceMesh->GetRefSkeleton().GetNum() == SourceRetargetSkeleton.BoneNames.Num();
 	const bool bTargetHasSameNumberOfBones = InTargetMesh->GetRefSkeleton().GetNum() == TargetRetargetSkeleton.BoneNames.Num();
+	if (!(bSourceHasSameNumberOfBones && bTargetHasSameNumberOfBones))
+	{
+		return false;
+	}
 	
-	return bSourceMatches && bTargetMatches && bAssetMatches && bSourceHasSameNumberOfBones && bTargetHasSameNumberOfBones;
+	return true;
 }
 
 void UIKRetargetProcessor::SetNeedsInitialized()
@@ -2361,7 +2376,7 @@ void UIKRetargetProcessor::ApplySettingsFromAsset()
 		IKRigProcessor->CopyAllInputsFromSourceAssetAtRuntime(TargetIKRig);
 	}
 	
-	// copy chain settings
+	// copy chain settings from the asset
 	const TArray<TObjectPtr<URetargetChainSettings>>& AllChainSettings = RetargeterAsset->GetAllChainSettings();
 	for (const TObjectPtr<URetargetChainSettings>& ChainSettings : AllChainSettings)
 	{
@@ -2369,7 +2384,7 @@ void UIKRetargetProcessor::ApplySettingsFromAsset()
 		{
 			if (Chain.TargetBoneChainName == ChainSettings->TargetChain)
 			{
-				Chain.Settings.CopySettingsFromAsset(ChainSettings);
+				Chain.Settings = ChainSettings->Settings.FK;
 			}
 		}
 		
@@ -2377,13 +2392,14 @@ void UIKRetargetProcessor::ApplySettingsFromAsset()
 		{
 			if (Chain.TargetBoneChainName == ChainSettings->TargetChain)
 			{
-				Chain.Settings.CopySettingsFromAsset(ChainSettings);
+				Chain.Settings = ChainSettings->Settings.IK;
+				Chain.SpeedPlantSettings = ChainSettings->Settings.SpeedPlanting;
 			}
 		}
 	}
 
 	// copy root settings
-	RootRetargeter.Settings.CopySettingsFromAsset(RetargeterAsset->GetRootSettingsUObject());
+	RootRetargeter.Settings = RetargeterAsset->GetRootSettingsUObject()->Settings;
 
 	// copy global settings
 	GlobalSettings = RetargeterAsset->GetGlobalSettings();
@@ -2433,7 +2449,7 @@ void UIKRetargetProcessor::ApplySettingsFromProfile(const FRetargetProfile& Prof
 			{
 				if (Chain.TargetBoneChainName == ChainSettings.Key)
 				{
-					Chain.Settings = ChainSettings.Value;
+					Chain.Settings = ChainSettings.Value.FK;
 				}
 			}
 		
@@ -2441,7 +2457,7 @@ void UIKRetargetProcessor::ApplySettingsFromProfile(const FRetargetProfile& Prof
 			{
 				if (Chain.TargetBoneChainName == ChainSettings.Key)
 				{
-					Chain.Settings = ChainSettings.Value;
+					Chain.Settings = ChainSettings.Value.IK;
 				}
 			}
 		}
diff --git a/Engine/Plugins/Animation/IKRig/Source/IKRig/Private/Retargeter/IKRetargetSettings.cpp b/Engine/Plugins/Animation/IKRig/Source/IKRig/Private/Retargeter/IKRetargetSettings.cpp
--- a/Engine/Plugins/Animation/IKRig/Source/IKRig/Private/Retargeter/IKRetargetSettings.cpp
+++ b/Engine/Plugins/Animation/IKRig/Source/IKRig/Private/Retargeter/IKRetargetSettings.cpp
@@ -38,17 +38,7 @@ bool FTargetChainIKSettings::operator==(const FTargetChainIKSettings& Other) con
 	&& bAffectedByIKWarping == Other.bAffectedByIKWarping;
 }
 
-void FTargetChainSettings::CopySettingsFromAsset(const URetargetChainSettings* AssetChainSettings)
-{
-	*this = AssetChainSettings->Settings;
-}
-
 bool FTargetChainSettings::operator==(const FTargetChainSettings& Other) const
 {
 	return FK == Other.FK && IK == Other.IK && SpeedPlanting == Other.SpeedPlanting;
 }
-
-void FTargetRootSettings::CopySettingsFromAsset(const URetargetRootSettings* AssetRootSettings)
-{
-	*this = AssetRootSettings->Settings;
-}
diff --git a/Engine/Plugins/Animation/IKRig/Source/IKRig/Public/Retargeter/IKRetargetProcessor.h b/Engine/Plugins/Animation/IKRig/Source/IKRig/Public/Retargeter/IKRetargetProcessor.h
--- a/Engine/Plugins/Animation/IKRig/Source/IKRig/Public/Retargeter/IKRetargetProcessor.h
+++ b/Engine/Plugins/Animation/IKRig/Source/IKRig/Public/Retargeter/IKRetargetProcessor.h
@@ -201,7 +201,7 @@ struct FPoleVectorMatcher
 		const FRetargetSkeleton& TargetSkeleton);
 
 	void MatchPoleVector(
-		const FTargetChainSettings& Settings,
+		const FTargetChainFKSettings& Settings,
 		const TArray<int32>& SourceIndices,
 		const TArray<int32>& TargetIndices,
 		const TArray<FTransform> &SourceGlobalPose,
@@ -291,14 +291,14 @@ struct FChainDecoderFK : public FChainFK
 	
 	void DecodePose(
 		const FRootRetargeter& RootRetargeter,
-		const FTargetChainSettings& Settings,
+		const FTargetChainFKSettings& Settings,
 		const TArray<int32>& TargetBoneIndices,
 		FChainEncoderFK& SourceChain,
 		const FTargetSkeleton& TargetSkeleton,
 		TArray<FTransform> &InOutGlobalPose);
 
 	void MatchPoleVector(
-		const FTargetChainSettings& Settings,
+		const FTargetChainFKSettings& Settings,
 		const TArray<int32>& TargetBoneIndices,
 		FChainEncoderFK& SourceChain,
 		const FTargetSkeleton& TargetSkeleton,
@@ -382,7 +382,8 @@ struct FChainRetargeterIK
 	void EncodePose(const TArray<FTransform> &SourceInputGlobalPose);
 	
 	void DecodePose(
-		const FTargetChainSettings& Settings,
+		const FTargetChainIKSettings& Settings,
+		const FTargetChainSpeedPlantSettings& SpeedPlantSettings,
 		const FRootRetargeter& RootRetargeter,
 		const TMap<FName, float>& SpeedValuesFromCurves,
 		const float DeltaTime,
@@ -393,8 +394,6 @@ struct FChainRetargeterIK
 
 struct FRetargetChainPair
 {
-	FTargetChainSettings Settings;
-	
 	TArray<int32> SourceBoneIndices;
 	TArray<int32> TargetBoneIndices;
 	
@@ -421,6 +420,7 @@ private:
 
 struct FRetargetChainPairFK : FRetargetChainPair
 {
+	FTargetChainFKSettings Settings;
 	FChainEncoderFK FKEncoder;
 	FChainDecoderFK FKDecoder;
 	FPoleVectorMatcher PoleVectorMatcher;
@@ -435,6 +435,8 @@ struct FRetargetChainPairFK : FRetargetChainPair
 
 struct FRetargetChainPairIK : FRetargetChainPair
 {
+	FTargetChainIKSettings Settings;
+	FTargetChainSpeedPlantSettings SpeedPlantSettings;
 	FChainRetargeterIK IKChainRetargeter;
 	FName IKGoalName;
 	FName PoleVectorGoalName;
@@ -504,7 +506,7 @@ public:
 	FName GetRetargetRoot(ERetargetSourceOrTarget SourceOrTarget) const;
 	
 	/** Get whether this processor is ready to call RunRetargeter() and generate new poses. */
-	bool IsInitialized() const { return bIsInitialized && (bRootsInitialized || bAtLeastOneValidBoneChainPair); };
+	bool IsInitialized() const { return bIsInitialized; };
 
 	/** Get whether this processor was initialized with these skeletal meshes and retarget asset*/
 	bool WasInitializedWithTheseAssets(
@@ -568,6 +570,7 @@ private:
 
 	/** Only true once Initialize() has successfully completed.*/
 	bool bIsInitialized = false;
+	int32 AssetVersionInitializedWith = -1;
 	/** true when roots are able to be retargeted */
 	bool bRootsInitialized = false;
 	/** true when at least one pair of bone chains is able to be retargeted */
diff --git a/Engine/Plugins/Animation/IKRig/Source/IKRig/Public/Retargeter/IKRetargetSettings.h b/Engine/Plugins/Animation/IKRig/Source/IKRig/Public/Retargeter/IKRetargetSettings.h
--- a/Engine/Plugins/Animation/IKRig/Source/IKRig/Public/Retargeter/IKRetargetSettings.h
+++ b/Engine/Plugins/Animation/IKRig/Source/IKRig/Public/Retargeter/IKRetargetSettings.h
@@ -180,8 +180,6 @@ public:
 	UPROPERTY(EditAnywhere, BlueprintReadWrite, Category = "Speed Planting")
 	FTargetChainSpeedPlantSettings SpeedPlanting;
 	
-	void CopySettingsFromAsset(const URetargetChainSettings* AssetChainSettings);
-
 	bool operator==(const FTargetChainSettings& Other) const;
 };
 
@@ -242,8 +240,6 @@ public:
 	UPROPERTY(EditAnywhere, BlueprintReadWrite, Category = "Root Retarget Settings", meta = (ClampMin = "0.0", ClampMax = "1.0", UIMin = "0.0", UIMax = "1.0", DisplayName = "Affect IK Vertical"))
 	float AffectIKVertical = 0.0f;
 
-	void CopySettingsFromAsset(const URetargetRootSettings* AssetChainSettings);
-
 	FVector GetAffectIKWeightVector() const
 	{
 		return FVector(AffectIKHorizontal, AffectIKHorizontal, AffectIKVertical);
diff --git a/Engine/Plugins/Animation/IKRig/Source/IKRig/Public/Retargeter/IKRetargeter.h b/Engine/Plugins/Animation/IKRig/Source/IKRig/Public/Retargeter/IKRetargeter.h
--- a/Engine/Plugins/Animation/IKRig/Source/IKRig/Public/Retargeter/IKRetargeter.h
+++ b/Engine/Plugins/Animation/IKRig/Source/IKRig/Public/Retargeter/IKRetargeter.h
@@ -256,6 +256,11 @@ public:
 	// Get the retarget profile by name (may be null) 
 	const FRetargetProfile* GetProfileByName(const FName& ProfileName) const;
 
+	// get current version of the data (to compare against running processor instances)
+	int32 GetVersion() const { return Version; };
+	// do this after any edit that would require running instance to reinitialize
+	void IncrementVersion() { ++Version; };
+
 	// BLUEPRINT GETTERS 
 
 	// Returns the chain settings associated with a given Goal in an IK Retargeter Asset using the given profile name (optional) 
@@ -374,6 +379,10 @@ public:
 private:
 
 	void CleanAndInitialize();
+	
+	// incremented by any edits that require re-initialization
+	UPROPERTY(Transient)
+	int32 Version = INDEX_NONE;
 
 	// The rig to copy animation FROM.
 	UPROPERTY(EditAnywhere, Category = Source)
@@ -451,12 +460,12 @@ public:
 	UPROPERTY()
 	float BoneDrawSize = 1.0f;
 
-	/** The controller responsible for managing this asset's data (all editor mutation goes through this) */
+	// The controller responsible for managing this asset's data (all editor mutation goes through this)
 	UPROPERTY(Transient, DuplicateTransient, NonTransactional )
 	TObjectPtr<UObject> Controller;
 	
 private:
-
+	
 	// only ask to fix the root height once, then warn thereafter (don't nag) 
 	TSet<TObjectPtr<USkeletalMesh>> MeshesAskedToFixRootHeightFor;
 #endif
diff --git a/Engine/Plugins/Animation/IKRig/Source/IKRigEditor/Private/RetargetEditor/IKRetargetDetails.cpp b/Engine/Plugins/Animation/IKRig/Source/IKRigEditor/Private/RetargetEditor/IKRetargetDetails.cpp
--- a/Engine/Plugins/Animation/IKRig/Source/IKRigEditor/Private/RetargetEditor/IKRetargetDetails.cpp
+++ b/Engine/Plugins/Animation/IKRig/Source/IKRigEditor/Private/RetargetEditor/IKRetargetDetails.cpp
@@ -1518,12 +1518,12 @@ void FRetargetOpStackCustomization::CustomizeDetails(IDetailLayoutBuilder& Detai
 			{
 				PropertyHandle->SetOnPropertyValueChanged(FSimpleDelegate::CreateLambda([this]()
 				{
-					Controller.Pin().Get()->ReinitializeProcessor();
+					Controller.Pin().Get()->HandleRetargeterNeedsInitialized();
 				}));
 				
 				PropertyHandle->SetOnChildPropertyValueChanged(FSimpleDelegate::CreateLambda([this]()
 				{
-						Controller.Pin().Get()->ReinitializeProcessor();
+					Controller.Pin().Get()->HandleRetargeterNeedsInitialized();
 				}));
 				
 				OpSettingCategory.AddProperty(PropertyHandle)
diff --git a/Engine/Plugins/Animation/IKRig/Source/IKRigEditor/Private/RetargetEditor/IKRetargetEditorController.cpp b/Engine/Plugins/Animation/IKRig/Source/IKRigEditor/Private/RetargetEditor/IKRetargetEditorController.cpp
--- a/Engine/Plugins/Animation/IKRig/Source/IKRigEditor/Private/RetargetEditor/IKRetargetEditorController.cpp
+++ b/Engine/Plugins/Animation/IKRig/Source/IKRigEditor/Private/RetargetEditor/IKRetargetEditorController.cpp
@@ -247,27 +247,14 @@ void FIKRetargetEditorController::HandleRetargeterNeedsInitialized() const
 	FixZeroHeightRetargetRoot(ERetargetSourceOrTarget::Source);
 	FixZeroHeightRetargetRoot(ERetargetSourceOrTarget::Target);
 	
-	ReinitializeProcessor();
-	
-	// refresh all the UI views
-	RefreshAllViews();	
-}
-
-void FIKRetargetEditorController::ReinitializeProcessor() const
-{
 	// clear the output log
 	ClearOutputLog();
+
+	// force running instances to reinitialize on next tick
+	AssetController->GetAsset()->IncrementVersion();
 	
-	// force reinit the retarget processor (also inits the target IK Rig processor)
-	if (UIKRetargetProcessor* Processor = GetRetargetProcessor())
-	{
-		constexpr bool bSuppressWarnings = false;
-		Processor->Initialize(
-			GetSkeletalMesh(ERetargetSourceOrTarget::Source),
-			GetSkeletalMesh(ERetargetSourceOrTarget::Target),
-			AssetController->GetAsset(),
-			bSuppressWarnings);
-	}
+	// refresh all the UI views
+	RefreshAllViews();	
 }
 
 void FIKRetargetEditorController::HandleIKRigReplaced(ERetargetSourceOrTarget SourceOrTarget)
diff --git a/Engine/Plugins/Animation/IKRig/Source/IKRigEditor/Public/RetargetEditor/IKRetargetEditorController.h b/Engine/Plugins/Animation/IKRig/Source/IKRigEditor/Public/RetargetEditor/IKRetargetEditorController.h
--- a/Engine/Plugins/Animation/IKRig/Source/IKRigEditor/Public/RetargetEditor/IKRetargetEditorController.h
+++ b/Engine/Plugins/Animation/IKRig/Source/IKRigEditor/Public/RetargetEditor/IKRetargetEditorController.h
@@ -103,7 +103,6 @@ public:
 	void HandleRetargetChainRemoved(UIKRigDefinition* ModifiedIKRig, const FName InChainRemoved) const;
 	// callback when IK Retargeter asset requires reinitialization
 	void HandleRetargeterNeedsInitialized() const;
-	void ReinitializeProcessor() const;
 	FDelegateHandle RetargeterReInitDelegateHandle;
 	// callback when IK Rig asset has been swapped out
 	void HandleIKRigReplaced(ERetargetSourceOrTarget SourceOrTarget);
diff --git a/Engine/Plugins/Animation/IKRig/Source/IKRigEditor/Public/RetargetEditor/IKRetargeterController.h b/Engine/Plugins/Animation/IKRig/Source/IKRigEditor/Public/RetargetEditor/IKRetargeterController.h
--- a/Engine/Plugins/Animation/IKRig/Source/IKRigEditor/Public/RetargetEditor/IKRetargeterController.h
+++ b/Engine/Plugins/Animation/IKRig/Source/IKRigEditor/Public/RetargetEditor/IKRetargeterController.h
@@ -5,6 +5,7 @@
 #include "UObject/Object.h"
 
 #include "RetargetEditor/IKRetargeterPoseGenerator.h"
+#include "Retargeter/IKRetargeter.h"
 
 #include "IKRetargeterController.generated.h"
 
@@ -373,6 +374,7 @@ struct FScopedReinitializeIKRetargeter
 	{
 		if (--Controller->ReinitializeScopeCounter == 0)
 		{
+			Controller->GetAsset()->IncrementVersion();
 			Controller->BroadcastNeedsReinitialized();
 		}
 	};
