Fix for Qualysis right arm not auto-characterizing

diff --git a/Engine/Plugins/Animation/IKRig/Source/IKRigEditor/Private/RigEditor/IKRigAutoCharacterizer.cpp b/Engine/Plugins/Animation/IKRig/Source/IKRigEditor/Private/RigEditor/IKRigAutoCharacterizer.cpp
--- a/Engine/Plugins/Animation/IKRig/Source/IKRigEditor/Private/RigEditor/IKRigAutoCharacterizer.cpp
+++ b/Engine/Plugins/Animation/IKRig/Source/IKRigEditor/Private/RigEditor/IKRigAutoCharacterizer.cpp
@@ -1069,8 +1069,8 @@ FKnownTemplateHierarchies::FKnownTemplateHierarchies()
 	// Qualisys
 	{
 		static FName QualisysName = "Qualisys";
-		static TArray<FName> QualisysBones = {"Reference", "Hips", "Spine", "Spine1", "Spine2", "Neck", "Head", "Head_end", "LeftShoulder", "LeftArm", "LeftForeArm", "LeftForeArmRoll", "LeftHand", "LeftHand_end", "LeftInHandThumb", "LeftHandThumb0", "LeftHandThumb1", "LeftHandThumb2", "LeftHandThumb3", "LeftHandThumb3_end", "LeftInHandIndex", "LeftHandIndex1", "LeftHandIndex2", "LeftHandIndex3", "LeftHandIndex3_end", "LeftInHandMiddle", "LeftHandMiddle1", "LeftHandMiddle2", "LeftHandMiddle3", "LeftHandMiddle3_end", "LeftInHandRing", "LeftHandRing1", "LeftHandRing2", "LeftHandRing3", "LeftHandRing3_end", "LeftInHandPinky", "LeftHandPinky1", "LeftHandPinky2", "LeftHandPinky3", "LeftHandPinky3_end", "RightShoulder", "LeftUpLeg", "LeftLeg", "LeftFoot", "LeftToeBase", "LeftToeBase_end", "RightUpLeg", "RightLeg", "RightFoot", "RightToeBase", "RightToeBase_end"};
-		static TArray<int32> QualisysParentIndices = {-1, 0, 1, 2, 3, 4, 5, 6, 4, 8, 9, 10, 11, 12, 12, 14, 15, 16, 17, 18, 12, 20, 21, 22, 23, 12, 25, 26, 27, 28, 12, 30, 31, 32, 33, 12, 35, 36, 37, 38, 4, 1, 41, 42, 43, 44, 1, 46, 47, 48, 49};
+		static TArray<FName> QualisysBones = {"Reference", "Hips", "Spine", "Spine1", "Spine2", "Neck", "Head", "Head_end", "LeftShoulder", "LeftArm", "LeftForeArm", "LeftForeArmRoll", "LeftHand", "LeftHand_end", "LeftInHandThumb", "LeftHandThumb0", "LeftHandThumb1", "LeftHandThumb2", "LeftHandThumb3", "LeftHandThumb3_end", "LeftInHandIndex", "LeftHandIndex1", "LeftHandIndex2", "LeftHandIndex3", "LeftHandIndex3_end", "LeftInHandMiddle", "LeftHandMiddle1", "LeftHandMiddle2", "LeftHandMiddle3", "LeftHandMiddle3_end", "LeftInHandRing", "LeftHandRing1", "LeftHandRing2", "LeftHandRing3", "LeftHandRing3_end", "LeftInHandPinky", "LeftHandPinky1", "LeftHandPinky2", "LeftHandPinky3", "LeftHandPinky3_end", "RightShoulder", "RightArm", "RightForeArm", "RightForeArmRoll", "RightHand", "RightHand_end", "RightInHandPinky", "RightHandPinky1", "RightHandPinky2", "RightHandPinky3", "RightHandPinky3_end", "RightInHandRing", "RightHandRing1", "RightHandRing2", "RightHandRing3", "RightHandRing3_end", "RightInHandMiddle", "RightHandMiddle1", "RightHandMiddle2", "RightHandMiddle3", "RightHandMiddle3_end", "RightInHandIndex", "RightHandIndex1", "RightHandIndex2", "RightHandIndex3", "RightHandIndex3_end", "RightInHandThumb", "RightInHandThumb0", "RightHandThumb1", "RightHandThumb2", "RightHandThumb3", "RightHandThumb3_end", "LeftUpLeg", "LeftLeg", "LeftFoot", "LeftToeBase", "LeftToeBase_end", "RightUpLeg", "RightLeg", "RightFoot", "RightToeBase", "RightToeBase_end"};
+		static TArray<int32> QualisysParentIndices = {-1, 0, 1, 2, 3, 4, 5, 6, 4, 8, 9, 10, 11, 12, 12, 14, 15, 16, 17, 18, 12, 20, 21, 22, 23, 12, 25, 26, 27, 28, 12, 30, 31, 32, 33, 12, 35, 36, 37, 38, 4, 40, 41, 42, 43, 44, 44, 46, 47, 48, 49, 44, 51, 52, 53, 54, 44, 56, 57, 58, 59, 44, 61, 62, 63, 64, 44, 66, 67, 68, 69, 70, 1, 72, 73, 74, 75, 1, 77, 78, 79, 80};
 		FTemplateHierarchy& Qualisys = AddTemplateHierarchy(QualisysName, QualisysBones, QualisysParentIndices);
 		FRetargetDefinition& QualisysRetarget = Qualisys.RetargetDefinition;
 		// core
