Fixes issue, retarget pose dropdown requires two clicks to activate.

diff --git a/Engine/Plugins/Animation/IKRig/Source/IKRigEditor/Private/RetargetEditor/IKRetargetEditorController.cpp b/Engine/Plugins/Animation/IKRig/Source/IKRigEditor/Private/RetargetEditor/IKRetargetEditorController.cpp
--- a/Engine/Plugins/Animation/IKRig/Source/IKRigEditor/Private/RetargetEditor/IKRetargetEditorController.cpp
+++ b/Engine/Plugins/Animation/IKRig/Source/IKRigEditor/Private/RetargetEditor/IKRetargetEditorController.cpp
@@ -154,6 +154,7 @@ void FIKRetargetEditorController::Initialize(TSharedPtr<FIKRetargetEditor> InEdi
 	PreviousMode = ERetargeterOutputMode::EditRetargetPose;
 	PoseExporter = MakeShared<FIKRetargetPoseExporter>();
 	PoseExporter->Initialize(SharedThis(this));
+	RefreshPoseList();
 
 	PlaybackManager = MakeUnique<FRetargetPlaybackManager>(SharedThis(this));
 
@@ -680,7 +681,6 @@ void FIKRetargetEditorController::RefreshHierarchyView() const
 {
 	if (HierarchyView.IsValid())
 	{
-		HierarchyView.Get()->RefreshPoseList();
 		HierarchyView.Get()->RefreshTreeView();
 	}
 }
@@ -766,6 +766,7 @@ void FIKRetargetEditorController::SetSourceOrTargetMode(ERetargetSourceOrTarget
 	}
 	
 	RefreshAllViews();
+	RefreshPoseList();
 }
 
 void FIKRetargetEditorController::SetSelectedMesh(UPrimitiveComponent* InMeshComponent)
diff --git a/Engine/Plugins/Animation/IKRig/Source/IKRigEditor/Private/RetargetEditor/SIKRetargetPoseEditor.cpp b/Engine/Plugins/Animation/IKRig/Source/IKRigEditor/Private/RetargetEditor/SIKRetargetPoseEditor.cpp
--- a/Engine/Plugins/Animation/IKRig/Source/IKRigEditor/Private/RetargetEditor/SIKRetargetPoseEditor.cpp
+++ b/Engine/Plugins/Animation/IKRig/Source/IKRigEditor/Private/RetargetEditor/SIKRetargetPoseEditor.cpp
@@ -48,13 +48,11 @@ void SIKRetargetPoseEditor::Construct(
 			[
 				SAssignNew(PoseListComboBox, SComboBox<TSharedPtr<FName>>)
 				.OptionsSource(&PoseNames)
-				.OnComboBoxOpening(this, &SIKRetargetPoseEditor::Refresh)
 				.OnGenerateWidget_Lambda([](TSharedPtr<FName> InItem)
 				{
 					return SNew(STextBlock).Text(FText::FromName(*InItem.Get()));
 				})
 				.OnSelectionChanged(Controller, &FIKRetargetEditorController::OnPoseSelected)
-				.Content()
 				[
 					SNew(STextBlock).Text(Controller, &FIKRetargetEditorController::GetCurrentPoseName)
 				]
@@ -102,8 +100,6 @@ void SIKRetargetPoseEditor::Refresh()
 	{
 		PoseNames.Add(MakeShareable(new FName(Pose.Key)));
 	}
-
-	PoseListComboBox->RefreshOptions();
 }
 
 TSharedRef<SWidget>  SIKRetargetPoseEditor::MakeToolbar(TSharedPtr<FUICommandList> Commands)
