Fixes IK Rig preview mesh change generates two undo transactions.

diff --git a/Engine/Plugins/Animation/IKRig/Source/IKRigEditor/Private/RigEditor/IKRigController.cpp b/Engine/Plugins/Animation/IKRig/Source/IKRigEditor/Private/RigEditor/IKRigController.cpp
--- a/Engine/Plugins/Animation/IKRig/Source/IKRigEditor/Private/RigEditor/IKRigController.cpp
+++ b/Engine/Plugins/Animation/IKRig/Source/IKRigEditor/Private/RigEditor/IKRigController.cpp
@@ -657,8 +657,9 @@ bool UIKRigController::SetSkeletalMesh(USkeletalMesh* SkeletalMesh) const
 		FScopedReinitializeIKRig Reinitialize(this);
 		return false;
 	}
-	
-	FScopedTransaction Transaction(LOCTEXT("SetSkeletalMesh_Label", "Set Skeletal Mesh"));
+
+	const bool bShouldActuallyTransact = Asset->PreviewSkeletalMesh != SkeletalMesh;
+	FScopedTransaction Transaction(LOCTEXT("SetSkeletalMesh_Label", "Set Skeletal Mesh"), bShouldActuallyTransact);
 	FScopedReinitializeIKRig Reinitialize(this);
 	Asset->Modify();
 	// update stored skeletal mesh used for previewing results
