Fixes retargeter chain rename undo bug.

diff --git a/Engine/Plugins/Animation/IKRig/Source/IKRigEditor/Private/RetargetEditor/IKRetargeterController.cpp b/Engine/Plugins/Animation/IKRig/Source/IKRigEditor/Private/RetargetEditor/IKRetargeterController.cpp
--- a/Engine/Plugins/Animation/IKRig/Source/IKRigEditor/Private/RetargetEditor/IKRetargeterController.cpp
+++ b/Engine/Plugins/Animation/IKRig/Source/IKRigEditor/Private/RetargetEditor/IKRetargeterController.cpp
@@ -566,6 +566,8 @@ void UIKRetargeterController::HandleRetargetChainRenamed(UIKRigDefinition* IKRig
 		FName& ChainNameToUpdate = bIsSourceRig ? ChainMap->SourceChain : ChainMap->TargetChain;
 		if (ChainNameToUpdate == OldChainName)
 		{
+			FScopedTransaction Transaction(LOCTEXT("RetargetChainRenamed_Label", "Retarget Chain Renamed"));
+			ChainMap->Modify();
 			ChainNameToUpdate = NewChainName;
 			FScopedReinitializeIKRetargeter Reinitialize(this);
 			return;
@@ -575,10 +577,13 @@ void UIKRetargeterController::HandleRetargetChainRenamed(UIKRigDefinition* IKRig
 
 void UIKRetargeterController::HandleRetargetChainRemoved(UIKRigDefinition* IKRig, const FName& InChainRemoved) const
 {
+	FScopedTransaction Transaction(LOCTEXT("RetargetChainRemoved_Label", "Retarget Chain Removed"));
+	Asset->Modify();
+	
 	const bool bIsSourceRig = IKRig == Asset->GetIKRig(ERetargetSourceOrTarget::Source);
 	const bool bIsTargetRig = IKRig == Asset->GetIKRig(ERetargetSourceOrTarget::Target);
 	check(bIsSourceRig || bIsTargetRig)
-
+	
 	// set source chain name to NONE if it has been deleted 
 	if (bIsSourceRig)
 	{
