Add proper support for root motion throughout IK Retargeter, editor exporter and batch exporter.

diff --git a/Engine/Plugins/Animation/IKRig/Source/IKRigEditor/Private/RetargetEditor/IKRetargetAnimInstanceProxy.cpp b/Engine/Plugins/Animation/IKRig/Source/IKRigEditor/Private/RetargetEditor/IKRetargetAnimInstanceProxy.cpp
--- a/Engine/Plugins/Animation/IKRig/Source/IKRigEditor/Private/RetargetEditor/IKRetargetAnimInstanceProxy.cpp
+++ b/Engine/Plugins/Animation/IKRig/Source/IKRigEditor/Private/RetargetEditor/IKRetargetAnimInstanceProxy.cpp
@@ -16,6 +16,8 @@ FIKRetargetAnimInstanceProxy::FIKRetargetAnimInstanceProxy(
 	RetargetNode(InRetargetNode),
 	OutputMode(ERetargeterOutputMode::EditRetargetPose)
 {
+	// retargeting is all done in world space, moving the source component breaks root motion retargeting
+	bIgnoreRootMotion = true;
 }
 
 void FIKRetargetAnimInstanceProxy::Initialize(UAnimInstance* InAnimInstance)
diff --git a/Engine/Plugins/Animation/IKRig/Source/IKRigEditor/Private/RetargetEditor/IKRetargetBatchOperation.cpp b/Engine/Plugins/Animation/IKRig/Source/IKRigEditor/Private/RetargetEditor/IKRetargetBatchOperation.cpp
--- a/Engine/Plugins/Animation/IKRig/Source/IKRigEditor/Private/RetargetEditor/IKRetargetBatchOperation.cpp
+++ b/Engine/Plugins/Animation/IKRig/Source/IKRigEditor/Private/RetargetEditor/IKRetargetBatchOperation.cpp
@@ -128,14 +128,6 @@ void UIKRetargetBatchOperation::DuplicateRetargetAssets(
 		
 		TMap<UAnimationAsset*, UAnimationAsset*> DuplicateMap = DuplicateAssets<UAnimationAsset>({Asset}, DestinationPackage, &Context.NameRule);
 		DuplicatedAnimAssets.Append(DuplicateMap);
-
-		// optionally let user override root lock on exported animation sequences,
-		// (by default it will inherit the bForceRootLock state from the duplicated source animation)
-		UAnimSequence* TargetSequence = Cast<UAnimSequence>(DuplicateMap[Asset]);
-		if (Context.RootLockMode != ERetargetRootLockMode::FromSourceAnimation && TargetSequence)
-		{
-			TargetSequence->bForceRootLock = Context.RootLockMode == ERetargetRootLockMode::ForceRootLocked ? true : false;
-		}
 	}
 	for (UAnimBlueprint* Asset : AnimBlueprintsToDuplicate)
 	{
@@ -387,9 +379,9 @@ void UIKRetargetBatchOperation::ConvertAnimation(
 		// ensure we evaluate the source animation using the skeletal mesh proportions that were evaluated in the viewport
 		FAnimPoseEvaluationOptions EvaluationOptions = FAnimPoseEvaluationOptions();
 		EvaluationOptions.OptionalSkeletalMesh = SourceSkeleton.SkeletalMesh;
-		// ensure WYSIWYG with editor by ensuring the same root motion is evaluated during export as during the editor preview
-		EvaluationOptions.bExtractRootMotion = !TargetSequence->bForceRootLock;
-		EvaluationOptions.bIncorporateRootMotionIntoPose = !TargetSequence->bForceRootLock;
+		// ensure WYSIWYG with editor by ensuring the same root motion is applied to the pose, not to the component
+		EvaluationOptions.bExtractRootMotion = false;
+		EvaluationOptions.bIncorporateRootMotionIntoPose = true;
 
 		// reset the planting state
 		Processor->ResetPlanting();
@@ -444,6 +436,11 @@ void UIKRetargetBatchOperation::ConvertAnimation(
 			// store key data for each bone
 			for (int32 TargetBoneIndex=0; TargetBoneIndex<NumTargetBones; ++TargetBoneIndex)
 			{
+				if (Context.bExportOnlyAnimatedBones && !Processor->IsBoneRetargeted(TargetBoneNames[TargetBoneIndex], ERetargetSourceOrTarget::Target))
+				{
+					continue;
+				}
+				
 				const FTransform& LocalPose = TargetLocalPose[TargetBoneIndex];
 				
 				FRawAnimSequenceTrack& BoneTrack = BoneTracks[TargetBoneIndex];
@@ -460,6 +457,11 @@ void UIKRetargetBatchOperation::ConvertAnimation(
 		{
 			const FName& TargetBoneName = TargetBoneNames[TargetBoneIndex];
 
+			if (Context.bExportOnlyAnimatedBones && !Processor->IsBoneRetargeted(TargetBoneName, ERetargetSourceOrTarget::Target))
+			{
+				continue;
+			}
+
 			const FRawAnimSequenceTrack& RawTrack = BoneTracks[TargetBoneIndex];
 			TargetSeqController.AddBoneCurve(TargetBoneName, bShouldTransact);
 			TargetSeqController.SetBoneTrackKeys(TargetBoneName, RawTrack.PosKeys, RawTrack.RotKeys, RawTrack.ScaleKeys, bShouldTransact);
diff --git a/Engine/Plugins/Animation/IKRig/Source/IKRigEditor/Private/RetargetEditor/IKRetargetEditor.cpp b/Engine/Plugins/Animation/IKRig/Source/IKRigEditor/Private/RetargetEditor/IKRetargetEditor.cpp
--- a/Engine/Plugins/Animation/IKRig/Source/IKRigEditor/Private/RetargetEditor/IKRetargetEditor.cpp
+++ b/Engine/Plugins/Animation/IKRig/Source/IKRigEditor/Private/RetargetEditor/IKRetargetEditor.cpp
@@ -611,6 +611,10 @@ void FIKRetargetEditor::HandlePreviewSceneCreated(const TSharedRef<IPersonaPrevi
 	EditorController->SourceSkelMeshComponent = NewObject<UDebugSkelMeshComponent>(Actor);
 	EditorController->TargetSkelMeshComponent = NewObject<UDebugSkelMeshComponent>(Actor);
 
+	// do not process root motion, we need all motion in world space for retargeting to work correctly
+	EditorController->SourceSkelMeshComponent->SetProcessRootMotionMode(EProcessRootMotionMode::Ignore);
+	EditorController->TargetSkelMeshComponent->SetProcessRootMotionMode(EProcessRootMotionMode::Ignore);
+
 	// hide skeletons, we want to do custom rendering
 	EditorController->SourceSkelMeshComponent->SkeletonDrawMode = ESkeletonDrawMode::Hidden;
 	EditorController->TargetSkelMeshComponent->SkeletonDrawMode = ESkeletonDrawMode::Hidden;
diff --git a/Engine/Plugins/Animation/IKRig/Source/IKRigEditor/Private/RetargetEditor/SRetargetAnimAssetsWindow.cpp b/Engine/Plugins/Animation/IKRig/Source/IKRigEditor/Private/RetargetEditor/SRetargetAnimAssetsWindow.cpp
--- a/Engine/Plugins/Animation/IKRig/Source/IKRigEditor/Private/RetargetEditor/SRetargetAnimAssetsWindow.cpp
+++ b/Engine/Plugins/Animation/IKRig/Source/IKRigEditor/Private/RetargetEditor/SRetargetAnimAssetsWindow.cpp
@@ -456,7 +456,6 @@ FReply SBatchExportOptionsDialog::OnButtonClick(EAppReturnType::Type ButtonID)
 	BatchContext->bRetargetAndConnectReferencedAssets = ExportOptions->bRetargetAndConnectReferencedAssets;
 	BatchContext->bOverwriteExistingFiles = ExportOptions->bOverwriteExistingFiles;
 	BatchContext->bExportOnlyAnimatedBones = ExportOptions->bExportOnlyAnimatedBones;
-	BatchContext->RootLockMode = ExportOptions->RootLockMode;
 	
 	UserResponse = ButtonID;
 	RequestDestroyWindow();
@@ -476,6 +475,14 @@ void SRetargetPoseViewport::Construct(const FArguments& InArgs)
 	SourceComponent->VisibilityBasedAnimTickOption = EVisibilityBasedAnimTickOption::AlwaysTickPoseAndRefreshBones;
 	TargetComponent->VisibilityBasedAnimTickOption = EVisibilityBasedAnimTickOption::AlwaysTickPoseAndRefreshBones;
 
+	SourceComponent->SetProcessRootMotionMode(EProcessRootMotionMode::Ignore);
+	TargetComponent->SetProcessRootMotionMode(EProcessRootMotionMode::Ignore);
+
+	// setup SOURCE anim instance running a preview node that suppresses root motion
+	SourceAnimInstance = NewObject<UIKRetargetAnimInstance>(SourceComponent, TEXT("IKRetargetSourceAnimScriptInstance"));
+	SourceAnimInstance->SetRetargetMode(ERetargeterOutputMode::RunRetarget);
+	SourceComponent->PreviewInstance = SourceAnimInstance.Get();
+
 	// setup TARGET anim instance running a retargeter that copies input pose from the source component
 	TargetAnimInstance = NewObject<UIKRetargetAnimInstance>(TargetComponent, TEXT("IKRetargetTargetAnimScriptInstance"));
 	TargetAnimInstance->SetRetargetMode(ERetargeterOutputMode::RunRetarget);
@@ -513,6 +520,7 @@ void SRetargetPoseViewport::SetSkeletalMesh(USkeletalMesh* InSkeletalMesh, EReta
 void SRetargetPoseViewport::SetRetargetAsset(UIKRetargeter* RetargetAsset)
 {
 	// apply the IK retargeter and give a reference to the source component
+	SourceAnimInstance->ConfigureAnimInstance(ERetargetSourceOrTarget::Source, RetargetAsset, nullptr);
 	TargetAnimInstance->ConfigureAnimInstance(ERetargetSourceOrTarget::Target, RetargetAsset, SourceComponent);
 }
 
diff --git a/Engine/Plugins/Animation/IKRig/Source/IKRigEditor/Public/RetargetEditor/IKRetargetBatchOperation.h b/Engine/Plugins/Animation/IKRig/Source/IKRigEditor/Public/RetargetEditor/IKRetargetBatchOperation.h
--- a/Engine/Plugins/Animation/IKRig/Source/IKRigEditor/Public/RetargetEditor/IKRetargetBatchOperation.h
+++ b/Engine/Plugins/Animation/IKRig/Source/IKRigEditor/Public/RetargetEditor/IKRetargetBatchOperation.h
@@ -57,12 +57,6 @@ public:
 	// Will not produce keys on bones that are not animated, reducing size on disk of the resulting files.
 	bool bExportOnlyAnimatedBones = true;
 
-	// Either uses the "ForceRootLock" setting from the source animation or optionally forces the root to be locked or unlocked,
-	// regardless of whether the source animation has the root locked or not.
-	// This affects only the exported animation sequences, not the editor preview.
-	// This setting has no effect on runtime retargeting where root motion is copied from the source component.
-	ERetargetRootLockMode RootLockMode;
-
 	// Reset all data (called when window re-opened
 	void Reset()
 	{
diff --git a/Engine/Plugins/Animation/IKRig/Source/IKRigEditor/Public/RetargetEditor/SRetargetAnimAssetsWindow.h b/Engine/Plugins/Animation/IKRig/Source/IKRigEditor/Public/RetargetEditor/SRetargetAnimAssetsWindow.h
--- a/Engine/Plugins/Animation/IKRig/Source/IKRigEditor/Public/RetargetEditor/SRetargetAnimAssetsWindow.h
+++ b/Engine/Plugins/Animation/IKRig/Source/IKRigEditor/Public/RetargetEditor/SRetargetAnimAssetsWindow.h
@@ -92,13 +92,6 @@ public:
 	UPROPERTY(EditAnywhere, BlueprintReadWrite, Config, Category = "Animation")
 	bool bExportOnlyAnimatedBones = true;
 
-	// Either uses the "ForceRootLock" setting from the source animation or optionally forces the root to be locked or unlocked,
-	// regardless of whether the source animation has the root locked or not.
-	// This affects only the exported animation sequences, not the editor preview.
-	// This setting has no effect on runtime retargeting where root motion is copied from the source component.
-	UPROPERTY(EditAnywhere, BlueprintReadWrite, Config, Category = "Animation")
-	ERetargetRootLockMode RootLockMode;
-
 	static UBatchExportOptions* GetInstance();
 
 private:
@@ -159,6 +152,7 @@ private:
 	FPreviewScene PreviewScene;
 	TObjectPtr<UDebugSkelMeshComponent> SourceComponent;
 	TObjectPtr<UDebugSkelMeshComponent> TargetComponent;
+	TObjectPtr<UIKRetargetAnimInstance> SourceAnimInstance;
 	TObjectPtr<UIKRetargetAnimInstance> TargetAnimInstance;
 
 	// the data model
diff --git a/Engine/Source/Runtime/Engine/Private/Animation/AnimSingleNodeInstanceProxy.cpp b/Engine/Source/Runtime/Engine/Private/Animation/AnimSingleNodeInstanceProxy.cpp
--- a/Engine/Source/Runtime/Engine/Private/Animation/AnimSingleNodeInstanceProxy.cpp
+++ b/Engine/Source/Runtime/Engine/Private/Animation/AnimSingleNodeInstanceProxy.cpp
@@ -318,6 +318,7 @@ void FAnimNode_SingleNode::Evaluate_AnyThread(FPoseContext& Output)
 			FAnimExtractContext ExtractionContext(static_cast<double>(Proxy->CurrentTime), Sequence->bEnableRootMotion, Proxy->DeltaTimeRecord, Proxy->bLooping);
 #if WITH_EDITOR
 			ExtractionContext.bIgnoreRootLock = Proxy->bIgnoreRootLock;
+			ExtractionContext.bExtractRootMotion = Proxy->bIgnoreRootMotion ? false : ExtractionContext.bExtractRootMotion;
 #endif
 
 			if (Sequence->IsValidAdditive())
diff --git a/Engine/Source/Runtime/Engine/Public/Animation/AnimSingleNodeInstanceProxy.h b/Engine/Source/Runtime/Engine/Public/Animation/AnimSingleNodeInstanceProxy.h
--- a/Engine/Source/Runtime/Engine/Public/Animation/AnimSingleNodeInstanceProxy.h
+++ b/Engine/Source/Runtime/Engine/Public/Animation/AnimSingleNodeInstanceProxy.h
@@ -54,6 +54,7 @@ public:
 #if WITH_EDITOR
 		bCanProcessAdditiveAnimations = false;
 		bIgnoreRootLock = false;
+		bIgnoreRootMotion = false;
 #endif
 	}
 
@@ -77,6 +78,7 @@ public:
 #if WITH_EDITOR
 		bCanProcessAdditiveAnimations = false;
 		bIgnoreRootLock = false;
+		bIgnoreRootMotion = false;
 #endif
 	}
 
@@ -203,8 +205,11 @@ protected:
 	additives. But we need to be able to override this for editor preview windows */
 	bool bCanProcessAdditiveAnimations;
 
-	/** Allows editor previews to ignore root lock and show root motion in the editor.*/
+	/** Allows editor previews to ignore root lock and show root motion in the editor, regardless of setting on sequence*/
 	bool bIgnoreRootLock;
+
+	/** Allows editor previews to ignore root motion in the editor, regardless of setting on sequence*/
+	bool bIgnoreRootMotion;
 #endif
 
 	/** Pose Weight value that can override curve data. In the future, we'd like to have UCurveSet that can play by default**/
