Added feedback to batch retargeter when no retargeter specified.

diff --git a/Engine/Plugins/Animation/IKRig/Source/IKRig/Private/Retargeter/IKRetargeter.cpp b/Engine/Plugins/Animation/IKRig/Source/IKRig/Private/Retargeter/IKRetargeter.cpp
--- a/Engine/Plugins/Animation/IKRig/Source/IKRig/Private/Retargeter/IKRetargeter.cpp
+++ b/Engine/Plugins/Animation/IKRig/Source/IKRig/Private/Retargeter/IKRetargeter.cpp
@@ -65,7 +65,7 @@ UIKRigDefinition* UIKRetargeter::GetIKRigWriteable(ERetargetSourceOrTarget Sourc
 }
 
 #if WITH_EDITORONLY_DATA
-const USkeletalMesh* UIKRetargeter::GetPreviewMesh(ERetargetSourceOrTarget SourceOrTarget) const
+USkeletalMesh* UIKRetargeter::GetPreviewMesh(ERetargetSourceOrTarget SourceOrTarget) const
 {
 	if (!IsInGameThread())
 	{
diff --git a/Engine/Plugins/Animation/IKRig/Source/IKRig/Public/Retargeter/IKRetargeter.h b/Engine/Plugins/Animation/IKRig/Source/IKRig/Public/Retargeter/IKRetargeter.h
--- a/Engine/Plugins/Animation/IKRig/Source/IKRig/Public/Retargeter/IKRetargeter.h
+++ b/Engine/Plugins/Animation/IKRig/Source/IKRig/Public/Retargeter/IKRetargeter.h
@@ -224,7 +224,7 @@ public:
 	UIKRigDefinition* GetIKRigWriteable(ERetargetSourceOrTarget SourceOrTarget) const;
 	#if WITH_EDITORONLY_DATA
 	// Get read-only access to preview meshes
-	const USkeletalMesh* GetPreviewMesh(ERetargetSourceOrTarget SourceOrTarget) const;
+	USkeletalMesh* GetPreviewMesh(ERetargetSourceOrTarget SourceOrTarget) const;
 	#endif
 
 	// Get read-only access to the chain mapping 
diff --git a/Engine/Plugins/Animation/IKRig/Source/IKRigEditor/Private/RetargetEditor/IKRetargetBatchOperation.cpp b/Engine/Plugins/Animation/IKRig/Source/IKRigEditor/Private/RetargetEditor/IKRetargetBatchOperation.cpp
--- a/Engine/Plugins/Animation/IKRig/Source/IKRigEditor/Private/RetargetEditor/IKRetargetBatchOperation.cpp
+++ b/Engine/Plugins/Animation/IKRig/Source/IKRigEditor/Private/RetargetEditor/IKRetargetBatchOperation.cpp
@@ -78,7 +78,7 @@ int32 UIKRetargetBatchOperation::GenerateAssetLists(const FIKRetargetBatchOperat
 		}
 	}
 
-	if (Context.bRetargetAndConnectReferencedAssets)
+	if (Context.bIncludeReferencedAssets)
 	{
 		// Grab assets from the blueprint.
 		// Do this first as it can add complex assets to the retarget array which will need to be processed next.
@@ -103,7 +103,7 @@ void UIKRetargetBatchOperation::DuplicateRetargetAssets(
 	FScopedSlowTask& Progress)
 {
 	Progress.EnterProgressFrame(1.f, FText(LOCTEXT("DuplicatingBatchRetarget", "Duplicating animation assets...")));
-
+	
 	UPackage* DestinationPackage = Context.TargetMesh->GetOutermost();
 
 	TArray<UAnimationAsset*> AnimationAssetsToDuplicate = AnimationAssetsToRetarget;
@@ -125,8 +125,15 @@ void UIKRetargetBatchOperation::DuplicateRetargetAssets(
 
 		FString AssetName = Asset->GetName();
 		Progress.EnterProgressFrame(1.f, FText::Format(LOCTEXT("DuplicatingAnimation", "Duplicating animation: {0}"), FText::FromString(AssetName)));
+
+		// if user wants to export files to the same location as the source, then replace the FolderPath in the duplication rule
+		FNameDuplicationRule NameRule = Context.NameRule;
+		if (Context.bUseSourcePath)
+		{
+			NameRule.FolderPath = FPackageName::GetLongPackagePath(Asset->GetPathName()) / TEXT("");
+		}
 		
-		TMap<UAnimationAsset*, UAnimationAsset*> DuplicateMap = DuplicateAssets<UAnimationAsset>({Asset}, DestinationPackage, &Context.NameRule);
+		TMap<UAnimationAsset*, UAnimationAsset*> DuplicateMap = DuplicateAssets<UAnimationAsset>({Asset}, DestinationPackage, &NameRule);
 		DuplicatedAnimAssets.Append(DuplicateMap);
 	}
 	for (UAnimBlueprint* Asset : AnimBlueprintsToDuplicate)
@@ -138,8 +145,15 @@ void UIKRetargetBatchOperation::DuplicateRetargetAssets(
 
 		FString AssetName = Asset->GetName();
 		Progress.EnterProgressFrame(1.f, FText::Format(LOCTEXT("DuplicatingBlueprint", "Duplicating blueprint: {0}"), FText::FromString(AssetName)));
+
+		// if user wants to export files to the same location as the source, then replace the FolderPath in the duplication rule
+		FNameDuplicationRule NameRule = Context.NameRule;
+		if (Context.bUseSourcePath)
+		{
+			NameRule.FolderPath = FPackageName::GetLongPackagePath(Asset->GetPathName()) / TEXT("");
+		}
 		
-		TMap<UAnimBlueprint*, UAnimBlueprint*> DuplicateMap = DuplicateAssets<UAnimBlueprint>({Asset}, DestinationPackage, &Context.NameRule);
+		TMap<UAnimBlueprint*, UAnimBlueprint*> DuplicateMap = DuplicateAssets<UAnimBlueprint>({Asset}, DestinationPackage, &NameRule);
 		DuplicatedBlueprints.Append(DuplicateMap);
 	}
 
@@ -601,30 +615,39 @@ void UIKRetargetBatchOperation::OverwriteExistingAssets(const FIKRetargetBatchOp
 		UAnimationAsset* NewAsset = Pair.Value;
 		
 		// get desired name
-		FString PathName = Context.NameRule.FolderPath;
 		FString DesiredObjectName = Context.NameRule.Rename(OldAsset);
 		if (NewAsset->GetName() == DesiredObjectName)
 		{
 			// asset was not renamed due to collision with existing asset, so there's nothing to replace
 			continue;
 		}
+
+		// destination path
+		FString PathName = FPackageName::GetLongPackagePath(NewAsset->GetPathName());
 		FString DesiredPackageName = PathName + "/" + DesiredObjectName;
 		FString DesiredObjectPath = DesiredPackageName + "." + DesiredObjectName;
 		FAssetData AssetDataToReplace = AssetRegistryModule.Get().GetAssetByObjectPath(FSoftObjectPath(DesiredObjectPath));
-		const bool bHasDuplicateToReplace = AssetDataToReplace.IsValid() && AssetDataToReplace.GetAsset() == OldAsset;
+		const bool bHasDuplicateToReplace = AssetDataToReplace.IsValid() && AssetDataToReplace.GetAsset()->GetClass() == OldAsset->GetClass();
 		if (!bHasDuplicateToReplace)
 		{
 			// this could happen if the desired name was already in use by a different asset type
 			continue;
 		}
 
+		UObject* AssetToReplace = AssetDataToReplace.GetAsset();
+		if (AssetToReplace == OldAsset)
+		{
+			// we only replace previously retargeted animations, never the original
+			continue;
+		}
+
 		// reroute all references from old asset to new asset
-		TArray<UObject*> AssetsToReplace = {OldAsset};
+		TArray<UObject*> AssetsToReplace = {AssetToReplace};
 		ObjectTools::ForceReplaceReferences(NewAsset, AssetsToReplace);
-
-		// delete the old asset
-		ObjectTools::ForceDeleteObjects({OldAsset}, false /*bShowConfirmation*/);
 		
+		// delete the old asset
+		ObjectTools::ForceDeleteObjects({AssetToReplace}, false /*bShowConfirmation*/);
+			
 		// rename the new asset with the desired name
 		FString CurrentAssetPath = NewAsset->GetPathName();
 		TArray<FAssetRenameData> AssetsToRename = { FAssetRenameData(CurrentAssetPath, DesiredObjectPath) };
@@ -725,7 +748,7 @@ TArray<FAssetData> UIKRetargetBatchOperation::DuplicateAndRetarget(
 	const FString& Replace,
 	const FString& Prefix,
 	const FString& Suffix,
-	const bool bRemapReferencedAssets)
+	const bool bIncludeReferencedAssets)
 {
 	// fill the context with all the data needed to run a batch retarget
 	FIKRetargetBatchOperationContext Context;
@@ -743,7 +766,7 @@ TArray<FAssetData> UIKRetargetBatchOperation::DuplicateAndRetarget(
 	Context.NameRule.Suffix = Suffix;
 	Context.NameRule.ReplaceFrom = Search;
 	Context.NameRule.ReplaceTo = Replace;
-	Context.bRetargetAndConnectReferencedAssets = bRemapReferencedAssets;
+	Context.bIncludeReferencedAssets = bIncludeReferencedAssets;
 
 	// actually run the batch operation
 	UIKRetargetBatchOperation* BatchOperation = NewObject<UIKRetargetBatchOperation>();
diff --git a/Engine/Plugins/Animation/IKRig/Source/IKRigEditor/Private/RetargetEditor/SIKRetargetAssetBrowser.cpp b/Engine/Plugins/Animation/IKRig/Source/IKRigEditor/Private/RetargetEditor/SIKRetargetAssetBrowser.cpp
--- a/Engine/Plugins/Animation/IKRig/Source/IKRigEditor/Private/RetargetEditor/SIKRetargetAssetBrowser.cpp
+++ b/Engine/Plugins/Animation/IKRig/Source/IKRigEditor/Private/RetargetEditor/SIKRetargetAssetBrowser.cpp
@@ -176,7 +176,7 @@ FReply SIKRetargetAssetBrowser::OnExportButtonClicked()
 	BatchContext.SourceMesh = Controller->GetSkeletalMesh(ERetargetSourceOrTarget::Source);
 	BatchContext.TargetMesh = Controller->GetSkeletalMesh(ERetargetSourceOrTarget::Target);
 	BatchContext.IKRetargetAsset = Controller->AssetController->GetAsset();
-	BatchContext.bRetargetAndConnectReferencedAssets = false;
+	BatchContext.bIncludeReferencedAssets = false;
 
 	// add selected assets to dup/retarget
 	TArray<FAssetData> SelectedAssets = GetCurrentSelectionDelegate.Execute();
@@ -300,7 +300,7 @@ void SBatchExportDialog::Construct(const FArguments& InArgs)
 
 	FPathPickerConfig PathPickerConfig;
 	PathPickerConfig.DefaultPath = AssetPath.ToString();
-	PathPickerConfig.OnPathSelected = FOnPathSelected::CreateSP(this, &SBatchExportDialog::OnPathChange);
+	PathPickerConfig.OnPathSelected = FOnPathSelected::CreateLambda([this](const FString& NewPath){AssetPath = FText::FromString(NewPath);});
 	PathPickerConfig.bAddDefaultPath = true;
 
 	FContentBrowserModule& ContentBrowserModule = FModuleManager::LoadModuleChecked<FContentBrowserModule>("ContentBrowser");
@@ -550,11 +550,6 @@ void SBatchExportDialog::UpdateExampleText()
 	ExampleText = FText::FromString(FString::Printf(TEXT("%s\n%s"), *ReplaceFrom, *ReplaceTo));
 }
 
-void SBatchExportDialog::OnPathChange(const FString& NewPath)
-{
-	AssetPath = FText::FromString(NewPath);
-}
-
 FReply SBatchExportDialog::OnButtonClick(EAppReturnType::Type ButtonID)
 {
 	UserResponse = ButtonID;
diff --git a/Engine/Plugins/Animation/IKRig/Source/IKRigEditor/Private/RetargetEditor/SRetargetAnimAssetsWindow.cpp b/Engine/Plugins/Animation/IKRig/Source/IKRigEditor/Private/RetargetEditor/SRetargetAnimAssetsWindow.cpp
--- a/Engine/Plugins/Animation/IKRig/Source/IKRigEditor/Private/RetargetEditor/SRetargetAnimAssetsWindow.cpp
+++ b/Engine/Plugins/Animation/IKRig/Source/IKRigEditor/Private/RetargetEditor/SRetargetAnimAssetsWindow.cpp
@@ -51,7 +51,7 @@ void SBatchExportPathDialog::Construct(const FArguments& InArgs)
 	const FContentBrowserModule& ContentBrowserModule = FModuleManager::LoadModuleChecked<FContentBrowserModule>("ContentBrowser");
 	FPathPickerConfig PathPickerConfig;
 	PathPickerConfig.DefaultPath = BatchContext->NameRule.FolderPath;
-	PathPickerConfig.OnPathSelected = FOnPathSelected::CreateSP(this, &SBatchExportPathDialog::OnPathChange);
+	PathPickerConfig.OnPathSelected = FOnPathSelected::CreateLambda([this](const FString& NewPath){BatchContext->NameRule.FolderPath = NewPath;});
 	PathPickerConfig.bAddDefaultPath = true;
 
 	// adjust UI based on if we're exporting retarget assets or animation assets
@@ -67,19 +67,20 @@ void SBatchExportPathDialog::Construct(const FArguments& InArgs)
 	.ClientSize(FVector2D(300, WindowHeight))
 	[
 		SNew(SVerticalBox)
-
+		
 		+ SVerticalBox::Slot()
 		.Padding(2)
 		[
 			SNew(SVerticalBox)
-			
+			.IsEnabled_Lambda([this](){return !BatchContext->bUseSourcePath;})
+
 			+SVerticalBox::Slot()
 			.FillHeight(1)
 			.Padding(3)
 			[
 				ContentBrowserModule.Get().CreatePathPicker(PathPickerConfig)
 			]
-
+			
 			+SVerticalBox::Slot()
 			.AutoHeight()
 			.Padding(2, 3)
@@ -98,10 +99,47 @@ void SBatchExportPathDialog::Construct(const FArguments& InArgs)
 				.FillWidth(1)
 				.HAlign(HAlign_Right)
 				[
-					SNew(STextBlock).Text(this, &SBatchExportPathDialog::GetFolderPath)
+					SNew(STextBlock).Text_Lambda([this]()
+					{
+						return FText::FromString(BatchContext->NameRule.FolderPath);
+					})
 				]
 			]
 		]
+
+		+SVerticalBox::Slot()
+		.AutoHeight()
+		.HAlign(HAlign_Right)
+		.Padding(3)
+		[
+			SNew(SHorizontalBox)
+			.Visibility_Lambda([this]{ return bExportingRetargetAssets ? EVisibility::Collapsed : EVisibility::Visible;})
+
+			+ SHorizontalBox::Slot()
+			.AutoWidth()
+			.HAlign(HAlign_Right)
+			.Padding(0.0f, 0.0f, 12.f, 0.f)
+			[
+				SNew(STextBlock)
+				.Text(LOCTEXT("UseSourcePath_Label", "Use Source Path: "))
+				.ToolTipText(LOCTEXT("UseSourcePathTooltip", "Places new animation assets in the same location as the source asset."))
+			]
+
+			+ SHorizontalBox::Slot()
+			.AutoWidth()
+			.HAlign(HAlign_Right)
+			[
+				SNew(SCheckBox)
+				.IsChecked_Lambda([this]()
+				{
+					return BatchContext->bUseSourcePath ? ECheckBoxState::Checked : ECheckBoxState::Unchecked;
+				})
+				.OnCheckStateChanged_Lambda([this](ECheckBoxState InCheckBoxState)
+				{
+					BatchContext->bUseSourcePath = InCheckBoxState == ECheckBoxState::Checked;
+				})
+			]
+		]
 		
 		+SVerticalBox::Slot()
 		.AutoHeight()
@@ -118,7 +156,7 @@ void SBatchExportPathDialog::Construct(const FArguments& InArgs)
 				[
 					SNew(STextBlock)
 					.Text(LOCTEXT("DuplicateAndRetarget_RenameLabel", "Rename New Assets"))
-					.Font(FCoreStyle::GetDefaultFontStyle("Regular", 14))
+					.Font(FAppStyle::Get().GetFontStyle("NormalFontBold"))
 				]
 
 				+SVerticalBox::Slot()
@@ -137,9 +175,16 @@ void SBatchExportPathDialog::Construct(const FArguments& InArgs)
 					+SHorizontalBox::Slot()
 					[
 						SNew(SEditableTextBox)
-							.Text(this, &SBatchExportPathDialog::GetPrefixName)
+							.Text_Lambda([this]()
+							{
+								return FText::FromString(BatchContext->NameRule.Prefix);
+							})
+							.OnTextChanged_Lambda([this](const FText& InText)
+							{
+								BatchContext->NameRule.Prefix = ConvertToCleanString(InText);
+								UpdateExampleText();	
+							})
 							.MinDesiredWidth(100)
-							.OnTextChanged(this, &SBatchExportPathDialog::SetPrefixName)
 							.IsReadOnly(false)
 							.RevertTextOnEscape(true)
 					]
@@ -161,9 +206,16 @@ void SBatchExportPathDialog::Construct(const FArguments& InArgs)
 					+SHorizontalBox::Slot()
 					[
 						SNew(SEditableTextBox)
-							.Text(this, &SBatchExportPathDialog::GetSuffixName)
+							.Text_Lambda([this]()
+							{
+								return FText::FromString(BatchContext->NameRule.Suffix);
+							})
+							.OnTextChanged_Lambda([this](const FText& InText)
+							{
+								BatchContext->NameRule.Suffix = ConvertToCleanString(InText);
+								UpdateExampleText();	
+							})
 							.MinDesiredWidth(100)
-							.OnTextChanged(this, &SBatchExportPathDialog::SetSuffixName)
 							.IsReadOnly(false)
 							.RevertTextOnEscape(true)
 					]
@@ -185,9 +237,16 @@ void SBatchExportPathDialog::Construct(const FArguments& InArgs)
 					+SHorizontalBox::Slot()
 					[
 						SNew(SEditableTextBox)
-							.Text(this, &SBatchExportPathDialog::GetReplaceFrom)
+							.Text_Lambda([this]()
+							{
+								return FText::FromString(BatchContext->NameRule.ReplaceFrom);
+							})
+							.OnTextChanged_Lambda([this](const FText& InText)
+							{
+								BatchContext->NameRule.ReplaceFrom = ConvertToCleanString(InText);
+								UpdateExampleText();	
+							})
 							.MinDesiredWidth(100)
-							.OnTextChanged(this, &SBatchExportPathDialog::SetReplaceFrom)
 							.IsReadOnly(false)
 							.RevertTextOnEscape(true)
 					]
@@ -209,9 +268,16 @@ void SBatchExportPathDialog::Construct(const FArguments& InArgs)
 					+SHorizontalBox::Slot()
 					[
 						SNew(SEditableTextBox)
-							.Text(this, &SBatchExportPathDialog::GetReplaceTo)
+							.Text_Lambda([this]()
+							{
+								return FText::FromString(BatchContext->NameRule.ReplaceTo);
+							})
+							.OnTextChanged_Lambda([this](const FText& InText)
+							{
+								BatchContext->NameRule.ReplaceTo = ConvertToCleanString(InText);
+								UpdateExampleText();	
+							})
 							.MinDesiredWidth(100)
-							.OnTextChanged(this, &SBatchExportPathDialog::SetReplaceTo)
 							.IsReadOnly(false)
 							.RevertTextOnEscape(true)
 					]
@@ -226,7 +292,7 @@ void SBatchExportPathDialog::Construct(const FArguments& InArgs)
 					.Padding(5, 5)
 					[
 						SNew(STextBlock)
-						.Text(this,  &SBatchExportPathDialog::GetExampleText)
+						.Text_Lambda([this](){return ExampleText; })
 						.Font(FAppStyle::GetFontStyle("Persona.RetargetManager.ItalicFont"))
 					]
 				]
@@ -267,11 +333,6 @@ void SBatchExportPathDialog::Construct(const FArguments& InArgs)
 	UpdateExampleText();
 }
 
-void SBatchExportPathDialog::OnPathChange(const FString& NewPath)
-{
-	BatchContext->NameRule.FolderPath = NewPath;
-}
-
 FReply SBatchExportPathDialog::OnButtonClick(EAppReturnType::Type ButtonID)
 {
 	UserResponse = ButtonID;
@@ -285,55 +346,6 @@ EAppReturnType::Type SBatchExportPathDialog::ShowModal()
 	return UserResponse;
 }
 
-FText SBatchExportPathDialog::GetPrefixName() const
-{
-	return FText::FromString(BatchContext->NameRule.Prefix);
-}
-
-void SBatchExportPathDialog::SetPrefixName(const FText &InText)
-{
-	BatchContext->NameRule.Prefix = ConvertToCleanString(InText);
-	UpdateExampleText();
-}
-
-FText SBatchExportPathDialog::GetSuffixName() const
-{
-	return FText::FromString(BatchContext->NameRule.Suffix);
-}
-
-void SBatchExportPathDialog::SetSuffixName(const FText &InText)
-{
-	BatchContext->NameRule.Suffix = ConvertToCleanString(InText);
-	UpdateExampleText();
-}
-
-FText SBatchExportPathDialog::GetReplaceFrom() const
-{
-	return FText::FromString(BatchContext->NameRule.ReplaceFrom);
-}
-
-void SBatchExportPathDialog::SetReplaceFrom(const FText &InText)
-{
-	BatchContext->NameRule.ReplaceFrom = ConvertToCleanString(InText);
-	UpdateExampleText();
-}
-
-FText SBatchExportPathDialog::GetReplaceTo() const
-{
-	return FText::FromString(BatchContext->NameRule.ReplaceTo);
-}
-
-void SBatchExportPathDialog::SetReplaceTo(const FText &InText)
-{
-	BatchContext->NameRule.ReplaceTo = ConvertToCleanString(InText);
-	UpdateExampleText();
-}
-
-FText SBatchExportPathDialog::GetExampleText() const
-{
-	return ExampleText;
-}
-
 void SBatchExportPathDialog::UpdateExampleText()
 {
 	const FString ReplaceFrom = FString::Printf(TEXT("Old Name : ***%s***"), *BatchContext->NameRule.ReplaceFrom);
@@ -453,7 +465,7 @@ FReply SBatchExportOptionsDialog::OnButtonClick(EAppReturnType::Type ButtonID)
 {
 	// update batch context with the user specified options
 	const UBatchExportOptions* ExportOptions = UBatchExportOptions::GetInstance();
-	BatchContext->bRetargetAndConnectReferencedAssets = ExportOptions->bRetargetAndConnectReferencedAssets;
+	BatchContext->bIncludeReferencedAssets = ExportOptions->bIncludeReferencedAssets;
 	BatchContext->bOverwriteExistingFiles = ExportOptions->bOverwriteExistingFiles;
 	BatchContext->bExportOnlyAnimatedBones = ExportOptions->bExportOnlyAnimatedBones;
 	
@@ -945,25 +957,22 @@ void SRetargetAnimAssetsWindow::OnFinishedChangingSelectionProperties(const FPro
 		return;
 	}
 
-	LogView->ClearLog();
-	
-	if (PropertyChangedEvent.Property->GetName() == "SourceSkeletalMesh")
-	{
-		SetSkeletalMesh(Settings->TargetSkeletalMesh, ERetargetSourceOrTarget::Target);
-		SetSkeletalMesh(Settings->SourceSkeletalMesh, ERetargetSourceOrTarget::Source);
-	}
-
-	if (PropertyChangedEvent.Property->GetName() == "TargetSkeletalMesh")
+	const bool ChangedMesh = PropertyChangedEvent.Property->GetName() == "SourceSkeletalMesh" || PropertyChangedEvent.Property->GetName() == "TargetSkeletalMesh";
+	if (ChangedMesh)
 	{
-		SetSkeletalMesh(Settings->SourceSkeletalMesh, ERetargetSourceOrTarget::Source);
-		SetSkeletalMesh(Settings->TargetSkeletalMesh, ERetargetSourceOrTarget::Target);
+		SetAssets(Settings->SourceSkeletalMesh, Settings->TargetSkeletalMesh, Settings->RetargetAsset);
 	}
 
-	const bool bEditedAutoGenCheckbox = PropertyChangedEvent.Property->GetName() == "bAutoGenerateRetargeter";
-	const bool bEditedRetargeter = PropertyChangedEvent.Property->GetName() == "bAutoGenerateRetargeter";
-	if (bEditedAutoGenCheckbox || bEditedRetargeter)
+	const bool ChangedRetargeter = PropertyChangedEvent.Property->GetName() == "RetargetAsset" || PropertyChangedEvent.Property->GetName() == "bAutoGenerateRetargeter";
+	if (ChangedRetargeter)
 	{
-		SetRetargetAsset(Settings->RetargetAsset);
+		USkeletalMesh* TargetMesh = Settings->TargetSkeletalMesh;
+		if (Settings->RetargetAsset && !Settings->TargetSkeletalMesh)
+		{
+			TargetMesh = Settings->RetargetAsset->GetPreviewMesh(ERetargetSourceOrTarget::Target);
+		}
+		
+		SetAssets(Settings->SourceSkeletalMesh, TargetMesh, Settings->RetargetAsset);
 	}
 }
 
@@ -1040,12 +1049,6 @@ FReply SRetargetAnimAssetsWindow::OnExportRetargetAssets()
 		AssetToSave->Rename(*AssetName, Package);
 		FAssetRegistryModule::AssetCreated(AssetToSave);
 		Package->MarkPackageDirty();
-	
-		// save the asset
-		FSavePackageArgs SaveArgs;
-		SaveArgs.TopLevelFlags = RF_Standalone;
-		const FString FullPackagePath = FPackageName::LongPackageNameToFilename(UniquePackageName, FPackageName::GetAssetPackageExtension());
-		UPackage::Save(Package, AssetToSave, *FullPackagePath, SaveArgs);
 
 		return AssetToSave;
 	};
@@ -1195,85 +1198,103 @@ void SRetargetAnimAssetsWindow::ShowWindow(TArray<UObject*> InSelectedAssets)
             }
 			
 			DialogWidget->Settings->SourceSkeletalMesh = Mesh;
-			DialogWidget->SetSkeletalMesh(Mesh, ERetargetSourceOrTarget::Source);
-			DialogWidget->SetSkeletalMesh(nullptr, ERetargetSourceOrTarget::Target);
+			DialogWidget->SetAssets(Mesh, nullptr, nullptr);
 		}
 	}
 	
 	FSlateApplication::Get().AddWindow(Window.ToSharedRef());
 }
 
-void SRetargetAnimAssetsWindow::SetSkeletalMesh(USkeletalMesh* Mesh, ERetargetSourceOrTarget SourceOrTarget)
+void SRetargetAnimAssetsWindow::SetAssets(
+	USkeletalMesh* SourceMesh,
+	USkeletalMesh* TargetMesh,
+	UIKRetargeter* Retargeter)
 {
-	if (SourceOrTarget == ERetargetSourceOrTarget::Source)
+	LogView->ClearLog();
+	
+	const bool ReplacedSourceMesh = BatchContext.SourceMesh != SourceMesh;
+
+	Settings->SourceSkeletalMesh = SourceMesh;
+	Settings->TargetSkeletalMesh = TargetMesh;
+	Settings->RetargetAsset = Retargeter;
+	
+	BatchContext.SourceMesh = SourceMesh;
+	BatchContext.TargetMesh = TargetMesh;
+	BatchContext.IKRetargetAsset = Retargeter;
+
+	// auto generate procedural assets
+	if (Settings->bAutoGenerateRetargeter)
+	{
+		// update procedurally generated IK Rig and retargeter (regenerates retarget pose with new mesh)
+		ProceduralAssets.AutoGenerateIKRigAsset(SourceMesh, ERetargetSourceOrTarget::Source);
+		ProceduralAssets.AutoGenerateIKRigAsset(TargetMesh, ERetargetSourceOrTarget::Target);
+		ProceduralAssets.AutoGenerateIKRetargetAsset();
+		BatchContext.IKRetargetAsset = ProceduralAssets.Retargeter;
+	}
+	
+	// update viewport world to show new meshes
+	Viewport->SetSkeletalMesh(SourceMesh, ERetargetSourceOrTarget::Source);
+	Viewport->SetSkeletalMesh(TargetMesh, ERetargetSourceOrTarget::Target);
+	Viewport->SetRetargetAsset(BatchContext.IKRetargetAsset);
+
+	ShowAssetWarnings();
+	
+	if (ReplacedSourceMesh)
 	{
-		BatchContext.SourceMesh = Mesh;
-		Settings->SourceSkeletalMesh = Mesh;
 		AssetBrowser.Get()->RefreshView();
 	}
-	else
+}
+
+void SRetargetAnimAssetsWindow::ShowAssetWarnings()
+{
+	// missing source mesh
+	if (!Settings->SourceSkeletalMesh)
 	{
-		BatchContext.TargetMesh = Mesh;
-		Settings->TargetSkeletalMesh = Mesh;
+		Log.LogError(LOCTEXT( "MissingMeshError", "No source mesh assigned."));
 	}
 
-	// update procedurally generated IK Rig and retargeter (regenerates retarget pose with new mesh)
-	ProceduralAssets.AutoGenerateIKRigAsset(Mesh, SourceOrTarget);
-	ProceduralAssets.AutoGenerateIKRetargetAsset();
-	// update viewport world to show new mesh
-	Viewport->SetSkeletalMesh(Mesh, SourceOrTarget);
+	// missing target mesh
+	if (!Settings->TargetSkeletalMesh)
+	{
+		Log.LogError(LOCTEXT( "MissingMeshError", "No target mesh assigned."));
+	}
 
-	// report relevant results
-	const FText SourceOrTargetLabel = SourceOrTarget == ERetargetSourceOrTarget::Source ? FText::FromString("source") : FText::FromString("target");
-	if (Mesh)
+	// auto retarget results
+	if (Settings->bAutoGenerateRetargeter)
 	{
-		const FAutoCharacterizeResults& Results = SourceOrTarget == ERetargetSourceOrTarget::Source ? ProceduralAssets.SourceCharacterizationResults : ProceduralAssets.TargetCharacterizationResults;
-		if (Results.bUsedTemplate)
+		// source auto-characterize results
+		if (Settings->SourceSkeletalMesh)
 		{
-			const FText TemplateName = FText::FromString(Results.BestTemplateName.ToString());
-			Log.LogInfo(FText::Format(LOCTEXT( "FoundTemplateInfo", "Using '{0}' template for {1} mesh."), TemplateName, SourceOrTargetLabel));
+			if (ProceduralAssets.SourceCharacterizationResults.bUsedTemplate)
+			{
+				Log.LogInfo(FText::Format(LOCTEXT("UsingSourceTemplate", "Using {0} template for source skeleton."), FText::FromName(ProceduralAssets.SourceCharacterizationResults.BestTemplateName)));
+			}
+			else
+			{
+				Log.LogError(LOCTEXT("NoSourceTemplate", "Trying to auto-generate a retargeter but no template was compatible with the source mesh."));
+			}
 		}
-		else
+		
+		// target auto-characterize results
+		if (Settings->TargetSkeletalMesh)
 		{
-			Log.LogError(FText::Format(LOCTEXT( "MissingTemplateError", "No template found for {0} mesh."), SourceOrTargetLabel));
+			if (ProceduralAssets.TargetCharacterizationResults.bUsedTemplate)
+			{
+				Log.LogInfo(FText::Format(LOCTEXT("UsingTargetTemplate", "Using {0} template for target skeleton."), FText::FromName(ProceduralAssets.TargetCharacterizationResults.BestTemplateName)));
+			}
+			else
+			{
+				Log.LogError(LOCTEXT("NoTargetTemplate", "Trying to auto-generate a retargeter but no template was compatible with the target mesh."));
+			}
 		}
 	}
 	else
 	{
-		Log.LogError(FText::Format(LOCTEXT( "MissingMeshError", "No {0} mesh assigned."), SourceOrTargetLabel));
-	}
-}
-
-void SRetargetAnimAssetsWindow::SetRetargetAsset(UIKRetargeter* RetargeterToUse)
-{
-	if (Settings->bAutoGenerateRetargeter)
-	{
-		// if user did NOT supply a custom retarget asset, procedurally generate one
-		ProceduralAssets.AutoGenerateIKRetargetAsset();
-		RetargeterToUse = ProceduralAssets.Retargeter;
-	}
-	else if (Settings->RetargetAsset)
-	{
-		// if user assigned a custom retarget asset, set the skeletal meshes based on that asset
-		// BUT, only do this if the user hasn't already specified their own skeletal mesh
-		const UIKRetargeterController* Controller = UIKRetargeterController::GetController(Settings->RetargetAsset);
-		if (!Settings->SourceSkeletalMesh)
-		{
-			USkeletalMesh* SourceMesh = Controller->GetPreviewMesh(ERetargetSourceOrTarget::Source);
-			SetSkeletalMesh(SourceMesh, ERetargetSourceOrTarget::Source);
-		}
-		if (!Settings->TargetSkeletalMesh)
+		if (!Settings->RetargetAsset && !Settings->bAutoGenerateRetargeter)
 		{
-			USkeletalMesh* TargetMesh = Controller->GetPreviewMesh(ERetargetSourceOrTarget::Target);
-			SetSkeletalMesh(TargetMesh, ERetargetSourceOrTarget::Target);
+			Log.LogError(LOCTEXT("NoRetargeter", "Not using auto-generated retargeter and no IK Retargeter asset was provided."));
 		}
 	}
-	
-	// store the asset to use in the context
-	BatchContext.IKRetargetAsset = RetargeterToUse;
-	
-	// update retargeter in viewport
-	Viewport->SetRetargetAsset(RetargeterToUse);
 }
 
 #undef LOCTEXT_NAMESPACE
diff --git a/Engine/Plugins/Animation/IKRig/Source/IKRigEditor/Public/RetargetEditor/IKRetargetBatchOperation.h b/Engine/Plugins/Animation/IKRig/Source/IKRigEditor/Public/RetargetEditor/IKRetargetBatchOperation.h
--- a/Engine/Plugins/Animation/IKRig/Source/IKRigEditor/Public/RetargetEditor/IKRetargetBatchOperation.h
+++ b/Engine/Plugins/Animation/IKRig/Source/IKRigEditor/Public/RetargetEditor/IKRetargetBatchOperation.h
@@ -47,12 +47,16 @@ public:
 	// Rename rules for duplicated assets
 	EditorAnimUtils::FNameDuplicationRule NameRule;
 
+	// Any files with the same name will be overwritten instead of creating a new file with a numeric suffix.
+	// This is useful when iterating on a batch process.
+	bool bUseSourcePath = false;
+	
 	// Any files with the same name will be overwritten instead of creating a new file with a numeric suffix.
 	// This is useful when iterating on a batch process.
 	bool bOverwriteExistingFiles = false;
 
 	// Duplicates and retargets any animation assets referenced by the input assets. For example, sequences in an animation blueprint or blendspace.
-	bool bRetargetAndConnectReferencedAssets = true;
+	bool bIncludeReferencedAssets = true;
 
 	// Will not produce keys on bones that are not animated, reducing size on disk of the resulting files.
 	bool bExportOnlyAnimatedBones = true;
@@ -63,7 +67,7 @@ public:
 		SourceMesh = nullptr;
 		TargetMesh = nullptr;
 		IKRetargetAsset = nullptr;
-		bRetargetAndConnectReferencedAssets = true;
+		bIncludeReferencedAssets = true;
 		NameRule.Prefix = "";
 		NameRule.Suffix = "";
 		NameRule.ReplaceFrom = "";
@@ -111,7 +115,7 @@ public:
 		const FString& Replace = "",
 		const FString& Prefix = "",
 		const FString& Suffix = "",
-		const bool bRemapReferencedAssets=true);
+		const bool bIncludeReferencedAssets=true);
 	
 	// Actually run the process to duplicate and retarget the assets for the given context
 	void RunRetarget(FIKRetargetBatchOperationContext& Context);
diff --git a/Engine/Plugins/Animation/IKRig/Source/IKRigEditor/Public/RetargetEditor/SIKRetargetAssetBrowser.h b/Engine/Plugins/Animation/IKRig/Source/IKRigEditor/Public/RetargetEditor/SIKRetargetAssetBrowser.h
--- a/Engine/Plugins/Animation/IKRig/Source/IKRigEditor/Public/RetargetEditor/SIKRetargetAssetBrowser.h
+++ b/Engine/Plugins/Animation/IKRig/Source/IKRigEditor/Public/RetargetEditor/SIKRetargetAssetBrowser.h
@@ -77,8 +77,7 @@ protected:
 	/** The rename rule sample text */
 	FText ExampleText;
 	void UpdateExampleText();
-
-	void OnPathChange(const FString& NewPath);
+	
 	FReply OnButtonClick(EAppReturnType::Type ButtonID);
 
 	EAppReturnType::Type UserResponse;
diff --git a/Engine/Plugins/Animation/IKRig/Source/IKRigEditor/Public/RetargetEditor/SRetargetAnimAssetsWindow.h b/Engine/Plugins/Animation/IKRig/Source/IKRigEditor/Public/RetargetEditor/SRetargetAnimAssetsWindow.h
--- a/Engine/Plugins/Animation/IKRig/Source/IKRigEditor/Public/RetargetEditor/SRetargetAnimAssetsWindow.h
+++ b/Engine/Plugins/Animation/IKRig/Source/IKRigEditor/Public/RetargetEditor/SRetargetAnimAssetsWindow.h
@@ -39,23 +39,11 @@ public:
 
 private:
 	
-	void OnPathChange(const FString& NewPath);
 	FReply OnButtonClick(EAppReturnType::Type ButtonID);
 	
-	// modifying Rename Prefix
-	FText GetPrefixName() const;
-	void SetPrefixName(const FText &InText);
-	// modifying Rename Suffix
-	FText GetSuffixName() const;
-	void SetSuffixName(const FText &InText);
-	// modifying Search/Replace text
-	FText GetReplaceFrom() const;
-	void SetReplaceFrom(const FText &InText);
-	FText GetReplaceTo() const;
-	void SetReplaceTo(const FText &InText);
 	// example rename text
-	FText GetExampleText() const;
 	void UpdateExampleText();
+	
 	// modify folder output path
 	FText GetFolderPath() const;
 
@@ -86,7 +74,7 @@ public:
 
 	// Duplicates and retargets any animation assets referenced by the input assets. For example, sequences in an animation blueprint or blendspace.
 	UPROPERTY(EditAnywhere, BlueprintReadWrite, Config, Category = "File")
-	bool bRetargetAndConnectReferencedAssets = true;
+	bool bIncludeReferencedAssets = true;
 
 	// Will not produce keys on bones that are not animated, reducing size on disk of the resulting files.
 	UPROPERTY(EditAnywhere, BlueprintReadWrite, Config, Category = "Animation")
@@ -295,8 +283,8 @@ public:
 
 private:
 
-	void SetRetargetAsset(UIKRetargeter* RetargeterToUse);
-	void SetSkeletalMesh(USkeletalMesh* Mesh, ERetargetSourceOrTarget SourceOrTarget);
+	void SetAssets(USkeletalMesh* SourceMesh, USkeletalMesh* TargetMesh, UIKRetargeter* Retargeter);
+	void ShowAssetWarnings();
 
 	void OnFinishedChangingSelectionProperties(const FPropertyChangedEvent& PropertyChangedEvent);
 
