Removing code that accidentally made it into CL 31528789

diff --git a/Engine/Plugins/Animation/IKRig/Source/IKRigEditor/Private/RetargetEditor/IKRetargetBatchOperation.cpp b/Engine/Plugins/Animation/IKRig/Source/IKRigEditor/Private/RetargetEditor/IKRetargetBatchOperation.cpp
--- a/Engine/Plugins/Animation/IKRig/Source/IKRigEditor/Private/RetargetEditor/IKRetargetBatchOperation.cpp
+++ b/Engine/Plugins/Animation/IKRig/Source/IKRigEditor/Private/RetargetEditor/IKRetargetBatchOperation.cpp
@@ -450,11 +450,6 @@ void UIKRetargetBatchOperation::ConvertAnimation(
 			// store key data for each bone
 			for (int32 TargetBoneIndex=0; TargetBoneIndex<NumTargetBones; ++TargetBoneIndex)
 			{
-				if (Context.bExportOnlyAnimatedBones && !Processor->IsBoneRetargeted(TargetBoneNames[TargetBoneIndex], ERetargetSourceOrTarget::Target))
-				{
-					continue;
-				}
-				
 				const FTransform& LocalPose = TargetLocalPose[TargetBoneIndex];
 				
 				FRawAnimSequenceTrack& BoneTrack = BoneTracks[TargetBoneIndex];
@@ -471,11 +466,6 @@ void UIKRetargetBatchOperation::ConvertAnimation(
 		{
 			const FName& TargetBoneName = TargetBoneNames[TargetBoneIndex];
 
-			if (Context.bExportOnlyAnimatedBones && !Processor->IsBoneRetargeted(TargetBoneName, ERetargetSourceOrTarget::Target))
-			{
-				continue;
-			}
-
 			const FRawAnimSequenceTrack& RawTrack = BoneTracks[TargetBoneIndex];
 			TargetSeqController.AddBoneCurve(TargetBoneName, bShouldTransact);
 			TargetSeqController.SetBoneTrackKeys(TargetBoneName, RawTrack.PosKeys, RawTrack.RotKeys, RawTrack.ScaleKeys, bShouldTransact);
diff --git a/Engine/Plugins/Animation/IKRig/Source/IKRigEditor/Private/RetargetEditor/SRetargetAnimAssetsWindow.cpp b/Engine/Plugins/Animation/IKRig/Source/IKRigEditor/Private/RetargetEditor/SRetargetAnimAssetsWindow.cpp
--- a/Engine/Plugins/Animation/IKRig/Source/IKRigEditor/Private/RetargetEditor/SRetargetAnimAssetsWindow.cpp
+++ b/Engine/Plugins/Animation/IKRig/Source/IKRigEditor/Private/RetargetEditor/SRetargetAnimAssetsWindow.cpp
@@ -467,7 +467,7 @@ FReply SBatchExportOptionsDialog::OnButtonClick(EAppReturnType::Type ButtonID)
 	const UBatchExportOptions* ExportOptions = UBatchExportOptions::GetInstance();
 	BatchContext->bIncludeReferencedAssets = ExportOptions->bIncludeReferencedAssets;
 	BatchContext->bOverwriteExistingFiles = ExportOptions->bOverwriteExistingFiles;
-	BatchContext->bExportOnlyAnimatedBones = ExportOptions->bExportOnlyAnimatedBones;
+	//BatchContext->bExportOnlyAnimatedBones = ExportOptions->bExportOnlyAnimatedBones;
 	
 	UserResponse = ButtonID;
 	RequestDestroyWindow();
diff --git a/Engine/Plugins/Animation/IKRig/Source/IKRigEditor/Public/RetargetEditor/SRetargetAnimAssetsWindow.h b/Engine/Plugins/Animation/IKRig/Source/IKRigEditor/Public/RetargetEditor/SRetargetAnimAssetsWindow.h
--- a/Engine/Plugins/Animation/IKRig/Source/IKRigEditor/Public/RetargetEditor/SRetargetAnimAssetsWindow.h
+++ b/Engine/Plugins/Animation/IKRig/Source/IKRigEditor/Public/RetargetEditor/SRetargetAnimAssetsWindow.h
@@ -76,9 +76,10 @@ public:
 	UPROPERTY(EditAnywhere, BlueprintReadWrite, Config, Category = "File")
 	bool bIncludeReferencedAssets = true;
 
+	// TODO - Kiaran Feb 2024 - Naively leaving out non-retargeted keys results in flipped skeletons, needs work.
 	// Will not produce keys on bones that are not animated, reducing size on disk of the resulting files.
-	UPROPERTY(EditAnywhere, BlueprintReadWrite, Config, Category = "Animation")
-	bool bExportOnlyAnimatedBones = true;
+	//UPROPERTY(EditAnywhere, BlueprintReadWrite, Config, Category = "Animation")
+	//bool bExportOnlyAnimatedBones = true;
 
 	static UBatchExportOptions* GetInstance();
 
