Adding new export dialog flow to IK Retarget editor.

diff --git a/Engine/Plugins/Animation/IKRig/Source/IKRigEditor/Private/RetargetEditor/SIKRetargetAssetBrowser.cpp b/Engine/Plugins/Animation/IKRig/Source/IKRigEditor/Private/RetargetEditor/SIKRetargetAssetBrowser.cpp
--- a/Engine/Plugins/Animation/IKRig/Source/IKRigEditor/Private/RetargetEditor/SIKRetargetAssetBrowser.cpp
+++ b/Engine/Plugins/Animation/IKRig/Source/IKRigEditor/Private/RetargetEditor/SIKRetargetAssetBrowser.cpp
@@ -3,7 +3,6 @@
 #include "RetargetEditor/SIKRetargetAssetBrowser.h"
 
 #include "SPositiveActionButton.h"
-#include "AnimPreviewInstance.h"
 #include "ContentBrowserDataSource.h"
 #include "ContentBrowserModule.h"
 #include "IContentBrowserSingleton.h"
@@ -11,14 +10,13 @@
 #include "Animation/AnimSequence.h"
 #include "Animation/PoseAsset.h"
 
-#include "IKRigEditor.h"
 #include "RetargetEditor/IKRetargetBatchOperation.h"
 #include "RetargetEditor/IKRetargetEditorController.h"
 #include "RetargetEditor/SRetargetAnimAssetsWindow.h"
 #include "Retargeter/IKRetargeter.h"
 #include "Retargeter/IKRetargetProcessor.h"
 #include "UObject/AssetRegistryTagsContext.h"
-#include "Widgets/Layout/SUniformGridPanel.h"
+#include "Animation/Skeleton.h"
 
 #define LOCTEXT_NAMESPACE "IKRetargeterAssetBrowser"
 
@@ -153,30 +151,36 @@ TSharedPtr<SWidget> SIKRetargetAssetBrowser::OnGetAssetContextMenu(const TArray<
 
 FReply SIKRetargetAssetBrowser::OnExportButtonClicked()
 {
-	FIKRetargetEditorController* Controller = EditorController.Pin().Get();
+	const FIKRetargetEditorController* Controller = EditorController.Pin().Get();
 	if (!Controller)
 	{
 		checkNoEntry();
 		return FReply::Handled();
 	}
 	
-	// prompt user for path to export animations to
-	TSharedRef<SBatchExportDialog> Dialog = SNew(SBatchExportDialog).DefaultAssetPath(FText::FromString(PrevBatchOutputPath));
-	if(Dialog->ShowModal() == EAppReturnType::Cancel)
+	// assemble the context for the assets we want to batch duplicate/retarget
+	FIKRetargetBatchOperationContext BatchContext;
+	BatchContext.NameRule.FolderPath = PrevBatchOutputPath;
+	BatchContext.SourceMesh = Controller->GetSkeletalMesh(ERetargetSourceOrTarget::Source);
+	BatchContext.TargetMesh = Controller->GetSkeletalMesh(ERetargetSourceOrTarget::Target);
+	BatchContext.IKRetargetAsset = Controller->AssetController->GetAsset();
+	
+	// get the export path from user
+	const TSharedRef<SBatchExportPathDialog> PathDialog = SNew(SBatchExportPathDialog).BatchContext(&BatchContext).ExportRetargetAssets(false);
+	if(PathDialog->ShowModal() == EAppReturnType::Cancel)
 	{
 		return FReply::Handled();
 	}
 
 	// store path for next time
-	PrevBatchOutputPath = Dialog->GetAssetPath();
+	PrevBatchOutputPath = BatchContext.NameRule.FolderPath;
 
-	// assemble the data for the assets we want to batch duplicate/retarget
-	FIKRetargetBatchOperationContext BatchContext = Dialog.Get().BatchContext;
-	BatchContext.NameRule.FolderPath = PrevBatchOutputPath;
-	BatchContext.SourceMesh = Controller->GetSkeletalMesh(ERetargetSourceOrTarget::Source);
-	BatchContext.TargetMesh = Controller->GetSkeletalMesh(ERetargetSourceOrTarget::Target);
-	BatchContext.IKRetargetAsset = Controller->AssetController->GetAsset();
-	BatchContext.bIncludeReferencedAssets = false;
+	// get the export options from user
+	const TSharedRef<SBatchExportOptionsDialog> OptionsDialog = SNew(SBatchExportOptionsDialog).BatchContext(&BatchContext);
+	if(OptionsDialog->ShowModal() == EAppReturnType::Cancel)
+	{
+		return FReply::Handled();
+	}
 
 	// add selected assets to dup/retarget
 	TArray<FAssetData> SelectedAssets = GetCurrentSelectionDelegate.Execute();
@@ -188,7 +192,7 @@ FReply SIKRetargetAssetBrowser::OnExportButtonClicked()
 		}
 	}
 
-	// actually run the retarget
+	// run the batch retarget
 	const TStrongObjectPtr<UIKRetargetBatchOperation> BatchOperation(NewObject<UIKRetargetBatchOperation>());
 	BatchOperation->RunRetarget(BatchContext);
 	
@@ -287,288 +291,4 @@ bool SIKRetargetAssetBrowser::OnShouldFilterAsset(const struct FAssetData& Asset
 	return !DesiredSkeleton->IsCompatibleForEditor(AssetData);
 }
 
-// ------------------------------------------BEGIN  SBatchExportDialog ----------------------------
-
-void SBatchExportDialog::Construct(const FArguments& InArgs)
-{
-	AssetPath = FText::FromString(FPackageName::GetLongPackagePath(InArgs._DefaultAssetPath.ToString()));
-
-	if(AssetPath.IsEmpty())
-	{
-		AssetPath = FText::FromString(TEXT("/Game"));
-	}
-
-	FPathPickerConfig PathPickerConfig;
-	PathPickerConfig.DefaultPath = AssetPath.ToString();
-	PathPickerConfig.OnPathSelected = FOnPathSelected::CreateLambda([this](const FString& NewPath){AssetPath = FText::FromString(NewPath);});
-	PathPickerConfig.bAddDefaultPath = true;
-
-	FContentBrowserModule& ContentBrowserModule = FModuleManager::LoadModuleChecked<FContentBrowserModule>("ContentBrowser");
-
-	constexpr int32 TextInputWidths = 200;
-
-	SWindow::Construct(
-		SWindow::FArguments()
-		.Title(LOCTEXT("SBatchExportDialog_Title", "Batch Export Retargeted Animations"))
-		.SupportsMinimize(false)
-		.SupportsMaximize(false)
-		.IsTopmostWindow(true)
-		.ClientSize(FVector2D(350, 600))
-		[
-			SNew(SVerticalBox)
-			
-			+ SVerticalBox::Slot()
-			.Padding(2)
-			[
-				SNew(SBorder)
-				.BorderImage(FAppStyle::GetBrush("ToolPanel.GroupBorder"))
-				[
-					SNew(SVerticalBox)
-
-					+SVerticalBox::Slot()
-					.Padding(2, 3)
-					.AutoHeight()
-					.HAlign(HAlign_Center)
-					[
-						SNew(STextBlock)
-						.Text(LOCTEXT("SelectPath", "Select Path"))
-						.Font(FCoreStyle::GetDefaultFontStyle("Regular", 14))
-					]
-
-					+SVerticalBox::Slot()
-					.FillHeight(1)
-					.Padding(3)
-					[
-						ContentBrowserModule.Get().CreatePathPicker(PathPickerConfig)
-					]
-				]
-			]
-
-			+ SVerticalBox::Slot()
-			.Padding(2)
-			.AutoHeight()
-			[
-				SNew(SBorder)
-				.BorderImage(FAppStyle::GetBrush("ToolPanel.GroupBorder"))
-				[
-					SNew(SVerticalBox)
-
-					+SVerticalBox::Slot()
-					.Padding(2, 3)
-					.AutoHeight()
-					.HAlign(HAlign_Center)
-					[
-						SNew(STextBlock)
-						.Text(LOCTEXT("DuplicateAndRetarget_RenameLabel", "Rename New Assets"))
-						.Font(FCoreStyle::GetDefaultFontStyle("Regular", 14))
-					]
-
-					+SVerticalBox::Slot()
-					.AutoHeight()
-					.Padding(2, 1)
-					.HAlign(HAlign_Right)
-					[
-						SNew(SHorizontalBox)
-						+SHorizontalBox::Slot()
-						.AutoWidth()
-						.VAlign(VAlign_Center)
-						.Padding(2, 1)
-						[
-							SNew(STextBlock).Text(LOCTEXT("DuplicateAndRetarget_Prefix", "Add Prefix:"))
-						]
-
-						+SHorizontalBox::Slot()
-						.AutoWidth()
-						[
-							SNew(SEditableTextBox)
-							.Text_Lambda([this]
-							{
-								return FText::FromString(BatchContext.NameRule.Prefix);
-							})
-							.MinDesiredWidth(TextInputWidths)
-							.OnTextChanged_Lambda([this](const FText& InText)
-							{
-								BatchContext.NameRule.Prefix = InText.ToString();
-								UpdateExampleText();
-							})
-							.IsReadOnly(false)
-							.RevertTextOnEscape(true)
-						]
-					]
-
-					+SVerticalBox::Slot()
-					.AutoHeight()
-					.Padding(2, 1)
-					.HAlign(HAlign_Right)
-					[
-						SNew(SHorizontalBox)
-						+SHorizontalBox::Slot()
-						.AutoWidth()
-						.VAlign(VAlign_Center)
-						.Padding(2, 1)
-						[
-							SNew(STextBlock).Text(LOCTEXT("DuplicateAndRetarget_Suffix", "Add Suffix:"))
-						]
-
-						+SHorizontalBox::Slot()
-						.AutoWidth()
-						[
-							SNew(SEditableTextBox)
-							.Text_Lambda([this]
-							{
-								return FText::FromString(BatchContext.NameRule.Suffix);
-							})
-							.MinDesiredWidth(TextInputWidths)
-							.OnTextChanged_Lambda([this](const FText& InText)
-							{
-								BatchContext.NameRule.Suffix = InText.ToString();
-								UpdateExampleText();
-							})
-							.IsReadOnly(false)
-							.RevertTextOnEscape(true)
-						]
-					]
-
-					+SVerticalBox::Slot()
-					.AutoHeight()
-					.Padding(2, 1)
-					.HAlign(HAlign_Right)
-					[
-						SNew(SHorizontalBox)
-						+SHorizontalBox::Slot()
-						.HAlign(HAlign_Right)
-						.VAlign(VAlign_Center)
-						[
-							SNew(SHorizontalBox)
-							+SHorizontalBox::Slot()
-							.Padding(2, 1)
-							[
-								SNew(STextBlock).Text(LOCTEXT("DuplicateAndRetarget_Search", "Search for:"))
-							]
-
-							+SHorizontalBox::Slot()
-							[
-								SNew(SEditableTextBox)
-								.Text_Lambda([this]
-								{
-									return FText::FromString(BatchContext.NameRule.ReplaceFrom);
-								})
-								.MinDesiredWidth(TextInputWidths)
-								.OnTextChanged_Lambda([this](const FText& InText)
-								{
-									BatchContext.NameRule.ReplaceFrom = InText.ToString();
-									UpdateExampleText();
-								})
-								.IsReadOnly(false)
-								.RevertTextOnEscape(true)
-							]
-						]
-						
-						+SHorizontalBox::Slot()
-						.HAlign(HAlign_Right)
-						.VAlign(VAlign_Center)
-						[
-							SNew(SHorizontalBox)
-							+SHorizontalBox::Slot()
-							.Padding(2, 1)
-							[
-								SNew(STextBlock).Text(LOCTEXT("DuplicateAndRetarget_Replace", "Replace with:"))
-							]
-
-							+SHorizontalBox::Slot()
-							[
-								SNew(SEditableTextBox)
-								.Text_Lambda([this]
-								{
-									return FText::FromString(BatchContext.NameRule.ReplaceTo);
-								})
-								.MinDesiredWidth(TextInputWidths)
-								.OnTextChanged_Lambda([this](const FText& InText)
-								{
-									BatchContext.NameRule.ReplaceTo = InText.ToString();
-									UpdateExampleText();
-								})
-								.IsReadOnly(false)
-								.RevertTextOnEscape(true)
-							]
-						]
-					]
-
-					+SVerticalBox::Slot()
-					.AutoHeight()
-					.Padding(2, 3)
-					.HAlign(HAlign_Right)
-					[
-						SNew(SHorizontalBox)
-						+SHorizontalBox::Slot()
-						.Padding(5, 5)
-						.AutoWidth()
-						.VAlign(VAlign_Center)
-						[
-							SNew(STextBlock)
-							.Text_Lambda([this]{ return ExampleText; })
-							.Font(FAppStyle::GetFontStyle("Persona.RetargetManager.ItalicFont"))
-						]
-					]
-				]
-			]
-			
-			+SVerticalBox::Slot()
-			.AutoHeight()
-			.HAlign(HAlign_Right)
-			.Padding(5)
-			[
-				SNew(SUniformGridPanel)
-				.SlotPadding(FAppStyle::GetMargin("StandardDialog.SlotPadding"))
-				.MinDesiredSlotWidth(FAppStyle::GetFloat("StandardDialog.MinDesiredSlotWidth"))
-				.MinDesiredSlotHeight(FAppStyle::GetFloat("StandardDialog.MinDesiredSlotHeight"))
-				+SUniformGridPanel::Slot(0, 0)
-				[
-					SNew(SButton)
-					.HAlign(HAlign_Center)
-					.ContentPadding(FAppStyle::GetMargin("StandardDialog.ContentPadding"))
-					.Text(LOCTEXT("Export", "Export"))
-					.OnClicked(this, &SBatchExportDialog::OnButtonClick, EAppReturnType::Ok)
-				]
-				+SUniformGridPanel::Slot(1, 0)
-				[
-					SNew(SButton)
-					.HAlign(HAlign_Center)
-					.ContentPadding(FAppStyle::GetMargin("StandardDialog.ContentPadding"))
-					.Text(LOCTEXT("Cancel", "Cancel"))
-					.OnClicked(this, &SBatchExportDialog::OnButtonClick, EAppReturnType::Cancel)
-				]
-			]
-		]);
-}
-
-void SBatchExportDialog::UpdateExampleText()
-{
-	const FString ReplaceFrom = FString::Printf(TEXT("Old Name : ###%s###"), *BatchContext.NameRule.ReplaceFrom);
-	const FString ReplaceTo = FString::Printf(TEXT("New Name : %s###%s###%s"), *BatchContext.NameRule.Prefix, *BatchContext.NameRule.ReplaceTo, *BatchContext.NameRule.Suffix);
-
-	ExampleText = FText::FromString(FString::Printf(TEXT("%s\n%s"), *ReplaceFrom, *ReplaceTo));
-}
-
-FReply SBatchExportDialog::OnButtonClick(EAppReturnType::Type ButtonID)
-{
-	UserResponse = ButtonID;
-
-	RequestDestroyWindow();
-
-	return FReply::Handled();
-}
-
-EAppReturnType::Type SBatchExportDialog::ShowModal()
-{
-	GEditor->EditorAddModalWindow(SharedThis(this));
-	return UserResponse;
-}
-
-FString SBatchExportDialog::GetAssetPath()
-{
-	return AssetPath.ToString();
-}
-
-
 #undef LOCTEXT_NAMESPACE
diff --git a/Engine/Plugins/Animation/IKRig/Source/IKRigEditor/Public/RetargetEditor/SIKRetargetAssetBrowser.h b/Engine/Plugins/Animation/IKRig/Source/IKRigEditor/Public/RetargetEditor/SIKRetargetAssetBrowser.h
--- a/Engine/Plugins/Animation/IKRig/Source/IKRigEditor/Public/RetargetEditor/SIKRetargetAssetBrowser.h
+++ b/Engine/Plugins/Animation/IKRig/Source/IKRigEditor/Public/RetargetEditor/SIKRetargetAssetBrowser.h
@@ -48,38 +48,3 @@ private:
 
 	friend FIKRetargetEditorController;
 };
-
-//** Dialog to select path to export to */
-class SBatchExportDialog: public SWindow
-{
-public:
-	SLATE_BEGIN_ARGS(SBatchExportDialog){}
-	SLATE_ARGUMENT(FText, DefaultAssetPath)
-	SLATE_END_ARGS()
-
-	SBatchExportDialog() : UserResponse(EAppReturnType::Cancel)
-	{
-	}
-
-	void Construct(const FArguments& InArgs);
-
-public:
-	/** Displays the dialog in a blocking fashion */
-	EAppReturnType::Type ShowModal();
-
-	/** Gets the resulting asset path */
-	FString GetAssetPath();
-
-	FIKRetargetBatchOperationContext BatchContext;
-
-protected:
-
-	/** The rename rule sample text */
-	FText ExampleText;
-	void UpdateExampleText();
-	
-	FReply OnButtonClick(EAppReturnType::Type ButtonID);
-
-	EAppReturnType::Type UserResponse;
-	FText AssetPath;
-};
