Refresh heirarchy view to show chain labels

diff --git a/Engine/Plugins/Animation/IKRig/Source/IKRig/Private/Retargeter/IKRetargetProcessor.cpp b/Engine/Plugins/Animation/IKRig/Source/IKRig/Private/Retargeter/IKRetargetProcessor.cpp
--- a/Engine/Plugins/Animation/IKRig/Source/IKRig/Private/Retargeter/IKRetargetProcessor.cpp
+++ b/Engine/Plugins/Animation/IKRig/Source/IKRig/Private/Retargeter/IKRetargetProcessor.cpp
@@ -1572,6 +1572,9 @@ void UIKRetargetProcessor::Initialize(
 	
 	bIsInitialized = bRootsInitialized || bAtLeastOneValidBoneChainPair;
 	AssetVersionInitializedWith = RetargeterAsset->GetVersion();
+#if WITH_EDITOR
+	RetargeterInitialized.Broadcast();
+#endif
 }
 
 bool UIKRetargetProcessor::InitializeRoots()
diff --git a/Engine/Plugins/Animation/IKRig/Source/IKRig/Public/Retargeter/IKRetargetProcessor.h b/Engine/Plugins/Animation/IKRig/Source/IKRig/Public/Retargeter/IKRetargetProcessor.h
--- a/Engine/Plugins/Animation/IKRig/Source/IKRig/Public/Retargeter/IKRetargetProcessor.h
+++ b/Engine/Plugins/Animation/IKRig/Source/IKRig/Public/Retargeter/IKRetargetProcessor.h
@@ -546,6 +546,10 @@ public:
 	void SetNeedsInitialized();
 	
 #if WITH_EDITOR
+private:
+	DECLARE_MULTICAST_DELEGATE(FOnRetargeterInitialized);
+	FOnRetargeterInitialized RetargeterInitialized;
+public:
 	/** Returns true if the bone is part of a retarget chain or root bone, false otherwise. */
 	bool IsBoneRetargeted(const FName BoneName, const ERetargetSourceOrTarget SourceOrTarget) const;
 	/** Returns index of the bone with the given name in either Source or Target skeleton. */
@@ -564,6 +568,8 @@ public:
 	FName GetMappedChainName(const FName InChainName, const ERetargetSourceOrTarget SourceOrTarget);
 	/** store data for debug drawing */
 	FRetargetDebugData DebugData;
+	/** Attach a delegate to be notified whenever this processor is initialized. */
+	FOnRetargeterInitialized& OnRetargeterInitialized(){ return RetargeterInitialized; };
 #endif
 
 private:
diff --git a/Engine/Plugins/Animation/IKRig/Source/IKRigEditor/Private/RetargetEditor/IKRetargetEditor.cpp b/Engine/Plugins/Animation/IKRig/Source/IKRigEditor/Private/RetargetEditor/IKRetargetEditor.cpp
--- a/Engine/Plugins/Animation/IKRig/Source/IKRigEditor/Private/RetargetEditor/IKRetargetEditor.cpp
+++ b/Engine/Plugins/Animation/IKRig/Source/IKRigEditor/Private/RetargetEditor/IKRetargetEditor.cpp
@@ -651,6 +651,9 @@ void FIKRetargetEditor::HandlePreviewSceneCreated(const TSharedRef<IPersonaPrevi
 	
 	EditorController->FixZeroHeightRetargetRoot(ERetargetSourceOrTarget::Source);
 	EditorController->FixZeroHeightRetargetRoot(ERetargetSourceOrTarget::Target);
+
+	// bind a callback to update the UI whenever the retarget processor in the target anim instance is initialized
+	EditorController->RetargeterInitializedDelegateHandle = EditorController->GetRetargetProcessor()->OnRetargeterInitialized().AddLambda([this](){ EditorController->RefreshAllViews(); });
 }
 
 void FIKRetargetEditor::SetupAnimInstance()
diff --git a/Engine/Plugins/Animation/IKRig/Source/IKRigEditor/Private/RetargetEditor/IKRetargetEditorController.cpp b/Engine/Plugins/Animation/IKRig/Source/IKRigEditor/Private/RetargetEditor/IKRetargetEditorController.cpp
--- a/Engine/Plugins/Animation/IKRig/Source/IKRigEditor/Private/RetargetEditor/IKRetargetEditorController.cpp
+++ b/Engine/Plugins/Animation/IKRig/Source/IKRigEditor/Private/RetargetEditor/IKRetargetEditorController.cpp
@@ -181,6 +181,7 @@ void FIKRetargetEditorController::Close()
 	AssetController->OnRetargeterNeedsInitialized().Remove(RetargeterReInitDelegateHandle);
 	AssetController->OnIKRigReplaced().Remove(IKRigReplacedDelegateHandle);
 	AssetController->OnPreviewMeshReplaced().Remove(PreviewMeshReplacedDelegateHandle);
+	GetRetargetProcessor()->OnRetargeterInitialized().Remove(RetargeterInitializedDelegateHandle);
 
 	for (const FBoundIKRig& BoundIKRig : BoundIKRigs)
 	{
diff --git a/Engine/Plugins/Animation/IKRig/Source/IKRigEditor/Public/RetargetEditor/IKRetargetEditorController.h b/Engine/Plugins/Animation/IKRig/Source/IKRigEditor/Public/RetargetEditor/IKRetargetEditorController.h
--- a/Engine/Plugins/Animation/IKRig/Source/IKRigEditor/Public/RetargetEditor/IKRetargetEditorController.h
+++ b/Engine/Plugins/Animation/IKRig/Source/IKRigEditor/Public/RetargetEditor/IKRetargetEditorController.h
@@ -110,6 +110,7 @@ public:
 	// callback when Preview Mesh asset has been swapped out
 	void HandlePreviewMeshReplaced(ERetargetSourceOrTarget SourceOrTarget);
 	FDelegateHandle PreviewMeshReplacedDelegateHandle;
+	FDelegateHandle RetargeterInitializedDelegateHandle;
 	
 	// all modifications to the data model should go through this controller
 	TObjectPtr<UIKRetargeterController> AssetController;
