Fixes retarget post op stack editing UI glitch

diff --git a/Engine/Plugins/Animation/IKRig/Source/IKRigEditor/Private/RetargetEditor/IKRetargetDetails.cpp b/Engine/Plugins/Animation/IKRig/Source/IKRigEditor/Private/RetargetEditor/IKRetargetDetails.cpp
--- a/Engine/Plugins/Animation/IKRig/Source/IKRigEditor/Private/RetargetEditor/IKRetargetDetails.cpp
+++ b/Engine/Plugins/Animation/IKRig/Source/IKRigEditor/Private/RetargetEditor/IKRetargetDetails.cpp
@@ -1258,7 +1258,8 @@ void FRetargetChainSettingsCustomization::AddSettingsSection(
 				.OnCheckStateChanged_Lambda([this, IsEnabledProperty, Controller](ECheckBoxState State)
 				{
 					IsEnabledProperty->SetValue(State == ECheckBoxState::Checked);
-					Controller->HandleRetargeterNeedsInitialized();
+					// clear the output log and trigger a reinitialization
+					Controller->ReinitializeRetargeterNoUIRefresh();
 				})
 			];
 	}
@@ -1518,12 +1519,12 @@ void FRetargetOpStackCustomization::CustomizeDetails(IDetailLayoutBuilder& Detai
 			{
 				PropertyHandle->SetOnPropertyValueChanged(FSimpleDelegate::CreateLambda([this]()
 				{
-					Controller.Pin().Get()->HandleRetargeterNeedsInitialized();
+					Controller.Pin().Get()->ReinitializeRetargeterNoUIRefresh();
 				}));
 				
 				PropertyHandle->SetOnChildPropertyValueChanged(FSimpleDelegate::CreateLambda([this]()
 				{
-					Controller.Pin().Get()->HandleRetargeterNeedsInitialized();
+					Controller.Pin().Get()->ReinitializeRetargeterNoUIRefresh();
 				}));
 				
 				OpSettingCategory.AddProperty(PropertyHandle)
diff --git a/Engine/Plugins/Animation/IKRig/Source/IKRigEditor/Private/RetargetEditor/IKRetargetEditor.cpp b/Engine/Plugins/Animation/IKRig/Source/IKRigEditor/Private/RetargetEditor/IKRetargetEditor.cpp
--- a/Engine/Plugins/Animation/IKRig/Source/IKRigEditor/Private/RetargetEditor/IKRetargetEditor.cpp
+++ b/Engine/Plugins/Animation/IKRig/Source/IKRigEditor/Private/RetargetEditor/IKRetargetEditor.cpp
@@ -653,7 +653,10 @@ void FIKRetargetEditor::HandlePreviewSceneCreated(const TSharedRef<IPersonaPrevi
 	EditorController->FixZeroHeightRetargetRoot(ERetargetSourceOrTarget::Target);
 
 	// bind a callback to update the UI whenever the retarget processor in the target anim instance is initialized
-	EditorController->RetargeterInitializedDelegateHandle = EditorController->GetRetargetProcessor()->OnRetargeterInitialized().AddLambda([this](){ EditorController->RefreshAllViews(); });
+	EditorController->RetargeterInitializedDelegateHandle = EditorController->GetRetargetProcessor()->OnRetargeterInitialized().AddLambda([this]()
+	{
+		EditorController->RefreshHierarchyView();
+	});
 }
 
 void FIKRetargetEditor::SetupAnimInstance()
diff --git a/Engine/Plugins/Animation/IKRig/Source/IKRigEditor/Private/RetargetEditor/IKRetargetEditorController.cpp b/Engine/Plugins/Animation/IKRig/Source/IKRigEditor/Private/RetargetEditor/IKRetargetEditorController.cpp
--- a/Engine/Plugins/Animation/IKRig/Source/IKRigEditor/Private/RetargetEditor/IKRetargetEditorController.cpp
+++ b/Engine/Plugins/Animation/IKRig/Source/IKRigEditor/Private/RetargetEditor/IKRetargetEditorController.cpp
@@ -249,14 +249,19 @@ void FIKRetargetEditorController::HandleRetargeterNeedsInitialized() const
 	FixZeroHeightRetargetRoot(ERetargetSourceOrTarget::Source);
 	FixZeroHeightRetargetRoot(ERetargetSourceOrTarget::Target);
 	
+	ReinitializeRetargeterNoUIRefresh();
+	
+	// refresh all the UI views
+	RefreshAllViews();	
+}
+
+void FIKRetargetEditorController::ReinitializeRetargeterNoUIRefresh() const
+{
 	// clear the output log
 	ClearOutputLog();
 
 	// force running instances to reinitialize on next tick
 	AssetController->GetAsset()->IncrementVersion();
-	
-	// refresh all the UI views
-	RefreshAllViews();	
 }
 
 void FIKRetargetEditorController::HandleIKRigReplaced(ERetargetSourceOrTarget SourceOrTarget)
diff --git a/Engine/Plugins/Animation/IKRig/Source/IKRigEditor/Private/RetargetEditor/SRetargetOpStack.cpp b/Engine/Plugins/Animation/IKRig/Source/IKRigEditor/Private/RetargetEditor/SRetargetOpStack.cpp
--- a/Engine/Plugins/Animation/IKRig/Source/IKRigEditor/Private/RetargetEditor/SRetargetOpStack.cpp
+++ b/Engine/Plugins/Animation/IKRig/Source/IKRigEditor/Private/RetargetEditor/SRetargetOpStack.cpp
@@ -78,7 +78,7 @@ void SRetargetOpStackItem::Construct(
 						{
 							const bool bIsChecked = InCheckBoxState == ECheckBoxState::Checked;
 							InOpStackWidget->EditorController.Pin()->AssetController->SetRetargetOpEnabled(InStackElement->IndexInStack, bIsChecked);
-							InOpStackWidget->EditorController.Pin()->HandleRetargeterNeedsInitialized();
+							InOpStackWidget->EditorController.Pin()->ReinitializeRetargeterNoUIRefresh();
 						}
 					})
 				]
@@ -335,7 +335,7 @@ void SRetargetOpStack::OnSelectionChanged(
 {
 	if (SelectInfo != ESelectInfo::Direct && InItem.IsValid())
 	{
-		LastSelectedItemIndex = InItem->IndexInStack;
+		EditorController.Pin()->LastSelectedOpIndex = InItem->IndexInStack;
 	}
 }
 
@@ -369,10 +369,10 @@ void SRetargetOpStack::RefreshStackView()
 		ListViewItems.Add(StackElement);
 	}
 
-	if (NumOps && ListViewItems.IsValidIndex(LastSelectedItemIndex))
+	if (NumOps && ListViewItems.IsValidIndex(EditorController.Pin()->LastSelectedOpIndex))
 	{
 		// restore selection
-		ListView->SetSelection(ListViewItems[LastSelectedItemIndex]);
+		ListView->SetSelection(ListViewItems[EditorController.Pin()->LastSelectedOpIndex]);
 	}
 	else
 	{
diff --git a/Engine/Plugins/Animation/IKRig/Source/IKRigEditor/Public/RetargetEditor/IKRetargetEditorController.h b/Engine/Plugins/Animation/IKRig/Source/IKRigEditor/Public/RetargetEditor/IKRetargetEditorController.h
--- a/Engine/Plugins/Animation/IKRig/Source/IKRigEditor/Public/RetargetEditor/IKRetargetEditorController.h
+++ b/Engine/Plugins/Animation/IKRig/Source/IKRigEditor/Public/RetargetEditor/IKRetargetEditorController.h
@@ -103,6 +103,8 @@ public:
 	void HandleRetargetChainRemoved(UIKRigDefinition* ModifiedIKRig, const FName InChainRemoved) const;
 	// callback when IK Retargeter asset requires reinitialization
 	void HandleRetargeterNeedsInitialized() const;
+	// reinitialize retargeter without refreshing UI
+	void ReinitializeRetargeterNoUIRefresh() const;
 	FDelegateHandle RetargeterReInitDelegateHandle;
 	// callback when IK Rig asset has been swapped out
 	void HandleIKRigReplaced(ERetargetSourceOrTarget SourceOrTarget);
@@ -249,6 +251,7 @@ public:
 
 	// op stack selection
 	URetargetOpBase* GetSelectedOp() const;
+	int32 LastSelectedOpIndex = 0;
 
 	// to frame selection when pressing "f" in viewport
 	bool GetCameraTargetForSelection(FSphere& OutTarget) const;
diff --git a/Engine/Plugins/Animation/IKRig/Source/IKRigEditor/Public/RetargetEditor/SRetargetOpStack.h b/Engine/Plugins/Animation/IKRig/Source/IKRigEditor/Public/RetargetEditor/SRetargetOpStack.h
--- a/Engine/Plugins/Animation/IKRig/Source/IKRigEditor/Public/RetargetEditor/SRetargetOpStack.h
+++ b/Engine/Plugins/Animation/IKRig/Source/IKRigEditor/Public/RetargetEditor/SRetargetOpStack.h
@@ -102,7 +102,6 @@ private:
 	void DeleteRetargetOp(TSharedPtr<FRetargetOpStackElement> OpToDelete);
 	// when an op is selected on in the stack view
 	void OnSelectionChanged(TSharedPtr<FRetargetOpStackElement> InItem, ESelectInfo::Type SelectInfo);
-	int32 LastSelectedItemIndex = 0;
 
 	// list view generate row callback
 	TSharedRef<ITableRow> MakeListRowWidget(TSharedPtr<FRetargetOpStackElement> InElement, const TSharedRef<STableViewBase>& OwnerTable);
