Fixes regression, cannot create IK Rig goal if solver present.

diff --git a/Engine/Plugins/Animation/IKRig/Source/IKRigEditor/Private/RigEditor/IKRigEditorController.cpp b/Engine/Plugins/Animation/IKRig/Source/IKRigEditor/Private/RigEditor/IKRigEditorController.cpp
--- a/Engine/Plugins/Animation/IKRig/Source/IKRigEditor/Private/RigEditor/IKRigEditorController.cpp
+++ b/Engine/Plugins/Animation/IKRig/Source/IKRigEditor/Private/RigEditor/IKRigEditorController.cpp
@@ -845,7 +845,7 @@ bool FIKRigEditorController::PromptToAddDefaultSolver() const
 {
 	if (AssetController->GetNumSolvers() > 0)
 	{
-		return true;
+		return false;
 	}
 
 	TArray<TSharedPtr<FIKRigSolverTypeAndName>> SolverTypes;
