Fixes crash when closing editor with open batch retargeter.

diff --git a/Engine/Plugins/Animation/IKRig/Source/IKRigEditor/Private/RetargetEditor/SRetargetAnimAssetsWindow.cpp b/Engine/Plugins/Animation/IKRig/Source/IKRigEditor/Private/RetargetEditor/SRetargetAnimAssetsWindow.cpp
--- a/Engine/Plugins/Animation/IKRig/Source/IKRigEditor/Private/RetargetEditor/SRetargetAnimAssetsWindow.cpp
+++ b/Engine/Plugins/Animation/IKRig/Source/IKRigEditor/Private/RetargetEditor/SRetargetAnimAssetsWindow.cpp
@@ -693,6 +693,18 @@ TSharedPtr<SWidget> SRetargetPoseViewport::MakeViewportToolbar()
 	return nullptr;
 }
 
+UBatchRetargetSettings* UBatchRetargetSettings::GetInstance()
+{
+	if (!SingletonInstance)
+	{
+		SingletonInstance = NewObject<UBatchRetargetSettings>(GetTransientPackage(), UBatchRetargetSettings::StaticClass());
+		SingletonInstance->AddToRoot();
+	}
+	return SingletonInstance;
+}
+
+UBatchRetargetSettings* UBatchRetargetSettings::SingletonInstance = nullptr;
+
 void SRetargetExporterAssetBrowser::Construct(const FArguments& InArgs, const TSharedRef<SRetargetAnimAssetsWindow> InRetargetWindow)
 {
 	RetargetWindow = InRetargetWindow;
@@ -796,7 +808,7 @@ bool SRetargetExporterAssetBrowser::OnShouldFilterAsset(const FAssetData& AssetD
 		return true;
 	}
 	
-	const TObjectPtr<UBatchRetargetSettings> BatchRetargetSettings = RetargetWindow.Get()->GetSettings();
+	const TObjectPtr<UBatchRetargetSettings> BatchRetargetSettings = UBatchRetargetSettings::GetInstance();
 	if (!ensure(BatchRetargetSettings))
 	{
 		return true;
@@ -824,8 +836,8 @@ bool SRetargetExporterAssetBrowser::OnShouldFilterAsset(const FAssetData& AssetD
 
 SRetargetAnimAssetsWindow::SRetargetAnimAssetsWindow()
 {
-	// create the settings uobject
-	Settings = NewObject<UBatchRetargetSettings>(GetTransientPackage(), UBatchRetargetSettings::StaticClass());
+	// get the global settings uobject
+	Settings = UBatchRetargetSettings::GetInstance();
 
 	// assign default retargeter
 	BatchContext.IKRetargetAsset = ProceduralAssets.Retargeter;
diff --git a/Engine/Plugins/Animation/IKRig/Source/IKRigEditor/Public/RetargetEditor/SRetargetAnimAssetsWindow.h b/Engine/Plugins/Animation/IKRig/Source/IKRigEditor/Public/RetargetEditor/SRetargetAnimAssetsWindow.h
--- a/Engine/Plugins/Animation/IKRig/Source/IKRigEditor/Public/RetargetEditor/SRetargetAnimAssetsWindow.h
+++ b/Engine/Plugins/Animation/IKRig/Source/IKRigEditor/Public/RetargetEditor/SRetargetAnimAssetsWindow.h
@@ -67,6 +67,9 @@ class UBatchExportOptions : public UObject
 
 public:
 
+	// singleton class
+	static UBatchExportOptions* GetInstance();
+
 	// Any files with the same name will be overwritten instead of creating a new file with a numeric suffix.
 	// This is useful when iterating on a batch process.
 	UPROPERTY(EditAnywhere, Config, Config, Category = "File")
@@ -81,8 +84,6 @@ public:
 	//UPROPERTY(EditAnywhere, BlueprintReadWrite, Config, Category = "Animation")
 	//bool bExportOnlyAnimatedBones = true;
 
-	static UBatchExportOptions* GetInstance();
-
 private:
 	static UBatchExportOptions* SingletonInstance;
 };
@@ -201,6 +202,9 @@ class UBatchRetargetSettings : public UObject
 
 public:
 
+	// singleton class
+	static UBatchRetargetSettings* GetInstance();
+
 	// The skeletal mesh with the proportions you want to copy animation FROM.
 	UPROPERTY(EditAnywhere, BlueprintReadWrite, Category = "Source")
 	TObjectPtr<USkeletalMesh> SourceSkeletalMesh;
@@ -220,6 +224,9 @@ public:
 	// You may also supply a custom IK Retargeter if needed.
 	UPROPERTY(EditAnywhere, BlueprintReadWrite, Category = "Retargeter", meta = (EditCondition = "!bAutoGenerateRetargeter"))
 	TObjectPtr<UIKRetargeter> RetargetAsset;
+
+private:
+	static UBatchRetargetSettings* SingletonInstance;
 };
 
 // asset browser for user to select animation assets to duplicate/retarget
