Fixed bug, stride warping splay moved feet vertically.
Fixed bug, stride warping chain method broken on straight spine.

diff --git a/Engine/Plugins/Animation/IKRig/Source/IKRig/Private/Retargeter/IKRetargetProcessor.cpp b/Engine/Plugins/Animation/IKRig/Source/IKRig/Private/Retargeter/IKRetargetProcessor.cpp
--- a/Engine/Plugins/Animation/IKRig/Source/IKRig/Private/Retargeter/IKRetargetProcessor.cpp
+++ b/Engine/Plugins/Animation/IKRig/Source/IKRig/Private/Retargeter/IKRetargetProcessor.cpp
@@ -2031,76 +2031,140 @@ void UIKRetargetProcessor::RunStrideWarping(const TArray<FTransform>& InTargetGl
 	{
 		return;
 	}
+
+	auto ConvertForwardVectorsToRotation = [](FVector& ForwardOrig, FVector& ForwardCurrent) -> FQuat
+	{
+		// project on floor (we only want Yaw) and normalize
+		ForwardOrig.Z = 0.0f;
+		ForwardCurrent.Z = 0.0f;
+		ForwardOrig.Normalize();
+		ForwardCurrent.Normalize();
+		// rotate from orig facing direction to current
+		return FQuat::FindBetweenNormals(ForwardOrig, ForwardCurrent);
+	};
+
+	// we need to determine the character's forward and side directions for warping
+	FVector InitialBodyPosition;
+	FVector CurrentBodyPosition;
+	FQuat CurrentRotation;
 	
 	// get sample points to use for "best fit" global body rotation
-	TArray<FVector> InitialPoints;
-	TArray<FVector> CurrentPoints;
-	if (GlobalSettings.DirectionSource == EWarpingDirectionSource::Goals)
+	switch(GlobalSettings.DirectionSource)
 	{
-		// use goals to determine the body's rotation
-		for (const FRetargetChainPairIK& ChainPair : ChainPairsIK)
+	case EWarpingDirectionSource::Goals:
 		{
-			if (!ChainPair.Settings.bAffectedByIKWarping)
+			TArray<FVector> InitialPoints;
+			TArray<FVector> CurrentPoints;
+			
+			// use goals to determine the body's rotation
+			for (const FRetargetChainPairIK& ChainPair : ChainPairsIK)
 			{
-				continue;
-			}
-			FVector A = ChainPair.IKChainRetargeter.Target.InitialEndPosition;
-			FVector B = ChainPair.IKChainRetargeter.Results.EndEffectorPosition;
+				if (!ChainPair.Settings.bAffectedByIKWarping)
+				{
+					continue;
+				}
+				FVector A = ChainPair.IKChainRetargeter.Target.InitialEndPosition;
+				FVector B = ChainPair.IKChainRetargeter.Results.EndEffectorPosition;
 
-			// flatten into 2D for more robust yaw construction (what really matters)
-			A.Z = 0.0f;
-			B.Z = 0.0f;
+				// flatten into 2D for more robust yaw construction (what really matters)
+				A.Z = 0.0f;
+				B.Z = 0.0f;
 		
-			InitialPoints.Add(A);
-			CurrentPoints.Add(B);
+				InitialPoints.Add(A);
+				CurrentPoints.Add(B);
+			}
+
+			// calculate "best fit" global body rotation based on deformation of sample points
+			CurrentRotation = GetRotationFromDeformedPoints(
+				InitialPoints,
+				CurrentPoints,
+				InitialBodyPosition,
+				CurrentBodyPosition);
+			
+			break;
 		}
-	}
-	else
-	{
-		// use chain to determine the body's rotation
-		const FRetargetChainPairFK* SpineChain = ChainPairsFK.FindByPredicate(
-				[this](const FRetargetChainPairFK& Other)
-				{
-					return Other.TargetBoneChainName == GlobalSettings.DirectionChain;
-				});
-		if (!SpineChain)
+	case EWarpingDirectionSource::Chain:
 		{
-			return;
-		}
+			// use chain to determine the body's rotation
+			const FRetargetChainPairFK* Chain = ChainPairsFK.FindByPredicate(
+					[this](const FRetargetChainPairFK& Other)
+					{
+						return Other.TargetBoneChainName == GlobalSettings.DirectionChain;
+					});
+			if (!Chain)
+			{
+				return;
+			}
+
+			const TArray<FTransform>& InitialChainTransforms = Chain->FKDecoder.InitialGlobalTransforms;
+			const TArray<FTransform>& CurrentChainTransforms = Chain->FKDecoder.CurrentGlobalTransforms;
+			if (!ensure(!InitialChainTransforms.IsEmpty()))
+			{
+				return;
+			}
 
-		const TArray<FTransform>& InitialChainTransforms = SpineChain->FKDecoder.InitialGlobalTransforms;
-		const TArray<FTransform>& CurrentChainTransforms = SpineChain->FKDecoder.CurrentGlobalTransforms;
+			// calculate initial and current centroids
+			InitialBodyPosition = FVector::ZeroVector;
+			CurrentBodyPosition = FVector::ZeroVector;
+			for (int32 TransformIndex=0; TransformIndex<InitialChainTransforms.Num(); ++TransformIndex)
+			{
+				InitialBodyPosition += InitialChainTransforms[TransformIndex].GetTranslation();
+				CurrentBodyPosition += CurrentChainTransforms[TransformIndex].GetTranslation();
+			}
+			const float InvNum = 1.0f / static_cast<float>(InitialChainTransforms.Num());
+			InitialBodyPosition *= InvNum;
+			CurrentBodyPosition *= InvNum;
+
+			// get the forward vectors of the chain
+			FVector ForwardOrig;
+			FVector ForwardCurrent;
+			
+			// get forward vectors: orig and current
+			if (InitialChainTransforms.Num() == 1)
+			{
+				// in case of single bone chain, we rotate the global forward vector with this bone and then project onto floor
+				const FQuat DeltaRotation = InitialChainTransforms[0].GetRotation() * CurrentChainTransforms[0].GetRotation().Inverse();
+				ForwardOrig = GlobalSettings.GetAxisVector(GlobalSettings.ForwardDirection);
+				ForwardCurrent = DeltaRotation.RotateVector(ForwardOrig);
+			}
+			else
+			{
+				// in case of multi-bone chain, we use the vector from the start to the end of the chain
+				ForwardOrig = InitialChainTransforms.Last().GetTranslation() - InitialChainTransforms[0].GetTranslation();
+				ForwardCurrent = CurrentChainTransforms.Last().GetTranslation() - CurrentChainTransforms[0].GetTranslation();
+			}
 
-		// calculate a "best fit" transform from deformed chain
-		for (int32 ChainIndex=0; ChainIndex<InitialChainTransforms.Num(); ++ChainIndex)
+			// get rotation from fwd vectors
+			CurrentRotation = ConvertForwardVectorsToRotation(ForwardOrig, ForwardCurrent);
+			break;
+	}
+	case EWarpingDirectionSource::RootBone:
 		{
-			FVector A = InitialChainTransforms[ChainIndex].GetLocation();
-			FVector B = CurrentChainTransforms[ChainIndex].GetLocation();
+			// use the root bone
+			InitialBodyPosition = TargetSkeleton.RetargetGlobalPose[0].GetTranslation();
+			CurrentBodyPosition = TargetSkeleton.OutputGlobalPose[0].GetTranslation();
 
-			// flatten into 2D for more robust yaw construction (what really matters)
-			A.Z = 0.0f;
-			B.Z = 0.0f;
-		
-			InitialPoints.Add(A);
-			CurrentPoints.Add(B);
+			// get the forward vectors of the root bone
+			const FQuat DeltaRotation = TargetSkeleton.RetargetGlobalPose[0].GetRotation() * TargetSkeleton.OutputGlobalPose[0].GetRotation().Inverse();
+			FVector ForwardOrig = GlobalSettings.GetAxisVector(GlobalSettings.ForwardDirection);
+			FVector ForwardCurrent = DeltaRotation.RotateVector(ForwardOrig);
+
+			// get rotation from fwd vectors
+			CurrentRotation = ConvertForwardVectorsToRotation(ForwardOrig, ForwardCurrent);
+			break;
+		}
+	default:
+		{
+			checkNoEntry();
+			return;
 		}
 	}
-
-	// calculate "best fit" global body rotation based on deformation of sample points
-	FVector InitialGoalCentroid;
-	FVector CurrentGoalCentroid;
-	const FQuat BestFitBodyRotation = GetRotationFromDeformedPoints(
-		InitialPoints,
-		CurrentPoints,
-		InitialGoalCentroid,
-		CurrentGoalCentroid);
-
-	// generate axes based on body rotation
-	FTransform InitialGlobalGoalTransform = FTransform(FQuat::Identity, InitialGoalCentroid);
-	FTransform CurrentGlobalGoalTransform = FTransform(BestFitBodyRotation, CurrentGoalCentroid);
-	const FVector& Fwd = CurrentGlobalGoalTransform.TransformVector(GlobalSettings.GetAxisVector(GlobalSettings.ForwardDirection));
-	const FVector& Side = FVector::CrossProduct(Fwd, FVector::ZAxisVector);
-	const FVector SideOrig = FVector::CrossProduct(GlobalSettings.GetAxisVector(GlobalSettings.ForwardDirection), FVector::ZAxisVector);
+	
+	FTransform CurrentBodyTransform = FTransform(CurrentRotation, CurrentBodyPosition);
+	FTransform InitialBodyTransform = FTransform(FQuat::Identity, InitialBodyPosition);
+	FVector Fwd = CurrentBodyTransform.TransformVector(GlobalSettings.GetAxisVector(GlobalSettings.ForwardDirection));
+	FVector Side = FVector::CrossProduct(Fwd, FVector::ZAxisVector);
+	FVector SideOrig = FVector::CrossProduct(GlobalSettings.GetAxisVector(GlobalSettings.ForwardDirection), FVector::ZAxisVector);
 
 	// warp goal positions...
 	for (FRetargetChainPairIK& ChainPair : ChainPairsIK)
@@ -2113,11 +2177,11 @@ void UIKRetargetProcessor::RunStrideWarping(const TArray<FTransform>& InTargetGl
 		FChainRetargeterIK& ChainRetargeter = ChainPair.IKChainRetargeter;
 
 		// warp this goal position in various ways...
-		FVector WarpedGoalPosition = ChainRetargeter.Results.EndEffectorPosition;
+		FVector& WarpedGoalPosition = ChainRetargeter.Results.EndEffectorPosition;
 		
 		// forward warping
-		const FVector InitialGoalInOrigSpace = InitialGlobalGoalTransform.InverseTransformPosition(ChainRetargeter.Target.InitialEndPosition);
-		const FVector IntialGoalInCurrentSpace = CurrentGlobalGoalTransform.TransformPosition(InitialGoalInOrigSpace);
+		const FVector InitialGoalInOrigSpace = InitialBodyTransform.InverseTransformPosition(ChainRetargeter.Target.InitialEndPosition);
+		const FVector IntialGoalInCurrentSpace = CurrentBodyTransform.TransformPosition(InitialGoalInOrigSpace);
 		const FPlane FwdPlane(IntialGoalInCurrentSpace, Fwd);
 		const FVector GoalProjOnFwdPlane = FwdPlane.PointPlaneProject(WarpedGoalPosition, FwdPlane);
 		WarpedGoalPosition = GoalProjOnFwdPlane + ((WarpedGoalPosition - GoalProjOnFwdPlane) * GlobalSettings.WarpForwards);
@@ -2129,14 +2193,13 @@ void UIKRetargetProcessor::RunStrideWarping(const TArray<FTransform>& InTargetGl
 		WarpedGoalPosition += Side * GlobalSettings.SidewaysOffset * GoalSideMultiplier;
 
 		// splay warping
-		WarpedGoalPosition = CurrentGoalCentroid + (WarpedGoalPosition - CurrentGoalCentroid) * GlobalSettings.WarpSplay; 
-
-		// assign warped position
-		ChainRetargeter.Results.EndEffectorPosition = WarpedGoalPosition;
+		FVector SplayOrigin = CurrentBodyPosition;
+		SplayOrigin.Z = WarpedGoalPosition.Z;
+		WarpedGoalPosition = SplayOrigin + (WarpedGoalPosition - SplayOrigin) * GlobalSettings.WarpSplay; 
 	}
 
 #if WITH_EDITOR
-	DebugData.StrideWarpingFrame = CurrentGlobalGoalTransform;
+	DebugData.StrideWarpingFrame = CurrentBodyTransform;
 #endif
 }
 
diff --git a/Engine/Plugins/Animation/IKRig/Source/IKRig/Public/Retargeter/IKRetargetSettings.h b/Engine/Plugins/Animation/IKRig/Source/IKRig/Public/Retargeter/IKRetargetSettings.h
--- a/Engine/Plugins/Animation/IKRig/Source/IKRig/Public/Retargeter/IKRetargetSettings.h
+++ b/Engine/Plugins/Animation/IKRig/Source/IKRig/Public/Retargeter/IKRetargetSettings.h
@@ -262,6 +262,7 @@ enum class EWarpingDirectionSource
 {
 	Goals,
 	Chain,
+	RootBone
 };
 
 USTRUCT(BlueprintType)
@@ -304,7 +305,8 @@ struct IKRIG_API FRetargetGlobalSettings
 	 * This global rotation is used to define the forward and sideways directions used when warping goals along those axes.
 	 * The options are:
 	 * Goals: uses the positions of the IK goals to approximate the facing direction. This is best used on characters with a vertical spine, like bipeds.
-	 * Chain: uses the positions of the bones in a retarget chain to approximate the facing direction. This is best for characters with a horizontal spine, like quadrupeds.
+	 * Chain: uses the positions of the bones in a retarget chain to approximate the facing direction. This is best when used with the spine chain for characters with a horizontal spine, like quadrupeds.
+	 * Root Bone: uses the rotation of the root bone of the skeleton. This is most robust, but character must have correct root motion with yaw rotation in movement direction.
 	 */
 	UPROPERTY(EditAnywhere, BlueprintReadWrite, Category = Warping)
 	EWarpingDirectionSource DirectionSource = EWarpingDirectionSource::Goals;
