Fixed bug, retarget pole vector offset only works if pole vector matching is non-zero.

diff --git a/Engine/Plugins/Animation/IKRig/Source/IKRig/Private/Retargeter/IKRetargetProcessor.cpp b/Engine/Plugins/Animation/IKRig/Source/IKRig/Private/Retargeter/IKRetargetProcessor.cpp
--- a/Engine/Plugins/Animation/IKRig/Source/IKRig/Private/Retargeter/IKRetargetProcessor.cpp
+++ b/Engine/Plugins/Animation/IKRig/Source/IKRig/Private/Retargeter/IKRetargetProcessor.cpp
@@ -1351,7 +1351,9 @@ void FPoleVectorMatcher::MatchPoleVector(
 	TArray<FTransform>& OutTargetGlobalPose,
 	FRetargetSkeleton& TargetSkeleton)
 {
-	if (Settings.PoleVectorMatching <= KINDA_SMALL_NUMBER && Settings.PoleVectorOffset <= KINDA_SMALL_NUMBER)
+	const bool bIsMatchingPoleVector = Settings.PoleVectorMatching > KINDA_SMALL_NUMBER;
+	const bool bIsOffsetingPoleVector = !FMath::IsNearlyZero(Settings.PoleVectorOffset);
+	if (!(bIsMatchingPoleVector || bIsOffsetingPoleVector))
 	{
 		return;
 	}
@@ -1371,20 +1373,33 @@ void FPoleVectorMatcher::MatchPoleVector(
 		LocalSpaces.Add(LocalTransform);
 	}
 
+	// normalized vector pointing from root to tip of chain
+	const FVector TargetChainAxisNorm = GetChainAxisNormalized(TargetIndices, OutTargetGlobalPose);
+
 	// calculate rotation to match the target to the source pole vector
-	const FVector SourcePoleVector = CalculatePoleVector(SourcePoleAxis, SourceIndices, SourceGlobalPose);
-	const FVector TargetPoleVector = CalculatePoleVector(TargetPoleAxis, TargetIndices, OutTargetGlobalPose);
+	FQuat MatchingRotation = FQuat::Identity;
+	if (bIsMatchingPoleVector)
+	{
+		const FVector SourcePoleVector = CalculatePoleVector(SourcePoleAxis, SourceIndices, SourceGlobalPose);
+		const FVector TargetPoleVector = CalculatePoleVector(TargetPoleAxis, TargetIndices, OutTargetGlobalPose);
 	
-	const float RotateTargetToSource = FMath::Acos(FVector::DotProduct(SourcePoleVector, TargetPoleVector));
-	const float MatchPoleAngle = RotateTargetToSource - TargetToSourceAngularOffsetAtRefPose + FMath::DegreesToRadians(Settings.PoleVectorOffset);
+		const float RotateTargetToSource = FMath::Acos(FVector::DotProduct(SourcePoleVector, TargetPoleVector));
+		const float MatchPoleAngle = RotateTargetToSource - (Settings.PoleVectorMaintainOffset ? TargetToSourceAngularOffsetAtRefPose : 0);
+		
+		MatchingRotation = FQuat(TargetChainAxisNorm, MatchPoleAngle);
+		MatchingRotation = FQuat::FastLerp(FQuat::Identity, MatchingRotation, Settings.PoleVectorMatching).GetNormalized();
+	}
 
-	const FVector TargetChainNormal = GetChainNormal(TargetIndices, OutTargetGlobalPose);
-	const FQuat PoleVectorMatchingRotation = FQuat(TargetChainNormal, MatchPoleAngle);
-	const FQuat FinalRotation = FQuat::FastLerp(FQuat::Identity, PoleVectorMatchingRotation, Settings.PoleVectorMatching).GetNormalized();
+	// manual offset rotation
+	FQuat OffsetRotation = FQuat::Identity;
+	if (bIsOffsetingPoleVector)
+	{
+		OffsetRotation = FQuat(TargetChainAxisNorm,FMath::DegreesToRadians(Settings.PoleVectorOffset));
+	}
 
 	// rotate the base of the chain to match the pole vectors
 	FTransform& BaseOfChain = OutTargetGlobalPose[TargetIndices[0]];
-	BaseOfChain.SetRotation(FinalRotation * BaseOfChain.GetRotation());
+	BaseOfChain.SetRotation(MatchingRotation * OffsetRotation * BaseOfChain.GetRotation());
 
 	// now update global pose of all bones within chain
 	for (int32 ChainIndex=0; ChainIndex < AllChildrenWithinChain.Num(); ++ChainIndex)
@@ -1431,7 +1446,7 @@ FVector FPoleVectorMatcher::CalculatePoleVector(
 {
 	check(!BoneIndices.IsEmpty())
 
-	const FVector ChainNormal = GetChainNormal(BoneIndices, GlobalPose);
+	const FVector ChainNormal = GetChainAxisNormalized(BoneIndices, GlobalPose);
 	const FVector UnitPoleAxis = GlobalPose[BoneIndices[0]].GetUnitAxis(PoleAxis);
 	const FVector PoleVector = FVector::VectorPlaneProject(UnitPoleAxis, ChainNormal);
 	return PoleVector.GetSafeNormal();
@@ -1458,7 +1473,7 @@ EAxis::Type FPoleVectorMatcher::GetMostDifferentAxis(
 	return MostDifferentAxis;
 }
 
-FVector FPoleVectorMatcher::GetChainNormal(
+FVector FPoleVectorMatcher::GetChainAxisNormalized(
 	const TArray<int32>& BoneIndices,
 	const TArray<FTransform>& GlobalPose)
 {
diff --git a/Engine/Plugins/Animation/IKRig/Source/IKRig/Public/Retargeter/IKRetargetProcessor.h b/Engine/Plugins/Animation/IKRig/Source/IKRig/Public/Retargeter/IKRetargetProcessor.h
--- a/Engine/Plugins/Animation/IKRig/Source/IKRig/Public/Retargeter/IKRetargetProcessor.h
+++ b/Engine/Plugins/Animation/IKRig/Source/IKRig/Public/Retargeter/IKRetargetProcessor.h
@@ -223,7 +223,7 @@ private:
 		const FTransform& Transform,
 		const FVector& InNormal);
 
-	static FVector GetChainNormal(
+	static FVector GetChainAxisNormalized(
 		const TArray<int32>& BoneIndices,
 		const TArray<FTransform>& GlobalPose);
 };
diff --git a/Engine/Plugins/Animation/IKRig/Source/IKRig/Public/Retargeter/IKRetargetSettings.h b/Engine/Plugins/Animation/IKRig/Source/IKRig/Public/Retargeter/IKRetargetSettings.h
--- a/Engine/Plugins/Animation/IKRig/Source/IKRig/Public/Retargeter/IKRetargetSettings.h
+++ b/Engine/Plugins/Animation/IKRig/Source/IKRig/Public/Retargeter/IKRetargetSettings.h
@@ -97,6 +97,10 @@ struct IKRIG_API FTargetChainFKSettings
 	UPROPERTY(EditAnywhere, BlueprintReadWrite, Category = "FK Adjustments", meta = (UIMin = "0.0", UIMax = "1.0", ClampMin = "0.0", ClampMax = "1.0"))
 	float PoleVectorMatching = 0.0f;
 
+	/** Default is False. When true, the original offset between the source/target pole vectors will be maintained when using Pole Vector Matching. */
+	UPROPERTY(EditAnywhere, BlueprintReadWrite, Category = "FK Adjustments")
+	bool PoleVectorMaintainOffset = false;
+
 	/** Range +/- 180. Default 0. An angular offset, in degrees, for the pole direction of the chain. */
 	UPROPERTY(EditAnywhere, BlueprintReadWrite, Category = "FK Adjustments", meta = (UIMin = "-180.0", UIMax = "180.0", ClampMin = "-180.0", ClampMax = "180.0"))
 	float PoleVectorOffset = 0.0f;
