Fixes bug, IK Rig editor warnings now come from running solver instances, not the editor instances.

diff --git a/Engine/Plugins/Animation/IKRig/Source/IKRig/Public/Rig/IKRigProcessor.h b/Engine/Plugins/Animation/IKRig/Source/IKRig/Public/Rig/IKRigProcessor.h
--- a/Engine/Plugins/Animation/IKRig/Source/IKRig/Public/Rig/IKRigProcessor.h
+++ b/Engine/Plugins/Animation/IKRig/Source/IKRig/Public/Rig/IKRigProcessor.h
@@ -110,6 +110,9 @@ public:
 
 	/** Used to propagate setting values from the source asset at runtime (settings that do not require re-initialization) */
 	void CopyAllInputsFromSourceAssetAtRuntime(const UIKRigDefinition* SourceAsset);
+
+	/** Get read-only access to currently running solvers */
+	const TArray<TObjectPtr<UIKRigSolver>>& GetSolvers(){ return Solvers; };
 	
 private:
 
diff --git a/Engine/Plugins/Animation/IKRig/Source/IKRigEditor/Private/RigEditor/SIKRigSolverStack.cpp b/Engine/Plugins/Animation/IKRig/Source/IKRigEditor/Private/RigEditor/SIKRigSolverStack.cpp
--- a/Engine/Plugins/Animation/IKRig/Source/IKRigEditor/Private/RigEditor/SIKRigSolverStack.cpp
+++ b/Engine/Plugins/Animation/IKRig/Source/IKRigEditor/Private/RigEditor/SIKRigSolverStack.cpp
@@ -16,6 +16,7 @@
 #include "Widgets/Input/SCheckBox.h"
 #include "Widgets/Text/SInlineEditableTextBlock.h"
 #include "SPositiveActionButton.h"
+#include "Rig/IKRigProcessor.h"
 
 #define LOCTEXT_NAMESPACE "SIKRigSolverStack"
 
@@ -156,9 +157,10 @@ void SIKRigSolverStackItem::Construct(
         ], OwnerTable);
 }
 
-bool SIKRigSolverStackItem::GetWarningMessage(FText& Message)
+bool SIKRigSolverStackItem::GetWarningMessage(FText& Message) const
 {
-	if (UIKRigSolver* Solver = GetSolver())
+	constexpr bool bFromAsset = false;
+	if (const UIKRigSolver* Solver = GetSolver(bFromAsset))
 	{
 		return Solver->GetWarningMessage(Message);
 	}
@@ -168,7 +170,8 @@ bool SIKRigSolverStackItem::GetWarningMessage(FText& Message)
 
 bool SIKRigSolverStackItem::IsSolverEnabled() const
 {
-	if (const UIKRigSolver* Solver = GetSolver())
+	constexpr bool bFromAsset = true;
+	if (const UIKRigSolver* Solver = GetSolver(bFromAsset))
 	{
 		return Solver->IsEnabled();
 	}
@@ -176,7 +179,7 @@ bool SIKRigSolverStackItem::IsSolverEnabled() const
 	return false;
 }
 
-UIKRigSolver* SIKRigSolverStackItem::GetSolver() const
+UIKRigSolver* SIKRigSolverStackItem::GetSolver(const bool bFromAsset) const
 {
 	if (!(StackElement.IsValid() && SolverStack.IsValid()))
 	{
@@ -186,8 +189,29 @@ UIKRigSolver* SIKRigSolverStackItem::GetSolver() const
 	{
 		return nullptr;
 	}
+	
 	const int32 SolverIndex = StackElement.Pin()->IndexInStack;
-	return SolverStack.Pin()->EditorController.Pin()->AssetController->GetSolverAtIndex(SolverIndex);
+	const TSharedPtr<FIKRigEditorController> EditorController = SolverStack.Pin()->EditorController.Pin();
+	
+	if (bFromAsset)
+	{
+		// get the solver stored in the asset
+		return EditorController->AssetController->GetSolverAtIndex(SolverIndex);
+	}
+	else
+	{
+		// get the currently running solver instance in the editor's runtime processor
+		// this is needed to report solver warnings dependent on being initialized
+		if (UIKRigProcessor* Processor = EditorController->GetIKRigProcessor())
+		{
+			if (Processor->GetSolvers().IsValidIndex(SolverIndex))
+			{
+				return Processor->GetSolvers()[SolverIndex];
+			}
+		}
+	}
+
+	return nullptr;
 }
 
 TSharedRef<FIKRigSolverStackDragDropOp> FIKRigSolverStackDragDropOp::New(TWeakPtr<FSolverStackElement> InElement)
diff --git a/Engine/Plugins/Animation/IKRig/Source/IKRigEditor/Public/RigEditor/SIKRigSolverStack.h b/Engine/Plugins/Animation/IKRig/Source/IKRigEditor/Public/RigEditor/SIKRigSolverStack.h
--- a/Engine/Plugins/Animation/IKRig/Source/IKRigEditor/Public/RigEditor/SIKRigSolverStack.h
+++ b/Engine/Plugins/Animation/IKRig/Source/IKRigEditor/Public/RigEditor/SIKRigSolverStack.h
@@ -55,12 +55,12 @@ public:
         TSharedRef<FSolverStackElement> InStackElement,
         TSharedPtr<SIKRigSolverStack> InSolverStack);
 
-	bool GetWarningMessage(FText& Message);
+	bool GetWarningMessage(FText& Message) const;
 
 	bool IsSolverEnabled() const;
 
 private:
-	UIKRigSolver* GetSolver() const;
+	UIKRigSolver* GetSolver(const bool bFromAsset) const;
 	TWeakPtr<FSolverStackElement> StackElement;
 	TWeakPtr<SIKRigSolverStack> SolverStack;
 };
