Fixes bug, root motion op doesn't work on characters with non-identity root orientation.

diff --git a/Engine/Plugins/Animation/IKRig/Source/IKRig/Private/Retargeter/RetargetOps/RootMotionGeneratorOp.cpp b/Engine/Plugins/Animation/IKRig/Source/IKRig/Private/Retargeter/RetargetOps/RootMotionGeneratorOp.cpp
--- a/Engine/Plugins/Animation/IKRig/Source/IKRig/Private/Retargeter/RetargetOps/RootMotionGeneratorOp.cpp
+++ b/Engine/Plugins/Animation/IKRig/Source/IKRig/Private/Retargeter/RetargetOps/RootMotionGeneratorOp.cpp
@@ -144,32 +144,34 @@ void URootMotionGeneratorOp::GenerateRootMotionFromTargetPelvis(
 {
 	// In this case, we are generating root motion "from scratch"
 	// using the target Pelvis bone as the source of the motion
-		
-	// calculate new root transform relative to current hips transform
-	OutRootTransform = InTargetGlobalPose[TargetPelvisIndex];
-	// optionally remove all rotation when calculating the offset
-	if (!bRotateWithPelvis)
-	{
-		OutRootTransform.SetRotation(TargetPelvisInRefPose.GetRotation());
-	}
 	
-	// optionally maintain the offset from the ref pose
 	if (bMaintainOffsetFromPelvis)
 	{
-		OutRootTransform = TargetPelvisRelativeToTargetRootRefPose * OutRootTransform;
+		// set root to the relative offset from the pelvis (recorded from ref pose)
+		OutRootTransform = TargetPelvisRelativeToTargetRootRefPose * InTargetGlobalPose[TargetPelvisIndex];
+	}
+	else
+	{
+		// snap root to the pelvis directly
+		OutRootTransform = InTargetGlobalPose[TargetPelvisIndex];
 	}
 
-	// optionally snap the root to the ground plane
+	// optionally remove all rotation (use static ref pose orientation)
+	if (!bRotateWithPelvis)
+	{
+		OutRootTransform.SetRotation(TargetRootInRefPose.GetRotation());
+	}
+
+	// adjust height of root
 	if (RootHeightSource == ERootMotionHeightSource::SnapToGround)
 	{
+		// snap the root to the ground plane
 		FVector RootTranslation = OutRootTransform.GetTranslation();
 		RootTranslation.Z = 0.f;
 		OutRootTransform.SetTranslation(RootTranslation);
-	}
-
-	// optionally snap the root to the height from the source
-	if (RootHeightSource == ERootMotionHeightSource::CopyHeightFromSource)
+	}else if (RootHeightSource == ERootMotionHeightSource::CopyHeightFromSource)
 	{
+		// snap the root to the height from the source
 		FVector RootTranslation = OutRootTransform.GetTranslation();
 		RootTranslation.Z = InSourceGlobalPose[SourceRootIndex].GetTranslation().Z;
 		OutRootTransform.SetTranslation(RootTranslation);
@@ -183,15 +185,16 @@ void URootMotionGeneratorOp::CopyRootMotionFromSourceRoot(
 {
 	// In this case, we are copying root motion from the source root
 	// But we also scale it based on the relative height of the source/target Pelvis
-		
-	// copy root motion from source root bone
-	OutRootTransform = InSourceGlobalPose[SourceRootIndex];
-	FVector NewRootLocation = OutRootTransform.GetTranslation();
-
-	// scale root motion by same scale factor applied to Pelvis (and modified by settings)
-	FVector RootDelta = NewRootLocation - SourceRootInRefPose.GetTranslation();
-	RootDelta *= Processor->GetRootRetargeter().GetGlobalScaleVector();
-	NewRootLocation = SourceRootInRefPose.GetTranslation() + RootDelta;
+
+	// rotation is the original target root rotation in ref pose plus the current rotation delta of the source
+	const FQuat SourceRootRotationDelta = InSourceGlobalPose[SourceRootIndex].GetRotation() * SourceRootInRefPose.GetRotation().Inverse();
+	OutRootTransform.SetRotation(SourceRootRotationDelta * TargetRootInRefPose.GetRotation());
+
+	// scale root translation by same scale factor applied to Pelvis (and modified by settings)
+	FVector NewRootLocation = InSourceGlobalPose[SourceRootIndex].GetLocation();
+	FVector RootTranslationDelta = NewRootLocation - SourceRootInRefPose.GetTranslation();
+	RootTranslationDelta *= Processor->GetRootRetargeter().GetGlobalScaleVector();
+	NewRootLocation = SourceRootInRefPose.GetTranslation() + RootTranslationDelta;
 
 	// optionally snap the root to the ground plane
 	if (RootHeightSource == ERootMotionHeightSource::SnapToGround)
