Hide SetRetargetDefinition from BP/Python until or if we make bone chains editable.

diff --git a/Engine/Plugins/Animation/IKRig/Source/IKRigEditor/Public/RigEditor/IKRigController.h b/Engine/Plugins/Animation/IKRig/Source/IKRigEditor/Public/RigEditor/IKRigController.h
--- a/Engine/Plugins/Animation/IKRig/Source/IKRigEditor/Public/RigEditor/IKRigController.h
+++ b/Engine/Plugins/Animation/IKRig/Source/IKRigEditor/Public/RigEditor/IKRigController.h
@@ -321,10 +321,6 @@ public:
 	UFUNCTION(BlueprintCallable, Category=IKRig)
 	FName GetRetargetRoot() const;
 
-	// Replace the entire retarget definition (includes all bone chains and the retarget root setting)
-	UFUNCTION(BlueprintCallable, Category=IKRig)
-	void SetRetargetDefinition(const FRetargetDefinition& RetargetDefinition) const;
-
 	// Analyse the skeleton to see if it matches a known template and automatically generates all retarget chains and sets the retarget root
 	// Returns true if a matching skeletal template was found and the retarget definition for it was applied.
 	UFUNCTION(BlueprintCallable, Category=IKRig)
@@ -339,6 +335,9 @@ public:
 	// RETARGETING C++ ONLY API
 	//
 
+	// Replace the entire retarget definition (includes all bone chains and the retarget root setting)
+	void SetRetargetDefinition(const FRetargetDefinition& RetargetDefinition) const;
+
 	// Auto generates a retarget definition and returns the results
 	void AutoGenerateRetargetDefinition(FAutoCharacterizeResults& Results) const;
 
