Fixes issue, IK Rig goals are greyed out in the details panel.

diff --git a/Engine/Plugins/Animation/IKRig/Source/IKRig/Public/AnimNodes/AnimNode_IKRig.h b/Engine/Plugins/Animation/IKRig/Source/IKRig/Public/AnimNodes/AnimNode_IKRig.h
--- a/Engine/Plugins/Animation/IKRig/Source/IKRig/Public/AnimNodes/AnimNode_IKRig.h
+++ b/Engine/Plugins/Animation/IKRig/Source/IKRig/Public/AnimNodes/AnimNode_IKRig.h
@@ -32,7 +32,7 @@ struct IKRIG_API FAnimNode_IKRig : public FAnimNode_CustomProperty
 	TObjectPtr<UIKRigDefinition> RigDefinitionAsset = nullptr;
 
 	/** The input goal transforms used by the IK Rig solvers.*/
-	UPROPERTY(EditAnywhere, BlueprintReadWrite, Category = Goal, meta=(NeverAsPin, EditCondition=false, EditConditionHides))
+	UPROPERTY(EditAnywhere, BlueprintReadWrite, Category = Goal, meta=(NeverAsPin))
 	TArray<FIKRigGoal> Goals;
 
 	/** optionally ignore the input pose and start from the reference pose each solve */
diff --git a/Engine/Plugins/Animation/IKRig/Source/IKRigDeveloper/Private/AnimGraphNode_IKRig.cpp b/Engine/Plugins/Animation/IKRig/Source/IKRigDeveloper/Private/AnimGraphNode_IKRig.cpp
--- a/Engine/Plugins/Animation/IKRig/Source/IKRigDeveloper/Private/AnimGraphNode_IKRig.cpp
+++ b/Engine/Plugins/Animation/IKRig/Source/IKRigDeveloper/Private/AnimGraphNode_IKRig.cpp
@@ -698,6 +698,8 @@ void UAnimGraphNode_IKRig::CustomizeDetails(IDetailLayoutBuilder& DetailBuilder)
 		IDetailCategoryBuilder& GoalsCategoryBuilder = DetailBuilder.EditCategory(GET_MEMBER_NAME_CHECKED(FAnimNode_IKRig, Goals));
 		GoalsCategoryBuilder.AddCustomBuilder(InputArgumentGroup);
 	}
+	// Hide normal goals properties
+	DetailBuilder.HideCategory("Goal");
 
 	// Handle property changed notification
 	const FSimpleDelegate OnValueChanged = FSimpleDelegate::CreateLambda([&DetailBuilder]()
