Fixes, source mesh not moving in retargeter.

diff --git a/Engine/Plugins/Animation/IKRig/Source/IKRigEditor/Private/RetargetEditor/IKRetargetDefaultMode.cpp b/Engine/Plugins/Animation/IKRig/Source/IKRigEditor/Private/RetargetEditor/IKRetargetDefaultMode.cpp
--- a/Engine/Plugins/Animation/IKRig/Source/IKRigEditor/Private/RetargetEditor/IKRetargetDefaultMode.cpp
+++ b/Engine/Plugins/Animation/IKRig/Source/IKRigEditor/Private/RetargetEditor/IKRetargetDefaultMode.cpp
@@ -40,10 +40,6 @@ void FIKRetargetDefaultMode::Initialize()
 		return; 
 	}
 
-	// update offsets on preview meshes
-	Controller->AddOffsetToMeshComponent(FVector::ZeroVector, Controller->SourceSkelMeshComponent);
-	Controller->AddOffsetToMeshComponent(FVector::ZeroVector, Controller->TargetSkelMeshComponent);
-
 	bIsInitialized = true;
 }
 
@@ -228,49 +224,6 @@ void FIKRetargetDefaultMode::RenderDebugProxies(FPrimitiveDrawInterface* PDI, co
 	PDI->SetHitProxy(nullptr);
 }
 
-bool FIKRetargetDefaultMode::AllowWidgetMove()
-{
-	return false;
-}
-
-bool FIKRetargetDefaultMode::ShouldDrawWidget() const
-{
-	return UsesTransformWidget(CurrentWidgetMode);
-}
-
-bool FIKRetargetDefaultMode::UsesTransformWidget() const
-{
-	return UsesTransformWidget(CurrentWidgetMode);
-}
-
-bool FIKRetargetDefaultMode::UsesTransformWidget(UE::Widget::EWidgetMode CheckMode) const
-{
-	const TSharedPtr<FIKRetargetEditorController> Controller = EditorController.Pin();
-	if (!Controller.IsValid())
-	{
-		return false; 
-	}
-	
-	const bool bTranslating = CheckMode == UE::Widget::EWidgetMode::WM_Translate;
-	return bTranslating && IsValid(Controller->GetSelectedMesh());
-}
-
-FVector FIKRetargetDefaultMode::GetWidgetLocation() const
-{
-	const TSharedPtr<FIKRetargetEditorController> Controller = EditorController.Pin();
-	if (!Controller.IsValid())
-	{
-		return FVector::ZeroVector; 
-	}
-	
-	if (!Controller->GetSelectedMesh())
-	{
-		return FVector::ZeroVector; // shouldn't get here
-	}
-
-	return Controller->GetSelectedMesh()->GetComponentTransform().GetLocation();
-}
-
 bool FIKRetargetDefaultMode::HandleClick(FEditorViewportClient* InViewportClient, HHitProxy* HitProxy, const FViewportClick& Click)
 {
 	const TSharedPtr<FIKRetargetEditorController> Controller = EditorController.Pin();
@@ -283,15 +236,6 @@ bool FIKRetargetDefaultMode::HandleClick(FEditorViewportClient* InViewportClient
 	const bool bCtrlOrShiftHeld = Click.IsControlDown() || Click.IsShiftDown();
 	const ESelectionEdit EditMode = bCtrlOrShiftHeld ? ESelectionEdit::Add : ESelectionEdit::Replace;
 	
-	// did we click on an actor in the viewport?
-	const bool bHitActor = HitProxy && HitProxy->IsA(HActor::StaticGetType());
-	if (bLeftButtonClicked && bHitActor)
-	{
-		const HActor* ActorProxy = static_cast<HActor*>(HitProxy);
-		Controller->SetSelectedMesh(ConstCast(ActorProxy->PrimComponent));
-		return true;
-	}
-
 	// did we click on a bone in the viewport?
 	const bool bHitBone = HitProxy && HitProxy->IsA(HIKRetargetEditorBoneProxy::StaticGetType());
 	if (bLeftButtonClicked && bHitBone)
@@ -327,117 +271,6 @@ bool FIKRetargetDefaultMode::HandleClick(FEditorViewportClient* InViewportClient
 	return true;
 }
 
-bool FIKRetargetDefaultMode::StartTracking(FEditorViewportClient* InViewportClient, FViewport* InViewport)
-{
-	return HandleBeginTransform(InViewportClient);
-}
-
-bool FIKRetargetDefaultMode::EndTracking(FEditorViewportClient* InViewportClient, FViewport* InViewport)
-{
-	return HandleEndTransform();
-}
-
-bool FIKRetargetDefaultMode::BeginTransform(const FGizmoState& InState)
-{
-	return HandleBeginTransform(Owner->GetFocusedViewportClient());
-}
-
-bool FIKRetargetDefaultMode::EndTransform(const FGizmoState& InState)
-{
-	return HandleEndTransform();
-}
-
-bool FIKRetargetDefaultMode::HandleBeginTransform(const FEditorViewportClient* InViewportClient)
-{
-	if (!InViewportClient)
-	{
-		return false;
-	}
-	
-	bIsTranslating = false;
-
-	// not manipulating any widget axes, so stop tracking
-	const EAxisList::Type CurrentAxis = InViewportClient->GetCurrentWidgetAxis();
-	if (CurrentAxis == EAxisList::None)
-	{
-		return false; 
-	}
-
-	const TSharedPtr<FIKRetargetEditorController> Controller = EditorController.Pin();
-	if (!Controller.IsValid())
-	{
-		return false; // invalid editor state
-	}
-
-	const bool bTranslating = InViewportClient->GetWidgetMode() == UE::Widget::EWidgetMode::WM_Translate;
-	if (bTranslating && IsValid(Controller->GetSelectedMesh()))
-	{
-		bIsTranslating = true;
-		GEditor->BeginTransaction(LOCTEXT("MovePreviewMesh", "Move Preview Mesh"));
-		Controller->AssetController->GetAsset()->Modify();
-		return true;
-	}
-
-	return false;
-}
-
-bool FIKRetargetDefaultMode::HandleEndTransform()
-{
-	GEditor->EndTransaction();
-	bIsTranslating = false;
-	return true;
-}
-
-bool FIKRetargetDefaultMode::InputDelta(
-	FEditorViewportClient* InViewportClient,
-	FViewport* InViewport,
-	FVector& InDrag,
-	FRotator& InRot,
-	FVector& InScale)
-{
-	const TSharedPtr<FIKRetargetEditorController> Controller = EditorController.Pin();
-	if (!Controller.IsValid())
-	{
-		return false; 
-	}
-	
-	if (!(bIsTranslating && IsValid(Controller->GetSelectedMesh())))
-	{
-		return false; // not handled
-	}
-
-	if(InViewportClient->GetWidgetMode() != UE::Widget::WM_Translate)
-	{
-		return false;
-	}
-
-	Controller->AddOffsetToMeshComponent(InDrag, Controller->GetSelectedMesh());
-	
-	return true;
-}
-
-bool FIKRetargetDefaultMode::GetCustomDrawingCoordinateSystem(FMatrix& InMatrix, void* InData)
-{
-	const TSharedPtr<FIKRetargetEditorController> Controller = EditorController.Pin();
-	if (!Controller.IsValid())
-	{
-		return false; 
-	}
-
-	if (!Controller->GetSelectedMesh())
-	{
-		return false;
-	}
-
-	InMatrix = Controller->GetSelectedMesh()->GetComponentTransform().ToMatrixNoScale().RemoveTranslation();
-	return true;
-}
-
-bool FIKRetargetDefaultMode::GetCustomInputCoordinateSystem(FMatrix& InMatrix, void* InData)
-{
-	return GetCustomDrawingCoordinateSystem(InMatrix, InData);
-}
-
 void FIKRetargetDefaultMode::Enter()
 {
 	IPersonaEditMode::Enter();
@@ -463,25 +296,6 @@ void FIKRetargetDefaultMode::Exit()
 	IPersonaEditMode::Exit();
 }
 
-UDebugSkelMeshComponent* FIKRetargetDefaultMode::GetCurrentlyEditedMesh() const
-{
-	const TSharedPtr<FIKRetargetEditorController> Controller = EditorController.Pin();
-	if (!Controller.IsValid())
-	{
-		return nullptr; 
-	}
-	
-	return SkeletonMode == ERetargetSourceOrTarget::Source ? Controller->SourceSkelMeshComponent : Controller->TargetSkelMeshComponent;
-}
-
-void FIKRetargetDefaultMode::ApplyOffsetToMeshTransform(const FVector& Offset, USceneComponent* Component)
-{
-	constexpr bool bSweep = false;
-	constexpr FHitResult* OutSweepHitResult = nullptr;
-	constexpr ETeleportType Teleport = ETeleportType::ResetPhysics;
-	Component->SetWorldLocation(Offset, bSweep, OutSweepHitResult, Teleport);
-}
-
 void FIKRetargetDefaultMode::Tick(FEditorViewportClient* ViewportClient, float DeltaTime)
 {
 	FEdMode::Tick(ViewportClient, DeltaTime);
diff --git a/Engine/Plugins/Animation/IKRig/Source/IKRigEditor/Private/RetargetEditor/IKRetargetEditor.cpp b/Engine/Plugins/Animation/IKRig/Source/IKRigEditor/Private/RetargetEditor/IKRetargetEditor.cpp
--- a/Engine/Plugins/Animation/IKRig/Source/IKRigEditor/Private/RetargetEditor/IKRetargetEditor.cpp
+++ b/Engine/Plugins/Animation/IKRig/Source/IKRigEditor/Private/RetargetEditor/IKRetargetEditor.cpp
@@ -512,8 +512,8 @@ FString FIKRetargetEditor::GetWorldCentricTabPrefix() const
 void FIKRetargetEditor::Tick(float DeltaTime)
 {
 	// update with latest offsets
-	EditorController->AddOffsetToMeshComponent(FVector::ZeroVector, EditorController->SourceSkelMeshComponent);
-	EditorController->AddOffsetToMeshComponent(FVector::ZeroVector, EditorController->TargetSkelMeshComponent);
+	EditorController->UpdateMeshOffset(ERetargetSourceOrTarget::Source);
+	EditorController->UpdateMeshOffset(ERetargetSourceOrTarget::Target);
 
 	// retargeter IK planting must be reset when time is reversed or playback jumps ahead 
 	const float CurrentTime = EditorController->SourceAnimInstance->GetCurrentTime();
@@ -633,6 +633,9 @@ void FIKRetargetEditor::HandlePreviewSceneCreated(const TSharedRef<IPersonaPrevi
 	EditorController->TargetSkelMeshComponent->SetAnimationMode(EAnimationMode::AnimationCustomMode);
 	
 	// must call AddComponent() BEFORE assigning the mesh to prevent auto-assignment of a default anim instance
+	EditorController->SourceRootComponent = NewObject<USceneComponent>(Actor);
+	EditorController->SourceSkelMeshComponent->AttachToComponent(EditorController->SourceRootComponent, FAttachmentTransformRules::KeepRelativeTransform);
+	InPersonaPreviewScene->AddComponent(EditorController->SourceRootComponent, FTransform::Identity);
 	InPersonaPreviewScene->AddComponent(EditorController->SourceSkelMeshComponent, FTransform::Identity);
     InPersonaPreviewScene->AddComponent(EditorController->TargetSkelMeshComponent, FTransform::Identity);
     
diff --git a/Engine/Plugins/Animation/IKRig/Source/IKRigEditor/Private/RetargetEditor/IKRetargetEditorController.cpp b/Engine/Plugins/Animation/IKRig/Source/IKRigEditor/Private/RetargetEditor/IKRetargetEditorController.cpp
--- a/Engine/Plugins/Animation/IKRig/Source/IKRigEditor/Private/RetargetEditor/IKRetargetEditorController.cpp
+++ b/Engine/Plugins/Animation/IKRig/Source/IKRigEditor/Private/RetargetEditor/IKRetargetEditorController.cpp
@@ -316,29 +316,28 @@ UIKRetargetAnimInstance* FIKRetargetEditorController::GetAnimInstance(
 	return SourceOrTarget == ERetargetSourceOrTarget::Source ? SourceAnimInstance.Get() : TargetAnimInstance.Get();
 }
 
-void FIKRetargetEditorController::AddOffsetToMeshComponent(const FVector& Offset, USceneComponent* MeshComponent) const
+void FIKRetargetEditorController::UpdateMeshOffset(ERetargetSourceOrTarget SourceOrTarget) const
 {
-	UIKRetargeter* Asset = AssetController->GetAsset();
-	FVector Position;
-	float Scale;
-	if (MeshComponent == TargetSkelMeshComponent)
-	{
-		Asset->TargetMeshOffset += Offset;
-		Position = Asset->TargetMeshOffset;
-		Scale = Asset->TargetMeshScale;
-	}
-	else
+	const UIKRetargeter* Asset = AssetController->GetAsset();
+	if (!Asset)
 	{
-		Asset->SourceMeshOffset += Offset;
-		Position = Asset->SourceMeshOffset;
-		Scale = 1.0f;
+		return;
 	}
 
+	const bool bIsSource  = SourceOrTarget == ERetargetSourceOrTarget::Source;
+	USceneComponent* Component = bIsSource ? SourceRootComponent : GetSkeletalMeshComponent(ERetargetSourceOrTarget::Target);
+    if (!Component)
+    {
+    	return;
+    }
+	
+	const FVector Position = bIsSource ? Asset->SourceMeshOffset : Asset->TargetMeshOffset;
+	const float Scale = bIsSource ? 1.0f : Asset->TargetMeshScale;
 	constexpr bool bSweep = false;
 	constexpr FHitResult* OutSweepHitResult = nullptr;
 	constexpr ETeleportType Teleport = ETeleportType::ResetPhysics;
-	MeshComponent->SetWorldLocation(Position, bSweep, OutSweepHitResult, Teleport);
-	MeshComponent->SetWorldScale3D(FVector(Scale,Scale,Scale));
+	Component->SetWorldLocation(Position, bSweep, OutSweepHitResult, Teleport);
+	Component->SetWorldScale3D(FVector(Scale,Scale,Scale));
 }
 
 bool FIKRetargetEditorController::GetCameraTargetForSelection(FSphere& OutTarget) const
@@ -775,20 +774,6 @@ void FIKRetargetEditorController::SetSourceOrTargetMode(ERetargetSourceOrTarget
 	RefreshPoseList();
 }
 
-void FIKRetargetEditorController::SetSelectedMesh(UPrimitiveComponent* InMeshComponent)
-{
-	SelectedMesh = InMeshComponent;
-	if (SelectedMesh)
-	{
-		LastSelectedItem = ERetargetSelectionType::MESH;
-	}
-}
-
-UPrimitiveComponent* FIKRetargetEditorController::GetSelectedMesh() const
-{
-	return SelectedMesh;
-}
-
 void FIKRetargetEditorController::EditBoneSelection(
 	const TArray<FName>& InBoneNames,
 	ESelectionEdit EditMode,
@@ -802,9 +787,7 @@ void FIKRetargetEditorController::EditBoneSelection(
 	}
 
 	LastSelectedItem = ERetargetSelectionType::BONE;
-
-	// deselect mesh
-	SetSelectedMesh(nullptr);
+	
 	SetRootSelected(false);
 	
 	switch (EditMode)
@@ -870,7 +853,6 @@ void FIKRetargetEditorController::EditChainSelection(
 	const bool bFromChainsView)
 {
 	// deselect others
-	SetSelectedMesh(nullptr);
 	SetRootSelected(false);
 
 	LastSelectedItem = ERetargetSelectionType::CHAIN;
@@ -971,7 +953,6 @@ void FIKRetargetEditorController::ClearSelection(const bool bKeepBoneSelection)
 {
 	// clear root and mesh selection
 	SetRootSelected(false);
-	SetSelectedMesh(nullptr);
 	
 	// deselect all chains
 	if (ChainsView.IsValid())
@@ -983,7 +964,6 @@ void FIKRetargetEditorController::ClearSelection(const bool bKeepBoneSelection)
 	// clear bone selection
 	if (!bKeepBoneSelection)
 	{
-		SetSelectedMesh(nullptr);
 		SetRootSelected(false);
 		SelectedBoneNames[ERetargetSourceOrTarget::Source].Reset();
 		SelectedBoneNames[ERetargetSourceOrTarget::Target].Reset();
diff --git a/Engine/Plugins/Animation/IKRig/Source/IKRigEditor/Public/RetargetEditor/IKRetargetDefaultMode.h b/Engine/Plugins/Animation/IKRig/Source/IKRigEditor/Public/RetargetEditor/IKRetargetDefaultMode.h
--- a/Engine/Plugins/Animation/IKRig/Source/IKRigEditor/Public/RetargetEditor/IKRetargetDefaultMode.h
+++ b/Engine/Plugins/Animation/IKRig/Source/IKRigEditor/Public/RetargetEditor/IKRetargetDefaultMode.h
@@ -26,43 +26,23 @@ public:
 
 	/** IPersonaEditMode interface */
 	virtual bool GetCameraTarget(FSphere& OutTarget) const override;
-	virtual class IPersonaPreviewScene& GetAnimPreviewScene() const override;
+	virtual IPersonaPreviewScene& GetAnimPreviewScene() const override;
 	/** END IPersonaEditMode interface */
 
 	/** FEdMode interface */
 	virtual void Initialize() override;
 	virtual void Tick(FEditorViewportClient* ViewportClient, float DeltaTime) override;
 	virtual void Render(const FSceneView* View, FViewport* Viewport, FPrimitiveDrawInterface* PDI) override;
-	virtual bool IsCompatibleWith(FEditorModeID OtherModeID) const override { return true; }
-	virtual bool AllowWidgetMove() override;
-	virtual bool ShouldDrawWidget() const override;
-	virtual bool UsesTransformWidget() const override;
-	virtual bool UsesTransformWidget(UE::Widget::EWidgetMode CheckMode) const override;
-	virtual FVector GetWidgetLocation() const override;
+	virtual bool IsCompatibleWith(FEditorModeID OtherModeID) const override { return true; };
 	virtual bool HandleClick(FEditorViewportClient* InViewportClient, HHitProxy* HitProxy, const FViewportClick& Click) override;
-	virtual bool StartTracking(FEditorViewportClient* InViewportClient, FViewport* InViewport) override;
-	virtual bool EndTracking(FEditorViewportClient* InViewportClient, FViewport* InViewport) override;
-	virtual bool BeginTransform(const FGizmoState& InState) override;
-	virtual bool EndTransform(const FGizmoState& InState) override;
-	virtual bool InputDelta(FEditorViewportClient* InViewportClient, FViewport* InViewport, FVector& InDrag, FRotator& InRot, FVector& InScale) override;
-	virtual bool GetCustomDrawingCoordinateSystem(FMatrix& InMatrix, void* InData) override;
-	virtual bool GetCustomInputCoordinateSystem(FMatrix& InMatrix, void* InData) override;
-
 	virtual void Enter() override;
 	virtual void Exit() override;
-	// IS THIS NEEDED
-	virtual bool IsSelectionAllowed( AActor* InActor, bool bInSelection ) const override { return true; }
 	/** END FEdMode interface */
 
 private:
 	void RenderDebugProxies(FPrimitiveDrawInterface* PDI, const FIKRetargetEditorController* Controller) const;
-	static void ApplyOffsetToMeshTransform(const FVector& Offset, USceneComponent* Component);
-
-	bool HandleBeginTransform(const FEditorViewportClient* InViewportClient);
-	bool HandleEndTransform();
 	
 	// the skeleton currently being edited
-	UDebugSkelMeshComponent* GetCurrentlyEditedMesh() const;
 	ERetargetSourceOrTarget SkeletonMode;
 	
 	/** The hosting app */
diff --git a/Engine/Plugins/Animation/IKRig/Source/IKRigEditor/Public/RetargetEditor/IKRetargetEditorController.h b/Engine/Plugins/Animation/IKRig/Source/IKRigEditor/Public/RetargetEditor/IKRetargetEditorController.h
--- a/Engine/Plugins/Animation/IKRig/Source/IKRigEditor/Public/RetargetEditor/IKRetargetEditorController.h
+++ b/Engine/Plugins/Animation/IKRig/Source/IKRigEditor/Public/RetargetEditor/IKRetargetEditorController.h
@@ -130,6 +130,9 @@ public:
 	UDebugSkelMeshComponent* GetSkeletalMeshComponent(const ERetargetSourceOrTarget SourceOrTarget) const;
 	UDebugSkelMeshComponent* SourceSkelMeshComponent;
 	UDebugSkelMeshComponent* TargetSkelMeshComponent;
+	// this root component is used as a parent of the source skeletal mesh to allow us to translate the source.
+	// we can't offset the source mesh component itself because that conflicts with root motion
+	USceneComponent* SourceRootComponent;
 
 	// viewport anim instance
 	UIKRetargetAnimInstance* GetAnimInstance(const ERetargetSourceOrTarget SourceOrTarget) const;
@@ -227,12 +230,6 @@ public:
 		const bool bFromHierarchyView = false);
 	const TArray<FName>& GetSelectedBones() const {return SelectedBoneNames[CurrentlyEditingSourceOrTarget]; };
 	// END bone selection
-	
-	// SELECTION - MESHES (viewport view)
-	void SetSelectedMesh(UPrimitiveComponent* InComponent);
-	UPrimitiveComponent* GetSelectedMesh() const;
-	void AddOffsetToMeshComponent(const FVector& Offset, USceneComponent* MeshComponent) const;
-	// END mesh selection
 
 	// SELECTION - CHAINS (viewport or chains view)
 	void EditChainSelection(
@@ -321,6 +318,8 @@ public:
 	// render the skeleton's in the viewport (either source or target)
 	void RenderSkeleton(FPrimitiveDrawInterface* PDI, ERetargetSourceOrTarget SourceOrTarget) const;
 
+	void UpdateMeshOffset(ERetargetSourceOrTarget SourceOrTarget) const;
+
 private:
 
 	TArray<FName> GetSelectedBonesAndChildren() const;
