Fixes crash in retarget editor when swapping target mesh while posing a bone on the mesh.

diff --git a/Engine/Plugins/Animation/IKRig/Source/IKRig/Private/Retargeter/IKRetargetProcessor.cpp b/Engine/Plugins/Animation/IKRig/Source/IKRig/Private/Retargeter/IKRetargetProcessor.cpp
--- a/Engine/Plugins/Animation/IKRig/Source/IKRig/Private/Retargeter/IKRetargetProcessor.cpp
+++ b/Engine/Plugins/Animation/IKRig/Source/IKRig/Private/Retargeter/IKRetargetProcessor.cpp
@@ -1494,6 +1494,12 @@ void UIKRetargetProcessor::Initialize(
 		UIKRetargeter* InRetargeterAsset,
 		const bool bSuppressWarnings)
 {
+	// don't attempt reinitialization with same assets
+	if (RetargeterAsset == InRetargeterAsset && AssetVersionInitializedWith == InRetargeterAsset->GetVersion())
+	{
+		return;
+	}
+	
 	// reset all initialized flags
 	bIsInitialized = false;
 	bRootsInitialized = false;
@@ -2249,10 +2255,13 @@ bool UIKRetargetProcessor::WasInitializedWithTheseAssets(
 	{
 		return false;
 	}
-	
-	// asset has been modified in a way that requires reinitialization
-	if (AssetVersionInitializedWith != InRetargetAsset->GetVersion())
+
+	// check that the retarget asset is the same as what we initialized with
+	const bool bSameAsset = InRetargetAsset == RetargeterAsset;
+	const bool bSameVersion = AssetVersionInitializedWith == InRetargetAsset->GetVersion();
+	if (!(bSameAsset && bSameVersion))
 	{
+		// asset has been modified in a way that requires reinitialization
 		return false;
 	}
 
@@ -2263,12 +2272,7 @@ bool UIKRetargetProcessor::WasInitializedWithTheseAssets(
 	const bool bTargetMatches = InTargetMesh == TargetRetargetSkeleton.SkeletalMesh;
 	if (!(bSourceMatches && bTargetMatches))
 	{
-		return false;
-	}
-
-	// check that the retarget asset is the same as what we initialized with
-	if (InRetargetAsset != RetargeterAsset)
-	{
+		// skeletal mesh swapped out
 		return false;
 	}
 
@@ -2277,6 +2281,7 @@ bool UIKRetargetProcessor::WasInitializedWithTheseAssets(
 	const bool bTargetHasSameNumberOfBones = InTargetMesh->GetRefSkeleton().GetNum() == TargetRetargetSkeleton.BoneNames.Num();
 	if (!(bSourceHasSameNumberOfBones && bTargetHasSameNumberOfBones))
 	{
+		// skeletal mesh modified
 		return false;
 	}
 	
diff --git a/Engine/Plugins/Animation/IKRig/Source/IKRigEditor/Private/RetargetEditor/IKRetargetEditPoseMode.cpp b/Engine/Plugins/Animation/IKRig/Source/IKRigEditor/Private/RetargetEditor/IKRetargetEditPoseMode.cpp
--- a/Engine/Plugins/Animation/IKRig/Source/IKRigEditor/Private/RetargetEditor/IKRetargetEditPoseMode.cpp
+++ b/Engine/Plugins/Animation/IKRig/Source/IKRigEditor/Private/RetargetEditor/IKRetargetEditPoseMode.cpp
@@ -435,21 +435,26 @@ void FIKRetargetEditPoseMode::UpdateWidgetTransform()
 		BoneEdit.GlobalTransform = FTransform::Identity;
 		return;
 	}
+
+	const FReferenceSkeleton& RefSkeleton = SkeletalMesh->GetRefSkeleton();
+	const int32 BoneIndex = RefSkeleton.FindBoneIndex(SelectedBones.Last());
+	if (BoneIndex == INDEX_NONE)
+	{
+		BoneEdit.GlobalTransform = FTransform::Identity;
+		return;
+	}
 	
 	float Scale;
 	FVector Offset;
 	GetEditedComponentScaleAndOffset(Scale,Offset);
 
-	const UIKRetargeterController* AssetController = Controller->AssetController;
-	const FIKRetargetPose& RetargetPose = AssetController->GetCurrentRetargetPose(SourceOrTarget);
-	const FReferenceSkeleton& RefSkeleton = SkeletalMesh->GetRefSkeleton();
-
 	BoneEdit.Name = SelectedBones.Last();
-	BoneEdit.Index = RefSkeleton.FindBoneIndex(BoneEdit.Name);
+	BoneEdit.Index = BoneIndex;
 	BoneEdit.GlobalTransform = Controller->GetGlobalRetargetPoseOfBone(SourceOrTarget, BoneEdit.Index, Scale, Offset);
 	BoneEdit.AccumulatedGlobalOffset = FQuat::Identity;
 
 	BoneEdit.PreviousDeltaRotation.Reset();
+	const FIKRetargetPose& RetargetPose =  Controller->AssetController->GetCurrentRetargetPose(SourceOrTarget);
 	for (int32 SelectionIndex=0; SelectionIndex<SelectedBones.Num(); ++SelectionIndex)
 	{
 		FQuat PrevDeltaRotation = RetargetPose.GetDeltaRotationForBone(SelectedBones[SelectionIndex]);
