Fixes crash when LODing out last bone in skeleton in Nanite skel mesh.

diff --git a/Engine/Plugins/Animation/IKRig/Source/IKRig/Private/Retargeter/IKRetargetProcessor.cpp b/Engine/Plugins/Animation/IKRig/Source/IKRig/Private/Retargeter/IKRetargetProcessor.cpp
--- a/Engine/Plugins/Animation/IKRig/Source/IKRig/Private/Retargeter/IKRetargetProcessor.cpp
+++ b/Engine/Plugins/Animation/IKRig/Source/IKRig/Private/Retargeter/IKRetargetProcessor.cpp
@@ -233,8 +233,8 @@ int32 FRetargetSkeleton::GetCachedEndOfBranchIndex(const int32 InBoneIndex) cons
 
 	const int32 NumBones = BoneNames.Num();
 	
-	// if we're asking for root's branch, get the last bone  
-	if (InBoneIndex == 0)
+	// if we're asking for the first or last bone, return the last bone  
+	if (InBoneIndex == 0 || InBoneIndex + 1 >= NumBones)
 	{
 		CachedEndOfBranchIndices[InBoneIndex] = NumBones-1;
 		return CachedEndOfBranchIndices[InBoneIndex];
@@ -247,14 +247,21 @@ int32 FRetargetSkeleton::GetCachedEndOfBranchIndex(const int32 InBoneIndex) cons
 
 	// if next child bone's parent is less than or equal to StartParentIndex,
 	// we are leaving the branch so no need to go further
-	while (ParentIndex > StartParentIndex && BoneIndex < NumBones)
+	int32 BoneIndexAtEndOfBranch = RETARGETSKELETON_INVALID_BRANCH_INDEX;
+	while (ParentIndex > StartParentIndex)
 	{
-		CachedEndOfBranchIndices[InBoneIndex] = BoneIndex;
-				
+		BoneIndexAtEndOfBranch = BoneIndex;
 		BoneIndex++;
+		if (BoneIndex >= NumBones)
+		{
+			break;
+		}
 		ParentIndex = GetParentIndex(BoneIndex);
 	}
 
+	// set once (outside of while loop above) to avoid potential race condition
+	CachedEndOfBranchIndices[InBoneIndex] = BoneIndexAtEndOfBranch;
+
 	return CachedEndOfBranchIndices[InBoneIndex];
 }
 
diff --git a/Engine/Plugins/Animation/IKRig/Source/IKRig/Private/Rig/IKRigSkeleton.cpp b/Engine/Plugins/Animation/IKRig/Source/IKRig/Private/Rig/IKRigSkeleton.cpp
--- a/Engine/Plugins/Animation/IKRig/Source/IKRig/Private/Rig/IKRigSkeleton.cpp
+++ b/Engine/Plugins/Animation/IKRig/Source/IKRig/Private/Rig/IKRigSkeleton.cpp
@@ -174,8 +174,8 @@ int32 FIKRigSkeleton::GetCachedEndOfBranchIndex(const int32 InBoneIndex) const
 
 	const int32 NumBones = BoneNames.Num();
 	
-	// if we're asking for root's branch, get the last bone  
-	if (InBoneIndex == 0)
+	// if we're asking for the first or last bone, return the last bone  
+	if (InBoneIndex == 0 || InBoneIndex + 1 >= NumBones)
 	{
 		CachedEndOfBranchIndices[InBoneIndex] = NumBones-1;
 		return CachedEndOfBranchIndices[InBoneIndex];
@@ -183,22 +183,29 @@ int32 FIKRigSkeleton::GetCachedEndOfBranchIndex(const int32 InBoneIndex) const
 
 	CachedEndOfBranchIndices[InBoneIndex] = INDEX_NONE;
 	
-	// store ref parent
-	const int32 RefParentIndex = GetParentIndex(InBoneIndex);
+	// store start parent
+	const int32 StartParentIndex = GetParentIndex(InBoneIndex);
 
 	int32 BoneIndex = InBoneIndex + 1;
 	int32 ParentIndex = GetParentIndex(BoneIndex);
 
-	// if next bellow bone's parent is less than or equal to RefParentIndex, we are leaving the branch
-	// so ne need to go further
-	while (ParentIndex > RefParentIndex && BoneIndex < NumBones)
+	// if next child bone's parent is less than or equal to StartParentIndex,
+	// we are leaving the branch so no need to go further
+	int32 BoneIndexAtEndOfBranch = FIKRIGSKELETON_INVALID_EO_BRANCH_INDEX;
+	while (ParentIndex > StartParentIndex)
 	{
-		CachedEndOfBranchIndices[InBoneIndex] = BoneIndex;
-				
+		BoneIndexAtEndOfBranch = BoneIndex;
 		BoneIndex++;
+		if (BoneIndex >= NumBones)
+		{
+			break;
+		}
 		ParentIndex = GetParentIndex(BoneIndex);
 	}
 
+	// set once (outside of while loop above) to avoid potential race condition
+	CachedEndOfBranchIndices[InBoneIndex] = BoneIndexAtEndOfBranch;
+
 	return CachedEndOfBranchIndices[InBoneIndex];
 }
 
diff --git a/Engine/Source/Runtime/Engine/Private/Animation/ReferenceSkeleton.cpp b/Engine/Source/Runtime/Engine/Private/Animation/ReferenceSkeleton.cpp
--- a/Engine/Source/Runtime/Engine/Private/Animation/ReferenceSkeleton.cpp
+++ b/Engine/Source/Runtime/Engine/Private/Animation/ReferenceSkeleton.cpp
@@ -762,8 +762,8 @@ int32 FReferenceSkeleton::GetCachedEndOfBranchIndex(const int32 InBoneIndex) con
 	
 	const int32 NumBones = GetRawBoneNum();
 	
-	// if we're asking for root's branch, get the last bone  
-	if (InBoneIndex == 0)
+	// if we're asking for the first or last bone, return the last bone  
+	if (InBoneIndex == 0 || InBoneIndex + 1 >= NumBones)
 	{
 		CachedEndOfBranchIndicesRaw[InBoneIndex] = GetRawBoneNum()-1;
 		return CachedEndOfBranchIndicesRaw[InBoneIndex];
@@ -776,10 +776,14 @@ int32 FReferenceSkeleton::GetCachedEndOfBranchIndex(const int32 InBoneIndex) con
 	// if next child bone's parent is less than or equal to StartParentIndex,
 	// we are leaving the branch so no need to go further
 	int32 BoneIndexAtEndOfBranch = BRANCH_CACHE_INVALID_INDEX;
-	while (ParentIndex > StartParentIndex && BoneIndex < NumBones)
+	while (ParentIndex > StartParentIndex)
 	{
 		BoneIndexAtEndOfBranch = BoneIndex;
 		BoneIndex++;
+		if (BoneIndex >= NumBones)
+		{
+			break;
+		}
 		ParentIndex = GetParentIndex(BoneIndex);
 	}
 
