Fixes issue, IK Retargeter doesn't initialize unless asset modified.

diff --git a/Engine/Plugins/Animation/IKRig/Source/IKRig/Private/Retargeter/IKRetargetProcessor.cpp b/Engine/Plugins/Animation/IKRig/Source/IKRig/Private/Retargeter/IKRetargetProcessor.cpp
--- a/Engine/Plugins/Animation/IKRig/Source/IKRig/Private/Retargeter/IKRetargetProcessor.cpp
+++ b/Engine/Plugins/Animation/IKRig/Source/IKRig/Private/Retargeter/IKRetargetProcessor.cpp
@@ -1501,8 +1501,12 @@ void UIKRetargetProcessor::Initialize(
 		UIKRetargeter* InRetargeterAsset,
 		const bool bSuppressWarnings)
 {
-	// don't attempt reinitialization with same assets
-	if (RetargeterAsset == InRetargeterAsset && AssetVersionInitializedWith == InRetargeterAsset->GetVersion())
+	// don't attempt reinitialization unless inputs have changed
+	const bool bSameTargetMesh = TargetSkeleton.SkeletalMesh == TargetSkeletalMesh;
+	const bool bSameSourceMesh = SourceSkeleton.SkeletalMesh == SourceSkeletalMesh;
+	const bool bSameAsset = RetargeterAsset == InRetargeterAsset;
+	const bool bSameAssetVersion = AssetVersionInitializedWith == InRetargeterAsset->GetVersion();
+	if (bSameAsset && bSameAssetVersion && bSameTargetMesh && bSameSourceMesh)
 	{
 		return;
 	}
diff --git a/Engine/Plugins/Animation/IKRig/Source/IKRig/Public/Retargeter/IKRetargetProcessor.h b/Engine/Plugins/Animation/IKRig/Source/IKRig/Public/Retargeter/IKRetargetProcessor.h
--- a/Engine/Plugins/Animation/IKRig/Source/IKRig/Public/Retargeter/IKRetargetProcessor.h
+++ b/Engine/Plugins/Animation/IKRig/Source/IKRig/Public/Retargeter/IKRetargetProcessor.h
@@ -576,7 +576,7 @@ private:
 
 	/** Only true once Initialize() has successfully completed.*/
 	bool bIsInitialized = false;
-	int32 AssetVersionInitializedWith = -1;
+	int32 AssetVersionInitializedWith = -2;
 	/** true when roots are able to be retargeted */
 	bool bRootsInitialized = false;
 	/** true when at least one pair of bone chains is able to be retargeted */
