Verify that after IK Retargeter asset is created from the context menu it can be opened in the Editor window.
#rnx

diff --git a/Engine/Plugins/Animation/IKRig/Source/IKRigEditor/Private/RetargetEditor/Tests/RetargetEditor.spec.cpp b/Engine/Plugins/Animation/IKRig/Source/IKRigEditor/Private/RetargetEditor/Tests/RetargetEditor.spec.cpp
new file mode 100644
--- /dev/null
+++ b/Engine/Plugins/Animation/IKRig/Source/IKRigEditor/Private/RetargetEditor/Tests/RetargetEditor.spec.cpp
@@ -0,0 +1,157 @@
+// Copyright Epic Games, Inc. All Rights Reserved.
+
+#if WITH_DEV_AUTOMATION_TESTS
+
+#if WITH_EDITOR
+
+#include "Misc/AutomationTest.h"
+#include "Editor.h"
+#include "Subsystems/AssetEditorSubsystem.h"
+#include "Subsystems/EditorAssetSubsystem.h"
+#include "Modules/ModuleManager.h"
+#include "AssetToolsModule.h"
+#include "AnimationEditorUtils.h"
+#include "Retargeter/IKRetargeter.h"
+#include "Rig/IKRigDefinition.h"
+#include "RetargetEditor/IKRetargetFactory.h"
+#include "RetargetEditor/IKRetargeterController.h"
+#include "UObject/UObjectGlobals.h"
+#include "ObjectTools.h"
+
+BEGIN_DEFINE_SPEC(
+	FRetargetEditorSpec, "Animation.Editor",
+	EAutomationTestFlags::EditorContext | EAutomationTestFlags::EngineFilter)
+
+	UIKRetargeter* IKRetargeterAsset = nullptr;
+	IAssetEditorInstance* EditorInstance = nullptr;
+
+	template <typename T>
+	T* CreateAsset(const FString&, UFactory* const);
+	UIKRetargeter* CreateUIKRetargeterAsset();
+	UIKRigDefinition* GetUIKRigDefinitionAsset() const;
+	bool VerifyMeshLoading(UIKRetargeter*, ERetargetSourceOrTarget) const;
+
+END_DEFINE_SPEC(FRetargetEditorSpec)
+
+void FRetargetEditorSpec::Define()
+{
+	Describe("IKRetargeter Asset", [this]()
+		{
+			BeforeEach([this]()
+				{
+					IKRetargeterAsset = CreateUIKRetargeterAsset();
+					TestNotNull("IK Retargeter asset Created", IKRetargeterAsset);
+				});
+
+			Describe("IKRetargeter Asset Editor", [this]()
+				{
+					BeforeEach([this]()
+						{
+							const bool bShowProgressWindow = false;
+							GEditor->GetEditorSubsystem<UAssetEditorSubsystem>()->OpenEditorForAsset(IKRetargeterAsset, EToolkitMode::Standalone, TSharedPtr<IToolkitHost>(), bShowProgressWindow);
+
+							const bool bFocusIfOpen = true;
+							EditorInstance = GEditor->GetEditorSubsystem<UAssetEditorSubsystem>()->FindEditorForAsset(IKRetargeterAsset, bFocusIfOpen);							
+						});
+
+					It("Verify that the IK Retargeter asset can be opened in the Editor window", [this]()
+						{
+							TestNotNull("IK Retargeter asset can be openned in the Editor window", EditorInstance);
+						});
+
+					AfterEach([this]()
+						{
+							if (IKRetargeterAsset)
+							{
+								GEditor->GetEditorSubsystem<UAssetEditorSubsystem>()->CloseAllEditorsForAsset(IKRetargeterAsset);
+							}
+						});
+				});
+
+			It("Verify that the new IK Retargeter asset initializes with no source mesh pre - loaded.", [this]()
+				{
+					TestTrue(TEXT("IK Retargeter asset initialized with no Source IK Rig asset."), IKRetargeterAsset->HasSourceIKRig() == false);
+				});
+
+			It("Verify that the new IK Retargeter asset initializes with no target mesh pre - loaded.", [this]()
+				{
+					TestTrue(TEXT("IK Retargeter asset initialized with no Target IK Rig asset."), IKRetargeterAsset->HasTargetIKRig() == false);
+				});
+
+			It("Verify that changing Source IK Rig asset updated IK Retargeter asset Hierarchy.", [this]()
+				{
+					TestTrue(TEXT("IK Retargeter asset Hierarchy updated with new Source IK Rig asset"), VerifyMeshLoading(IKRetargeterAsset, ERetargetSourceOrTarget::Source));
+				});
+
+			It("Verify that changing Target IK Rig asset updated IK Retargeter asset Hierarchy.", [this]()
+				{
+					TestTrue(TEXT("IK Retargeter asset Hierarchy updated with new Target IK Rig asset"), VerifyMeshLoading(IKRetargeterAsset, ERetargetSourceOrTarget::Target));
+				});
+
+			AfterEach([this]()
+				{
+					if (IKRetargeterAsset)
+					{
+						ObjectTools::DeleteObjectsUnchecked({ IKRetargeterAsset });
+					}
+				});
+		});
+}
+
+template <typename T>
+T* FRetargetEditorSpec::CreateAsset(const FString& InName, UFactory* const Factory)
+{
+	FString Name;
+	FString PackageName;
+	AnimationEditorUtils::CreateUniqueAssetName(TEXT("/Game/Tests/Animation/Temp") + InName, "", PackageName, Name);
+
+	FAssetToolsModule& AssetToolsModule = FModuleManager::GetModuleChecked<FAssetToolsModule>("AssetTools");
+	T* NewAsset = CastChecked<T>(AssetToolsModule.Get().CreateAsset(Name, FPackageName::GetLongPackagePath(PackageName), T::StaticClass(), Factory));
+	NewAsset->MarkPackageDirty();
+	return NewAsset;
+}
+
+UIKRetargeter* FRetargetEditorSpec::CreateUIKRetargeterAsset()
+{
+	const FString InName = "IKRetargeret";
+	TStrongObjectPtr<UIKRetargetFactory> Factory(NewObject<UIKRetargetFactory>());
+
+	UIKRetargeter* NewAsset = CreateAsset<UIKRetargeter>(InName, Factory.Get());
+	return NewAsset;
+}
+
+UIKRigDefinition* FRetargetEditorSpec::GetUIKRigDefinitionAsset() const
+{
+	const FString AssetName = "IK_Mannequin_Retarget";
+	const FString AssetPath = FString::Printf(TEXT("/Game/Tests/Animation/Retargeting/%s"), *AssetName);
+	const FAssetData AssetData = GEditor->GetEditorSubsystem<UEditorAssetSubsystem>()->FindAssetData(*AssetPath);
+	if (AssetData.IsValid())
+	{
+		return CastChecked<UIKRigDefinition>(AssetData.GetAsset());
+	}
+
+	return nullptr;
+}
+
+bool FRetargetEditorSpec::VerifyMeshLoading(UIKRetargeter* Asset, ERetargetSourceOrTarget type) const
+{
+	if(UIKRetargeterController* Controller = UIKRetargeterController::GetController(Asset))
+	{
+		if(UIKRigDefinition* RigDefinition = GetUIKRigDefinitionAsset())
+		{
+			Controller->SetIKRig(type, RigDefinition);
+
+			if(const UIKRigDefinition* IKRig = Controller->GetIKRig(type))
+			{
+				const FIKRigSkeleton& Skeleton = IKRig->GetSkeleton();
+				const int32 BoneNum = Skeleton.BoneNames.Num();
+				return BoneNum > 0;
+			}
+		}
+	}
+	return false;
+}
+
+#endif // WITH_DEV_AUTOMATION_TESTS
+
+#endif //WITH_EDITOR
\ No newline at end of file
