Fixes issue, IK Rig plugin not in the Animation category.

diff --git a/Engine/Plugins/Animation/IKRig/IKRig.uplugin b/Engine/Plugins/Animation/IKRig/IKRig.uplugin
--- a/Engine/Plugins/Animation/IKRig/IKRig.uplugin
+++ b/Engine/Plugins/Animation/IKRig/IKRig.uplugin
@@ -4,7 +4,7 @@
 	"VersionName": "1.0",
 	"FriendlyName": "IK Rig",
 	"Description": "",
-	"Category": "Other",
+	"Category": "Animation",
 	"CreatedBy": "Epic Games, Inc.",
 	"CreatedByURL": "https://epicgames.com",
 	"DocsURL": "",
