Fixes issue, retargeter freezes after editing retarget pose.

diff --git a/Engine/Plugins/Animation/IKRig/Source/IKRig/Private/Retargeter/IKRetargetProcessor.cpp b/Engine/Plugins/Animation/IKRig/Source/IKRig/Private/Retargeter/IKRetargetProcessor.cpp
--- a/Engine/Plugins/Animation/IKRig/Source/IKRig/Private/Retargeter/IKRetargetProcessor.cpp
+++ b/Engine/Plugins/Animation/IKRig/Source/IKRig/Private/Retargeter/IKRetargetProcessor.cpp
@@ -2302,6 +2302,12 @@ bool UIKRetargetProcessor::WasInitializedWithTheseAssets(
 void UIKRetargetProcessor::SetNeedsInitialized()
 {
 	bIsInitialized = false;
+
+	if (RetargeterAsset)
+	{
+		RetargeterAsset->IncrementVersion(); // triggers re-init
+	}
+	
 	if (IKRigProcessor)
 	{
 		// may not be initialized yet (during setup as prerequisites are being created)
