Fixes crash when batch retargeting animations with scale values.

diff --git a/Engine/Plugins/Animation/IKRig/Source/IKRigEditor/Private/RetargetEditor/IKRetargetBatchOperation.cpp b/Engine/Plugins/Animation/IKRig/Source/IKRigEditor/Private/RetargetEditor/IKRetargetBatchOperation.cpp
--- a/Engine/Plugins/Animation/IKRig/Source/IKRigEditor/Private/RetargetEditor/IKRetargetBatchOperation.cpp
+++ b/Engine/Plugins/Animation/IKRig/Source/IKRigEditor/Private/RetargetEditor/IKRetargetBatchOperation.cpp
@@ -420,6 +420,12 @@ void UIKRetargetBatchOperation::ConvertAnimation(
 				const FName& BoneName = SourceBoneNames[BoneIndex];
 				SourceComponentPose[BoneIndex] = UAnimPoseExtensions::GetBonePose(SourcePoseAtFrame, BoneName, EAnimPoseSpaces::World);
 			}
+			
+			// strip all scale out of the pose values, the translation of a component-space pose has incorporated scale values
+			for (FTransform& Transform : SourceComponentPose)
+			{
+				Transform.SetScale3D(FVector::OneVector);
+			}
 
 			// update goals 
 			Processor->ApplySettingsFromAsset();
