Fixed issue, "create retarget pose" menu greyed out when not "edit pose" mode.

diff --git a/Engine/Plugins/Animation/IKRig/Source/IKRigEditor/Private/RetargetEditor/IKRetargetEditor.cpp b/Engine/Plugins/Animation/IKRig/Source/IKRigEditor/Private/RetargetEditor/IKRetargetEditor.cpp
--- a/Engine/Plugins/Animation/IKRig/Source/IKRigEditor/Private/RetargetEditor/IKRetargetEditor.cpp
+++ b/Engine/Plugins/Animation/IKRig/Source/IKRigEditor/Private/RetargetEditor/IKRetargetEditor.cpp
@@ -201,13 +201,13 @@ void FIKRetargetEditor::BindCommands()
 	ToolkitCommands->MapAction(
 		Commands.NewRetargetPose,
 		FExecuteAction::CreateSP(EditorController, &FIKRetargetEditorController::HandleNewPose),
-		FCanExecuteAction::CreateSP(EditorController, &FIKRetargetEditorController::IsEditingPose),
+		FCanExecuteAction(),
 		EUIActionRepeatMode::RepeatDisabled);
 
 	ToolkitCommands->MapAction(
 		Commands.DuplicateRetargetPose,
 		FExecuteAction::CreateSP(EditorController, &FIKRetargetEditorController::HandleDuplicatePose),
-		FCanExecuteAction::CreateSP(EditorController, &FIKRetargetEditorController::IsEditingPose),
+		FCanExecuteAction(),
 		EUIActionRepeatMode::RepeatDisabled);
 
 	ToolkitCommands->MapAction(
@@ -265,19 +265,19 @@ void FIKRetargetEditor::BindCommands()
 	ToolkitCommands->MapAction(
 		Commands.ImportRetargetPose,
 		FExecuteAction::CreateSP(PoseExporterRef, &FIKRetargetPoseExporter::HandleImportFromPoseAsset),
-		FCanExecuteAction::CreateSP(EditorController, &FIKRetargetEditorController::IsEditingPose),
+		FCanExecuteAction(),
 		EUIActionRepeatMode::RepeatDisabled);
 
 	ToolkitCommands->MapAction(
 		Commands.ImportRetargetPoseFromAnim,
 		FExecuteAction::CreateSP(PoseExporterRef, &FIKRetargetPoseExporter::HandleImportFromSequenceAsset),
-		FCanExecuteAction::CreateSP(EditorController, &FIKRetargetEditorController::IsEditingPose),
+		FCanExecuteAction(),
 		EUIActionRepeatMode::RepeatDisabled);
 
 	ToolkitCommands->MapAction(
 		Commands.ExportRetargetPose,
 		FExecuteAction::CreateSP(PoseExporterRef, &FIKRetargetPoseExporter::HandleExportPoseAsset),
-		FCanExecuteAction::CreateSP(EditorController, &FIKRetargetEditorController::IsEditingPose),
+		FCanExecuteAction(),
 		EUIActionRepeatMode::RepeatDisabled);
 }
 
