Fixes issue, IK Retarget profile "Enable IK" does not update at runtime.

diff --git a/Engine/Plugins/Animation/IKRig/Source/IKRig/Private/AnimNodes/AnimNode_RetargetPoseFromMesh.cpp b/Engine/Plugins/Animation/IKRig/Source/IKRig/Private/AnimNodes/AnimNode_RetargetPoseFromMesh.cpp
--- a/Engine/Plugins/Animation/IKRig/Source/IKRig/Private/AnimNodes/AnimNode_RetargetPoseFromMesh.cpp
+++ b/Engine/Plugins/Animation/IKRig/Source/IKRig/Private/AnimNodes/AnimNode_RetargetPoseFromMesh.cpp
@@ -105,29 +105,23 @@ void FAnimNode_RetargetPoseFromMesh::Evaluate_AnyThread(FPoseContext& Output)
 	{
 		// live preview source asset settings in the retarget, editor only
 		// NOTE: this copies goal targets as well, but these are overwritten by IK chain goals
-		Processor->ApplySettingsFromAsset();
+		Processor->CopyIKRigSettingsFromAsset();
 	}
 	#endif
 
-	// apply custom profile settings to the processor
-	Processor->ApplySettingsFromProfile(CustomRetargetProfile);
-	
 	// LOD off the IK pass
-	if (LODThresholdForIK != INDEX_NONE && Output.AnimInstanceProxy->GetLODLevel() > LODThresholdForIK)
-	{
-		// override the custom profile's global settings but with IK forcibly turned off
-		FRetargetProfile TurnIKOffProfile;
-		TurnIKOffProfile.bApplyGlobalSettings = true;
-		TurnIKOffProfile.GlobalSettings = Processor->GetGlobalSettings();
-		TurnIKOffProfile.GlobalSettings.bEnableIK = false;
-		Processor->ApplySettingsFromProfile(TurnIKOffProfile);
-	}
+	bool bForceIKOff = LODThresholdForIK != INDEX_NONE && Output.AnimInstanceProxy->GetLODLevel() > LODThresholdForIK;
+	FRetargetProfile RetargetProfileToUse = GetMergedRetargetProfile(bForceIKOff);
 
 	// run the retargeter
-	const TArray<FTransform>& RetargetedPose = Processor->RunRetargeter(SourceMeshComponentSpaceBoneTransforms, SpeedValuesFromCurves, DeltaTime);
+	const TArray<FTransform>& RetargetedPose = Processor->RunRetargeter(
+		SourceMeshComponentSpaceBoneTransforms,
+		SpeedValuesFromCurves,
+		DeltaTime,
+		RetargetProfileToUse);
 	DeltaTime = 0.0f;
-
-	// copy pose back
+	
+	// convert pose to local space and apply to output
 	FCSPose<FCompactPose> ComponentPose;
 	ComponentPose.InitPose(Output.Pose);
 	const FCompactPose& CompactPose = ComponentPose.GetPose();
@@ -329,8 +323,10 @@ bool FAnimNode_RetargetPoseFromMesh::EnsureProcessorIsInitialized(const TObjectP
 	if (!Processor->WasInitializedWithTheseAssets(SourceMesh, TargetMesh, IKRetargeterAsset))
 	{
 		// initialize retarget processor with source and target skeletal meshes
-		// (asset is passed in as outer UObject for new UIKRigProcessor) 
-		Processor->Initialize(SourceMesh,	TargetMesh,IKRetargeterAsset);
+		// (asset is passed in as outer UObject for new UIKRigProcessor)
+		bool bForceIKOff = false;
+		FRetargetProfile RetargetProfileToUse = GetMergedRetargetProfile(bForceIKOff);
+		Processor->Initialize(SourceMesh, TargetMesh, IKRetargeterAsset, RetargetProfileToUse);
 	}
 
 	return Processor->IsInitialized();
@@ -394,3 +390,19 @@ TObjectPtr<USkeletalMeshComponent> FAnimNode_RetargetPoseFromMesh::GetComponentT
 	return SourceMeshComponent.Get();
 }
 
+FRetargetProfile FAnimNode_RetargetPoseFromMesh::GetMergedRetargetProfile(bool bForceIKOff) const
+{
+	// collect settings to retarget with starting with asset settings and overriding with custom profile
+	FRetargetProfile Profile;
+	IKRetargeterAsset->FillProfileWithAssetSettings(Profile);
+	// load custom profile plugged into the anim node
+	Profile.MergeWithOtherProfile(CustomRetargetProfile);
+	// force all IK off (skips IK solve)
+	if (bForceIKOff)
+	{
+		Profile.GlobalSettings.bEnableIK = false;
+	}
+
+	return Profile;
+}
+
diff --git a/Engine/Plugins/Animation/IKRig/Source/IKRig/Private/Retargeter/IKRetargetProcessor.cpp b/Engine/Plugins/Animation/IKRig/Source/IKRig/Private/Retargeter/IKRetargetProcessor.cpp
--- a/Engine/Plugins/Animation/IKRig/Source/IKRig/Private/Retargeter/IKRetargetProcessor.cpp
+++ b/Engine/Plugins/Animation/IKRig/Source/IKRig/Private/Retargeter/IKRetargetProcessor.cpp
@@ -1499,6 +1499,7 @@ void UIKRetargetProcessor::Initialize(
 		USkeletalMesh* SourceSkeletalMesh,
 		USkeletalMesh* TargetSkeletalMesh,
 		UIKRetargeter* InRetargeterAsset,
+		const FRetargetProfile& Settings,
 		const bool bSuppressWarnings)
 {
 	// don't attempt reinitialization unless inputs have changed
@@ -1563,7 +1564,7 @@ void UIKRetargetProcessor::Initialize(
 	bRootsInitialized = InitializeRoots();
 
 	// initialize pairs of bone chains
-	bAtLeastOneValidBoneChainPair = InitializeBoneChainPairs();
+	bAtLeastOneValidBoneChainPair = InitializeBoneChainPairs(Settings.ChainSettings);
 	if (!bAtLeastOneValidBoneChainPair)
 	{
 		// couldn't match up any BoneChain pairs, no limb retargeting possible
@@ -1600,7 +1601,7 @@ void UIKRetargetProcessor::Initialize(
 	}
 
 	// copy the initial settings from the asset
-	ApplySettingsFromAsset();
+	CopyIKRigSettingsFromAsset();
 	
 	bIsInitialized = bRootsInitialized || bAtLeastOneValidBoneChainPair;
 	AssetVersionInitializedWith = RetargeterAsset->GetVersion();
@@ -1637,7 +1638,7 @@ bool UIKRetargetProcessor::InitializeRoots()
 	return bRootEncoderInit && bRootDecoderInit;
 }
 
-bool UIKRetargetProcessor::InitializeBoneChainPairs()
+bool UIKRetargetProcessor::InitializeBoneChainPairs(const TMap<FName, FTargetChainSettings>& ChainSettings)
 {
 	ChainPairsFK.Reset();
 	ChainPairsIK.Reset();
@@ -1684,13 +1685,10 @@ bool UIKRetargetProcessor::InitializeBoneChainPairs()
 		}
 		
 		// load IK chain
-		if (ChainMap->Settings.IK.EnableIK)
+		FRetargetChainPairIK ChainPairIK;
+		if (ChainPairIK.Initialize(*SourceBoneChain, *TargetBoneChain, SourceSkeleton, TargetSkeleton, Log))
 		{
-			FRetargetChainPairIK ChainPairIK;
-			if (ChainPairIK.Initialize(*SourceBoneChain, *TargetBoneChain, SourceSkeleton, TargetSkeleton, Log))
-			{
-				ChainPairsIK.Add(ChainPairIK);
-			}
+			ChainPairsIK.Add(ChainPairIK);
 		}
 
 		// warn user if IK goal is not on the END bone of the target chain. It will still work, but may produce bad results.
@@ -1710,6 +1708,22 @@ bool UIKRetargetProcessor::InitializeBoneChainPairs()
 		}
 	}
 
+	// set initial chain settings
+	for (FRetargetChainPairFK& FKChain : ChainPairsFK)
+	{
+		if (const FTargetChainSettings* Settings = ChainSettings.Find(FKChain.TargetBoneChainName))
+		{
+			FKChain.Settings = Settings->FK;
+		}
+	}
+	for (FRetargetChainPairIK& IKChain : ChainPairsIK)
+	{
+		if (const FTargetChainSettings* Settings = ChainSettings.Find(IKChain.TargetBoneChainName))
+		{
+			IKChain.Settings = Settings->IK;
+		}
+	}
+
 	// sort the chains based on their StartBone's index
 	auto ChainsSorter = [this](const FRetargetChainPair& A, const FRetargetChainPair& B)
 	{
@@ -1761,20 +1775,11 @@ bool UIKRetargetProcessor::InitializeIKRig(UObject* Outer, const USkeletalMesh*
 	
 	// gather list of excluded goals based on any chain mapping that has it's IK disabled
 	TArray<FName> GoalsToExclude;
-	const TArray<TObjectPtr<URetargetChainSettings>>& ChainMapping = RetargeterAsset->GetAllChainSettings();
-	const TArray<FBoneChain>& IKRigChains = IKRig->GetRetargetChains();
-	for (const URetargetChainSettings* ChainMap : ChainMapping)
+	for (FRetargetChainPairIK ChainPairIK : ChainPairsIK)
 	{
-		if (ChainMap->Settings.IK.EnableIK)
+		if (!ChainPairIK.Settings.EnableIK)
 		{
-			continue;
-		}
-		for (const FBoneChain& Chain : IKRigChains)
-		{
-			if (Chain.ChainName == ChainMap->TargetChain)
-			{
-				GoalsToExclude.Add(Chain.IKGoalName);
-			}
+			GoalsToExclude.Add(ChainPairIK.IKGoalName);
 		}
 	}
 	
@@ -1789,6 +1794,12 @@ bool UIKRetargetProcessor::InitializeIKRig(UObject* Outer, const USkeletalMesh*
 	// validate that all IK bone chains have an associated Goal
 	for (FRetargetChainPairIK& ChainPair : ChainPairsIK)
 	{
+		// skip chains that have had their IK disabled, these goals won't be in the IK Rig
+		if (GoalsToExclude.Contains(ChainPair.IKGoalName))
+		{
+			continue;
+		}
+		
 		// does the IK rig have the IK goal this bone chain requires?
 		if (!IKRigProcessor->GetGoalContainer().FindGoalByName(ChainPair.IKGoalName))
 		{
@@ -1896,10 +1907,23 @@ bool UIKRetargetProcessor::InitializeOpStack(const TArray<TObjectPtr<URetargetOp
 TArray<FTransform>&  UIKRetargetProcessor::RunRetargeter(
 	const TArray<FTransform>& InSourceGlobalPose,
 	const TMap<FName, float>& SpeedValuesFromCurves,
-	const float DeltaTime)
+	const float DeltaTime,
+	const FRetargetProfile& Settings)
 {
 	check(bIsInitialized);
 
+	// apply the retargeting settings
+	ApplySettingsFromProfile(Settings);
+	// applying settings can cause the retargeter to require reinitialization (when enabling/disabling IK on a chain)
+	// since reinitialization can only currently be done on the main thread, we have to wait for the next frame
+	// in the meantime we return the last generated pose which we know the processor has because it wouldn't have made it past the check()
+	// unless it were already initialized once.
+	// NOTE: In the future hopefully we can reinitialize from any thread by avoiding UObject creation in UIKRetargetProcessor::Initialize()
+	if (!IsInitialized())
+	{
+		return TargetSkeleton.OutputGlobalPose;
+	}
+
 #if WITH_EDITOR
 	// validate system running the retargeter has stripped all the scale out of the incoming pose
 	for (const FTransform& Transform : InSourceGlobalPose)
@@ -2471,54 +2495,13 @@ FName UIKRetargetProcessor::GetMappedChainName(
 
 #endif
 
-void UIKRetargetProcessor::ApplySettingsFromAsset()
+void UIKRetargetProcessor::CopyIKRigSettingsFromAsset()
 {
 	// copy IK Rig settings
 	if (const UIKRigDefinition* TargetIKRig = RetargeterAsset->GetIKRig(ERetargetSourceOrTarget::Target))
 	{
 		IKRigProcessor->CopyAllInputsFromSourceAssetAtRuntime(TargetIKRig);
 	}
-	
-	// copy chain settings from the asset
-	const TArray<TObjectPtr<URetargetChainSettings>>& AllChainSettings = RetargeterAsset->GetAllChainSettings();
-	for (const TObjectPtr<URetargetChainSettings>& ChainSettings : AllChainSettings)
-	{
-		for (FRetargetChainPairFK& Chain : ChainPairsFK)
-		{
-			if (Chain.TargetBoneChainName == ChainSettings->TargetChain)
-			{
-				Chain.Settings = ChainSettings->Settings.FK;
-			}
-		}
-		
-		for (FRetargetChainPairIK& Chain : ChainPairsIK)
-		{
-			if (Chain.TargetBoneChainName == ChainSettings->TargetChain)
-			{
-				Chain.Settings = ChainSettings->Settings.IK;
-				Chain.SpeedPlantSettings = ChainSettings->Settings.SpeedPlanting;
-			}
-		}
-	}
-
-	// copy root settings
-	RootRetargeter.Settings = RetargeterAsset->GetRootSettingsUObject()->Settings;
-
-	// copy global settings
-	GlobalSettings = RetargeterAsset->GetGlobalSettings();
-
-	// apply current retarget poses (only applied if the pose has been switched to a different one OR if the current pose was modified)
-	const FName SourcePose = RetargeterAsset->GetCurrentRetargetPoseName(ERetargetSourceOrTarget::Source);
-	const FName TargetPose = RetargeterAsset->GetCurrentRetargetPoseName(ERetargetSourceOrTarget::Target);
-	UpdateRetargetPoseAtRuntime(SourcePose, ERetargetSourceOrTarget::Source);
-	UpdateRetargetPoseAtRuntime(TargetPose, ERetargetSourceOrTarget::Target);
-
-	// apply the current profile
-	// (this is always applied last so that profile overrides take precedence over asset settings)
-	if (const FRetargetProfile* CurrentProfile = RetargeterAsset->GetCurrentProfile())
-	{
-		ApplySettingsFromProfile(*CurrentProfile);
-	}
 }
 
 const FRetargetSkeleton& UIKRetargetProcessor::GetSkeleton(ERetargetSourceOrTarget SourceOrTarget) const
@@ -2534,49 +2517,20 @@ FName UIKRetargetProcessor::GetRetargetRoot(ERetargetSourceOrTarget SourceOrTarg
 void UIKRetargetProcessor::ApplySettingsFromProfile(const FRetargetProfile& Profile)
 {
 	// assign retarget poses specified in the profile
-	if (Profile.bApplySourceRetargetPose)
-	{
-		UpdateRetargetPoseAtRuntime(Profile.SourceRetargetPoseName, ERetargetSourceOrTarget::Source);
-	}
-	if (Profile.bApplyTargetRetargetPose)
-	{
-		UpdateRetargetPoseAtRuntime(Profile.TargetRetargetPoseName, ERetargetSourceOrTarget::Target);
-	}
+	UpdateRetargetPoseAtRuntime(Profile.SourceRetargetPoseName, ERetargetSourceOrTarget::Source);
+	UpdateRetargetPoseAtRuntime(Profile.TargetRetargetPoseName, ERetargetSourceOrTarget::Target);
 
 	// assign chain settings
-	if (Profile.bApplyChainSettings)
+	for (const TTuple<FName, FTargetChainSettings>& ChainSettings : Profile.ChainSettings)
 	{
-		for (const TTuple<FName, FTargetChainSettings>& ChainSettings : Profile.ChainSettings)
-		{
-			for (FRetargetChainPairFK& Chain : ChainPairsFK)
-			{
-				if (Chain.TargetBoneChainName == ChainSettings.Key)
-				{
-					Chain.Settings = ChainSettings.Value.FK;
-				}
-			}
-		
-			for (FRetargetChainPairIK& Chain : ChainPairsIK)
-			{
-				if (Chain.TargetBoneChainName == ChainSettings.Key)
-				{
-					Chain.Settings = ChainSettings.Value.IK;
-				}
-			}
-		}
+		UpdateChainSettingsAtRuntime(ChainSettings.Key, ChainSettings.Value);
 	}
 
 	// assign root settings
-	if (Profile.bApplyRootSettings)
-	{
-		RootRetargeter.Settings = Profile.RootSettings;
-	}
-
+	RootRetargeter.Settings = Profile.RootSettings;
+	
 	// assign global settings
-	if (Profile.bApplyGlobalSettings)
-	{
-		GlobalSettings = Profile.GlobalSettings;
-	}
+	GlobalSettings = Profile.GlobalSettings;
 }
 
 void UIKRetargetProcessor::UpdateRetargetPoseAtRuntime(
@@ -2632,5 +2586,34 @@ void UIKRetargetProcessor::UpdateRetargetPoseAtRuntime(
 	RootRetargeter.InitializeTarget(RootRetargeter.Target.BoneName, TargetSkeleton, Log);
 }
 
+void UIKRetargetProcessor::UpdateChainSettingsAtRuntime(const FName ChainName, const FTargetChainSettings& NewChainSettings)
+{
+	// update FK chain
+	for (FRetargetChainPairFK& Chain : ChainPairsFK)
+	{
+		if (Chain.TargetBoneChainName == ChainName)
+		{
+			Chain.Settings = NewChainSettings.FK;
+			break;
+		}
+	}
+
+	// update IK chain
+	for (FRetargetChainPairIK& Chain : ChainPairsIK)
+	{
+		if (Chain.TargetBoneChainName == ChainName)
+		{
+			if (Chain.Settings.EnableIK != NewChainSettings.IK.EnableIK)
+			{
+				SetNeedsInitialized();
+			}
+			
+			Chain.Settings = NewChainSettings.IK;
+			Chain.SpeedPlantSettings = NewChainSettings.SpeedPlanting;
+			break;
+		}
+	}
+}
+
 #undef LOCTEXT_NAMESPACE
 
diff --git a/Engine/Plugins/Animation/IKRig/Source/IKRig/Private/Retargeter/IKRetargeter.cpp b/Engine/Plugins/Animation/IKRig/Source/IKRig/Private/Retargeter/IKRetargeter.cpp
--- a/Engine/Plugins/Animation/IKRig/Source/IKRig/Private/Retargeter/IKRetargeter.cpp
+++ b/Engine/Plugins/Animation/IKRig/Source/IKRig/Private/Retargeter/IKRetargeter.cpp
@@ -400,6 +400,30 @@ const FName UIKRetargeter::GetDefaultPoseName()
 	return DefaultPoseName;
 }
 
+void UIKRetargeter::FillProfileWithAssetSettings(FRetargetProfile& InOutProfile) const
+{
+	// first copy all the asset settings into the profile
+	InOutProfile.bApplyTargetRetargetPose = true;
+	InOutProfile.TargetRetargetPoseName = GetCurrentRetargetPoseName(ERetargetSourceOrTarget::Target);
+	InOutProfile.bApplySourceRetargetPose = true;
+	InOutProfile.SourceRetargetPoseName = GetCurrentRetargetPoseName(ERetargetSourceOrTarget::Source);
+	InOutProfile.bApplyChainSettings = true;
+	for (const TObjectPtr<URetargetChainSettings>& Chain : ChainSettings)
+	{
+		InOutProfile.ChainSettings.Add(Chain->TargetChain, Chain->Settings);
+	}
+	InOutProfile.bApplyRootSettings = true;
+	InOutProfile.RootSettings = RootSettings->Settings;
+	InOutProfile.bApplyGlobalSettings = true;
+	InOutProfile.GlobalSettings = GlobalSettings->Settings;
+
+	// now override any settings in the asset's current profile
+	if (const FRetargetProfile* ProfileToUse = GetCurrentProfile())
+	{
+		InOutProfile.MergeWithOtherProfile(*ProfileToUse);
+	}
+}
+
 const FRetargetProfile* UIKRetargeter::GetCurrentProfile() const
 {
 	return GetProfileByName(CurrentProfile);
diff --git a/Engine/Plugins/Animation/IKRig/Source/IKRig/Public/AnimNodes/AnimNode_RetargetPoseFromMesh.h b/Engine/Plugins/Animation/IKRig/Source/IKRig/Public/AnimNodes/AnimNode_RetargetPoseFromMesh.h
--- a/Engine/Plugins/Animation/IKRig/Source/IKRig/Public/AnimNodes/AnimNode_RetargetPoseFromMesh.h
+++ b/Engine/Plugins/Animation/IKRig/Source/IKRig/Public/AnimNodes/AnimNode_RetargetPoseFromMesh.h
@@ -81,6 +81,8 @@ private:
 	void CopyBoneTransformsFromSource(USkeletalMeshComponent* TargetMeshComponent);
 	/** indirection to account for Leader Pose Component setup */
 	TObjectPtr<USkeletalMeshComponent> GetComponentToCopyPoseFrom() const;
+	// get a retarget profile that includes settings from the asset, plus any custom profile overrides
+	FRetargetProfile GetMergedRetargetProfile(bool bEnableIK) const;
 
 	/** the runtime processor used to run the retarget and generate new poses */
 	UPROPERTY(Transient)
diff --git a/Engine/Plugins/Animation/IKRig/Source/IKRig/Public/Retargeter/IKRetargetProcessor.h b/Engine/Plugins/Animation/IKRig/Source/IKRig/Public/Retargeter/IKRetargetProcessor.h
--- a/Engine/Plugins/Animation/IKRig/Source/IKRig/Public/Retargeter/IKRetargetProcessor.h
+++ b/Engine/Plugins/Animation/IKRig/Source/IKRig/Public/Retargeter/IKRetargetProcessor.h
@@ -478,23 +478,29 @@ public:
 		USkeletalMesh *SourceSkeleton,
 		USkeletalMesh *TargetSkeleton,
 		UIKRetargeter* InRetargeterAsset,
+		const FRetargetProfile& Settings,
 		const bool bSuppressWarnings=false);
 
 	/**
 	* Run the retarget to generate a new pose.
 	* @param InSourceGlobalPose -  is the source mesh input pose in Component/Global space
+	* @param SpeedValuesFromCurves - the speed of each curve used by speed planting (blended in anim graph)
+	* @param DeltaTime -  time since last tick in seconds (used by speed planting)
+	* @param Settings -  the retarget profile to use for this update
 	* @return The retargeted Component/Global space pose for the target skeleton
 	*/
 	TArray<FTransform>& RunRetargeter(
 		const TArray<FTransform>& InSourceGlobalPose,
 		const TMap<FName, float>& SpeedValuesFromCurves,
-		const float DeltaTime);
-
-	/** Apply the settings stored in a retarget profile. Call this before RunRetargeter() to use the settings stored in a profile. */
-	void ApplySettingsFromProfile(const FRetargetProfile& Profile);
+		const float DeltaTime,
+		const FRetargetProfile& Settings);
 	
-	/** Apply the settings stored in the retargeter asset. */
-	void ApplySettingsFromAsset();
+	/** Apply the settings stored in the IKRig asset. */
+	void CopyIKRigSettingsFromAsset();
+
+	/** Does a partial reinitialization (at runtime) whenever the retarget pose is swapped to a different or if the
+	 * pose has been modified. Does nothing if the pose has not changed. */
+	void UpdateRetargetPoseAtRuntime(const FName NewRetargetPoseName, ERetargetSourceOrTarget SourceOrTarget);
 
 	/** Get read-only access to either source or target skeleton. */
 	const FRetargetSkeleton& GetSkeleton(ERetargetSourceOrTarget SourceOrTarget) const;
@@ -532,10 +538,6 @@ public:
 	/** Reset the IK planting state. */
 	void ResetPlanting();
 
-	/** Does a partial reinitialization (at runtime) whenever the retarget pose is swapped to a different or if the
-	 * pose has been modified. Does nothing if the pose has not changed. */
-	void UpdateRetargetPoseAtRuntime(const FName NewRetargetPoseName, ERetargetSourceOrTarget SourceOrTarget);
-
 	/** logging system */
 	FIKRigLogger Log;
 
@@ -609,12 +611,19 @@ private:
 	/** The collection of operations to run in the final phase of retargeting */
 	UPROPERTY(Transient) // must be property to keep from being GC'd
 	TArray<TObjectPtr<URetargetOpBase>> OpStack;
+
+	/** Apply the settings stored in a retarget profile. Called inside RunRetargeter(). */
+	void ApplySettingsFromProfile(const FRetargetProfile& Profile);
+	
+	/** Update chain settings at runtime (for use with a profile)
+	 * Will queue an IK Rig reinitialization if the "Enable IK" state is modified. */
+	void UpdateChainSettingsAtRuntime(const FName ChainName, const FTargetChainSettings& NewChainSettings);
 	
 	/** Initializes the FRootRetargeter */
 	bool InitializeRoots();
 
 	/** Initializes the all the chain pairs */
-	bool InitializeBoneChainPairs();
+	bool InitializeBoneChainPairs(const TMap<FName, FTargetChainSettings>& ChainSettings);
 
 	/** Initializes the IK Rig that evaluates the IK solve for the target IK chains */
 	bool InitializeIKRig(UObject* Outer, const USkeletalMesh* InSkeletalMesh);
diff --git a/Engine/Plugins/Animation/IKRig/Source/IKRig/Public/Retargeter/IKRetargetProfile.h b/Engine/Plugins/Animation/IKRig/Source/IKRig/Public/Retargeter/IKRetargetProfile.h
--- a/Engine/Plugins/Animation/IKRig/Source/IKRig/Public/Retargeter/IKRetargetProfile.h
+++ b/Engine/Plugins/Animation/IKRig/Source/IKRig/Public/Retargeter/IKRetargetProfile.h
@@ -58,4 +58,35 @@ public:
 	// Retarget settings to control global behavior, like Stride Warping (not applied unless bApplyGlobalSettings is true)
 	UPROPERTY(EditAnywhere, BlueprintReadWrite, Category=GlobalSettings, meta=(EditCondition="bApplyGlobalSettings"))
 	FRetargetGlobalSettings GlobalSettings;
+
+	void MergeWithOtherProfile(const FRetargetProfile& OtherProfile)
+	{
+		if (OtherProfile.bApplyTargetRetargetPose)
+		{
+			TargetRetargetPoseName = OtherProfile.TargetRetargetPoseName;
+		}
+
+		if (OtherProfile.bApplySourceRetargetPose)
+		{
+			SourceRetargetPoseName = OtherProfile.SourceRetargetPoseName;
+		}
+
+		if (OtherProfile.bApplyChainSettings)
+		{
+			for (const TPair<FName,FTargetChainSettings>& Pair : OtherProfile.ChainSettings)
+			{
+				ChainSettings.Add(Pair.Key, Pair.Value);
+			}
+		}
+
+		if (OtherProfile.bApplyRootSettings)
+		{
+			RootSettings = OtherProfile.RootSettings;
+		}
+
+		if (OtherProfile.bApplyGlobalSettings)
+		{
+			GlobalSettings = OtherProfile.GlobalSettings;
+		}
+	}
 };
diff --git a/Engine/Plugins/Animation/IKRig/Source/IKRig/Public/Retargeter/IKRetargeter.h b/Engine/Plugins/Animation/IKRig/Source/IKRig/Public/Retargeter/IKRetargeter.h
--- a/Engine/Plugins/Animation/IKRig/Source/IKRig/Public/Retargeter/IKRetargeter.h
+++ b/Engine/Plugins/Animation/IKRig/Source/IKRig/Public/Retargeter/IKRetargeter.h
@@ -251,6 +251,8 @@ public:
 	// Get name of default pose 
 	static const FName GetDefaultPoseName();
 	
+	// Fill the provided profile with settings from this asset
+	void FillProfileWithAssetSettings(FRetargetProfile& InOutProfile) const;
 	// Get the current retarget profile (may be null) 
 	const FRetargetProfile* GetCurrentProfile() const;
 	// Get the retarget profile by name (may be null) 
diff --git a/Engine/Plugins/Animation/IKRig/Source/IKRigEditor/Private/RetargetEditor/IKRetargetBatchOperation.cpp b/Engine/Plugins/Animation/IKRig/Source/IKRigEditor/Private/RetargetEditor/IKRetargetBatchOperation.cpp
--- a/Engine/Plugins/Animation/IKRig/Source/IKRigEditor/Private/RetargetEditor/IKRetargetBatchOperation.cpp
+++ b/Engine/Plugins/Animation/IKRig/Source/IKRigEditor/Private/RetargetEditor/IKRetargetBatchOperation.cpp
@@ -326,7 +326,9 @@ void UIKRetargetBatchOperation::ConvertAnimation(
 	// initialize the retargeter
 	UObject* TransientOuter = Cast<UObject>(GetTransientPackage());
 	UIKRetargetProcessor* Processor = NewObject<UIKRetargetProcessor>(TransientOuter);
-	Processor->Initialize(Context.SourceMesh, Context.TargetMesh, Context.IKRetargetAsset);
+	FRetargetProfile RetargetProfile;
+	Context.IKRetargetAsset->FillProfileWithAssetSettings(RetargetProfile);
+	Processor->Initialize(Context.SourceMesh, Context.TargetMesh, Context.IKRetargetAsset, RetargetProfile);
 	if (!Processor->IsInitialized())
 	{
 		UE_LOG(LogTemp, Warning, TEXT("Unable to initialize the IK Retargeter. Newly created animations were not retargeted!"));
@@ -428,7 +430,7 @@ void UIKRetargetBatchOperation::ConvertAnimation(
 			}
 
 			// update goals 
-			Processor->ApplySettingsFromAsset();
+			Processor->CopyIKRigSettingsFromAsset();
 			
 			// calculate the delta time
 			const float TimeAtCurrentFrame = SourceSequence->GetTimeAtFrame(FrameIndex);
@@ -446,8 +448,12 @@ void UIKRetargetBatchOperation::ConvertAnimation(
 				SpeedCurveValues.Add(SpeedCurveName, SourceSequence->EvaluateCurveData(SpeedCurveName, TimeAtCurrentFrame));
 			}
 
+			// get the settings profile
+			FRetargetProfile SettingsProfile;
+			Context.IKRetargetAsset->FillProfileWithAssetSettings(SettingsProfile);
+
 			// run the retargeter
-			const TArray<FTransform>& TargetComponentPose = Processor->RunRetargeter(SourceComponentPose, SpeedCurveValues, DeltaTime);
+			const TArray<FTransform>& TargetComponentPose = Processor->RunRetargeter(SourceComponentPose, SpeedCurveValues, DeltaTime, SettingsProfile);
 
 			// convert to a local-space pose
 			TArray<FTransform> TargetLocalPose = TargetComponentPose;
diff --git a/Engine/Plugins/Animation/IKRig/Source/IKRigEditor/Private/RetargetEditor/IKRetargeterPoseGenerator.cpp b/Engine/Plugins/Animation/IKRig/Source/IKRigEditor/Private/RetargetEditor/IKRetargeterPoseGenerator.cpp
--- a/Engine/Plugins/Animation/IKRig/Source/IKRigEditor/Private/RetargetEditor/IKRetargeterPoseGenerator.cpp
+++ b/Engine/Plugins/Animation/IKRig/Source/IKRigEditor/Private/RetargetEditor/IKRetargeterPoseGenerator.cpp
@@ -562,7 +562,9 @@ bool FRetargetAutoPoseGenerator::CheckReadyToAlignBones() const
 	USkeletalMesh* TargetSkeletalMesh = Controller->GetPreviewMesh(ERetargetSourceOrTarget::Target);
 	UIKRetargeter* RetargeterAsset = Controller->GetAssetPtr();
 	constexpr bool bSuppressWarnings = true;
-	Processor->Initialize(SourceSkeletalMesh, TargetSkeletalMesh, RetargeterAsset, bSuppressWarnings);
+	FRetargetProfile RetargetProfile;
+	RetargeterAsset->FillProfileWithAssetSettings(RetargetProfile);
+	Processor->Initialize(SourceSkeletalMesh, TargetSkeletalMesh, RetargeterAsset, RetargetProfile, bSuppressWarnings);
 	
 	// can't auto align until processor is initialized because we need a fully resolved chain mapping on source/target
 	// this could happen if the user attempts to edit the retarget pose on an incompatible IK Rig
