IKRig: Exposed the FBIK solver's effectors and bone settings to BP scripting.
These classes already have a BP API but and can now be accessed from the IK rig definition or controller
#rnx

diff --git a/Engine/Plugins/Animation/IKRig/Source/IKRig/Public/Rig/Solvers/IKRig_FBIKSolver.h b/Engine/Plugins/Animation/IKRig/Source/IKRig/Public/Rig/Solvers/IKRig_FBIKSolver.h
--- a/Engine/Plugins/Animation/IKRig/Source/IKRig/Public/Rig/Solvers/IKRig_FBIKSolver.h
+++ b/Engine/Plugins/Animation/IKRig/Source/IKRig/Public/Rig/Solvers/IKRig_FBIKSolver.h
@@ -240,7 +240,19 @@ public:
 	 * In very limited circumstances, it can be beneficial to use the pose from the previous frame.*/
 	UPROPERTY()
 	bool bStartSolveFromInputPose_DEPRECATED = true;
-	
+
+	UFUNCTION(BlueprintCallable, Category = "IKRig")
+	TArray<UIKRig_FBIKEffector*> GetEffectors()
+	{
+		return Effectors;
+	}
+
+	UFUNCTION(BlueprintCallable, Category = "IKRig")
+	TArray<UIKRig_FBIKBoneSettings*> GetBoneSettings()
+	{
+		return BoneSettings;
+	}
+
 	UPROPERTY()
 	TArray<TObjectPtr<UIKRig_FBIKEffector>> Effectors;
 
