Improved retargeter debug drawing

diff --git a/Engine/Plugins/Animation/IKRig/Source/IKRig/Public/Retargeter/IKRetargeter.h b/Engine/Plugins/Animation/IKRig/Source/IKRig/Public/Retargeter/IKRetargeter.h
--- a/Engine/Plugins/Animation/IKRig/Source/IKRig/Public/Retargeter/IKRetargeter.h
+++ b/Engine/Plugins/Animation/IKRig/Source/IKRig/Public/Retargeter/IKRetargeter.h
@@ -442,6 +442,14 @@ public:
 	UPROPERTY(EditAnywhere, Category = DebugSettings)
 	bool bDebugDraw = true;
 
+	// Draw lines on each bone chain. 
+	UPROPERTY(EditAnywhere, Category = DebugSettings)
+	bool bDrawChainLines = true;
+
+	// Draw spheres on single bone chains. 
+	UPROPERTY(EditAnywhere, Category = DebugSettings)
+	bool bDrawSingleBoneChains = false;
+	
 	// Draw final IK goal locations. 
 	UPROPERTY(EditAnywhere, Category = DebugSettings)
 	bool bDrawFinalGoals = true;
@@ -449,6 +457,14 @@ public:
 	// Draw goal locations from source skeleton. 
 	UPROPERTY(EditAnywhere, Category = DebugSettings)
 	bool bDrawSourceLocations = true;
+
+	// Draw circle on the floor below the retarget root. 
+	UPROPERTY(EditAnywhere, Category = DebugSettings)
+	bool bDrawRootCircle = true;
+
+	// Draw coordinate frame used to define stride warping directions.  
+	UPROPERTY(EditAnywhere, Category = DebugSettings)
+	bool bDrawWarpingFrame = false;
 	
 	// The visual size of the IK goals in the viewport. 
 	UPROPERTY(EditAnywhere, Category = DebugSettings)
diff --git a/Engine/Plugins/Animation/IKRig/Source/IKRigEditor/Private/RetargetEditor/IKRetargetDefaultMode.cpp b/Engine/Plugins/Animation/IKRig/Source/IKRigEditor/Private/RetargetEditor/IKRetargetDefaultMode.cpp
--- a/Engine/Plugins/Animation/IKRig/Source/IKRigEditor/Private/RetargetEditor/IKRetargetDefaultMode.cpp
+++ b/Engine/Plugins/Animation/IKRig/Source/IKRigEditor/Private/RetargetEditor/IKRetargetDefaultMode.cpp
@@ -78,7 +78,7 @@ void FIKRetargetDefaultMode::RenderDebugProxies(FPrimitiveDrawInterface* PDI, co
 
 	UDebugSkelMeshComponent* TargetSkelMesh = Controller->GetSkeletalMeshComponent(ERetargetSourceOrTarget::Target);
 	const FTransform ComponentTransform = TargetSkelMesh->GetComponentTransform();
-	const FVector ComponentOffset = ComponentTransform.GetTranslation();
+	const float ComponentScale = ComponentTransform.GetScale3D().GetMax();
 
 	const TArray<FName>& SelectedChains = Controller->GetSelectedChains();
 
@@ -88,140 +88,160 @@ void FIKRetargetDefaultMode::RenderDebugProxies(FPrimitiveDrawInterface* PDI, co
 	const FLinearColor MainColor = FLinearColor::Green;
 	const FLinearColor NonSelected = FLinearColor::Gray * 0.3f;
 
-	// get the root modification
-	const FRootRetargeter& RootRetargeter = RetargetProcessor->GetRootRetargeter();
-	const FVector RootModification = RootRetargeter.Target.RootTranslationDelta * RootRetargeter.Settings.GetAffectIKWeightVector();
-
 	// draw IK goals on each IK chain
-	const TArray<FRetargetChainPairIK>& IKChainPairs = RetargetProcessor->GetIKChainPairs();
-	for (const FRetargetChainPairIK& IKChainPair : IKChainPairs)
+	if (Asset->bDrawFinalGoals || Asset->bDrawSourceLocations)
 	{
-		const FChainDebugData& ChainDebugData = IKChainPair.IKChainRetargeter.DebugData;
-		FTransform FinalTransform = ChainDebugData.OutputTransformEnd;
-		FinalTransform.AddToTranslation(ComponentOffset);
-
-		const bool bIsSelected = SelectedChains.Contains(IKChainPair.TargetBoneChainName);
-
-		PDI->SetHitProxy(new HIKRetargetEditorChainProxy(IKChainPair.TargetBoneChainName));
-
-		if (Asset->bDrawFinalGoals)
-		{
-			IKRigDebugRendering::DrawWireCube(
-			PDI,
-			FinalTransform,
-			bIsSelected ? GoalColor : GoalColor * NonSelected,
-			Asset->ChainDrawSize,
-			Asset->ChainDrawThickness);
-		}
+		// get the root modification
+		const FRootRetargeter& RootRetargeter = RetargetProcessor->GetRootRetargeter();
+		const FVector RootModification = RootRetargeter.Target.RootTranslationDelta * RootRetargeter.Settings.GetAffectIKWeightVector();
 		
-		if (Asset->bDrawSourceLocations)
+		// spin through all IK chains
+		const TArray<FRetargetChainPairIK>& IKChainPairs = RetargetProcessor->GetIKChainPairs();
+		for (const FRetargetChainPairIK& IKChainPair : IKChainPairs)
 		{
-			const FSourceChainIK& SourceChain = IKChainPair.IKChainRetargeter.Source;
-			FTransform SourceGoalTransform;
-			SourceGoalTransform.SetTranslation(SourceChain.CurrentEndPosition + ComponentOffset + RootModification);
-			SourceGoalTransform.SetRotation(SourceChain.CurrentEndRotation);
+			const FChainDebugData& ChainDebugData = IKChainPair.IKChainRetargeter.DebugData;
+			FTransform FinalTransform = ChainDebugData.OutputTransformEnd * ComponentTransform;
 
-			FLinearColor Color = bIsSelected ? SourceColor : SourceColor * NonSelected;
+			const bool bIsSelected = SelectedChains.Contains(IKChainPair.TargetBoneChainName);
 
-			DrawWireSphere(
-				PDI,
-				SourceGoalTransform,
-				Color,
-				Asset->ChainDrawSize * 0.5f,
-				12,
-				SDPG_World,
-				0.0f,
-				0.001f,
-				false);
+			PDI->SetHitProxy(new HIKRetargetEditorChainProxy(IKChainPair.TargetBoneChainName));
 
 			if (Asset->bDrawFinalGoals)
 			{
-				DrawDashedLine(
+				IKRigDebugRendering::DrawWireCube(
+				PDI,
+				FinalTransform,
+				bIsSelected ? GoalColor : GoalColor * NonSelected,
+				Asset->ChainDrawSize,
+				Asset->ChainDrawThickness * ComponentScale);
+			}
+		
+			if (Asset->bDrawSourceLocations)
+			{
+				const FSourceChainIK& SourceChain = IKChainPair.IKChainRetargeter.Source;
+				FTransform SourceGoalTransform;
+				SourceGoalTransform.SetTranslation(SourceChain.CurrentEndPosition + RootModification);
+				SourceGoalTransform.SetRotation(SourceChain.CurrentEndRotation);
+				SourceGoalTransform *= ComponentTransform;
+
+				FLinearColor Color = bIsSelected ? SourceColor : SourceColor * NonSelected;
+
+				DrawWireSphere(
 					PDI,
-					SourceGoalTransform.GetLocation(),
-					FinalTransform.GetLocation(),
+					SourceGoalTransform,
 					Color,
-					1.0f,
-					SDPG_Foreground);
+					Asset->ChainDrawSize * 0.5f,
+					12,
+					SDPG_World,
+					0.0f,
+					0.001f,
+					false);
+
+				if (Asset->bDrawFinalGoals)
+				{
+					DrawDashedLine(
+						PDI,
+						SourceGoalTransform.GetLocation(),
+						FinalTransform.GetLocation(),
+						Color,
+						1.0f,
+						SDPG_Foreground);
+				}
 			}
-		}
 
-		// done drawing chain proxies
-		PDI->SetHitProxy(nullptr);
+			// done drawing chain proxies
+			PDI->SetHitProxy(nullptr);
+		}
 	}
+	
 
 	// draw lines on each FK chain
-	const TArray<FRetargetChainPairFK>& FKChainPairs = RetargetProcessor->GetFKChainPairs();
-	for (const FRetargetChainPairFK& FKChainPair : FKChainPairs)
+	if (Asset->bDrawChainLines || Asset->bDrawSingleBoneChains)
 	{
-		const TArray<int32>& TargetChainBoneIndices = FKChainPair.FKDecoder.BoneIndices;
-		if (TargetChainBoneIndices.IsEmpty())
+		const TArray<FRetargetChainPairFK>& FKChainPairs = RetargetProcessor->GetFKChainPairs();
+		for (const FRetargetChainPairFK& FKChainPair : FKChainPairs)
 		{
-			continue;
-		}
+			const TArray<int32>& TargetChainBoneIndices = FKChainPair.FKDecoder.BoneIndices;
+			if (TargetChainBoneIndices.IsEmpty())
+			{
+				continue;
+			}
 		
-		const bool bIsSelected = SelectedChains.Contains(FKChainPair.TargetBoneChainName);
+			const bool bIsSelected = SelectedChains.Contains(FKChainPair.TargetBoneChainName);
+			FLinearColor Color = bIsSelected ? MainColor : MainColor * NonSelected;
 		
-		// draw a line from start to end of chain, or in the case of a chain with only 1 bone in it, draw a sphere
-		PDI->SetHitProxy(new HIKRetargetEditorChainProxy(FKChainPair.TargetBoneChainName));
-		if (TargetChainBoneIndices.Num() > 1)
-		{
-			FTransform StartTransform = TargetSkelMesh->GetBoneTransform(TargetChainBoneIndices[0], ComponentTransform);
-			FTransform EndTransform = TargetSkelMesh->GetBoneTransform(TargetChainBoneIndices.Last(), ComponentTransform);
-			PDI->DrawLine(
-			StartTransform.GetLocation(),
-			EndTransform.GetLocation(),
-			bIsSelected ? MainColor : MainColor * NonSelected,
-			SDPG_Foreground,
-			Asset->ChainDrawThickness);
-		}
-		else
-		{
-			FTransform StartTransform = TargetSkelMesh->GetBoneTransform(TargetChainBoneIndices[0], ComponentTransform);
-			
-			DrawWireSphere(
-				PDI,
+			// draw a line from start to end of chain, or in the case of a chain with only 1 bone in it, draw a sphere
+			PDI->SetHitProxy(new HIKRetargetEditorChainProxy(FKChainPair.TargetBoneChainName));
+			if (Asset->bDrawChainLines && TargetChainBoneIndices.Num() > 1)
+			{
+				FTransform StartTransform = TargetSkelMesh->GetBoneTransform(TargetChainBoneIndices[0], ComponentTransform);
+				FTransform EndTransform = TargetSkelMesh->GetBoneTransform(TargetChainBoneIndices.Last(), ComponentTransform);
+				PDI->DrawLine(
 				StartTransform.GetLocation(),
-				bIsSelected ? MainColor : MainColor * NonSelected,
-				Asset->ChainDrawSize,
-				12,
-				SDPG_World,
-				Asset->ChainDrawThickness,
-				0.001f,
-				false);
-		}
+				EndTransform.GetLocation(),
+				Color,
+				SDPG_Foreground,
+				Asset->ChainDrawThickness * ComponentScale);
+			}
+			else if (Asset->bDrawSingleBoneChains)
+			{
+				// single bone chain, just draw a sphere on the bone
+				FTransform BoneTransform = TargetSkelMesh->GetBoneTransform(TargetChainBoneIndices[0], ComponentTransform);
+				DrawWireSphere(
+					PDI,
+					BoneTransform,
+					Color,
+					Asset->ChainDrawSize,
+					12,
+					SDPG_World,
+					Asset->ChainDrawThickness * ComponentScale,
+					0.001f,
+					false);
+			}
 		
-		PDI->SetHitProxy(nullptr);
+			PDI->SetHitProxy(nullptr);
+		}
 	}
 
 	// draw stride warping frame
-	FTransform WarpingFrame = RetargetProcessor->DebugData.StrideWarpingFrame;
-	DrawCoordinateSystem(PDI, WarpingFrame.GetLocation(), WarpingFrame.GetRotation().Rotator(), Asset->ChainDrawSize, SDPG_World, Asset->ChainDrawThickness);
+	if (Asset->bDrawWarpingFrame)
+	{
+		FTransform WarpingFrame = RetargetProcessor->DebugData.StrideWarpingFrame * ComponentTransform;
+		DrawCoordinateSystem(
+			PDI,
+			WarpingFrame.GetLocation(),
+			WarpingFrame.GetRotation().Rotator(),
+			Asset->ChainDrawSize * ComponentScale,
+			SDPG_World,
+			Asset->ChainDrawThickness * ComponentScale);	
+	}
 
 	// root bone name
-	const FName RootBoneName = Controller->AssetController->GetRetargetRootBone(ERetargetSourceOrTarget::Target);
-	const int32 RootBoneIndex = TargetSkelMesh->GetReferenceSkeleton().FindBoneIndex(RootBoneName);
-	if (RootBoneIndex == INDEX_NONE)
+	if (Asset->bDrawRootCircle)
 	{
-		return;
-	}
-	const FTransform RootTransform = TargetSkelMesh->GetBoneTransform(RootBoneIndex, ComponentTransform);
-	const FVector RootCircleLocation = RootTransform.GetLocation() * FVector(1,1,0);
-	const bool bIsSelected = EditorController.Pin()->GetRootSelected();
-	const FLinearColor RootColor = bIsSelected ? MainColor : MainColor * NonSelected;
+		const FName RootBoneName = Controller->AssetController->GetRetargetRootBone(ERetargetSourceOrTarget::Target);
+		const int32 RootBoneIndex = TargetSkelMesh->GetReferenceSkeleton().FindBoneIndex(RootBoneName);
+		if (RootBoneIndex != INDEX_NONE)
+		{
+			const FTransform RootTransform = TargetSkelMesh->GetBoneTransform(RootBoneIndex, ComponentTransform);
+			const FVector RootCircleLocation = RootTransform.GetLocation() * FVector(1,1,0);
+			const bool bIsSelected = EditorController.Pin()->GetRootSelected();
+			const FLinearColor RootColor = bIsSelected ? MainColor : MainColor * NonSelected;
 	
-	PDI->SetHitProxy(new HIKRetargetEditorRootProxy());
-	DrawCircle(
-		PDI,
-		RootCircleLocation,
-		FVector(1, 0, 0),
-		FVector(0, 1, 0),
-		RootColor,
-		Asset->ChainDrawSize * 10.f,
-		12,
-		SDPG_World,
-		Asset->ChainDrawThickness * 2.0f);
-	PDI->SetHitProxy(nullptr);
+			PDI->SetHitProxy(new HIKRetargetEditorRootProxy());
+			DrawCircle(
+				PDI,
+				RootCircleLocation,
+				FVector(1, 0, 0),
+				FVector(0, 1, 0),
+				RootColor,
+				Asset->ChainDrawSize * 10.f * ComponentTransform.GetScale3D().GetMax(),
+				30,
+				SDPG_World,
+				Asset->ChainDrawThickness * 2.0f * ComponentScale);
+			PDI->SetHitProxy(nullptr);
+		}	
+	}
 }
 
 bool FIKRetargetDefaultMode::HandleClick(FEditorViewportClient* InViewportClient, HHitProxy* HitProxy, const FViewportClick& Click)
