Change to allow external access to batch retarget operation from other plugins (PR from github)

PR #9107: Update to IKRetargetBatchOperation to allow external access

diff --git a/Engine/Plugins/Animation/IKRig/Source/IKRigEditor/Public/RetargetEditor/IKRetargetBatchOperation.h b/Engine/Plugins/Animation/IKRig/Source/IKRigEditor/Public/RetargetEditor/IKRetargetBatchOperation.h
--- a/Engine/Plugins/Animation/IKRig/Source/IKRigEditor/Public/RetargetEditor/IKRetargetBatchOperation.h
+++ b/Engine/Plugins/Animation/IKRig/Source/IKRigEditor/Public/RetargetEditor/IKRetargetBatchOperation.h
@@ -27,7 +27,7 @@ enum class ERetargetRootLockMode : uint8
 };
 
 /// Data needed to run a batch "duplicate and retarget" operation on a set of animation assets
-struct FIKRetargetBatchOperationContext
+struct IKRIGEDITOR_API FIKRetargetBatchOperationContext
 {
 	
 public:
