Resolves slow auto-retargeter UI in projects with many animation blueprints.

diff --git a/Engine/Plugins/Animation/IKRig/Source/IKRigEditor/Private/RetargetEditor/SRetargetAnimAssetsWindow.cpp b/Engine/Plugins/Animation/IKRig/Source/IKRigEditor/Private/RetargetEditor/SRetargetAnimAssetsWindow.cpp
--- a/Engine/Plugins/Animation/IKRig/Source/IKRigEditor/Private/RetargetEditor/SRetargetAnimAssetsWindow.cpp
+++ b/Engine/Plugins/Animation/IKRig/Source/IKRigEditor/Private/RetargetEditor/SRetargetAnimAssetsWindow.cpp
@@ -875,7 +875,7 @@ bool SRetargetExporterAssetBrowser::OnShouldFilterAsset(const FAssetData& AssetD
 	{
 		return true;
 	}
-	
+
 	const USkeleton* DesiredSkeleton = BatchRetargetSettings->SourceSkeletalMesh->GetSkeleton();
 	if (!DesiredSkeleton)
 	{
@@ -884,10 +884,15 @@ bool SRetargetExporterAssetBrowser::OnShouldFilterAsset(const FAssetData& AssetD
 
 	if (bIsAnimBlueprint)
 	{
-		TObjectPtr<USkeleton> ABPSkeleton = Cast<UAnimBlueprint>(AssetData.GetAsset())->TargetSkeleton;
-		return !DesiredSkeleton->IsCompatibleForEditor(ABPSkeleton);
+		const FAssetDataTagMapSharedView::FFindTagResult Result = AssetData.TagsAndValues.FindTag("TargetSkeleton");
+		if (!Result.IsSet())
+		{
+			return true;
+		}
+
+		return DesiredSkeleton->IsCompatibleForEditor(Result.GetValue());
 	}
-	
+
 	return !DesiredSkeleton->IsCompatibleForEditor(AssetData);
 }
 
