Batch Retargeter: Added ability to retain additive settings of sequence asset.

diff --git a/Engine/Plugins/Animation/IKRig/Source/IKRigEditor/Private/RetargetEditor/IKRetargetBatchOperation.cpp b/Engine/Plugins/Animation/IKRig/Source/IKRigEditor/Private/RetargetEditor/IKRetargetBatchOperation.cpp
--- a/Engine/Plugins/Animation/IKRig/Source/IKRigEditor/Private/RetargetEditor/IKRetargetBatchOperation.cpp
+++ b/Engine/Plugins/Animation/IKRig/Source/IKRigEditor/Private/RetargetEditor/IKRetargetBatchOperation.cpp
@@ -207,6 +207,8 @@ void UIKRetargetBatchOperation::RetargetAssets(
 {
 	USkeleton* OldSkeleton = Context.SourceMesh->GetSkeleton();
 	USkeleton* NewSkeleton = Context.TargetMesh->GetSkeleton();
+
+	TArray<FAdditiveRetargetSettings> SettingsToRestoreAfterRetarget;
 	
 	for (UAnimationAsset* AssetToRetarget : AnimationAssetsToRetarget)
 	{
@@ -234,10 +236,9 @@ void UIKRetargetBatchOperation::RetargetAssets(
 			Controller.RemoveAllBoneTracks(bShouldTransact);
 
 			// reset all additive animation properties to ensure WYSIWYG playback of additive anims between retargeter and sequence
-			AnimSequenceToRetarget->AdditiveAnimType = EAdditiveAnimationType::AAT_None;
-			AnimSequenceToRetarget->RefPoseType = EAdditiveBasePoseType::ABPT_None;
-			AnimSequenceToRetarget->RefFrameIndex = 0;
-			AnimSequenceToRetarget->RefPoseSeq = nullptr;
+			FAdditiveRetargetSettings SequenceSettings;
+			SequenceSettings.PrepareForRetarget(AnimSequenceToRetarget);
+			SettingsToRestoreAfterRetarget.Add(SequenceSettings);
 			
 			// set the retarget source to the target skeletal mesh
 			AnimSequenceToRetarget->RetargetSource = NAME_None;
@@ -280,6 +281,15 @@ void UIKRetargetBatchOperation::RetargetAssets(
 	// convert the animation using the IK retargeter
 	ConvertAnimation(Context,Progress);
 
+	// optionally restore the additive flags
+	if (Context.bRetainAdditiveFlags)
+	{
+		for (FAdditiveRetargetSettings SettingsToRestore : SettingsToRestoreAfterRetarget)
+		{
+			SettingsToRestore.RestoreOnAsset();
+		}
+	}
+
 	// convert all Animation Blueprints and compile 
 	for (UAnimBlueprint* AnimBlueprint : AnimBlueprintsToRetarget)
 	{
@@ -741,6 +751,41 @@ void UIKRetargetBatchOperation::CleanupIfCancelled(const FScopedSlowTask& Progre
 	ObjectTools::DeleteObjects(NewAssets, bShowConfirmation);
 }
 
+void FAdditiveRetargetSettings::PrepareForRetarget(UAnimSequence* InSequenceAsset)
+{
+	if (!ensure(InSequenceAsset))
+	{
+		return;
+	}
+
+	SequenceAsset = InSequenceAsset;
+
+	// store setting values
+	AdditiveAnimType = SequenceAsset->AdditiveAnimType;
+	RefPoseType = SequenceAsset->RefPoseType;
+	RefFrameIndex = SequenceAsset->RefFrameIndex;
+	RefPoseSeq = SequenceAsset->RefPoseSeq;
+
+	// remove all additive settings so that retarget happens on base motion
+	SequenceAsset->AdditiveAnimType = EAdditiveAnimationType::AAT_None;
+	SequenceAsset->RefPoseType = EAdditiveBasePoseType::ABPT_None;
+	SequenceAsset->RefFrameIndex = 0;
+	SequenceAsset->RefPoseSeq = nullptr;
+}
+
+void FAdditiveRetargetSettings::RestoreOnAsset() const
+{
+	if (!ensure(SequenceAsset))
+	{
+		return;
+	}
+		
+	SequenceAsset->AdditiveAnimType = AdditiveAnimType;
+	SequenceAsset->RefPoseType = RefPoseType;
+	SequenceAsset->RefFrameIndex = RefFrameIndex;
+	SequenceAsset->RefPoseSeq = RefPoseSeq;
+}
+
 TArray<FAssetData> UIKRetargetBatchOperation::DuplicateAndRetarget(
 	const TArray<FAssetData>& AssetsToRetarget,
 	USkeletalMesh* SourceMesh,
diff --git a/Engine/Plugins/Animation/IKRig/Source/IKRigEditor/Private/RetargetEditor/SRetargetAnimAssetsWindow.cpp b/Engine/Plugins/Animation/IKRig/Source/IKRigEditor/Private/RetargetEditor/SRetargetAnimAssetsWindow.cpp
--- a/Engine/Plugins/Animation/IKRig/Source/IKRigEditor/Private/RetargetEditor/SRetargetAnimAssetsWindow.cpp
+++ b/Engine/Plugins/Animation/IKRig/Source/IKRigEditor/Private/RetargetEditor/SRetargetAnimAssetsWindow.cpp
@@ -469,6 +469,7 @@ FReply SBatchExportOptionsDialog::OnButtonClick(EAppReturnType::Type ButtonID)
 	const UBatchExportOptions* ExportOptions = UBatchExportOptions::GetInstance();
 	BatchContext->bIncludeReferencedAssets = ExportOptions->bIncludeReferencedAssets;
 	BatchContext->bOverwriteExistingFiles = ExportOptions->bOverwriteExistingFiles;
+	BatchContext->bRetainAdditiveFlags = ExportOptions->bRetainAdditiveFlags;
 	//BatchContext->bExportOnlyAnimatedBones = ExportOptions->bExportOnlyAnimatedBones;
 	
 	UserResponse = ButtonID;
diff --git a/Engine/Plugins/Animation/IKRig/Source/IKRigEditor/Public/RetargetEditor/IKRetargetBatchOperation.h b/Engine/Plugins/Animation/IKRig/Source/IKRigEditor/Public/RetargetEditor/IKRetargetBatchOperation.h
--- a/Engine/Plugins/Animation/IKRig/Source/IKRigEditor/Public/RetargetEditor/IKRetargetBatchOperation.h
+++ b/Engine/Plugins/Animation/IKRig/Source/IKRigEditor/Public/RetargetEditor/IKRetargetBatchOperation.h
@@ -13,6 +13,23 @@ struct FScopedSlowTask;
 class USkeletalMesh;
 class UAnimationAsset;
 class UAnimBlueprint;
+class UAnimSequence;
+
+
+struct FAdditiveRetargetSettings
+{
+	// the asset this operates on
+	TObjectPtr<UAnimSequence> SequenceAsset;
+	
+	// the settings to save/restore before/after retargeting
+	TEnumAsByte<EAdditiveAnimationType> AdditiveAnimType;
+	TEnumAsByte<EAdditiveBasePoseType> RefPoseType;
+	int32 RefFrameIndex;
+	TObjectPtr<UAnimSequence> RefPoseSeq;
+	
+	void PrepareForRetarget(UAnimSequence* InSequenceAsset);
+	void RestoreOnAsset() const;
+};
 
 // which skeleton are we referring to?
 UENUM()
@@ -61,6 +78,10 @@ public:
 	// Will not produce keys on bones that are not animated, reducing size on disk of the resulting files.
 	bool bExportOnlyAnimatedBones = true;
 
+	// Keep the additive animation sequence attributes on the retargeted results.
+	// NOTE: results may not be WYSIWYG with the editor, but the behavior should remain intact.
+	bool bRetainAdditiveFlags = true;
+
 	// Reset all data (called when window re-opened
 	void Reset()
 	{
diff --git a/Engine/Plugins/Animation/IKRig/Source/IKRigEditor/Public/RetargetEditor/SRetargetAnimAssetsWindow.h b/Engine/Plugins/Animation/IKRig/Source/IKRigEditor/Public/RetargetEditor/SRetargetAnimAssetsWindow.h
--- a/Engine/Plugins/Animation/IKRig/Source/IKRigEditor/Public/RetargetEditor/SRetargetAnimAssetsWindow.h
+++ b/Engine/Plugins/Animation/IKRig/Source/IKRigEditor/Public/RetargetEditor/SRetargetAnimAssetsWindow.h
@@ -79,6 +79,11 @@ public:
 	UPROPERTY(EditAnywhere, BlueprintReadWrite, Config, Category = "File")
 	bool bIncludeReferencedAssets = true;
 
+	// If retargeting additive animations, they will have their additive settings reset so that the retargeter evaluates the motion in a non-additive way.
+	// Setting this flag to true will ensure that the resulting animation sequences will retain their additive settings after the retarget operation.
+	UPROPERTY(EditAnywhere, BlueprintReadWrite, Config, Category = "File")
+	bool bRetainAdditiveFlags = true;
+
 	// TODO - Kiaran Feb 2024 - Naively leaving out non-retargeted keys results in flipped skeletons, needs work.
 	// Will not produce keys on bones that are not animated, reducing size on disk of the resulting files.
 	//UPROPERTY(EditAnywhere, BlueprintReadWrite, Config, Category = "Animation")
