Pin Bone tooltip changes.

diff --git a/Engine/Plugins/Animation/IKRig/Source/IKRig/Public/Retargeter/RetargetOps/PinBoneOp.h b/Engine/Plugins/Animation/IKRig/Source/IKRig/Public/Retargeter/RetargetOps/PinBoneOp.h
--- a/Engine/Plugins/Animation/IKRig/Source/IKRig/Public/Retargeter/RetargetOps/PinBoneOp.h
+++ b/Engine/Plugins/Animation/IKRig/Source/IKRig/Public/Retargeter/RetargetOps/PinBoneOp.h
@@ -32,12 +32,12 @@ struct FPinBoneData
 	, BoneToPinToIndex(INDEX_NONE)
 	, OffsetInRefPose(FTransform::Identity){}
 
-	// The bone to be affected by this op. Will have it's transform modified to match the BoneToPinTo
-	UPROPERTY(EditAnywhere, Category=Settings)
+	// The bone copy animation onto.
+	UPROPERTY(EditAnywhere, Category=Settings, DisplayName="CopyToBone")
 	FName BoneToPin;
 
-	// The bone, on the target skeleton to pin to.
-	UPROPERTY(EditAnywhere, Category=Settings)
+	// The bone to copy animation from.
+	UPROPERTY(EditAnywhere, Category=Settings, DisplayName="CopyFromBone")
 	FName BoneToPinTo;
 	
 	int32 BoneToPinIndex;
@@ -74,7 +74,7 @@ public:
 	UPROPERTY(EditAnywhere, Category=Settings)
 	EPinBoneType PinType = EPinBoneType::FullTransform;
 
-	// Maintain the original offset between the BoneToPin and BoneToPinTo
+	// Maintain the original offset between the source and target bone
 	UPROPERTY(EditAnywhere, Category=Settings)
 	bool bMaintainOffset = true;
 
