Retargeter: added support for auto-retargeting Rokoko v2 skeleton

diff --git a/Engine/Plugins/Animation/IKRig/Source/IKRigEditor/Private/RigEditor/IKRigAutoCharacterizer.cpp b/Engine/Plugins/Animation/IKRig/Source/IKRigEditor/Private/RigEditor/IKRigAutoCharacterizer.cpp
--- a/Engine/Plugins/Animation/IKRig/Source/IKRigEditor/Private/RigEditor/IKRigAutoCharacterizer.cpp
+++ b/Engine/Plugins/Animation/IKRig/Source/IKRigEditor/Private/RigEditor/IKRigAutoCharacterizer.cpp
@@ -1462,6 +1462,59 @@ FKnownTemplateHierarchies::FKnownTemplateHierarchies()
 		FNBiped.AutoRetargetDefinition.BonesToExcludeFromAutoPose.Add("L_BipedLeg_A_Ankle_Jnt");
 		FNBiped.AutoRetargetDefinition.BonesToExcludeFromAutoPose.Add("R_BipedLeg_A_Ankle_Jnt");
 	}
+
+	// Rokoko v2
+	{
+		static FName RokokoName = "Rokoko";
+		static TArray<FName> RokokoBones = {"Root", "Hips", "Spine1", "Spine2", "Spine3", "Spine4", "Neck", "Head", "RightShoulder", "RightArm", "RightForeArm", "RightHand", "RightFinger1Metacarpal", "RightFinger1Proximal", "RightFinger1Distal", "RightFinger5Metacarpal", "RightFinger5Proximal", "RightFinger5Medial", "RightFinger5Distal", "RightFinger4Metacarpal", "RightFinger4Proximal", "RightFinger4Medial", "RightFinger4Distal", "RightFinger3Metacarpal", "RightFinger3Proximal", "RightFinger3Medial", "RightFinger3Distal", "RightFinger2Metacarpal", "RightFinger2Proximal", "RightFinger2Medial", "RightFinger2Distal", "LeftShoulder", "LeftArm", "LeftForeArm", "LeftHand", "LeftFinger5Metacarpal", "LeftFinger5Proximal", "LeftFinger5Medial", "LeftFinger5Distal", "LeftFinger4Metacarpal", "LeftFinger4Proximal", "LeftFinger4Medial", "LeftFinger4Distal", "LeftFinger3Metacarpal", "LeftFinger3Proximal", "LeftFinger3Medial", "LeftFinger3Distal", "LeftFinger2Metacarpal", "LeftFinger2Proximal", "LeftFinger2Medial", "LeftFinger2Distal", "LeftFinger1Metacarpal", "LeftFinger1Proximal", "LeftFinger1Distal", "RightThigh", "RightShin", "RightFoot", "RightToe", "LeftThigh", "LeftShin", "LeftFoot", "LeftToe"};
+		static TArray<int32> RokokoParentIndices = {-1, 0, 1, 2, 3, 4, 5, 6, 5, 8, 9, 10, 11, 12, 13, 11, 15, 16, 17, 11, 19, 20, 21, 11, 23, 24, 25, 11, 27, 28, 29, 5, 31, 32, 33, 34, 35, 36, 37, 34, 39, 40, 41, 34, 43, 44, 45, 34, 47, 48, 49, 34, 51, 52, 1, 54, 55, 56, 1, 58, 59, 60};
+		FTemplateHierarchy& Rokoko = AddTemplateHierarchy(RokokoName, RokokoBones, RokokoParentIndices);
+		FRetargetDefinition& RokokoRetarget = Rokoko.AutoRetargetDefinition.RetargetDefinition;
+		// core
+		RokokoRetarget.RootBone = FName("Hips");
+		RokokoRetarget.AddBoneChain(FCharacterizationStandard::Spine, FName("Spine1"), FName("Spine4"));
+		RokokoRetarget.AddBoneChain(FCharacterizationStandard::Neck, FName("Neck"), FName("Neck"));
+		RokokoRetarget.AddBoneChain(FCharacterizationStandard::Head, FName("Head"), FName("Head"));
+		// left
+		RokokoRetarget.AddBoneChain(FCharacterizationStandard::LeftLeg, FName("LeftThigh"), FName("LeftToe"),FCharacterizationStandard::LeftFootIKGoal);
+		RokokoRetarget.AddBoneChain(FCharacterizationStandard::LeftClavicle, FName("LeftShoulder"), FName("LeftShoulder"));
+		RokokoRetarget.AddBoneChain(FCharacterizationStandard::LeftArm, FName("LeftArm"), FName("LeftHand"),FCharacterizationStandard::LeftHandIKGoal);
+		RokokoRetarget.AddBoneChain(FCharacterizationStandard::LeftThumb, FName("LeftFinger1Metacarpal"), FName("LeftFinger1Distal"));
+		RokokoRetarget.AddBoneChain(FCharacterizationStandard::LeftIndex, FName("LeftFinger2Proximal"), FName("LeftFinger2Distal"));
+		RokokoRetarget.AddBoneChain(FCharacterizationStandard::LeftMiddle, FName("LeftFinger3Proximal"), FName("LeftFinger3Distal"));
+		RokokoRetarget.AddBoneChain(FCharacterizationStandard::LeftRing, FName("LeftFinger4Proximal"), FName("LeftFinger4Distal"));
+		RokokoRetarget.AddBoneChain(FCharacterizationStandard::LeftPinky, FName("LeftFinger5Proximal"), FName("LeftFinger5Distal"));
+		RokokoRetarget.AddBoneChain(FCharacterizationStandard::LeftIndexMetacarpal, FName("LeftFinger2Metacarpal"), FName("LeftFinger2Metacarpal"));
+		RokokoRetarget.AddBoneChain(FCharacterizationStandard::LeftMiddleMetacarpal, FName("LeftFinger3Metacarpal"), FName("LeftFinger3Metacarpal"));
+		RokokoRetarget.AddBoneChain(FCharacterizationStandard::LeftRingMetacarpal, FName("LeftFinger4Metacarpal"), FName("LeftFinger4Metacarpal"));
+		RokokoRetarget.AddBoneChain(FCharacterizationStandard::LeftPinkyMetacarpal, FName("LeftFinger5Metacarpal"), FName("LeftFinger5Metacarpal"));
+		// right
+		RokokoRetarget.AddBoneChain(FCharacterizationStandard::RightLeg, FName("RightThigh"), FName("RightToe"),FCharacterizationStandard::RightFootIKGoal);
+		RokokoRetarget.AddBoneChain(FCharacterizationStandard::RightClavicle, FName("RightShoulder"), FName("RightShoulder"));
+		RokokoRetarget.AddBoneChain(FCharacterizationStandard::RightArm, FName("RightArm"), FName("RightHand"),FCharacterizationStandard::RightHandIKGoal);
+		RokokoRetarget.AddBoneChain(FCharacterizationStandard::RightThumb, FName("RightFinger1Metacarpal"), FName("RightFinger1Distal"));
+		RokokoRetarget.AddBoneChain(FCharacterizationStandard::RightIndex, FName("RightFinger2Proximal"), FName("RightFinger2Distal"));
+		RokokoRetarget.AddBoneChain(FCharacterizationStandard::RightMiddle, FName("RightFinger3Proximal"), FName("RightFinger3Distal"));
+		RokokoRetarget.AddBoneChain(FCharacterizationStandard::RightRing, FName("RightFinger4Proximal"), FName("RightFinger4Distal"));
+		RokokoRetarget.AddBoneChain(FCharacterizationStandard::RightPinky, FName("RightFinger5Proximal"), FName("RightFinger5Distal"));
+		RokokoRetarget.AddBoneChain(FCharacterizationStandard::RightIndexMetacarpal, FName("RightFinger2Metacarpal"), FName("RightFinger2Metacarpal"));
+		RokokoRetarget.AddBoneChain(FCharacterizationStandard::RightMiddleMetacarpal, FName("RightFinger3Metacarpal"), FName("RightFinger3Metacarpal"));
+		RokokoRetarget.AddBoneChain(FCharacterizationStandard::RightRingMetacarpal, FName("RightFinger4Metacarpal"), FName("RightFinger4Metacarpal"));
+		RokokoRetarget.AddBoneChain(FCharacterizationStandard::RightPinkyMetacarpal, FName("RightFinger5Metacarpal"), FName("RightFinger5Metacarpal"));
+		// bone settings for IK
+		Rokoko.AutoRetargetDefinition.BoneSettingsForIK.SetRotationStiffness(FName("Hips"), FCharacterizationStandard::PelvisRotationStiffness);
+		Rokoko.AutoRetargetDefinition.BoneSettingsForIK.SetRotationStiffness(FName("LeftShoulder"), FCharacterizationStandard::ClavicleRotationStiffness);
+		Rokoko.AutoRetargetDefinition.BoneSettingsForIK.SetRotationStiffness(FName("RightShoulder"), FCharacterizationStandard::ClavicleRotationStiffness);
+		Rokoko.AutoRetargetDefinition.BoneSettingsForIK.SetRotationStiffness(FName("LeftHand"), FCharacterizationStandard::FootRotationStiffness);
+		Rokoko.AutoRetargetDefinition.BoneSettingsForIK.SetRotationStiffness(FName("RightHand"), FCharacterizationStandard::FootRotationStiffness);
+		Rokoko.AutoRetargetDefinition.BoneSettingsForIK.SetPreferredAxis(FName("LeftLeg"), EPreferredAxis::NegativeZ);
+		Rokoko.AutoRetargetDefinition.BoneSettingsForIK.SetPreferredAxis(FName("RightLeg"), EPreferredAxis::PositiveZ);
+		Rokoko.AutoRetargetDefinition.BoneSettingsForIK.SetPreferredAxis(FName("LeftForeArm"), EPreferredAxis::PositiveZ);
+		Rokoko.AutoRetargetDefinition.BoneSettingsForIK.SetPreferredAxis(FName("RightForeArm"), EPreferredAxis::NegativeZ);
+		// exclude feet from auto-pose
+		Rokoko.AutoRetargetDefinition.BonesToExcludeFromAutoPose.Add("LeftFoot");
+		Rokoko.AutoRetargetDefinition.BonesToExcludeFromAutoPose.Add("RightFoot");
+	}
 }
 
 void FKnownTemplateHierarchies::GetClosestMatchingKnownHierarchy(
