New TRS Gizmo & IK Retargeter asset editor: don't use the current vpc axis in the new trs gizmo path, as the BeginTransform function will never be called if no axis is selected (this is directly treated at the gizmo level)

#rnx

diff --git a/Engine/Plugins/Animation/IKRig/Source/IKRigEditor/Private/RetargetEditor/IKRetargetEditPoseMode.cpp b/Engine/Plugins/Animation/IKRig/Source/IKRigEditor/Private/RetargetEditor/IKRetargetEditPoseMode.cpp
--- a/Engine/Plugins/Animation/IKRig/Source/IKRigEditor/Private/RetargetEditor/IKRetargetEditPoseMode.cpp
+++ b/Engine/Plugins/Animation/IKRig/Source/IKRigEditor/Private/RetargetEditor/IKRetargetEditPoseMode.cpp
@@ -164,6 +164,15 @@ bool FIKRetargetEditPoseMode::HandleClick(FEditorViewportClient* InViewportClien
 
 bool FIKRetargetEditPoseMode::StartTracking(FEditorViewportClient* InViewportClient, FViewport* InViewport)
 {
+	TrackingState = FIKRetargetTrackingState::None;
+	
+	// not manipulating any widget axes, so stop tracking
+	const EAxisList::Type CurrentAxis = InViewportClient->GetCurrentWidgetAxis();
+	if (CurrentAxis == EAxisList::None)
+	{
+		return false; 
+	}
+	
 	return HandleBeginTransform(InViewportClient);
 }
 
@@ -174,6 +183,7 @@ bool FIKRetargetEditPoseMode::EndTracking(FEditorViewportClient* InViewportClien
 
 bool FIKRetargetEditPoseMode::BeginTransform(const FGizmoState& InState)
 {
+	TrackingState = FIKRetargetTrackingState::None;
 	return HandleBeginTransform(Owner->GetFocusedViewportClient());
 }
 
@@ -188,15 +198,6 @@ bool FIKRetargetEditPoseMode::HandleBeginTransform(const FEditorViewportClient*
 	{
 		return false;
 	}
-	
-	TrackingState = FIKRetargetTrackingState::None;
-
-	// not manipulating any widget axes, so stop tracking
-	const EAxisList::Type CurrentAxis = InViewportClient->GetCurrentWidgetAxis();
-	if (CurrentAxis == EAxisList::None)
-	{
-		return false; 
-	}
 
 	const TSharedPtr<FIKRetargetEditorController> Controller = EditorController.Pin();
 	if (!Controller.IsValid())
