Fixes crash in IK Retarget editor when clearing source mesh

diff --git a/Engine/Plugins/Animation/IKRig/Source/IKRigEditor/Private/RetargetEditor/IKRetargetEditorController.cpp b/Engine/Plugins/Animation/IKRig/Source/IKRigEditor/Private/RetargetEditor/IKRetargetEditorController.cpp
--- a/Engine/Plugins/Animation/IKRig/Source/IKRigEditor/Private/RetargetEditor/IKRetargetEditorController.cpp
+++ b/Engine/Plugins/Animation/IKRig/Source/IKRigEditor/Private/RetargetEditor/IKRetargetEditorController.cpp
@@ -1397,6 +1397,11 @@ void FIKRetargetEditorController::AddReferencedObjects(FReferenceCollector& Coll
 void FIKRetargetEditorController::RenderSkeleton(FPrimitiveDrawInterface* PDI, ERetargetSourceOrTarget InSourceOrTarget) const
 {
 	const UDebugSkelMeshComponent* MeshComponent = GetSkeletalMeshComponent(InSourceOrTarget);
+	if (!MeshComponent->GetSkeletalMeshAsset())
+	{
+		return;
+	}
+	
 	const FTransform ComponentTransform = MeshComponent->GetComponentTransform();
 	const FReferenceSkeleton& RefSkeleton = MeshComponent->GetReferenceSkeleton();
 	const int32 NumBones = RefSkeleton.GetNum();
