Retarget Editor: Adding ability to align bones based on local rotation axes.

diff --git a/Engine/Plugins/Animation/IKRig/Source/IKRig/Private/Retargeter/IKRetargetProcessor.cpp b/Engine/Plugins/Animation/IKRig/Source/IKRig/Private/Retargeter/IKRetargetProcessor.cpp
--- a/Engine/Plugins/Animation/IKRig/Source/IKRig/Private/Retargeter/IKRetargetProcessor.cpp
+++ b/Engine/Plugins/Animation/IKRig/Source/IKRig/Private/Retargeter/IKRetargetProcessor.cpp
@@ -2419,6 +2419,21 @@ FTransform UIKRetargetProcessor::GetGlobalRetargetPoseAtParam(
 	return FTransform::Identity;
 }
 
+FTransform UIKRetargetProcessor::GetRetargetPoseBoneTransform(
+	const FName InBoneName,
+	const ERetargetSourceOrTarget SourceOrTarget,
+	ERetargetBoneSpace BoneSpace) const
+{
+	const int32 BoneIndex = GetBoneIndexFromName(InBoneName, SourceOrTarget);
+	if (!ensure(BoneIndex != INDEX_NONE))
+	{
+		return FTransform::Identity;
+	}
+
+	const FRetargetSkeleton& Skeleton = GetSkeleton(SourceOrTarget);
+	return BoneSpace == ERetargetBoneSpace::Global ? Skeleton.RetargetGlobalPose[BoneIndex] : Skeleton.RetargetLocalPose[BoneIndex];
+}
+
 const FChainFK* UIKRetargetProcessor::GetChain(const FName InChainName, const ERetargetSourceOrTarget SourceOrTarget) const
 {
 	const bool bIsSource = SourceOrTarget == ERetargetSourceOrTarget::Source;
@@ -2452,7 +2467,7 @@ float UIKRetargetProcessor::GetParamOfBoneInChain(
 	return Chain->Params[BoneIndexInChain];
 }
 
-FName UIKRetargetProcessor::GetBoneAtParam(
+FName UIKRetargetProcessor::GetClosestBoneToParam(
 	const FName InChainName, 
 	const float InParam,
 	const ERetargetSourceOrTarget SourceOrTarget) const
@@ -2463,17 +2478,20 @@ FName UIKRetargetProcessor::GetBoneAtParam(
 	{
 		return Skeleton.BoneNames[Chain->BoneIndices[0]];
 	}
-	
-	for (int32 ChainIndex = Chain->Params.Num()-1; ChainIndex >= 0 ; --ChainIndex)
+
+	float ClosestDistance = TNumericLimits<float>::Max();
+	int32 ChainIndexOfClosestBone = 0;
+	for (int32 ChainIndex = 0; ChainIndex < Chain->Params.Num(); ++ChainIndex)
 	{
-		if (Chain->Params[ChainIndex] <= InParam)
+		const float DistanceToParam = FMath::Abs(Chain->Params[ChainIndex] - InParam);
+		if (DistanceToParam <= ClosestDistance)
 		{
-			return Skeleton.BoneNames[Chain->BoneIndices[ChainIndex]];
+			ChainIndexOfClosestBone = ChainIndex;
+			ClosestDistance = DistanceToParam;
 		}
 	}
 
-	checkNoEntry();
-	return NAME_None;
+	return Skeleton.BoneNames[Chain->BoneIndices[ChainIndexOfClosestBone]];
 }
 
 FName UIKRetargetProcessor::GetMappedChainName(
diff --git a/Engine/Plugins/Animation/IKRig/Source/IKRig/Public/Retargeter/IKRetargetProcessor.h b/Engine/Plugins/Animation/IKRig/Source/IKRig/Public/Retargeter/IKRetargetProcessor.h
--- a/Engine/Plugins/Animation/IKRig/Source/IKRig/Public/Retargeter/IKRetargetProcessor.h
+++ b/Engine/Plugins/Animation/IKRig/Source/IKRig/Public/Retargeter/IKRetargetProcessor.h
@@ -451,6 +451,12 @@ struct FRetargetDebugData
 	FTransform StrideWarpingFrame;
 };
 
+enum class ERetargetBoneSpace : uint8
+{
+	Global,
+	Local
+};
+
 /** The runtime processor that converts an input pose from a source skeleton into an output pose on a target skeleton.
  * To use:
  * 1. Initialize a processor with a Source/Target skeletal mesh and a UIKRetargeter asset.
@@ -557,12 +563,14 @@ public:
 	FName GetChainNameForBone(const FName BoneName, const ERetargetSourceOrTarget SourceOrTarget) const;
 	/** Get a transform at a given param in a chain */
 	FTransform GetGlobalRetargetPoseAtParam(const FName InChainName, const float Param, const ERetargetSourceOrTarget SourceOrTarget) const;
+	/** Get a transform of a bone in the retarget pose */
+	FTransform GetRetargetPoseBoneTransform(const FName InBoneName, const ERetargetSourceOrTarget SourceOrTarget, ERetargetBoneSpace BoneSpace) const;
 	/** Get access to the internal chain on the given skeleton */
 	const FChainFK* GetChain(const FName InChainName, const ERetargetSourceOrTarget SourceOrTarget) const;
 	/** Get the param of the bone in it's retarget chain. Ranges from 0 to NumBonesInChain. */
 	float GetParamOfBoneInChain(const FName InBoneName, const ERetargetSourceOrTarget SourceOrTarget) const;
 	/** Get the bone in the chain at the given param. */
-	FName GetBoneAtParam(const FName InChainName, const float InParam, const ERetargetSourceOrTarget SourceOrTarget) const;
+	FName GetClosestBoneToParam(const FName InChainName, const float InParam, const ERetargetSourceOrTarget SourceOrTarget) const;
 	/** Get the chain mapped to this one. */
 	FName GetMappedChainName(const FName InChainName, const ERetargetSourceOrTarget SourceOrTarget);
 	/** store data for debug drawing */
diff --git a/Engine/Plugins/Animation/IKRig/Source/IKRigEditor/Private/RetargetEditor/IKRetargetCommands.cpp b/Engine/Plugins/Animation/IKRig/Source/IKRigEditor/Private/RetargetEditor/IKRetargetCommands.cpp
--- a/Engine/Plugins/Animation/IKRig/Source/IKRigEditor/Private/RetargetEditor/IKRetargetCommands.cpp
+++ b/Engine/Plugins/Animation/IKRig/Source/IKRigEditor/Private/RetargetEditor/IKRetargetCommands.cpp
@@ -19,11 +19,10 @@ void FIKRetargetCommands::RegisterCommands()
 	UI_COMMAND(ResetSelectedBones, "Reset Selected Bones", "Sets the selected bones to the mesh reference pose.", EUserInterfaceActionType::Button, FInputChord());
 	UI_COMMAND(ResetSelectedAndChildrenBones, "Reset Selected and Children Bones", "Sets the selected bones (and all children recursively) to the mesh reference pose.", EUserInterfaceActionType::Button, FInputChord());
 
-	UI_COMMAND(AutoAlignAllBones, "Align All Bones", "Goes through all bones contained in mapped retarget chains and rotates them to match the worldspace orientation of the equivalent bone on the other skeleton.", EUserInterfaceActionType::Button, FInputChord());
-	UI_COMMAND(AlignSelected, "Align Selected", "Goes through the selected bones contained in mapped retarget chains and rotates them to match the worldspace orientation of the equivalent bone on the other skeleton.", EUserInterfaceActionType::Button, FInputChord());
-	UI_COMMAND(AlignSelectedAndChildren, "Align Selected and Children", "Goes through the selected bones (and their children). If the bone is contained in a mapped retarget chains, it is rotated to match the worldspace orientation of the equivalent bone on the other skeleton.", EUserInterfaceActionType::Button, FInputChord());
-	UI_COMMAND(AlignSelectedUsingMesh, "Align Selected Using Mesh", "Align the selected bones using the mesh to determine the bone directions.", EUserInterfaceActionType::Button, FInputChord());
-	UI_COMMAND(SnapCharacterToGround, "Snap Character to Ground", "Move the character vertically to restore original height from the ground. Uses the selected bone as the reference point for a grounded limb. Otherwise searches the skeleton for the lowest retargeted bone.", EUserInterfaceActionType::Button, FInputChord());
+	UI_COMMAND(AlignAllBones, "Align All Bones", "Resets the retarget pose and then auto-aligns ALL retargeted bones using the specified method.", EUserInterfaceActionType::Button, FInputChord());
+	UI_COMMAND(AlignSelected, "Align Selected", "Auto-aligns only selected bones using the specified method.", EUserInterfaceActionType::Button, FInputChord());
+	UI_COMMAND(AlignSelectedAndChildren, "Align Selected and Children", "Auto-aligns selected bones and their children (recursive) using the specified method.", EUserInterfaceActionType::Button, FInputChord());
+	UI_COMMAND(SnapCharacterToGround, "Snap Character to Ground", "Translates the whole skeleton vertically to restore original height from the ground. Uses the selected bone as the reference point for a grounded limb. Otherwise searches the skeleton for the lowest retargeted bone.", EUserInterfaceActionType::Button, FInputChord());
 
 	UI_COMMAND(NewRetargetPose, "Create", "Create a new retarget pose.", EUserInterfaceActionType::Button, FInputChord());
 	UI_COMMAND(DuplicateRetargetPose, "Duplicate Current", "Duplicate the current retarget pose.", EUserInterfaceActionType::Button, FInputChord());
diff --git a/Engine/Plugins/Animation/IKRig/Source/IKRigEditor/Private/RetargetEditor/IKRetargetEditor.cpp b/Engine/Plugins/Animation/IKRig/Source/IKRigEditor/Private/RetargetEditor/IKRetargetEditor.cpp
--- a/Engine/Plugins/Animation/IKRig/Source/IKRigEditor/Private/RetargetEditor/IKRetargetEditor.cpp
+++ b/Engine/Plugins/Animation/IKRig/Source/IKRigEditor/Private/RetargetEditor/IKRetargetEditor.cpp
@@ -227,26 +227,20 @@ void FIKRetargetEditor::BindCommands()
 	//
 	
 	ToolkitCommands->MapAction(
-		Commands.AutoAlignAllBones,
-		FExecuteAction::CreateSP(EditorController, &FIKRetargetEditorController::HandleAlignAllBones),
+		Commands.AlignAllBones,
+		FExecuteAction::CreateSP(EditorController, &FIKRetargetEditorController::HandleAlignBones, false /* include children*/, true /*all bones*/),
 		FCanExecuteAction::CreateSP(EditorController, &FIKRetargetEditorController::IsEditingPose),
 		EUIActionRepeatMode::RepeatDisabled);
 	
 	ToolkitCommands->MapAction(
 		Commands.AlignSelected,
-		FExecuteAction::CreateSP(EditorController, &FIKRetargetEditorController::HandleAlignSelectedBones, ERetargetAutoAlignMethod::ChainToChain, false /* no children*/),
+		FExecuteAction::CreateSP(EditorController, &FIKRetargetEditorController::HandleAlignBones, false /* include children*/, false /*all bones*/),
 		FCanExecuteAction::CreateSP(EditorController, &FIKRetargetEditorController::IsEditingPoseWithAnyBoneSelected),
 		EUIActionRepeatMode::RepeatDisabled);
 
 	ToolkitCommands->MapAction(
 		Commands.AlignSelectedAndChildren,
-		FExecuteAction::CreateSP(EditorController, &FIKRetargetEditorController::HandleAlignSelectedBones, ERetargetAutoAlignMethod::ChainToChain, true /* include children*/),
-		FCanExecuteAction::CreateSP(EditorController, &FIKRetargetEditorController::IsEditingPoseWithAnyBoneSelected),
-		EUIActionRepeatMode::RepeatDisabled);
-
-	ToolkitCommands->MapAction(
-		Commands.AlignSelectedUsingMesh,
-		FExecuteAction::CreateSP(EditorController, &FIKRetargetEditorController::HandleAlignSelectedBones, ERetargetAutoAlignMethod::MeshToMesh, false /* no children*/),
+		FExecuteAction::CreateSP(EditorController, &FIKRetargetEditorController::HandleAlignBones, true /* include children*/, false /*all bones*/),
 		FCanExecuteAction::CreateSP(EditorController, &FIKRetargetEditorController::IsEditingPoseWithAnyBoneSelected),
 		EUIActionRepeatMode::RepeatDisabled);
 
diff --git a/Engine/Plugins/Animation/IKRig/Source/IKRigEditor/Private/RetargetEditor/IKRetargetEditorController.cpp b/Engine/Plugins/Animation/IKRig/Source/IKRigEditor/Private/RetargetEditor/IKRetargetEditorController.cpp
--- a/Engine/Plugins/Animation/IKRig/Source/IKRigEditor/Private/RetargetEditor/IKRetargetEditorController.cpp
+++ b/Engine/Plugins/Animation/IKRig/Source/IKRigEditor/Private/RetargetEditor/IKRetargetEditorController.cpp
@@ -1281,15 +1281,18 @@ void FIKRetargetEditorController::HandleResetSelectedAndChildrenBones() const
 	AssetController->ResetRetargetPose(CurrentPose, BonesToReset, GetSourceOrTarget());
 }
 
-void FIKRetargetEditorController::HandleAlignAllBones() const
-{
-	AssetController->AutoAlignAllBones(GetSourceOrTarget());
-}
-
-void FIKRetargetEditorController::HandleAlignSelectedBones(const ERetargetAutoAlignMethod Method, const bool bIncludeChildren) const
+void FIKRetargetEditorController::HandleAlignBones(
+	const bool bIncludeChildren,
+	const bool bIncludeAllBones) const
 {
+	if (bIncludeAllBones)
+	{
+		AssetController->AutoAlignAllBones(GetSourceOrTarget());
+		return;
+	}
+	
 	const TArray<FName> BonesToAlign = bIncludeChildren ? GetSelectedBonesAndChildren() : GetSelectedBones();
-	AssetController->AutoAlignBones(BonesToAlign, Method, GetSourceOrTarget());
+	AssetController->AutoAlignBones(BonesToAlign, CurrentPoseAlignmentMode, GetSourceOrTarget());
 }
 
 void FIKRetargetEditorController::HandleSnapToGround() const
diff --git a/Engine/Plugins/Animation/IKRig/Source/IKRigEditor/Private/RetargetEditor/IKRetargeterController.cpp b/Engine/Plugins/Animation/IKRig/Source/IKRigEditor/Private/RetargetEditor/IKRetargeterController.cpp
--- a/Engine/Plugins/Animation/IKRig/Source/IKRigEditor/Private/RetargetEditor/IKRetargeterController.cpp
+++ b/Engine/Plugins/Animation/IKRig/Source/IKRigEditor/Private/RetargetEditor/IKRetargeterController.cpp
@@ -928,7 +928,7 @@ FVector UIKRetargeterController::GetRootOffsetInRetargetPose(
 	return GetCurrentRetargetPose(SourceOrTarget).GetRootTranslationDelta();
 }
 
-void UIKRetargeterController::AutoAlignAllBones(ERetargetSourceOrTarget SourceOrTarget) const
+void UIKRetargeterController::AutoAlignAllBones(ERetargetSourceOrTarget SourceOrTarget, const ERetargetAutoAlignMethod Method) const
 {
 	// undo transaction
 	constexpr bool bShouldTransact = true;
@@ -939,10 +939,14 @@ void UIKRetargeterController::AutoAlignAllBones(ERetargetSourceOrTarget SourceOr
 
 	// first reset the entire retarget pose
 	ResetRetargetPose(GetCurrentRetargetPoseName(SourceOrTarget), TArray<FName>(), SourceOrTarget);
-	
+
 	// suppress warnings about bones that cannot be aligned when aligning ALL bones
 	constexpr bool bSuppressWarnings = true;
-	AutoPoseGenerator.Get()->AlignAllBones(SourceOrTarget, bSuppressWarnings);
+	AutoPoseGenerator.Get()->AlignBones(
+		TArray<FName>(), // empty list means "all bones"
+		Method,
+		SourceOrTarget,
+		bSuppressWarnings);
 }
 
 void UIKRetargeterController::AutoAlignBones(
diff --git a/Engine/Plugins/Animation/IKRig/Source/IKRigEditor/Private/RetargetEditor/IKRetargeterPoseGenerator.cpp b/Engine/Plugins/Animation/IKRig/Source/IKRigEditor/Private/RetargetEditor/IKRetargeterPoseGenerator.cpp
--- a/Engine/Plugins/Animation/IKRig/Source/IKRigEditor/Private/RetargetEditor/IKRetargeterPoseGenerator.cpp
+++ b/Engine/Plugins/Animation/IKRig/Source/IKRigEditor/Private/RetargetEditor/IKRetargeterPoseGenerator.cpp
@@ -31,20 +31,6 @@ void FRetargetAutoPoseGenerator::AddReferencedObjects(FReferenceCollector& Colle
 	Collector.AddReferencedObject(Processor);
 }
 
-void FRetargetAutoPoseGenerator::AlignAllBones(
-	const ERetargetSourceOrTarget SourceOrTarget,
-	const bool bSuppressWarnings) const
-{
-	// cannot get all bones until we have an initialized processor
-	if (!CheckReadyToAlignBones())
-	{
-		return;
-	}
-	
-	const TArray<FName>& AllBones = Processor->GetSkeleton(SourceOrTarget).BoneNames;
-	AlignBones(AllBones, ERetargetAutoAlignMethod::ChainToChain, SourceOrTarget, bSuppressWarnings);
-}
-
 void FRetargetAutoPoseGenerator::AlignBones(
 	const TArray<FName>& BonesToAlign,
 	const ERetargetAutoAlignMethod Method, 
@@ -56,9 +42,16 @@ void FRetargetAutoPoseGenerator::AlignBones(
 	{
 		return;
 	}
-
-	// sort hierarchically (we want to align from root to leaf order)
+	
+	// make copy to sort hierarchically (we want to align from root to leaf order)
 	TArray<FName> SortedBones = BonesToAlign;
+
+	// if user passed in no bones, then align ALL bones
+	if (SortedBones.IsEmpty())
+	{
+		SortedBones = Processor->GetSkeleton(SourceOrTarget).BoneNames;
+	}
+	
 	const FRetargetSkeleton& Skeleton = Processor->GetSkeleton(SourceOrTarget);
 	SortedBones.Sort([Skeleton](FName A, FName B) -> bool
 	{
@@ -151,6 +144,29 @@ void FRetargetAutoPoseGenerator::AlignBone(
 			bFoundMatchVector = GetChainTangentFromMesh(ChainToMatch, Param, SideToMatch, DirectionToMatch);
 			break;
 		}
+
+	case ERetargetAutoAlignMethod::LocalRotationAxes:
+		{
+			// get local rotation of equivalent bone to match
+			FName BoneToMatch = Processor->GetClosestBoneToParam(ChainToMatch, Param, SideToMatch);
+			if (BoneToMatch == NAME_None)
+			{
+				return;
+			}
+			FTransform LocalRtgPoseBoneToMatch = Processor->GetRetargetPoseBoneTransform(BoneToMatch, SideToMatch, ERetargetBoneSpace::Local);
+
+			// get the local ref pose of the bone to align
+			const FRetargetSkeleton& SkeletonToAlign = Processor->GetSkeleton(SideToAlign);
+			const FReferenceSkeleton& RefSkeletonToAlign = SkeletonToAlign.SkeletalMesh->GetRefSkeleton();
+			const int32 BoneToAlignIndex = SkeletonToAlign.FindBoneIndexByName(BoneToAlign);
+			FTransform LocalRefBoneToAlign = RefSkeletonToAlign.GetRefBonePose()[BoneToAlignIndex];
+
+			// calc delta to rotate from the ref pose to match the LRA of the other bone's retarget pose
+			const FQuat LocalRotationOffset = LocalRefBoneToAlign.GetRotation() * LocalRtgPoseBoneToMatch.GetRotation().Inverse();
+			Controller->SetRotationOffsetForRetargetPoseBone(BoneToAlign, LocalRotationOffset, SideToAlign);
+			Processor->UpdateRetargetPoseAtRuntime(CurrentPoseName, SideToAlign);
+			return;
+		}
 		
 	default:
 		{
@@ -340,7 +356,7 @@ bool FRetargetAutoPoseGenerator::GetChainTangentFromMesh(
 	FVector& OutDirection) const
 {
 	// get the bone at this location in the chain
-	const FName BoneName = Processor->GetBoneAtParam(ChainName, Param, SourceOrTarget);
+	const FName BoneName = Processor->GetClosestBoneToParam(ChainName, Param, SourceOrTarget);
 	
 	// get all children of bone
 	TArray<FName> BonesToConsider = {BoneName};
@@ -531,7 +547,7 @@ bool FRetargetAutoPoseGenerator::IsARetargetLeaf(
 	const float Param,
 	ERetargetSourceOrTarget SourceOrTarget) const
 {
-	const FName BoneName = Processor->GetBoneAtParam(ChainName, Param, SourceOrTarget);
+	const FName BoneName = Processor->GetClosestBoneToParam(ChainName, Param, SourceOrTarget);
 	TArray<FName> AllChildren;
 	GetAllChildrenOfBone(BoneName, SourceOrTarget, AllChildren);
 	for (const FName Child : AllChildren)
diff --git a/Engine/Plugins/Animation/IKRig/Source/IKRigEditor/Private/RetargetEditor/SIKRetargetPoseEditor.cpp b/Engine/Plugins/Animation/IKRig/Source/IKRigEditor/Private/RetargetEditor/SIKRetargetPoseEditor.cpp
--- a/Engine/Plugins/Animation/IKRig/Source/IKRigEditor/Private/RetargetEditor/SIKRetargetPoseEditor.cpp
+++ b/Engine/Plugins/Animation/IKRig/Source/IKRigEditor/Private/RetargetEditor/SIKRetargetPoseEditor.cpp
@@ -16,7 +16,7 @@ void SIKRetargetPoseEditor::Construct(
 	TSharedRef<FIKRetargetEditorController> InEditorController)
 {
 	EditorController = InEditorController;
-
+	
 	// the editor controller
 	FIKRetargetEditorController* Controller = EditorController.Pin().Get();
 
@@ -166,17 +166,59 @@ TSharedRef<SWidget> SIKRetargetPoseEditor::GenerateEditMenuContent(TSharedPtr<FU
 {
 	FMenuBuilder MenuBuilder(true, Commands);
 
-	MenuBuilder.AddMenuEntry(FIKRetargetCommands::Get().AutoAlignAllBones);
-	MenuBuilder.AddSeparator();
-	
+	static TArray<TSharedPtr<ERetargetAutoAlignMethod>> AlignModes =
+	{
+		MakeShared<ERetargetAutoAlignMethod>(ERetargetAutoAlignMethod::ChainToChain),
+		MakeShared<ERetargetAutoAlignMethod>(ERetargetAutoAlignMethod::LocalRotationAxes),
+		MakeShared<ERetargetAutoAlignMethod>(ERetargetAutoAlignMethod::MeshToMesh),
+	};
+
+	static auto ToDisplayText = [](ERetargetAutoAlignMethod Type)
+	{
+		switch (Type)
+		{
+		case ERetargetAutoAlignMethod::ChainToChain: return LOCTEXT("DirectionTypeLabel", "Direction");
+		case ERetargetAutoAlignMethod::LocalRotationAxes: return LOCTEXT("LocalAxisTypeLabel", "Local Rotation Axes");
+		case ERetargetAutoAlignMethod::MeshToMesh: return LOCTEXT("MeshTypeLabel", "Mesh");
+		default: return FText::GetEmpty();
+		}
+	};
+
+	MenuBuilder.BeginSection("AlignBonesHeader", LOCTEXT("AlignBonesHeader", "Auto-Align Bones"));
+	MenuBuilder.AddMenuEntry(FIKRetargetCommands::Get().AlignAllBones);
 	MenuBuilder.AddMenuEntry(FIKRetargetCommands::Get().AlignSelected);
 	MenuBuilder.AddMenuEntry(FIKRetargetCommands::Get().AlignSelectedAndChildren);
-	MenuBuilder.AddSeparator();
-	
-	MenuBuilder.AddMenuEntry(FIKRetargetCommands::Get().AlignSelectedUsingMesh);
-	MenuBuilder.AddSeparator();
-	
+	MenuBuilder.AddWidget(
+		SNew(SBox)
+		[
+			SNew(SComboBox<TSharedPtr<ERetargetAutoAlignMethod>>)
+			.OptionsSource(&AlignModes)
+			.ToolTipText(LOCTEXT("PoseAlignmentTooltipText",
+				"Direction: aligns the direction of the bone to match that of the equivalent bone in the other skeleton. Uses the chain hierarchy to define a direction vector. \n"
+				"Local Rotation Axes: aligns the local axes of the bone to match those of the equivalent bone in the other skeleton. May produce nonsensical results on skeletons with different rotation axes.\n"
+				"Mesh: Generates a direction vector for the bone based on the principle axis of the vertices weighted to the bone."))
+			.OnGenerateWidget_Lambda([](TSharedPtr<ERetargetAutoAlignMethod> InOption)
+			{
+				return SNew(STextBlock).Text(ToDisplayText(*InOption));
+			})
+			.OnSelectionChanged_Lambda([this](TSharedPtr<ERetargetAutoAlignMethod> NewSelection, ESelectInfo::Type SelectInfo)
+			{
+				EditorController.Pin()->CurrentPoseAlignmentMode = *NewSelection;
+			})
+			[
+				SNew(STextBlock)
+				.Text_Lambda([this]()
+				{
+					return ToDisplayText(EditorController.Pin()->CurrentPoseAlignmentMode);
+				})
+			]
+		]
+		, LOCTEXT("AlignModes_TitleText", "Alignment Method"), true /*noindent*/);
+	MenuBuilder.EndSection();
+
+	MenuBuilder.BeginSection("RootTranslationHeader", LOCTEXT("RootTranslationHeader", "Root Translation"));
 	MenuBuilder.AddMenuEntry(FIKRetargetCommands::Get().SnapCharacterToGround);
+	MenuBuilder.EndSection();
 
 	return MenuBuilder.MakeWidget();
 }
diff --git a/Engine/Plugins/Animation/IKRig/Source/IKRigEditor/Public/RetargetEditor/IKRetargetCommands.h b/Engine/Plugins/Animation/IKRig/Source/IKRigEditor/Public/RetargetEditor/IKRetargetCommands.h
--- a/Engine/Plugins/Animation/IKRig/Source/IKRigEditor/Public/RetargetEditor/IKRetargetCommands.h
+++ b/Engine/Plugins/Animation/IKRig/Source/IKRigEditor/Public/RetargetEditor/IKRetargetCommands.h
@@ -50,15 +50,13 @@ public:
 	// reset retarget pose to ref pose
 	TSharedPtr< FUICommandInfo > ResetSelectedAndChildrenBones;
 
-	// auto generate retarget pose
-	TSharedPtr< FUICommandInfo > AutoAlignAllBones;
+	// auto align all bones in skeleton
+	TSharedPtr< FUICommandInfo > AlignAllBones;
 	// auto align selected bones
 	TSharedPtr< FUICommandInfo > AlignSelected;
 	// auto align selected bones and children
 	TSharedPtr< FUICommandInfo > AlignSelectedAndChildren;
 	// auto align selected bones
-	TSharedPtr< FUICommandInfo > AlignSelectedUsingMesh;
-	// auto align selected bones
 	TSharedPtr< FUICommandInfo > SnapCharacterToGround;
 	
 	// delete retarget pose
diff --git a/Engine/Plugins/Animation/IKRig/Source/IKRigEditor/Public/RetargetEditor/IKRetargetEditorController.h b/Engine/Plugins/Animation/IKRig/Source/IKRigEditor/Public/RetargetEditor/IKRetargetEditorController.h
--- a/Engine/Plugins/Animation/IKRig/Source/IKRigEditor/Public/RetargetEditor/IKRetargetEditorController.h
+++ b/Engine/Plugins/Animation/IKRig/Source/IKRigEditor/Public/RetargetEditor/IKRetargetEditorController.h
@@ -279,9 +279,9 @@ public:
 	void HandleResetSelectedAndChildrenBones() const;
 	
 	// auto generate retarget pose
-	void HandleAlignAllBones() const;
-	void HandleAlignSelectedBones(const ERetargetAutoAlignMethod Method, const bool bIncludeChildren) const;
+	void HandleAlignBones(const bool bIncludeChildren, const bool bIncludeAllBones) const;
 	void HandleSnapToGround() const;
+	ERetargetAutoAlignMethod CurrentPoseAlignmentMode;
 
 	// create new retarget pose
 	void HandleNewPose();
diff --git a/Engine/Plugins/Animation/IKRig/Source/IKRigEditor/Public/RetargetEditor/IKRetargeterController.h b/Engine/Plugins/Animation/IKRig/Source/IKRigEditor/Public/RetargetEditor/IKRetargeterController.h
--- a/Engine/Plugins/Animation/IKRig/Source/IKRigEditor/Public/RetargetEditor/IKRetargeterController.h
+++ b/Engine/Plugins/Animation/IKRig/Source/IKRigEditor/Public/RetargetEditor/IKRetargeterController.h
@@ -288,7 +288,7 @@ public:
 
 	// Automatically align all bones in mapped chains and store in the current retarget pose.
 	UFUNCTION(BlueprintCallable, Category=RetargetOps)
-	void AutoAlignAllBones(ERetargetSourceOrTarget SourceOrTarget) const;
+	void AutoAlignAllBones(ERetargetSourceOrTarget SourceOrTarget, const ERetargetAutoAlignMethod Method = ERetargetAutoAlignMethod::ChainToChain) const;
 
 	// Automatically align an array of bones and store in the current retarget pose. Bones not in a mapped chain will be ignored.
 	UFUNCTION(BlueprintCallable, Category=RetargetOps)
diff --git a/Engine/Plugins/Animation/IKRig/Source/IKRigEditor/Public/RetargetEditor/IKRetargeterPoseGenerator.h b/Engine/Plugins/Animation/IKRig/Source/IKRigEditor/Public/RetargetEditor/IKRetargeterPoseGenerator.h
--- a/Engine/Plugins/Animation/IKRig/Source/IKRigEditor/Public/RetargetEditor/IKRetargeterPoseGenerator.h
+++ b/Engine/Plugins/Animation/IKRig/Source/IKRigEditor/Public/RetargetEditor/IKRetargeterPoseGenerator.h
@@ -32,6 +32,8 @@ enum class ERetargetAutoAlignMethod : uint8
 	ChainToChain,
 	// use mesh direction as source, and mesh direction as the target to align
 	MeshToMesh,
+	// match the local rotation axes of each bone individually (will be horribly wrong if characters don't share axis orientations)
+	LocalRotationAxes
 };
 
 // a sub-system used by the IKRetargetEditorController for automatically generating a retarget pose
@@ -44,11 +46,9 @@ struct FRetargetAutoPoseGenerator : public FGCObject
 	virtual void AddReferencedObjects(FReferenceCollector& Collector) override;
 	virtual FString GetReferencerName() const override { return TEXT("FRetargetAutoPoseGenerator"); }
 	// End FGCObject interface
-
-	// align all bones bones on the specified skeleton to the OTHER skeleton
-	void AlignAllBones(const ERetargetSourceOrTarget SourceOrTarget, const bool bSuppressWarnings) const;
 	
 	// align the specified bones on the current skeleton to the OTHER skeleton, using the specified method
+	// NOTE: if BonesToAlign is empty, it will align ALL bones in the skeleton
 	void AlignBones(
 		const TArray<FName>& BonesToAlign,
 		const ERetargetAutoAlignMethod Method,
