IK Retargeter: Added "Global Rotation Axes" pose alignment mode

diff --git a/Engine/Plugins/Animation/IKRig/Source/IKRigEditor/Private/RetargetEditor/IKRetargetEditorController.cpp b/Engine/Plugins/Animation/IKRig/Source/IKRigEditor/Private/RetargetEditor/IKRetargetEditorController.cpp
--- a/Engine/Plugins/Animation/IKRig/Source/IKRigEditor/Private/RetargetEditor/IKRetargetEditorController.cpp
+++ b/Engine/Plugins/Animation/IKRig/Source/IKRigEditor/Private/RetargetEditor/IKRetargetEditorController.cpp
@@ -1287,7 +1287,7 @@ void FIKRetargetEditorController::HandleAlignBones(
 {
 	if (bIncludeAllBones)
 	{
-		AssetController->AutoAlignAllBones(GetSourceOrTarget());
+		AssetController->AutoAlignAllBones(GetSourceOrTarget(), CurrentPoseAlignmentMode);
 		return;
 	}
 	
diff --git a/Engine/Plugins/Animation/IKRig/Source/IKRigEditor/Private/RetargetEditor/IKRetargeterController.cpp b/Engine/Plugins/Animation/IKRig/Source/IKRigEditor/Private/RetargetEditor/IKRetargeterController.cpp
--- a/Engine/Plugins/Animation/IKRig/Source/IKRigEditor/Private/RetargetEditor/IKRetargeterController.cpp
+++ b/Engine/Plugins/Animation/IKRig/Source/IKRigEditor/Private/RetargetEditor/IKRetargeterController.cpp
@@ -956,7 +956,7 @@ void UIKRetargeterController::AutoAlignBones(
 {
 	// undo transaction
 	constexpr bool bShouldTransact = true;
-	FScopedTransaction Transaction(LOCTEXT("AutoAlignSelectedBones", "Auto Align Selected Bones"), bShouldTransact);
+	FScopedTransaction Transaction(LOCTEXT("AutoAlignBones", "Auto Align Bones"), bShouldTransact);
 	Asset->Modify();
 
 	FScopedReinitializeIKRetargeter Reinitialize(this);
diff --git a/Engine/Plugins/Animation/IKRig/Source/IKRigEditor/Private/RetargetEditor/IKRetargeterPoseGenerator.cpp b/Engine/Plugins/Animation/IKRig/Source/IKRigEditor/Private/RetargetEditor/IKRetargeterPoseGenerator.cpp
--- a/Engine/Plugins/Animation/IKRig/Source/IKRigEditor/Private/RetargetEditor/IKRetargeterPoseGenerator.cpp
+++ b/Engine/Plugins/Animation/IKRig/Source/IKRigEditor/Private/RetargetEditor/IKRetargeterPoseGenerator.cpp
@@ -168,6 +168,44 @@ void FRetargetAutoPoseGenerator::AlignBone(
 			return;
 		}
 		
+	case ERetargetAutoAlignMethod::GlobalRotationAxes:
+		{
+			// get current retarget pose global transform with no delta of the bone to align
+			const FRetargetSkeleton& SkeletonToAlign = Processor->GetSkeleton(SideToAlign);
+			const FReferenceSkeleton& RefSkeletonToAlign = SkeletonToAlign.SkeletalMesh->GetRefSkeleton();
+			const int32 BoneToAlignIndex = SkeletonToAlign.FindBoneIndexByName(BoneToAlign);
+			FTransform ParentOfToAlignGlobal = FTransform::Identity;
+			const int32 ParentOfToAlignIndex = RefSkeletonToAlign.GetParentIndex(BoneToAlignIndex);
+			if (ParentOfToAlignIndex != INDEX_NONE)
+			{
+				const FName ParentOfToAlignName = RefSkeletonToAlign.GetBoneName(ParentOfToAlignIndex);
+				ParentOfToAlignGlobal = Processor->GetRetargetPoseBoneTransform(ParentOfToAlignName, SideToAlign, ERetargetBoneSpace::Global);
+			}
+			FTransform BoneToAlignLocalRefPose = RefSkeletonToAlign.GetRefBonePose()[BoneToAlignIndex];
+			FTransform BoneToAlignGlobalCurrentNoDelta = BoneToAlignLocalRefPose * ParentOfToAlignGlobal;
+
+			// get current retarget pose global transform of bone to match
+			FName BoneToMatch = Processor->GetClosestBoneToParam(ChainToMatch, Param, SideToMatch);
+			if (BoneToMatch == NAME_None)
+			{
+				return;
+			}
+			FTransform BoneToMatchGlobalCurrentRtgPose = Processor->GetRetargetPoseBoneTransform(BoneToMatch, SideToMatch, ERetargetBoneSpace::Global);
+
+			// transform BoneToAlign and BoneToMatch into the parent's local space
+			FTransform ParentOfToAlignInv = ParentOfToAlignGlobal.Inverse();
+			FTransform LocalToAlign = BoneToAlignGlobalCurrentNoDelta * ParentOfToAlignInv;
+			FTransform LocalToMatch = BoneToMatchGlobalCurrentRtgPose * ParentOfToAlignInv;
+
+			// compute the delta transform in local space
+			FTransform LocalDeltaTransform = LocalToMatch * LocalToAlign.Inverse();
+
+			Controller->SetRotationOffsetForRetargetPoseBone(BoneToAlign, LocalDeltaTransform.GetRotation(), SideToAlign);
+			Processor->UpdateRetargetPoseAtRuntime(CurrentPoseName, SideToAlign);
+			
+			return;
+		}
+		
 	default:
 		{
 			checkNoEntry();
diff --git a/Engine/Plugins/Animation/IKRig/Source/IKRigEditor/Private/RetargetEditor/SIKRetargetPoseEditor.cpp b/Engine/Plugins/Animation/IKRig/Source/IKRigEditor/Private/RetargetEditor/SIKRetargetPoseEditor.cpp
--- a/Engine/Plugins/Animation/IKRig/Source/IKRigEditor/Private/RetargetEditor/SIKRetargetPoseEditor.cpp
+++ b/Engine/Plugins/Animation/IKRig/Source/IKRigEditor/Private/RetargetEditor/SIKRetargetPoseEditor.cpp
@@ -170,6 +170,7 @@ TSharedRef<SWidget> SIKRetargetPoseEditor::GenerateEditMenuContent(TSharedPtr<FU
 	{
 		MakeShared<ERetargetAutoAlignMethod>(ERetargetAutoAlignMethod::ChainToChain),
 		MakeShared<ERetargetAutoAlignMethod>(ERetargetAutoAlignMethod::LocalRotationAxes),
+		MakeShared<ERetargetAutoAlignMethod>(ERetargetAutoAlignMethod::GlobalRotationAxes),
 		MakeShared<ERetargetAutoAlignMethod>(ERetargetAutoAlignMethod::MeshToMesh),
 	};
 
@@ -179,6 +180,7 @@ TSharedRef<SWidget> SIKRetargetPoseEditor::GenerateEditMenuContent(TSharedPtr<FU
 		{
 		case ERetargetAutoAlignMethod::ChainToChain: return LOCTEXT("DirectionTypeLabel", "Direction");
 		case ERetargetAutoAlignMethod::LocalRotationAxes: return LOCTEXT("LocalAxisTypeLabel", "Local Rotation Axes");
+		case ERetargetAutoAlignMethod::GlobalRotationAxes: return LOCTEXT("GlobalAxisTypeLabel", "Global Rotation Axes");
 		case ERetargetAutoAlignMethod::MeshToMesh: return LOCTEXT("MeshTypeLabel", "Mesh");
 		default: return FText::GetEmpty();
 		}
@@ -196,6 +198,7 @@ TSharedRef<SWidget> SIKRetargetPoseEditor::GenerateEditMenuContent(TSharedPtr<FU
 			.ToolTipText(LOCTEXT("PoseAlignmentTooltipText",
 				"Direction: aligns the direction of the bone to match that of the equivalent bone in the other skeleton. Uses the chain hierarchy to define a direction vector. \n"
 				"Local Rotation Axes: aligns the local axes of the bone to match those of the equivalent bone in the other skeleton. May produce nonsensical results on skeletons with different rotation axes.\n"
+				"Global Rotation Axes: aligns the global axes of the bone to match those of the equivalent bone in the other skeleton. May produce nonsensical results on skeletons with different rotation axes.\n"
 				"Mesh: Generates a direction vector for the bone based on the principle axis of the vertices weighted to the bone."))
 			.OnGenerateWidget_Lambda([](TSharedPtr<ERetargetAutoAlignMethod> InOption)
 			{
diff --git a/Engine/Plugins/Animation/IKRig/Source/IKRigEditor/Public/RetargetEditor/IKRetargeterPoseGenerator.h b/Engine/Plugins/Animation/IKRig/Source/IKRigEditor/Public/RetargetEditor/IKRetargeterPoseGenerator.h
--- a/Engine/Plugins/Animation/IKRig/Source/IKRigEditor/Public/RetargetEditor/IKRetargeterPoseGenerator.h
+++ b/Engine/Plugins/Animation/IKRig/Source/IKRigEditor/Public/RetargetEditor/IKRetargeterPoseGenerator.h
@@ -33,7 +33,9 @@ enum class ERetargetAutoAlignMethod : uint8
 	// use mesh direction as source, and mesh direction as the target to align
 	MeshToMesh,
 	// match the local rotation axes of each bone individually (will be horribly wrong if characters don't share axis orientations)
-	LocalRotationAxes
+	LocalRotationAxes,
+	// match the global rotation axes of each bone individually (will be horribly wrong if characters don't share axis orientations)
+	GlobalRotationAxes
 };
 
 // a sub-system used by the IKRetargetEditorController for automatically generating a retarget pose
