IK Retargeter: fixed batch retarget freezing IK goals when using speed planting with a missing curve.

diff --git a/Engine/Plugins/Animation/IKRig/Source/IKRigEditor/Private/RetargetEditor/IKRetargetBatchOperation.cpp b/Engine/Plugins/Animation/IKRig/Source/IKRigEditor/Private/RetargetEditor/IKRetargetBatchOperation.cpp
--- a/Engine/Plugins/Animation/IKRig/Source/IKRigEditor/Private/RetargetEditor/IKRetargetBatchOperation.cpp
+++ b/Engine/Plugins/Animation/IKRig/Source/IKRigEditor/Private/RetargetEditor/IKRetargetBatchOperation.cpp
@@ -362,6 +362,8 @@ void UIKRetargetBatchOperation::ConvertAnimation(
 	// get names of the curves the retargeter is looking for
 	TArray<FName> SpeedCurveNames;
 	Context.IKRetargetAsset->GetSpeedCurveNames(SpeedCurveNames);
+	// while retargeting, record a map of (sequence name, curve name) for any curve missing in a sequence
+	TMap<FString, TSet<FName>> MissingSpeedCurves;
 	
 	// for each pair of source / target animation sequences
 	for (TPair<UAnimationAsset*, UAnimationAsset*>& Pair : DuplicatedAnimAssets)
@@ -450,8 +452,16 @@ void UIKRetargetBatchOperation::ConvertAnimation(
 			
 			// get the curve values from the source sequence (for speed-based IK planting)
 			TMap<FName, float> SpeedCurveValues;
+			constexpr bool bForceUseRawData = false;
 			for (const FName& SpeedCurveName : SpeedCurveNames)
 			{
+				if (!SourceSequence->HasCurveData(SpeedCurveName, bForceUseRawData))
+				{
+					TSet<FName>& MissingCurves = MissingSpeedCurves.FindOrAdd(SourceSequence->GetName());
+					MissingCurves.Add(SpeedCurveName);
+					continue;
+				}
+
 				SpeedCurveValues.Add(SpeedCurveName, SourceSequence->EvaluateCurveData(SpeedCurveName, TimeAtCurrentFrame));
 			}
 
@@ -492,6 +502,15 @@ void UIKRetargetBatchOperation::ConvertAnimation(
 
 		TargetSeqController.CloseBracket(bShouldTransact);
 	}
+
+	// warn about any missing curves
+	for (const TPair<FString, TSet<FName>>& Pair : MissingSpeedCurves)
+	{
+		for (const FName& MissingCurve : Pair.Value)
+		{
+			UE_LOG(LogTemp, Warning, TEXT("IK Retarget Batch Operation: Missing speed curve, %s on sequence %s."), *MissingCurve.ToString(), *Pair.Key);	
+		}
+	}
 }
 
 void UIKRetargetBatchOperation::RemapCurves(const FIKRetargetBatchOperationContext& Context, FScopedSlowTask& Progress)
