IK Retargeter: fixes chain column in hierarchy display after editing retarget pose.

diff --git a/Engine/Plugins/Animation/IKRig/Source/IKRigEditor/Private/RetargetEditor/SIKRetargetHierarchy.cpp b/Engine/Plugins/Animation/IKRig/Source/IKRigEditor/Private/RetargetEditor/SIKRetargetHierarchy.cpp
--- a/Engine/Plugins/Animation/IKRig/Source/IKRigEditor/Private/RetargetEditor/SIKRetargetHierarchy.cpp
+++ b/Engine/Plugins/Animation/IKRig/Source/IKRigEditor/Private/RetargetEditor/SIKRetargetHierarchy.cpp
@@ -21,18 +21,26 @@ static const FName ChainColumnName(TEXT("RetargetChainName"));
 
 FIKRetargetHierarchyElement::FIKRetargetHierarchyElement(
 	const FName& InName,
-	const FName& InChainName,
 	const TSharedRef<FIKRetargetEditorController>& InEditorController)
 	: Key(FText::FromName(InName)),
 	Name(InName),
-	ChainName(InChainName),
 	EditorController(InEditorController)
 {}
 
+FName FIKRetargetHierarchyElement::GetChainName() const
+{
+	TSharedPtr<FIKRetargetEditorController> Controller = EditorController.Pin();
+	if (!Controller)
+	{
+		return NAME_None;
+	}
+
+	return Controller->GetChainNameFromBone(Name, Controller->GetSourceOrTarget());
+}
+
 TSharedRef<SWidget> SIKRetargetHierarchyRow::GenerateWidgetForColumn(const FName& ColumnName)
 {
 	const FName BoneName = WeakTreeElement.Pin()->Name;
-	const FName ChainName = WeakTreeElement.Pin()->ChainName;
 	ERetargetSourceOrTarget CurrentSkeleton = EditorController.Pin()->GetSourceOrTarget();
 	const bool bIsBoneRetargeted = EditorController.Pin()->IsBoneRetargeted(BoneName, CurrentSkeleton);
 	
@@ -106,7 +114,7 @@ TSharedRef<SWidget> SIKRetargetHierarchyRow::GenerateWidgetForColumn(const FName
 		.AutoWidth()
 		[
 			SNew(STextBlock)
-			.Text(FText::FromName(ChainName))
+			.Text_Lambda([this](){ return FText::FromName(WeakTreeElement.Pin()->GetChainName()); })
 			.Font(TextFont)
 			.ColorAndOpacity(TextColor)
 		];
@@ -433,10 +441,12 @@ void SIKRetargetHierarchy::RefreshTreeView(bool IsInitialSetup)
 	for (const FName BoneName : BoneNames)
 	{
 		// create "Bone" tree element for this bone
-		FName ChainName = Controller->GetChainNameFromBone(BoneName, Controller->GetSourceOrTarget());
-		TSharedPtr<FIKRetargetHierarchyElement> BoneElement = MakeShared<FIKRetargetHierarchyElement>(BoneName, ChainName, ControllerRef);
+		TSharedPtr<FIKRetargetHierarchyElement> BoneElement = MakeShared<FIKRetargetHierarchyElement>(BoneName, ControllerRef);
 		const int32 BoneElementIndex = AllElements.Add(BoneElement);
 		BoneTreeElementIndices.Add(BoneName, BoneElementIndex);
+
+		// get the chain this bone belongs to
+		FName ChainName = Controller->GetChainNameFromBone(BoneName, Controller->GetSourceOrTarget());
 		
 		// apply text filter to items
 		if (!(TextFilter->GetFilterText().IsEmpty() ||
@@ -448,7 +458,7 @@ void SIKRetargetHierarchy::RefreshTreeView(bool IsInitialSetup)
 		}
 
 		// apply filter "hide bones not in chain"
-		if (FilterOptions.bHideBonesNotInChain && BoneElement->ChainName == NAME_None)
+		if (FilterOptions.bHideBonesNotInChain && ChainName == NAME_None)
 		{
 			BoneElement->bIsHidden = true;
 		}
diff --git a/Engine/Plugins/Animation/IKRig/Source/IKRigEditor/Public/RetargetEditor/SIKRetargetHierarchy.h b/Engine/Plugins/Animation/IKRig/Source/IKRigEditor/Public/RetargetEditor/SIKRetargetHierarchy.h
--- a/Engine/Plugins/Animation/IKRig/Source/IKRigEditor/Public/RetargetEditor/SIKRetargetHierarchy.h
+++ b/Engine/Plugins/Animation/IKRig/Source/IKRigEditor/Public/RetargetEditor/SIKRetargetHierarchy.h
@@ -17,18 +17,16 @@ class FIKRetargetHierarchyElement : public TSharedFromThis<FIKRetargetHierarchyE
 {
 public:
 	
-	FIKRetargetHierarchyElement(
-		const FName& InName,
-		const FName& InChainName,
-		const TSharedRef<FIKRetargetEditorController>& InEditorController);
+	FIKRetargetHierarchyElement(const FName& InName, const TSharedRef<FIKRetargetEditorController>& InEditorController);
 
 	FText Key;
 	TSharedPtr<FIKRetargetHierarchyElement> Parent;
 	TArray<TSharedPtr<FIKRetargetHierarchyElement>> Children;
 	FName Name;
-	FName ChainName;
 	bool bIsHidden = false;
 
+	FName GetChainName() const;
+
 private:
 	
 	TWeakPtr<FIKRetargetEditorController> EditorController;
