IK Retargeter: fixes ensure from non-normalized rotations after applying global pose match

diff --git a/Engine/Plugins/Animation/IKRig/Source/IKRigEditor/Private/RetargetEditor/IKRetargeterPoseGenerator.cpp b/Engine/Plugins/Animation/IKRig/Source/IKRigEditor/Private/RetargetEditor/IKRetargeterPoseGenerator.cpp
--- a/Engine/Plugins/Animation/IKRig/Source/IKRigEditor/Private/RetargetEditor/IKRetargeterPoseGenerator.cpp
+++ b/Engine/Plugins/Animation/IKRig/Source/IKRigEditor/Private/RetargetEditor/IKRetargeterPoseGenerator.cpp
@@ -200,7 +200,7 @@ void FRetargetAutoPoseGenerator::AlignBone(
 			// compute the delta transform in local space
 			FTransform LocalDeltaTransform = LocalToMatch * LocalToAlign.Inverse();
 
-			Controller->SetRotationOffsetForRetargetPoseBone(BoneToAlign, LocalDeltaTransform.GetRotation(), SideToAlign);
+			Controller->SetRotationOffsetForRetargetPoseBone(BoneToAlign, LocalDeltaTransform.GetRotation().GetNormalized(), SideToAlign);
 			Processor->UpdateRetargetPoseAtRuntime(CurrentPoseName, SideToAlign);
 			
 			return;
