IK Rig: Added "Import Hierarchy" button

diff --git a/Engine/Plugins/Animation/IKRig/Source/IKRigEditor/Private/RigEditor/SIKRigHierarchy.cpp b/Engine/Plugins/Animation/IKRig/Source/IKRigEditor/Private/RigEditor/SIKRigHierarchy.cpp
--- a/Engine/Plugins/Animation/IKRig/Source/IKRigEditor/Private/RigEditor/SIKRigHierarchy.cpp
+++ b/Engine/Plugins/Animation/IKRig/Source/IKRigEditor/Private/RigEditor/SIKRigHierarchy.cpp
@@ -395,6 +395,8 @@ void SIKRigHierarchy::Construct(const FArguments& InArgs, TSharedRef<FIKRigEdito
 		.AutoHeight()
 		[
 			SNew(SHorizontalBox)
+			.Visibility_Lambda([this](){return IsSkeletonLoaded() ? EVisibility::Visible : EVisibility::Collapsed; })
+			
 			+SHorizontalBox::Slot()
 			.AutoWidth()
 			.VAlign(VAlign_Center)
@@ -434,11 +436,33 @@ void SIKRigHierarchy::Construct(const FArguments& InArgs, TSharedRef<FIKRigEdito
 				]
 			]
 		]
+
+    	+ SVerticalBox::Slot()
+		.AutoHeight()
+		.VAlign(VAlign_Top)
+		[
+			SNew(SHorizontalBox)
+			.Visibility_Lambda([this](){return IsSkeletonLoaded() ? EVisibility::Collapsed : EVisibility::Visible; })
+
+			+ SHorizontalBox::Slot()
+			.VAlign(VAlign_Center)
+			.MaxWidth(180.0f)
+			.Padding(3.0f, 1.0f)
+			[
+				SNew(SButton)
+				.ButtonStyle(FAppStyle::Get(), "FlatButton.Success")
+				.ForegroundColor(FLinearColor::White)
+				.OnClicked(FOnClicked::CreateSP(this, &SIKRigHierarchy::OnImportSkeletonClicked))
+				.Text(LOCTEXT("ImportSkeletonLabel", "Import Hierarchy"))
+				.IsEnabled_Lambda([this](){ return !IsSkeletonLoaded(); })
+			]
+		]
 		
         +SVerticalBox::Slot()
         .Padding(0.0f, 0.0f)
         [
             SNew(SBorder)
+        	.Visibility_Lambda([this](){return IsSkeletonLoaded() ? EVisibility::Visible : EVisibility::Collapsed; })
             .Padding(2.0f)
             .BorderImage(FAppStyle::GetBrush("SCSEditor.TreePanel"))
             [
@@ -699,6 +723,42 @@ bool SIKRigHierarchy::IsElementExcludedBone(TWeakPtr<FIKRigTreeElement> TreeElem
 	return AssetController->GetBoneExcluded(TreeElement.Pin()->BoneName);
 }
 
+bool SIKRigHierarchy::IsSkeletonLoaded() const
+{
+	const UIKRigController* AssetController = EditorController.Pin()->AssetController;
+	return AssetController->GetAsset()->GetPreviewMesh() != nullptr;
+}
+
+FReply SIKRigHierarchy::OnImportSkeletonClicked() const
+{
+	FIKRigHierarchyImportSettings Settings;
+	TSharedPtr<FStructOnScope> StructToDisplay = MakeShareable(new FStructOnScope(FIKRigHierarchyImportSettings::StaticStruct(), (uint8*)&Settings));
+
+	TSharedRef<SKismetInspector> KismetInspector = SNew(SKismetInspector);
+	KismetInspector->ShowSingleStruct(StructToDisplay);
+
+	SGenericDialogWidget::FArguments DialogArguments;
+	DialogArguments.OnOkPressed_Lambda([&Settings, this] ()
+	{
+		if (Settings.Mesh == nullptr)
+		{
+			return;
+		}
+
+		FIKRigEditorController* Controller = EditorController.Pin().Get();
+		if (!Controller)
+		{
+			return;
+		}
+
+		Controller->AssetController->SetSkeletalMesh(Settings.Mesh);
+	});
+
+	SGenericDialogWidget::OpenDialog(LOCTEXT("IKRigHierarchyImport", "Import Hierarchy"), KismetInspector, DialogArguments, true);
+
+	return FReply::Handled();
+}
+
 void SIKRigHierarchy::BindCommands()
 {
 	const FIKRigSkeletonCommands& Commands = FIKRigSkeletonCommands::Get();
diff --git a/Engine/Plugins/Animation/IKRig/Source/IKRigEditor/Public/RigEditor/SIKRigHierarchy.h b/Engine/Plugins/Animation/IKRig/Source/IKRigEditor/Public/RigEditor/SIKRigHierarchy.h
--- a/Engine/Plugins/Animation/IKRig/Source/IKRigEditor/Public/RigEditor/SIKRigHierarchy.h
+++ b/Engine/Plugins/Animation/IKRig/Source/IKRigEditor/Public/RigEditor/SIKRigHierarchy.h
@@ -9,9 +9,10 @@
 #include "DragAndDrop/DecoratedDragDropOp.h"
 #include "Widgets/SCompoundWidget.h"
 #include "Widgets/Views/STreeView.h"
-
 #include "SIKRigSolverStack.h"
 
+#include "SIKRigHierarchy.generated.h"
+
 class FIKRigEditorController;
 class SIKRigHierarchy;
 class FIKRigEditorToolkit;
@@ -164,6 +165,19 @@ struct FIKRigHierarchyFilterOptions
 	}
 };
 
+USTRUCT()
+struct FIKRigHierarchyImportSettings
+{
+	GENERATED_BODY()
+
+	FIKRigHierarchyImportSettings()
+	: Mesh(nullptr)
+	{}
+
+	UPROPERTY(EditAnywhere, Category = "Hierachy Import")
+	TObjectPtr<USkeletalMesh> Mesh;
+};
+
 typedef SBaseHierarchyTreeView<FIKRigTreeElement> SIKRigSkeletonTreeView;
 
 class SIKRigHierarchy : public SCompoundWidget, public FEditorUndoClient
@@ -200,6 +214,11 @@ public:
 	/** determine if the element is an excluded bone*/
 	bool IsElementExcludedBone(TWeakPtr<FIKRigTreeElement> TreeElement);
 
+	/** return true if a skeleton hierarchy has been loaded from an asset*/
+	bool IsSkeletonLoaded() const;
+	/** import skeletal mesh into the asset */
+	FReply OnImportSkeletonClicked() const;
+
 private:
 	/** SWidget interface */
 	virtual FReply OnKeyDown(const FGeometry& MyGeometry, const FKeyEvent& InKeyEvent);
