IK Retargeter: Separated translation and rotation blending when using "Match To Source"

diff --git a/Engine/Plugins/Animation/IKRig/Source/IKRig/Private/Retargeter/IKRetargetProcessor.cpp b/Engine/Plugins/Animation/IKRig/Source/IKRig/Private/Retargeter/IKRetargetProcessor.cpp
--- a/Engine/Plugins/Animation/IKRig/Source/IKRig/Private/Retargeter/IKRetargetProcessor.cpp
+++ b/Engine/Plugins/Animation/IKRig/Source/IKRig/Private/Retargeter/IKRetargetProcessor.cpp
@@ -936,6 +936,10 @@ void FChainRetargeterIK::DecodePose(
 	const float DeltaTime,
     const TArray<FTransform>& InGlobalPose)
 {
+	//
+	// calculate ROTATION of IK goal ...
+	//
+	
 	// apply delta rotation from input
 	const FQuat DeltaRotation = Source.CurrentEndRotation * Source.InitialEndRotation.Inverse();
 	FQuat GoalRotation = DeltaRotation * Target.InitialEndRotation;
@@ -943,8 +947,16 @@ void FChainRetargeterIK::DecodePose(
 	// apply static rotation offset in the local space of the foot
 	GoalRotation = GoalRotation * Settings.StaticRotationOffset.Quaternion();
 
+	// blend to source rotation
+	const float BlendToSourceRotation = Settings.BlendToSource * Settings.BlendToSourceRotation;
+	if (BlendToSourceRotation > KINDA_SMALL_NUMBER)
+	{
+		GoalRotation = FQuat::FastLerp(GoalRotation, Source.CurrentEndRotation, BlendToSourceRotation);
+		GoalRotation.Normalize();
+	}
+
 	//
-	// calculate position of IK goal ...
+	// calculate POSITION of IK goal ...
 	//
 	
 	// set position to length-scaled direction from source limb
@@ -954,10 +966,11 @@ void FChainRetargeterIK::DecodePose(
 	FVector GoalPosition = Start + (Source.CurrentEndDirectionNormalized * Target.InitialLength);
 
 	// blend to source location
-	if (Settings.BlendToSource > KINDA_SMALL_NUMBER)
+	const float BlendToSourceTranslation = Settings.BlendToSource * Settings.BlendToSourceTranslation;
+	if (BlendToSourceTranslation > KINDA_SMALL_NUMBER)
 	{
 		const FVector RootModification = RootRetargeter.Target.RootTranslationDelta * RootRetargeter.Settings.GetAffectIKWeightVector();
-		const FVector Weight = Settings.BlendToSource * Settings.BlendToSourceWeights;
+		const FVector Weight = BlendToSourceTranslation * Settings.BlendToSourceWeights;
 		const FVector SourceLocation = Source.CurrentEndPosition + RootModification;
 		GoalPosition.X = FMath::Lerp(GoalPosition.X, SourceLocation.X, Weight.X);
 		GoalPosition.Y = FMath::Lerp(GoalPosition.Y, SourceLocation.Y, Weight.Y);
diff --git a/Engine/Plugins/Animation/IKRig/Source/IKRig/Public/Retargeter/IKRetargetSettings.h b/Engine/Plugins/Animation/IKRig/Source/IKRig/Public/Retargeter/IKRetargetSettings.h
--- a/Engine/Plugins/Animation/IKRig/Source/IKRig/Public/Retargeter/IKRetargetSettings.h
+++ b/Engine/Plugins/Animation/IKRig/Source/IKRig/Public/Retargeter/IKRetargetSettings.h
@@ -119,11 +119,23 @@ struct IKRIG_API FTargetChainIKSettings
 	UPROPERTY(EditAnywhere, BlueprintReadWrite, Category = "IK Adjustments")
 	bool EnableIK = true;
 
-	/** Range 0 to 1. Default 0. Blends IK goal position from retargeted location (0) to source bone location (1).
+	/** Range 0 to 1. Default 0. Blends IK goal transform from retargeted transform (0) to source bone transform (1).
+	*  At 0 the goal is placed at the retargeted location and rotation.
+	*  At 1 the goal is placed at the location and rotation of the source chain's end bone. */
+	UPROPERTY(EditAnywhere, BlueprintReadWrite, Category = "IK Adjustments", meta = (ClampMin = "0.0", ClampMax = "1.0", UIMin = "0.0", UIMax = "1.0"))
+	float BlendToSource = 0.0f;
+
+	/** Range 0 to 1. Default 1. Blends the translational component of BlendToSource on/off.
 	*  At 0 the goal is placed at the retargeted location.
 	*  At 1 the goal is placed at the location of the source chain's end bone. */
 	UPROPERTY(EditAnywhere, BlueprintReadWrite, Category = "IK Adjustments", meta = (ClampMin = "0.0", ClampMax = "1.0", UIMin = "0.0", UIMax = "1.0"))
-	float BlendToSource = 0.0f;
+	float BlendToSourceTranslation = 1.0f;
+
+	/** Range 0 to 1. Default 0. Blends the rotational component of BlendToSource on/off.
+	*  At 0 the goal is oriented to the retargeted rotation.
+	*  At 1 the goal is oriented to the source chain's end bone rotation. */
+	UPROPERTY(EditAnywhere, BlueprintReadWrite, Category = "IK Adjustments", meta = (ClampMin = "0.0", ClampMax = "1.0", UIMin = "0.0", UIMax = "1.0"))
+	float BlendToSourceRotation = 0.0f;
 	
 	/** Range 0 to 1. Default 1. Weight each axis separately when using Blend To Source.
 	*  At 0 the goal is placed at the retargeted location.
