IK Retargeter: Fixed incorrect property markup and doc string on ScaleVertical property

diff --git a/Engine/Plugins/Animation/IKRig/Source/IKRig/Public/Retargeter/IKRetargetSettings.h b/Engine/Plugins/Animation/IKRig/Source/IKRig/Public/Retargeter/IKRetargetSettings.h
--- a/Engine/Plugins/Animation/IKRig/Source/IKRig/Public/Retargeter/IKRetargetSettings.h
+++ b/Engine/Plugins/Animation/IKRig/Source/IKRig/Public/Retargeter/IKRetargetSettings.h
@@ -155,8 +155,8 @@ struct IKRIG_API FTargetChainIKSettings
 	UPROPERTY(EditAnywhere, BlueprintReadWrite, Category = "IK Adjustments")
 	FRotator StaticRotationOffset = FRotator::ZeroRotator;
 
-	/** Range +-Infinity. Default 1. Scales the vertical component of the IK goal's position.*/
-	UPROPERTY(EditAnywhere, BlueprintReadWrite, Category = "IK Adjustments", meta = (UIMin = "0.0", UIMax = "5.0"))
+	/** Range 0 to infinity. Default 1. Scales the vertical component of the IK goal's position. */
+	UPROPERTY(EditAnywhere, BlueprintReadWrite, Category = "IK Adjustments", meta = (ClampMin = "0.0", UIMin = "0.0", UIMax = "5.0"))
 	float ScaleVertical = 1.0f;
 	
 	/** Range 0 to 5. Default 1. Brings IK goal closer (0) or further (1+) from origin of chain.
