IK Retargeter: upgraded Pin Bone Op behavior and UI
Replaced "Pin Type" with checkboxes for trans/rot/scale
Renamed BoneToPin and BoneToPinTo
Added FBoneReference widget to bone fields
Split offset modes into enums
Added new translation offset modes

diff --git a/Engine/Plugins/Animation/IKRig/Source/IKRig/Private/Retargeter/IKRetargetProcessor.cpp b/Engine/Plugins/Animation/IKRig/Source/IKRig/Private/Retargeter/IKRetargetProcessor.cpp
--- a/Engine/Plugins/Animation/IKRig/Source/IKRig/Private/Retargeter/IKRetargetProcessor.cpp
+++ b/Engine/Plugins/Animation/IKRig/Source/IKRig/Private/Retargeter/IKRetargetProcessor.cpp
@@ -1912,6 +1912,11 @@ bool UIKRetargetProcessor::InitializeOpStack(const TArray<TObjectPtr<URetargetOp
 		NewOp->bIsInitialized = NewOp->Initialize(this, SourceSkeleton, TargetSkeleton, Log);
 		bAllOpsInitialized &= NewOp->bIsInitialized;
 		OpStack.Add(NewOp);
+
+		// asset instance gets reference to running editor instance
+		#if WITH_EDITOR
+		OpFromAsset->EditorInstance = NewOp;
+		#endif
 	}
 
 	return bAllOpsInitialized;
diff --git a/Engine/Plugins/Animation/IKRig/Source/IKRig/Private/Retargeter/RetargetOps/PinBoneOp.cpp b/Engine/Plugins/Animation/IKRig/Source/IKRig/Private/Retargeter/RetargetOps/PinBoneOp.cpp
--- a/Engine/Plugins/Animation/IKRig/Source/IKRig/Private/Retargeter/RetargetOps/PinBoneOp.cpp
+++ b/Engine/Plugins/Animation/IKRig/Source/IKRig/Private/Retargeter/RetargetOps/PinBoneOp.cpp
@@ -1,9 +1,43 @@
 ﻿// Copyright Epic Games, Inc. All Rights Reserved.
 
 #include "Retargeter/RetargetOps/PinBoneOp.h"
+#include "Engine/SkeletalMesh.h"
+
+#include "IKRigObjectVersion.h"
 
 #define LOCTEXT_NAMESPACE "UPinBoneOp"
 
+void FPinBoneData::PostSerialize(const FArchive& Ar)
+{
+	
+#if WITH_EDITORONLY_DATA
+	
+	if (Ar.IsLoading())
+	{
+		if (Ar.CustomVer(FIKRigObjectVersion::GUID) < FIKRigObjectVersion::PinBoneTypeAndOffsetsUpgraded)
+		{
+			PRAGMA_DISABLE_DEPRECATION_WARNINGS
+			BoneToCopyTo.BoneName = BoneToPin_DEPRECATED;
+			BoneToCopyFrom.BoneName = BoneToPinTo_DEPRECATED;
+			PRAGMA_ENABLE_DEPRECATION_WARNINGS
+		}
+	}
+	
+#endif // WITH_EDITORONLY_DATA
+	
+}
+
+const FTransform& FPinBoneData::GetParentTransform(int32 BoneIndex, const FRetargetSkeleton& Skeleton, const TArray<FTransform>& Pose)
+{
+	int32 ParentIndex = Skeleton.GetParentIndex(BoneIndex);
+	if (ParentIndex == INDEX_NONE)
+	{
+		return FTransform::Identity;
+	}
+
+	return Pose[ParentIndex];
+}
+
 bool UPinBoneOp::Initialize(
 	const UIKRetargetProcessor* Processor,
 	const FRetargetSkeleton& SourceSkeleton,
@@ -11,26 +45,27 @@ bool UPinBoneOp::Initialize(
 	FIKRigLogger& Log)
 {
 	bool bFullyInitialized = true;
-	for (FPinBoneData& BoneToPin : BonesToPin)
+	for (FPinBoneData& PinData : BonesToPin)
 	{
-		BoneToPin.BoneToPinIndex = TargetSkeleton.FindBoneIndexByName(BoneToPin.BoneToPin);
+		PinData.BoneToCopyTo.BoneIndex = TargetSkeleton.FindBoneIndexByName(PinData.BoneToCopyTo.BoneName);
+		
 		if (PinTo == ERetargetSourceOrTarget::Source)
 		{
-			BoneToPin.BoneToPinToIndex = SourceSkeleton.FindBoneIndexByName(BoneToPin.BoneToPinTo);
+			PinData.BoneToCopyFrom.BoneIndex = SourceSkeleton.FindBoneIndexByName(PinData.BoneToCopyFrom.BoneName);
 		}
 		else
 		{
-			BoneToPin.BoneToPinToIndex = TargetSkeleton.FindBoneIndexByName(BoneToPin.BoneToPinTo);
+			PinData.BoneToCopyFrom.BoneIndex = TargetSkeleton.FindBoneIndexByName(PinData.BoneToCopyFrom.BoneName);
 		}
 
-		const bool bFoundBoneToPin = BoneToPin.BoneToPinIndex != INDEX_NONE;
-		const bool bFoundBoneToPinTo = BoneToPin.BoneToPinToIndex != INDEX_NONE;
+		const bool bFoundBoneToPin = PinData.BoneToCopyTo.BoneIndex != INDEX_NONE;
+		const bool bFoundBoneToPinTo = PinData.BoneToCopyFrom.BoneIndex != INDEX_NONE;
 		if (!bFoundBoneToPin)
 		{
 			bFullyInitialized = false;
 			Log.LogWarning(FText::Format(
 				LOCTEXT("MissingSourceBone", "Pin Bone retarget op refers to non-existant bone to pin, {0}."),
-				FText::FromName(BoneToPin.BoneToPin)));
+				FText::FromName(PinData.BoneToCopyTo.BoneName)));
 		}
 
 		if (!bFoundBoneToPinTo)
@@ -38,16 +73,28 @@ bool UPinBoneOp::Initialize(
 			bFullyInitialized = false;
 			Log.LogWarning(FText::Format(
 				LOCTEXT("MissingTargetBone", "Pin Bone retarget op refers to non-existant bone to pin to, {0}."),
-				FText::FromName(BoneToPin.BoneToPinTo)));
+				FText::FromName(PinData.BoneToCopyFrom.BoneName)));
 		}
 
-		if (bFoundBoneToPin && bFoundBoneToPinTo)
+		if (!(bFoundBoneToPin && bFoundBoneToPinTo))
 		{
-			const FTransform& BoneToPinTransform = TargetSkeleton.RetargetGlobalPose[BoneToPin.BoneToPinIndex];
-			const TArray<FTransform>& PinToPose = PinTo == ERetargetSourceOrTarget::Source ? SourceSkeleton.RetargetGlobalPose : TargetSkeleton.RetargetGlobalPose;
-			const FTransform& BoneToPinToTransform = PinToPose[BoneToPin.BoneToPinToIndex];
-			BoneToPin.OffsetInRefPose =  BoneToPinToTransform.GetRelativeTransform(BoneToPinTransform);
+			continue;
 		}
+
+		// cache offsets
+		const TArray<FTransform>& PoseToCopyFrom = PinTo == ERetargetSourceOrTarget::Source ? SourceSkeleton.RetargetGlobalPose : TargetSkeleton.RetargetGlobalPose;
+		const FRetargetSkeleton& SkeletonToCopyFrom = Processor->GetSkeleton(PinTo);
+		const FRetargetSkeleton& SkeletonToCopyTo = Processor->GetSkeleton(ERetargetSourceOrTarget::Target);
+		const FTransform& BoneToCopyFromRefPoseGlobal = SkeletonToCopyFrom.RetargetGlobalPose[PinData.BoneToCopyFrom.BoneIndex];
+		const FTransform& BoneToCopyToRefPoseGlobal = SkeletonToCopyTo.RetargetGlobalPose[PinData.BoneToCopyTo.BoneIndex];
+		// cache offset from BoneToCopyFrom to BoneToCopyTo in ref pose
+		PinData.OffsetFromBoneToCopyFromInRefPose = BoneToCopyFromRefPoseGlobal.GetRelativeTransform(BoneToCopyToRefPoseGlobal);
+		// cache local transform of BoneToCopyFrom in ref pose
+		const FTransform& BoneToCopyFromParentRefPoseGlobal = FPinBoneData::GetParentTransform(PinData.BoneToCopyFrom.BoneIndex, SkeletonToCopyFrom, PoseToCopyFrom);
+		PinData.LocalRefPoseBoneToCopyFrom = BoneToCopyFromParentRefPoseGlobal.GetRelativeTransform(BoneToCopyFromRefPoseGlobal);
+		// cache local transform of BoneToCopyTo in ref pose
+		const FTransform& BoneToCopyToParentRefPoseGlobal = FPinBoneData::GetParentTransform(PinData.BoneToCopyFrom.BoneIndex, SkeletonToCopyFrom, PoseToCopyFrom);
+		PinData.LocalRefPoseBoneToCopyTo = BoneToCopyToParentRefPoseGlobal.GetRelativeTransform(BoneToCopyToRefPoseGlobal);
 	}
 
 #if WITH_EDITOR
@@ -59,6 +106,10 @@ bool UPinBoneOp::Initialize(
 	{
 		Message = FText(LOCTEXT("MissingBonesWarning", "Bone(s) not found. See output log."));
 	}
+
+	// record skeletons for UI bone selector widgets
+	SourceSkeletonAsset = SourceSkeleton.SkeletalMesh->GetSkeleton();
+	TargetSkeletonAsset = TargetSkeleton.SkeletalMesh->GetSkeleton();
 #endif
 
 	// always treat this op as "initialized", individual pins will only execute if their prerequisites are met
@@ -70,39 +121,235 @@ void UPinBoneOp::Run(
 	const TArray<FTransform>& InSourceGlobalPose,
 	TArray<FTransform>& OutTargetGlobalPose)
 {
-	for (const FPinBoneData& BoneToPin : BonesToPin)
+	const TArray<FTransform>& PoseToCopyFrom = PinTo == ERetargetSourceOrTarget::Source ? InSourceGlobalPose : OutTargetGlobalPose;
+	
+	for (const FPinBoneData& PinData : BonesToPin)
 	{
-		if (BoneToPin.BoneToPinIndex == INDEX_NONE || BoneToPin.BoneToPinToIndex == INDEX_NONE)
+		bool bValidBoneToCopyFrom = PoseToCopyFrom.IsValidIndex(PinData.BoneToCopyFrom.BoneIndex);
+		bool bValidBoneToCopyTo = OutTargetGlobalPose.IsValidIndex(PinData.BoneToCopyTo.BoneIndex);
+		if (!bValidBoneToCopyFrom || !bValidBoneToCopyTo)
 		{
 			continue; // disabled or not successfully initialized
 		}
 
 		// calculate new transform for bone to pin
-		const TArray<FTransform>& PinToPose = PinTo == ERetargetSourceOrTarget::Source ? InSourceGlobalPose : OutTargetGlobalPose;
-		const FTransform& BoneToPinToTransform = PinToPose[BoneToPin.BoneToPinToIndex];
-		const FTransform OffsetToMaintain = bMaintainOffset ? BoneToPin.OffsetInRefPose : FTransform::Identity;
-		FTransform Result = LocalOffset * ((OffsetToMaintain * BoneToPinToTransform) * GlobalOffset);
+		const FTransform NewTransform = GetNewBoneTransform(PinData, Processor, InSourceGlobalPose, OutTargetGlobalPose);
+		// apply static local and global offsets
+		const FTransform Result = LocalOffset * (NewTransform * GlobalOffset);
+
+		// store result in output pose
+		FTransform& BoneToCopyToTransform = OutTargetGlobalPose[PinData.BoneToCopyTo.BoneIndex];
+		if (bCopyTranslation)
+		{
+			BoneToCopyToTransform.SetTranslation(Result.GetLocation());
+		}
+		if (bCopyRotation)
+		{
+			BoneToCopyToTransform.SetRotation(Result.GetRotation());
+		}
+		if (bCopyScale)
+		{
+			BoneToCopyToTransform.SetScale3D(Result.GetScale3D());
+		}
+	}
+}
+
+void UPinBoneOp::Serialize(FArchive& Ar)
+{
+	Super::Serialize(Ar);
+	Ar.UsingCustomVersion(FIKRigObjectVersion::GUID);
+
+#if WITH_EDITORONLY_DATA
+	if (Ar.IsLoading())
+	{
+		
+		if (Ar.CustomVer(FIKRigObjectVersion::GUID) < FIKRigObjectVersion::PinBoneTypeAndOffsetsUpgraded)
+		{
+			PRAGMA_DISABLE_DEPRECATION_WARNINGS
+
+			// load the boolean "bMaintainOffset" and convert it into the equivalent translation and rotation offset modes
+			if (bMaintainOffset_DEPRECATED)
+			{
+				TranslationMode = EPinBoneTranslationMode::CopyGlobalPositionAndMaintainOffset;
+				RotationMode = EPinBoneRotationMode::MaintainOffsetFromBoneToCopyFrom;
+			}
+			else
+			{
+				TranslationMode = EPinBoneTranslationMode::CopyGlobalPosition;
+				RotationMode = EPinBoneRotationMode::CopyGlobalRotation;
+			}
+
+			// load the "PinType" enum and convert it into the equivalent trans/rot/scale toggles
+			switch (PinType_DEPRECATED)
+			{
+			case EPinBoneType::FullTransform:
+				{
+					bCopyTranslation = true;
+					bCopyRotation = true;
+					bCopyScale = true;
+					break;
+				}
+			case EPinBoneType::TranslateOnly:
+				{
+					bCopyTranslation = true;
+					bCopyRotation = false;
+					bCopyScale = false;
+					break;
+				}
+			case EPinBoneType::RotateOnly:
+				{
+					bCopyTranslation = false;
+					bCopyRotation = true;
+					bCopyScale = false;
+					break;
+				}
+			case EPinBoneType::ScaleOnly:
+				{
+					bCopyTranslation = false;
+					bCopyRotation = false;
+					bCopyScale = true;
+					break;
+				}
+			default:
+				break;
+			}
+			PRAGMA_ENABLE_DEPRECATION_WARNINGS
+		}
+	}
+#endif
+}
+
+USkeleton* UPinBoneOp::GetSkeleton(bool& bInvalidSkeletonIsError, const IPropertyHandle* PropertyHandle)
+{
+#if WITH_EDITOR
+	if (UPinBoneOp* RunningInstance = Cast<UPinBoneOp>(EditorInstance))
+	{
+		FName PropertyName = PropertyHandle->GetProperty()->GetFName();
+		if (PropertyName == GET_MEMBER_NAME_CHECKED(FPinBoneData, BoneToCopyTo))
+		{
+			return RunningInstance->PinTo == ERetargetSourceOrTarget::Source ? RunningInstance->TargetSkeletonAsset : RunningInstance->SourceSkeletonAsset;
+		}
 
-		// apply result
-		FTransform& BoneToPinTransform = OutTargetGlobalPose[BoneToPin.BoneToPinIndex];
-		switch (PinType)
+		if (PropertyName == GET_MEMBER_NAME_CHECKED(FPinBoneData, BoneToCopyFrom))
 		{
-		case EPinBoneType::FullTransform:
-			BoneToPinTransform = Result;
+			return RunningInstance->PinTo == ERetargetSourceOrTarget::Target ? RunningInstance->TargetSkeletonAsset : RunningInstance->SourceSkeletonAsset;
+		}
+	}
+#endif
+	
+	ensureMsgf(false, TEXT("UPinBoneOp unable to get skeleton for UI widget."));
+	return nullptr;
+}
+
+FTransform UPinBoneOp::GetNewBoneTransform(
+	const FPinBoneData& PinData,
+	const UIKRetargetProcessor* Processor,
+	const TArray<FTransform>& InSourceGlobalPose,
+	TArray<FTransform>& OutTargetGlobalPose) const
+{
+	const TArray<FTransform>& PoseToCopyFrom = PinTo == ERetargetSourceOrTarget::Source ? InSourceGlobalPose : OutTargetGlobalPose;
+	const FRetargetSkeleton& SkeletonToCopyFrom = Processor->GetSkeleton(PinTo);
+	const FRetargetSkeleton& SkeletonToCopyTo = Processor->GetSkeleton(ERetargetSourceOrTarget::Target);
+
+	// get the current translational offset of BoneToCopyFrom relative to it's parent
+	auto GetBoneToCopyFromCurrentVector = [PoseToCopyFrom, SkeletonToCopyFrom, PinData]()
+	{
+		const FTransform& BoneToCopyFromCurrent = PoseToCopyFrom[PinData.BoneToCopyFrom.BoneIndex];
+		const FTransform& ParentOfBoneToCopyFromCurrent = FPinBoneData::GetParentTransform(PinData.BoneToCopyFrom.BoneIndex, SkeletonToCopyFrom, PoseToCopyFrom);
+		return BoneToCopyFromCurrent.GetTranslation() - ParentOfBoneToCopyFromCurrent.GetTranslation();
+	};
+	
+	// the resulting offset to return
+	FTransform Result = FTransform::Identity;
+
+	// generate translational offset
+	switch (TranslationMode)
+	{
+	case EPinBoneTranslationMode::CopyGlobalPosition:
+		{
+			Result.SetTranslation(PoseToCopyFrom[PinData.BoneToCopyFrom.BoneIndex].GetTranslation());
 			break;
-		case EPinBoneType::TranslateOnly:
-			BoneToPinTransform.SetTranslation(Result.GetLocation());
+		}
+		
+	case EPinBoneTranslationMode::CopyGlobalPositionAndMaintainOffset:
+		{
+			const FTransform BoneToCopyFromTransform = PoseToCopyFrom[PinData.BoneToCopyFrom.BoneIndex];
+			Result.SetTranslation((PinData.OffsetFromBoneToCopyFromInRefPose * BoneToCopyFromTransform).GetTranslation());
+			break;
+		}
+		
+	case EPinBoneTranslationMode::CopyLocalPosition:
+		{
+			const FVector BoneToCopyFromCurrentVector = GetBoneToCopyFromCurrentVector();
+			const FTransform& ParentOfBoneToCopyToCurrent = FPinBoneData::GetParentTransform(PinData.BoneToCopyTo.BoneIndex, SkeletonToCopyTo, OutTargetGlobalPose);
+			Result.SetTranslation(ParentOfBoneToCopyToCurrent.GetTranslation() + BoneToCopyFromCurrentVector);
+			break;
+		}
+		
+	case EPinBoneTranslationMode::CopyLocalPositionRelativeOffset:
+		{
+			const double RestPoseLengthDifference = PinData.LocalRefPoseBoneToCopyTo.GetTranslation().Size() - PinData.LocalRefPoseBoneToCopyFrom.GetTranslation().Size();
+			const FVector BoneToCopyFromTranslationCurrent = GetBoneToCopyFromCurrentVector();
+			double CurrentLength;
+			FVector BoneToCopyFromTranslationNorm;
+			BoneToCopyFromTranslationCurrent.ToDirectionAndLength(BoneToCopyFromTranslationNorm, CurrentLength);
+			const FVector BoneToCopyFromRelativeOffset = BoneToCopyFromTranslationNorm * (CurrentLength + RestPoseLengthDifference);
+		
+			const FTransform& ParentOfBoneToCopyToCurrent = FPinBoneData::GetParentTransform(PinData.BoneToCopyTo.BoneIndex, SkeletonToCopyTo, OutTargetGlobalPose);
+			Result.SetTranslation(ParentOfBoneToCopyToCurrent.GetTranslation() + BoneToCopyFromRelativeOffset);
+			break;
+		}
+		
+	case EPinBoneTranslationMode::CopyLocalPositionRelativeScaled:
+		{
+			const double LengthOfBoneToCopyToInRefPose = PinData.LocalRefPoseBoneToCopyTo.GetTranslation().Size();
+			const double LengthOfBoneToCopyFromInRefPose = PinData.LocalRefPoseBoneToCopyFrom.GetTranslation().Size();
+			double RestPoseLengthRatio = 1.0f;
+			if (LengthOfBoneToCopyFromInRefPose > KINDA_SMALL_NUMBER)
+			{
+				RestPoseLengthRatio = LengthOfBoneToCopyToInRefPose / LengthOfBoneToCopyFromInRefPose;
+			}
+		
+			const FVector BoneToCopyFromTranslationCurrent = GetBoneToCopyFromCurrentVector();
+			double CurrentLength;
+			FVector BoneToCopyFromTranslationNorm;
+			BoneToCopyFromTranslationCurrent.ToDirectionAndLength(BoneToCopyFromTranslationNorm, CurrentLength);
+			const FVector BoneToCopyFromRelativeOffset = BoneToCopyFromTranslationNorm * (CurrentLength * RestPoseLengthRatio);
+		
+			const FTransform& ParentOfBoneToCopyToCurrent = FPinBoneData::GetParentTransform(PinData.BoneToCopyTo.BoneIndex, SkeletonToCopyTo, OutTargetGlobalPose);
+			Result.SetTranslation(ParentOfBoneToCopyToCurrent.GetTranslation() + BoneToCopyFromRelativeOffset);
 			break;
-		case EPinBoneType::RotateOnly:
-			BoneToPinTransform.SetRotation(Result.GetRotation());
+		}
+		
+	default:
+		ensureAlwaysMsgf(false, TEXT("Missing translation offset mode."));
+		break;
+	}
+
+	// generate rotational offset
+	switch (RotationMode)
+	{
+	case EPinBoneRotationMode::CopyGlobalRotation:
+		{
+			Result.SetRotation(PoseToCopyFrom[PinData.BoneToCopyFrom.BoneIndex].GetRotation());
 			break;
-		case EPinBoneType::ScaleOnly:
-			BoneToPinTransform.SetScale3D(Result.GetScale3D());
+		}
+	case EPinBoneRotationMode::MaintainOffsetFromBoneToCopyFrom:
+		{
+			const FTransform& BoneToCopyFromRefPoseGlobal = SkeletonToCopyFrom.RetargetGlobalPose[PinData.BoneToCopyFrom.BoneIndex];
+			const FTransform& BoneToCopyToRefPoseGlobal = SkeletonToCopyTo.RetargetGlobalPose[PinData.BoneToCopyTo.BoneIndex];
+			const FTransform Offset = BoneToCopyFromRefPoseGlobal.GetRelativeTransform(BoneToCopyToRefPoseGlobal);
+			const FTransform BoneToCopyFromTransform = PoseToCopyFrom[PinData.BoneToCopyFrom.BoneIndex];
+			Result.SetRotation((Offset * BoneToCopyFromTransform).GetRotation());
 			break;
-		default:
-			checkNoEntry();
 		}
+		
+	default:
+		ensureAlwaysMsgf(false, TEXT("Missing rotation offset mode."));
+		break;
 	}
+
+	return Result;
 }
 
-#undef LOCTEXT_NAMESPACE
\ No newline at end of file
+#undef LOCTEXT_NAMESPACE
diff --git a/Engine/Plugins/Animation/IKRig/Source/IKRig/Public/IKRigObjectVersion.h b/Engine/Plugins/Animation/IKRig/Source/IKRig/Public/IKRigObjectVersion.h
--- a/Engine/Plugins/Animation/IKRig/Source/IKRig/Public/IKRigObjectVersion.h
+++ b/Engine/Plugins/Animation/IKRig/Source/IKRig/Public/IKRigObjectVersion.h
@@ -19,6 +19,10 @@ struct IKRIG_API FIKRigObjectVersion
 
 		// Chain settings moved to struct to be used in profiles
 		ChainSettingsConvertedToStruct,
+		
+		// Pin Bone "Pin Type" converted into separate trans/rot/scale toggles
+		// Pin Bone "Maintain Offset" converted from bool to enum
+		PinBoneTypeAndOffsetsUpgraded,
 
 		// -----<new versions can be added above this line>-------------------------------------------------
 		VersionPlusOne,
diff --git a/Engine/Plugins/Animation/IKRig/Source/IKRig/Public/Retargeter/IKRetargetOps.h b/Engine/Plugins/Animation/IKRig/Source/IKRig/Public/Retargeter/IKRetargetOps.h
--- a/Engine/Plugins/Animation/IKRig/Source/IKRig/Public/Retargeter/IKRetargetOps.h
+++ b/Engine/Plugins/Animation/IKRig/Source/IKRig/Public/Retargeter/IKRetargetOps.h
@@ -29,7 +29,7 @@ public:
 	// override to evaluate this operation and modify the output pose
 	virtual void Run(
 		const UIKRetargetProcessor* Processor,
-		const TArray<FTransform>& InSourceGlobalPose,
+		const TArray<FTransform>& InOutSourceGlobalPose,
 		TArray<FTransform>& OutTargetGlobalPose){};
 
 	UPROPERTY()
@@ -44,6 +44,8 @@ public:
 	virtual FText GetNiceName() const { return FText::FromString(TEXT("Default Op Name")); };
 	// override to display a warning message in the op stack
 	virtual FText WarningMessage() const { return FText::GetEmpty(); };
+	// the instance of this op currently running in the editor
+	mutable URetargetOpBase* EditorInstance;
 #endif
 };
 
diff --git a/Engine/Plugins/Animation/IKRig/Source/IKRig/Public/Retargeter/RetargetOps/PinBoneOp.h b/Engine/Plugins/Animation/IKRig/Source/IKRig/Public/Retargeter/RetargetOps/PinBoneOp.h
--- a/Engine/Plugins/Animation/IKRig/Source/IKRig/Public/Retargeter/RetargetOps/PinBoneOp.h
+++ b/Engine/Plugins/Animation/IKRig/Source/IKRig/Public/Retargeter/RetargetOps/PinBoneOp.h
@@ -2,15 +2,15 @@
 
 #pragma once
 
+#include "Interfaces/Interface_BoneReferenceSkeletonProvider.h"
 #include "Retargeter/IKRetargetOps.h"
 
 #include "PinBoneOp.generated.h"
 
 #define LOCTEXT_NAMESPACE "RetargetOps"
 
-// which skeleton are we referring to?
 UENUM()
-enum class EPinBoneType : uint8
+enum class UE_DEPRECATED(5.6, "Uses bools to include translation/rotation/scale (or combinations therof) separately.") EPinBoneType : uint8
 {
 	FullTransform,
 	TranslateOnly,
@@ -23,30 +23,66 @@ struct FPinBoneData
 {
 	GENERATED_BODY()
 
-	FPinBoneData() = default;
+	void PostSerialize(const FArchive& Ar);
 	
-	FPinBoneData(FName InBoneToPin, FName InBoneToPinTo)
-	: BoneToPin(InBoneToPin)
-	, BoneToPinTo(InBoneToPinTo)
-	, BoneToPinIndex(INDEX_NONE)
-	, BoneToPinToIndex(INDEX_NONE)
-	, OffsetInRefPose(FTransform::Identity){}
-
-	// The bone copy animation onto.
-	UPROPERTY(EditAnywhere, Category=Settings, DisplayName="CopyToBone")
-	FName BoneToPin;
-
-	// The bone to copy animation from.
-	UPROPERTY(EditAnywhere, Category=Settings, DisplayName="CopyFromBone")
-	FName BoneToPinTo;
-	
-	int32 BoneToPinIndex;
-	int32 BoneToPinToIndex;
-	FTransform OffsetInRefPose;
+	static const FTransform& GetParentTransform(int32 BoneIndex, const FRetargetSkeleton& Skeleton, const TArray<FTransform>& Pose);
+
+	// The bone to be affected by this op. Will have its transform modified to match the BoneToCopyFrom
+	UPROPERTY(EditAnywhere, Category=Settings)
+	FBoneReference BoneToCopyTo;
+	// The bone to copy transform from.
+	UPROPERTY(EditAnywhere, Category=Settings)
+	FBoneReference BoneToCopyFrom;
+
+	// the local transform of the BoneToCopyTo in the ref pose
+	FTransform LocalRefPoseBoneToCopyTo;
+	// the local transform of the BoneToCopyTo in the ref pose
+	FTransform LocalRefPoseBoneToCopyFrom;
+	// the relative transform between the BoneToCopyFrom and BoneToCopyTo in the reference pose
+	FTransform OffsetFromBoneToCopyFromInRefPose;
+
+	// deprecated properties
+	UPROPERTY(meta=(DeprecatedProperty))
+	FName BoneToPin_DEPRECATED;
+	UPROPERTY(meta=(DeprecatedProperty))
+	FName BoneToPinTo_DEPRECATED;
+};
+
+template<>
+struct TStructOpsTypeTraits<FPinBoneData> : public TStructOpsTypeTraitsBase2<FPinBoneData>
+{
+	enum
+	{
+		WithPostSerialize = true,
+	};
+};
+
+UENUM()
+enum class EPinBoneTranslationMode : uint8
+{
+	// Copy the global position of BoneToCopyFrom
+	CopyGlobalPosition,
+	// Copy the local translation (relative to parent), in global space, from BoneToCopyFrom to BoneToCopyTo
+	CopyLocalPosition,
+	// Copy the local translation, in global space, from BoneToCopyFrom to BoneToCopyTo and add the difference in lengths in the reference pose
+	CopyLocalPositionRelativeOffset,
+	// Copy the local translation, in global space, from BoneToCopyFrom to BoneToCopyTo and scale by the relative lengths in the reference pose
+	CopyLocalPositionRelativeScaled,
+	// Copy the global position of BoneToCopyFrom, and add the local offset between BoneToCopyFrom and BoneToCopyTo in the reference pose
+	CopyGlobalPositionAndMaintainOffset,
+};
+
+UENUM()
+enum class EPinBoneRotationMode : uint8
+{
+	// no offset is maintained
+	CopyGlobalRotation,
+	// maintains rotation offset between BoneToCopyTo and BoneToCopyFrom in the reference pose
+	MaintainOffsetFromBoneToCopyFrom
 };
 
 UCLASS(BlueprintType, EditInlineNew)
-class IKRIG_API UPinBoneOp : public URetargetOpBase
+class IKRIG_API UPinBoneOp : public URetargetOpBase, public IBoneReferenceSkeletonProvider
 {
 	GENERATED_BODY()
 
@@ -63,20 +99,41 @@ public:
 		const TArray<FTransform>& InSourceGlobalPose,
 		TArray<FTransform>& OutTargetGlobalPose) override;
 
-	UPROPERTY(EditAnywhere, Category="OpSettings")
-	TArray<FPinBoneData> BonesToPin;
+	// UObject interface
+	virtual void Serialize(FArchive& Ar) override;
+	// ~END UObject interface
 
-	// The bone, on the target skeleton to pin to.
+	// IBoneReferenceSkeletonProvider overrides.
+	virtual USkeleton* GetSkeleton(bool& bInvalidSkeletonIsError, const IPropertyHandle* PropertyHandle) override;
+	// ~END IBoneReferenceSkeletonProvider overrides.
+
+	// A list of bone-pairs to copy transforms between
 	UPROPERTY(EditAnywhere, Category=Settings)
+	TArray<FPinBoneData> BonesToPin;
+
+	// Which skeleton to use to copy animation from (usually the target).
+	UPROPERTY(EditAnywhere, Category=Settings, meta=(DisplayName="Skeleton To Copy From"))
 	ERetargetSourceOrTarget PinTo = ERetargetSourceOrTarget::Target;
 
-	// Apply this pin to the full transform, or just translation or rotation only.
+	// Copy the translation of the bone.
+	UPROPERTY(EditAnywhere, Category=Settings)
+	bool bCopyTranslation = true;
+
+	// Copy the translation of the bone.
+	UPROPERTY(EditAnywhere, Category=Settings)
+	bool bCopyRotation = true;
+	
+	// Copy the translation of the bone.
+	UPROPERTY(EditAnywhere, Category=Settings)
+	bool bCopyScale = true;
+
+	// The method used to calculate the translation of the bone.
 	UPROPERTY(EditAnywhere, Category=Settings)
-	EPinBoneType PinType = EPinBoneType::FullTransform;
+	EPinBoneTranslationMode TranslationMode;
 
-	// Maintain the original offset between the source and target bone
+	// The method used to calculate the rotation of the bone.
 	UPROPERTY(EditAnywhere, Category=Settings)
-	bool bMaintainOffset = true;
+	EPinBoneRotationMode RotationMode;
 
 	// A manual offset to apply in global space
 	UPROPERTY(EditAnywhere, Category=Settings)
@@ -90,7 +147,29 @@ public:
 	virtual FText GetNiceName() const override { return FText(LOCTEXT("OpName", "Pin Bone")); };
 	virtual FText WarningMessage() const override { return Message; };
 	FText Message;
+
+private:
+	USkeleton* SourceSkeletonAsset;
+	USkeleton* TargetSkeletonAsset;
 #endif
+
+private:
+
+	FTransform GetNewBoneTransform(
+		const FPinBoneData& PinData,
+		const UIKRetargetProcessor* Processor,
+		const TArray<FTransform>& InSourceGlobalPose,
+		TArray<FTransform>& OutTargetGlobalPose) const;
+
+	UE_DEPRECATED(5.6, "Uses separate translation/rotation offset modes")
+	UPROPERTY()
+	bool bMaintainOffset_DEPRECATED = true;
+
+	PRAGMA_DISABLE_DEPRECATION_WARNINGS
+	UE_DEPRECATED(5.6, "Uses separate translate/rotate/scale checkboxes")
+	UPROPERTY()
+	EPinBoneType PinType_DEPRECATED;
+	PRAGMA_ENABLE_DEPRECATION_WARNINGS
 };
 
 #undef LOCTEXT_NAMESPACE
diff --git a/Engine/Plugins/Animation/IKRig/Source/IKRigEditor/Private/RetargetEditor/SRetargetAnimAssetsWindow.cpp b/Engine/Plugins/Animation/IKRig/Source/IKRigEditor/Private/RetargetEditor/SRetargetAnimAssetsWindow.cpp
--- a/Engine/Plugins/Animation/IKRig/Source/IKRigEditor/Private/RetargetEditor/SRetargetAnimAssetsWindow.cpp
+++ b/Engine/Plugins/Animation/IKRig/Source/IKRigEditor/Private/RetargetEditor/SRetargetAnimAssetsWindow.cpp
@@ -723,7 +723,8 @@ void FProceduralRetargetAssets::AutoGenerateIKRetargetAsset()
 	// setup a "Pin Bone Op" to pin IK bones to the new target locations
 	const int32 PinBonesOpIndex = RetargetController->AddRetargetOp(UPinBoneOp::StaticClass());
 	UPinBoneOp* PinBoneOp = CastChecked<UPinBoneOp>(RetargetController->GetRetargetOpAtIndex(PinBonesOpIndex));
-	PinBoneOp->bMaintainOffset = false;
+	PinBoneOp->TranslationMode = EPinBoneTranslationMode::CopyGlobalPosition;
+	PinBoneOp->RotationMode = EPinBoneRotationMode::CopyGlobalRotation;
 	const TArray<FBoneToPin>& AllBonesToPin = TargetCharacterizationResults.AutoRetargetDefinition.BonesToPin.GetBonesToPin();
 	for (const FBoneToPin& BonesToPin : AllBonesToPin)
 	{
