Fixing build error on Mac in IK Retargeter

diff --git a/Engine/Plugins/Animation/IKRig/Source/IKRigEditor/Private/RetargetEditor/SRetargetAnimAssetsWindow.cpp b/Engine/Plugins/Animation/IKRig/Source/IKRigEditor/Private/RetargetEditor/SRetargetAnimAssetsWindow.cpp
--- a/Engine/Plugins/Animation/IKRig/Source/IKRigEditor/Private/RetargetEditor/SRetargetAnimAssetsWindow.cpp
+++ b/Engine/Plugins/Animation/IKRig/Source/IKRigEditor/Private/RetargetEditor/SRetargetAnimAssetsWindow.cpp
@@ -728,7 +728,10 @@ void FProceduralRetargetAssets::AutoGenerateIKRetargetAsset()
 	const TArray<FBoneToPin>& AllBonesToPin = TargetCharacterizationResults.AutoRetargetDefinition.BonesToPin.GetBonesToPin();
 	for (const FBoneToPin& BonesToPin : AllBonesToPin)
 	{
-		PinBoneOp->BonesToPin.Emplace(BonesToPin.BoneToPin, BonesToPin.BoneToPinTo);
+		FPinBoneData PinBoneData;
+		PinBoneData.BoneToCopyFrom = BonesToPin.BoneToPinTo;
+		PinBoneData.BoneToCopyTo = BonesToPin.BoneToPin;
+		PinBoneOp->BonesToPin.Add(PinBoneData);
 	}
 }
 
