Adding ability to scale source motion in IK Retargeter.

diff --git a/Engine/Plugins/Animation/IKRig/Source/IKRig/Private/Retargeter/IKRetargetProcessor.cpp b/Engine/Plugins/Animation/IKRig/Source/IKRig/Private/Retargeter/IKRetargetProcessor.cpp
--- a/Engine/Plugins/Animation/IKRig/Source/IKRig/Private/Retargeter/IKRetargetProcessor.cpp
+++ b/Engine/Plugins/Animation/IKRig/Source/IKRig/Private/Retargeter/IKRetargetProcessor.cpp
@@ -27,7 +27,8 @@ void FRetargetSkeleton::Initialize(
 	const TArray<FBoneChain>& BoneChains,
 	const FName InRetargetPoseName,
 	const FIKRetargetPose* RetargetPose,
-	const FName RetargetRootBone)
+	const FName RetargetRootBone,
+	const float ScaleFactor)
 {
 	// reset all skeleton data
 	Reset();
@@ -63,7 +64,7 @@ void FRetargetSkeleton::Initialize(
 
 	// update retarget pose to reflect custom offsets (applies stored offsets)
 	// NOTE: this must be done AFTER generating IsBoneInAnyTargetChain array above
-	GenerateRetargetPose(InRetargetPoseName, RetargetPose, RetargetRootBone);
+	GenerateRetargetPose(InRetargetPoseName, RetargetPose, RetargetRootBone, ScaleFactor);
 }
 
 void FRetargetSkeleton::Reset()
@@ -78,11 +79,13 @@ void FRetargetSkeleton::Reset()
 void FRetargetSkeleton::GenerateRetargetPose(
 	const FName InRetargetPoseName,
 	const FIKRetargetPose* InRetargetPose,
-	const FName RetargetRootBone)
+	const FName RetargetRootBone,
+	const float ScaleFactor)
 {
 	// record the name of the retarget pose (prevents re-initialization if profile swaps it)
 	RetargetPoseName = InRetargetPoseName;
 	RetargetPoseVersion = InRetargetPose->GetVersion();
+	RetargetPoseScale = ScaleFactor;
 	
 	// initialize retarget pose to the skeletal mesh reference pose
 	RetargetLocalPose = SkeletalMesh->GetRefSkeleton().GetRefBonePose();
@@ -129,7 +132,21 @@ void FRetargetSkeleton::GenerateRetargetPose(
 		RetargetLocalPose[BoneIndex].SetRotation(LocalBoneRotation);
 	}
 
+	// update global transforms based on local pose modified by the retarget pose offsets
 	UpdateGlobalTransformsBelowBone(-1, RetargetLocalPose, RetargetGlobalPose);
+
+	// scale the retarget pose
+	if (!FMath::IsNearlyEqual(ScaleFactor, 1.0f))
+	{
+		// scale translations of global transforms
+		for (FTransform& BoneTransform : RetargetGlobalPose)
+		{
+			BoneTransform.SetTranslation(BoneTransform.GetTranslation() * ScaleFactor);
+		}
+
+		// update the local transforms
+		UpdateLocalTransformsBelowBone(-1, RetargetLocalPose, RetargetGlobalPose);
+	}
 }
 
 int32 FRetargetSkeleton::FindBoneIndexByName(const FName InName) const
@@ -160,7 +177,7 @@ void FRetargetSkeleton::UpdateLocalTransformsBelowBone(
 	TArray<FTransform>& OutLocalPose,
 	const TArray<FTransform>& InGlobalPose) const
 {
-	check(BoneNames.IsValidIndex(StartBoneIndex));
+	check(BoneNames.IsValidIndex(StartBoneIndex+1));
 	check(BoneNames.Num() == OutLocalPose.Num());
 	check(BoneNames.Num() == InGlobalPose.Num());
 	
@@ -332,8 +349,9 @@ void FTargetSkeleton::Initialize(
 	const FName RetargetRootBone)
 {
 	Reset();
-	
-	FRetargetSkeleton::Initialize(InSkeletalMesh, BoneChains, InRetargetPoseName, RetargetPose, RetargetRootBone);
+
+	constexpr float TargetSkeletonScaleFactor = 1.0f;
+	FRetargetSkeleton::Initialize(InSkeletalMesh, BoneChains, InRetargetPoseName, RetargetPose, RetargetRootBone, TargetSkeletonScaleFactor);
 
 	// make storage for per-bone "Is Retargeted" flag (used for hierarchy updates)
 	// these are bones that are in a target chain that is mapped to a source chain (ie, will actually be retargeted)
@@ -1565,7 +1583,8 @@ void UIKRetargetProcessor::Initialize(
 		SourceIKRig->GetRetargetChains(),
 		RetargeterAsset->GetCurrentRetargetPoseName(ERetargetSourceOrTarget::Source),
 		RetargeterAsset->GetCurrentRetargetPose(ERetargetSourceOrTarget::Source),
-		SourceIKRig->GetRetargetRoot());
+		SourceIKRig->GetRetargetRoot(),
+		GlobalSettings.SourceScaleFactor);
 	TargetSkeleton.Initialize(
 		TargetSkeletalMesh,
 		TargetIKRig->GetRetargetChains(),
@@ -1923,7 +1942,7 @@ bool UIKRetargetProcessor::InitializeOpStack(const TArray<TObjectPtr<URetargetOp
 }
 
 TArray<FTransform>&  UIKRetargetProcessor::RunRetargeter(
-	const TArray<FTransform>& InSourceGlobalPose,
+	TArray<FTransform>& InSourceGlobalPose,
 	const TMap<FName, float>& SpeedValuesFromCurves,
 	const float DeltaTime,
 	const FRetargetProfile& Settings)
@@ -1950,6 +1969,15 @@ TArray<FTransform>&  UIKRetargetProcessor::RunRetargeter(
 		ensureMsgf(bHasNoScale, TEXT("Found scale values on incoming pose in retarget processor. Scale should be baked into translation and set to 1,1,1. "));
 	}
 #endif
+
+	// scale the input pose
+	if (!FMath::IsNearlyEqual(GlobalSettings.SourceScaleFactor, 1.0f))
+	{
+		for (FTransform& SourceBoneTransform : InSourceGlobalPose)
+		{
+			SourceBoneTransform.SetTranslation(SourceBoneTransform.GetTranslation() * GlobalSettings.SourceScaleFactor);
+		}
+	}
 	
 	// start from retarget pose
 	TargetSkeleton.OutputGlobalPose = TargetSkeleton.RetargetGlobalPose;
@@ -2552,6 +2580,10 @@ FName UIKRetargetProcessor::GetRetargetRoot(ERetargetSourceOrTarget SourceOrTarg
 
 void UIKRetargetProcessor::ApplySettingsFromProfile(const FRetargetProfile& Profile)
 {
+	// assign global settings
+	// NOTE: must be done before updating retarget pose because it uses the scale factors in the global settings
+	GlobalSettings = Profile.GlobalSettings;
+	
 	// assign retarget poses specified in the profile
 	UpdateRetargetPoseAtRuntime(Profile.SourceRetargetPoseName, ERetargetSourceOrTarget::Source);
 	UpdateRetargetPoseAtRuntime(Profile.TargetRetargetPoseName, ERetargetSourceOrTarget::Target);
@@ -2564,9 +2596,6 @@ void UIKRetargetProcessor::ApplySettingsFromProfile(const FRetargetProfile& Prof
 
 	// assign root settings
 	RootRetargeter.Settings = Profile.RootSettings;
-	
-	// assign global settings
-	GlobalSettings = Profile.GlobalSettings;
 }
 
 void UIKRetargetProcessor::UpdateRetargetPoseAtRuntime(
@@ -2583,7 +2612,8 @@ void UIKRetargetProcessor::UpdateRetargetPoseAtRuntime(
 	FRetargetSkeleton& RetargetSkeleton = bIsSource ? SourceSkeleton : TargetSkeleton;
 	const bool bSameRetargetPose = RetargetSkeleton.RetargetPoseName == NewRetargetPoseName;
 	const bool bSameVersion = RetargetSkeleton.RetargetPoseVersion == NewRetargetPose->GetVersion();
-	if (bSameRetargetPose && bSameVersion)
+	const bool bSameGlobalScale = bIsSource ? FMath::IsNearlyEqual(RetargetSkeleton.RetargetPoseScale, GlobalSettings.SourceScaleFactor) : true;
+	if (bSameRetargetPose && bSameVersion && bSameGlobalScale)
 	{
 		return; // retarget pose has not changed since it was initialized
 	}
@@ -2595,7 +2625,8 @@ void UIKRetargetProcessor::UpdateRetargetPoseAtRuntime(
 	}
 
 	// re-generate the retarget pose
-	RetargetSkeleton.GenerateRetargetPose(NewRetargetPoseName, NewRetargetPose, RootBoneName);
+	const float ScaleFactor = SourceOrTarget == ERetargetSourceOrTarget::Source ? GlobalSettings.SourceScaleFactor : 1.0f;
+	RetargetSkeleton.GenerateRetargetPose(NewRetargetPoseName, NewRetargetPose, RootBoneName, ScaleFactor);
 
 	// re-initialize the bone chains using the newly generated retarget pose
 	const UIKRigDefinition* SourceIKRig = RetargeterAsset->GetIKRig(ERetargetSourceOrTarget::Source);
diff --git a/Engine/Plugins/Animation/IKRig/Source/IKRig/Public/Retargeter/IKRetargetProcessor.h b/Engine/Plugins/Animation/IKRig/Source/IKRig/Public/Retargeter/IKRetargetProcessor.h
--- a/Engine/Plugins/Animation/IKRig/Source/IKRig/Public/Retargeter/IKRetargetProcessor.h
+++ b/Engine/Plugins/Animation/IKRig/Source/IKRig/Public/Retargeter/IKRetargetProcessor.h
@@ -30,8 +30,11 @@ struct IKRIG_API FRetargetSkeleton
 	TArray<int32> ParentIndices;			// per-bone indices of parent bones (the hierarchy)
 	TArray<FTransform> RetargetLocalPose;	// local space retarget pose
 	TArray<FTransform> RetargetGlobalPose;	// global space retarget pose
+	
 	FName RetargetPoseName;					// the name of the retarget pose this was initialized with
 	int32 RetargetPoseVersion;				// the version of the retarget pose this was initialized with (transient)
+	float RetargetPoseScale;				// the scale of the retarget pose this was initialized with (transient)
+	
 	USkeletalMesh* SkeletalMesh;			// the skeletal mesh this was initialized with
 	TArray<FName> ChainThatContainsBone;	// record which chain is actually controlling each bone
 
@@ -40,14 +43,16 @@ struct IKRIG_API FRetargetSkeleton
 		const TArray<FBoneChain>& BoneChains,
 		const FName InRetargetPoseName,
 		const FIKRetargetPose* RetargetPose,
-		const FName RetargetRootBone);
+		const FName RetargetRootBone,
+		const float ScaleFactor);
 
 	void Reset();
 
 	void GenerateRetargetPose(
 		const FName InRetargetPoseName,
 		const FIKRetargetPose* InRetargetPose,
-		const FName RetargetRootBone);
+		const FName RetargetRootBone,
+		const float ScaleFactor);
 
 	int32 FindBoneIndexByName(const FName InName) const;
 
@@ -496,7 +501,7 @@ public:
 	* @return The retargeted Component/Global space pose for the target skeleton
 	*/
 	TArray<FTransform>& RunRetargeter(
-		const TArray<FTransform>& InSourceGlobalPose,
+		TArray<FTransform>& InSourceGlobalPose,
 		const TMap<FName, float>& SpeedValuesFromCurves,
 		const float DeltaTime,
 		const FRetargetProfile& Settings);
diff --git a/Engine/Plugins/Animation/IKRig/Source/IKRig/Public/Retargeter/IKRetargetSettings.h b/Engine/Plugins/Animation/IKRig/Source/IKRig/Public/Retargeter/IKRetargetSettings.h
--- a/Engine/Plugins/Animation/IKRig/Source/IKRig/Public/Retargeter/IKRetargetSettings.h
+++ b/Engine/Plugins/Animation/IKRig/Source/IKRig/Public/Retargeter/IKRetargetSettings.h
@@ -310,6 +310,10 @@ struct IKRIG_API FRetargetGlobalSettings
 	UPROPERTY(EditAnywhere, BlueprintReadWrite, Category = RetargetPhases)
 	bool bEnablePost = true;
 
+	/** Range 0.01 to +inf. Default 1. Scales the incoming source pose. Affects entire skeleton and all IK goals.*/
+	UPROPERTY(EditAnywhere, BlueprintReadWrite, Category = Scale, meta = (ClampMin = "0.01", UIMin = "0.01", UIMax = "10.0"))
+	float SourceScaleFactor = 1.0f;
+
 	/** Enable IK Warping.
 	 * These options allow for global modifications to all IK Goals that have "Affected by IK Warping" turned on (the default).
 	 * "Affected by IK Warping" can be found in the IK settings of a retarget chain. */
diff --git a/Engine/Plugins/Animation/IKRig/Source/IKRigEditor/Private/RetargetEditor/IKRetargetDetails.cpp b/Engine/Plugins/Animation/IKRig/Source/IKRigEditor/Private/RetargetEditor/IKRetargetDetails.cpp
--- a/Engine/Plugins/Animation/IKRig/Source/IKRigEditor/Private/RetargetEditor/IKRetargetDetails.cpp
+++ b/Engine/Plugins/Animation/IKRig/Source/IKRigEditor/Private/RetargetEditor/IKRetargetDetails.cpp
@@ -1391,6 +1391,10 @@ void FRetargetGlobalSettingsCustomization::CustomizeDetails(IDetailLayoutBuilder
 	// replace the "Settings" struct property
 	IDetailCategoryBuilder& SettingsCategory = DetailBuilder.EditCategory(FName("Global Settings"));
 
+	// global scale group
+	IDetailGroup& SourceScaleGroup = SettingsCategory.AddGroup("GlobalScaleFactor", LOCTEXT("GlobalScale_Label", "Global Scale Factor"), false, true);
+	SourceScaleGroup.AddPropertyRow( GetPropertyHandle(GET_MEMBER_NAME_STRING_CHECKED(FRetargetGlobalSettings, SourceScaleFactor)));
+	
 	// phases group
 	IDetailGroup& AlphaGroup = SettingsCategory.AddGroup("Phases", LOCTEXT("Phases_Label", "Global Retarget Phases"), false, true);
 	AlphaGroup.AddPropertyRow(GetPropertyHandle(GET_MEMBER_NAME_STRING_CHECKED(FRetargetGlobalSettings, bEnableRoot)));
diff --git a/Engine/Plugins/Animation/IKRig/Source/IKRigEditor/Private/RetargetEditor/IKRetargetEditor.cpp b/Engine/Plugins/Animation/IKRig/Source/IKRigEditor/Private/RetargetEditor/IKRetargetEditor.cpp
--- a/Engine/Plugins/Animation/IKRig/Source/IKRigEditor/Private/RetargetEditor/IKRetargetEditor.cpp
+++ b/Engine/Plugins/Animation/IKRig/Source/IKRigEditor/Private/RetargetEditor/IKRetargetEditor.cpp
@@ -505,9 +505,8 @@ FString FIKRetargetEditor::GetWorldCentricTabPrefix() const
 
 void FIKRetargetEditor::Tick(float DeltaTime)
 {
-	// update with latest offsets
-	EditorController->UpdateMeshOffset(ERetargetSourceOrTarget::Source);
-	EditorController->UpdateMeshOffset(ERetargetSourceOrTarget::Target);
+	// update with latest offsets and scale values
+	EditorController->UpdateSkeletalMeshComponents();
 
 	// retargeter IK planting must be reset when time is reversed or playback jumps ahead 
 	const float CurrentTime = EditorController->SourceAnimInstance->GetCurrentTime();
diff --git a/Engine/Plugins/Animation/IKRig/Source/IKRigEditor/Private/RetargetEditor/IKRetargetEditorController.cpp b/Engine/Plugins/Animation/IKRig/Source/IKRigEditor/Private/RetargetEditor/IKRetargetEditorController.cpp
--- a/Engine/Plugins/Animation/IKRig/Source/IKRigEditor/Private/RetargetEditor/IKRetargetEditorController.cpp
+++ b/Engine/Plugins/Animation/IKRig/Source/IKRigEditor/Private/RetargetEditor/IKRetargetEditorController.cpp
@@ -316,28 +316,34 @@ UIKRetargetAnimInstance* FIKRetargetEditorController::GetAnimInstance(
 	return SourceOrTarget == ERetargetSourceOrTarget::Source ? SourceAnimInstance.Get() : TargetAnimInstance.Get();
 }
 
-void FIKRetargetEditorController::UpdateMeshOffset(ERetargetSourceOrTarget SourceOrTarget) const
+void FIKRetargetEditorController::UpdateSkeletalMeshComponents() const
 {
-	const UIKRetargeter* Asset = AssetController->GetAsset();
-	if (!Asset)
+	auto UpdateMeshComponent = [this](ERetargetSourceOrTarget SourceOrTarget)
 	{
-		return;
-	}
+		const UIKRetargeter* Asset = AssetController->GetAsset();
+		if (!Asset)
+		{
+			return;
+		}
+		
+		USceneComponent* Component = GetSkeletalMeshComponent(SourceOrTarget);
+		if (!Component)
+		{
+			return;
+		}
 
-	const bool bIsSource  = SourceOrTarget == ERetargetSourceOrTarget::Source;
-	USceneComponent* Component = bIsSource ? SourceRootComponent : GetSkeletalMeshComponent(ERetargetSourceOrTarget::Target);
-    if (!Component)
-    {
-    	return;
-    }
-	
-	const FVector Position = bIsSource ? Asset->SourceMeshOffset : Asset->TargetMeshOffset;
-	const float Scale = bIsSource ? 1.0f : Asset->TargetMeshScale;
-	constexpr bool bSweep = false;
-	constexpr FHitResult* OutSweepHitResult = nullptr;
-	constexpr ETeleportType Teleport = ETeleportType::ResetPhysics;
-	Component->SetWorldLocation(Position, bSweep, OutSweepHitResult, Teleport);
-	Component->SetWorldScale3D(FVector(Scale,Scale,Scale));
+		const bool bIsSource  = SourceOrTarget == ERetargetSourceOrTarget::Source;
+		const FVector Position = bIsSource ? Asset->SourceMeshOffset : Asset->TargetMeshOffset;
+		const float Scale = bIsSource ? Asset->GetGlobalSettings().SourceScaleFactor : Asset->TargetMeshScale;
+		constexpr bool bSweep = false;
+		constexpr FHitResult* OutSweepHitResult = nullptr;
+		constexpr ETeleportType Teleport = ETeleportType::TeleportPhysics;
+		Component->SetWorldLocation(Position, bSweep, OutSweepHitResult, Teleport);
+		Component->SetWorldScale3D(FVector(Scale,Scale,Scale));
+	};
+
+	UpdateMeshComponent(ERetargetSourceOrTarget::Source);
+	UpdateMeshComponent(ERetargetSourceOrTarget::Target);
 }
 
 bool FIKRetargetEditorController::GetCameraTargetForSelection(FSphere& OutTarget) const
@@ -1566,13 +1572,15 @@ void FIKRetargetEditorController::FixZeroHeightRetargetRoot(ERetargetSourceOrTar
 		return;
 	}
 
+	const float ScaleFactor = SourceOrTarget == ERetargetSourceOrTarget::Source ? GetGlobalSettings().SourceScaleFactor : 1.0f;
 	FRetargetSkeleton DummySkeleton;
 	DummySkeleton.Initialize(
 		SkeletalMesh,
 		TArray<FBoneChain>(),
 		CurrentRetargetPoseName,
 		&CurrentRetargetPose,
-		RetargetRootBoneName);
+		RetargetRootBoneName,
+		ScaleFactor);
 
 	const int32 RootBoneIndex = DummySkeleton.FindBoneIndexByName(RetargetRootBoneName);
 	if (RootBoneIndex == INDEX_NONE)
diff --git a/Engine/Plugins/Animation/IKRig/Source/IKRigEditor/Public/RetargetEditor/IKRetargetEditorController.h b/Engine/Plugins/Animation/IKRig/Source/IKRigEditor/Public/RetargetEditor/IKRetargetEditorController.h
--- a/Engine/Plugins/Animation/IKRig/Source/IKRigEditor/Public/RetargetEditor/IKRetargetEditorController.h
+++ b/Engine/Plugins/Animation/IKRig/Source/IKRigEditor/Public/RetargetEditor/IKRetargetEditorController.h
@@ -318,7 +318,7 @@ public:
 	// render the skeleton's in the viewport (either source or target)
 	void RenderSkeleton(FPrimitiveDrawInterface* PDI, ERetargetSourceOrTarget SourceOrTarget) const;
 
-	void UpdateMeshOffset(ERetargetSourceOrTarget SourceOrTarget) const;
+	void UpdateSkeletalMeshComponents() const;
 
 private:
 
