IK Retargeter pin bone un-deprecated enum that is needed for backwards compatibility of python scripts

diff --git a/Engine/Plugins/Animation/IKRig/Source/IKRig/Public/Retargeter/RetargetOps/PinBoneOp.h b/Engine/Plugins/Animation/IKRig/Source/IKRig/Public/Retargeter/RetargetOps/PinBoneOp.h
--- a/Engine/Plugins/Animation/IKRig/Source/IKRig/Public/Retargeter/RetargetOps/PinBoneOp.h
+++ b/Engine/Plugins/Animation/IKRig/Source/IKRig/Public/Retargeter/RetargetOps/PinBoneOp.h
@@ -10,7 +10,7 @@
 #define LOCTEXT_NAMESPACE "RetargetOps"
 
 UENUM()
-enum class UE_DEPRECATED(5.6, "Uses bools to include translation/rotation/scale (or combinations therof) separately.") EPinBoneType : uint8
+enum class EPinBoneType : uint8
 {
 	FullTransform,
 	TranslateOnly,
@@ -119,11 +119,11 @@ public:
 	UPROPERTY(EditAnywhere, Category=Settings)
 	bool bCopyTranslation = true;
 
-	// Copy the translation of the bone.
+	// Copy the rotation of the bone.
 	UPROPERTY(EditAnywhere, Category=Settings)
 	bool bCopyRotation = true;
 	
-	// Copy the translation of the bone.
+	// Copy the scale of the bone.
 	UPROPERTY(EditAnywhere, Category=Settings)
 	bool bCopyScale = true;
 
