IK Retargeter hierarchy view displaying no chain data after pressing Undo

diff --git a/Engine/Plugins/Animation/IKRig/Source/IKRig/Private/AnimNodes/AnimNode_RetargetPoseFromMesh.cpp b/Engine/Plugins/Animation/IKRig/Source/IKRig/Private/AnimNodes/AnimNode_RetargetPoseFromMesh.cpp
--- a/Engine/Plugins/Animation/IKRig/Source/IKRig/Private/AnimNodes/AnimNode_RetargetPoseFromMesh.cpp
+++ b/Engine/Plugins/Animation/IKRig/Source/IKRig/Private/AnimNodes/AnimNode_RetargetPoseFromMesh.cpp
@@ -170,7 +170,7 @@ void FAnimNode_RetargetPoseFromMesh::PreUpdate(const UAnimInstance* InAnimInstan
 	}
 
 	const TObjectPtr<USkeletalMeshComponent> TargetMeshComponent = InAnimInstance->GetSkelMeshComponent();
-	if (EnsureProcessorIsInitialized(TargetMeshComponent))
+	if (InitializeRetargetProcessor(TargetMeshComponent))
 	{
 		CopyBoneTransformsFromSource(TargetMeshComponent);
 
@@ -278,7 +278,7 @@ UIKRetargetProcessor* FAnimNode_RetargetPoseFromMesh::GetRetargetProcessor() con
 	return Processor;
 }
 
-bool FAnimNode_RetargetPoseFromMesh::EnsureProcessorIsInitialized(const TObjectPtr<USkeletalMeshComponent> TargetMeshComponent)
+bool FAnimNode_RetargetPoseFromMesh::InitializeRetargetProcessor(const TObjectPtr<USkeletalMeshComponent> TargetMeshComponent)
 {
 	// has user supplied a retargeter asset?
 	if (!IKRetargeterAsset)
diff --git a/Engine/Plugins/Animation/IKRig/Source/IKRig/Public/AnimNodes/AnimNode_RetargetPoseFromMesh.h b/Engine/Plugins/Animation/IKRig/Source/IKRig/Public/AnimNodes/AnimNode_RetargetPoseFromMesh.h
--- a/Engine/Plugins/Animation/IKRig/Source/IKRig/Public/AnimNodes/AnimNode_RetargetPoseFromMesh.h
+++ b/Engine/Plugins/Animation/IKRig/Source/IKRig/Public/AnimNodes/AnimNode_RetargetPoseFromMesh.h
@@ -73,10 +73,11 @@ struct IKRIG_API FAnimNode_RetargetPoseFromMesh : public FAnimNode_Base
 	void CreateRetargetProcessorIfNeeded(UObject* Outer);
 	UIKRetargetProcessor* GetRetargetProcessor() const;
 
-private:
-
 	/** returns true if processor is setup and ready to go, false otherwise */
-	bool EnsureProcessorIsInitialized(const TObjectPtr<USkeletalMeshComponent> TargetMeshComponent);
+	bool InitializeRetargetProcessor(const TObjectPtr<USkeletalMeshComponent> TargetMeshComponent);
+
+private:
+	
 	/** copies the source mesh pose (on main thread) */
 	void CopyBoneTransformsFromSource(USkeletalMeshComponent* TargetMeshComponent);
 	/** indirection to account for Leader Pose Component setup */
diff --git a/Engine/Plugins/Animation/IKRig/Source/IKRigEditor/Private/RetargetEditor/IKRetargetAnimInstance.cpp b/Engine/Plugins/Animation/IKRig/Source/IKRigEditor/Private/RetargetEditor/IKRetargetAnimInstance.cpp
--- a/Engine/Plugins/Animation/IKRig/Source/IKRigEditor/Private/RetargetEditor/IKRetargetAnimInstance.cpp
+++ b/Engine/Plugins/Animation/IKRig/Source/IKRigEditor/Private/RetargetEditor/IKRetargetAnimInstance.cpp
@@ -159,6 +159,12 @@ UIKRetargetProcessor* UIKRetargetAnimInstance::GetRetargetProcessor()
 	return RetargetNode.GetRetargetProcessor();
 }
 
+void UIKRetargetAnimInstance::ForceInitializeProcessor(const TObjectPtr<USkeletalMeshComponent> TargetComponent)
+{
+	RetargetNode.IKRetargeterAsset->IncrementVersion();
+	RetargetNode.InitializeRetargetProcessor(TargetComponent);
+}
+
 FAnimInstanceProxy* UIKRetargetAnimInstance::CreateAnimInstanceProxy()
 {
 	LLM_SCOPE_BYNAME(TEXT("Animation/IKRig"));
diff --git a/Engine/Plugins/Animation/IKRig/Source/IKRigEditor/Private/RetargetEditor/IKRetargetEditorController.cpp b/Engine/Plugins/Animation/IKRig/Source/IKRigEditor/Private/RetargetEditor/IKRetargetEditorController.cpp
--- a/Engine/Plugins/Animation/IKRig/Source/IKRigEditor/Private/RetargetEditor/IKRetargetEditorController.cpp
+++ b/Engine/Plugins/Animation/IKRig/Source/IKRigEditor/Private/RetargetEditor/IKRetargetEditorController.cpp
@@ -260,8 +260,10 @@ void FIKRetargetEditorController::ReinitializeRetargeterNoUIRefresh() const
 	// clear the output log
 	ClearOutputLog();
 
-	// force running instances to reinitialize on next tick
-	AssetController->GetAsset()->IncrementVersion();
+	if(UIKRetargetAnimInstance* AnimInstance = TargetAnimInstance.Get())
+	{
+		return AnimInstance->ForceInitializeProcessor(GetSkeletalMeshComponent(ERetargetSourceOrTarget::Target));
+	}
 }
 
 void FIKRetargetEditorController::HandleIKRigReplaced(ERetargetSourceOrTarget SourceOrTarget)
diff --git a/Engine/Plugins/Animation/IKRig/Source/IKRigEditor/Public/RetargetEditor/IKRetargetAnimInstance.h b/Engine/Plugins/Animation/IKRig/Source/IKRigEditor/Public/RetargetEditor/IKRetargetAnimInstance.h
--- a/Engine/Plugins/Animation/IKRig/Source/IKRigEditor/Public/RetargetEditor/IKRetargetAnimInstance.h
+++ b/Engine/Plugins/Animation/IKRig/Source/IKRigEditor/Public/RetargetEditor/IKRetargetAnimInstance.h
@@ -62,6 +62,8 @@ public:
 
 	UIKRetargetProcessor* GetRetargetProcessor();
 
+	void ForceInitializeProcessor(const TObjectPtr<USkeletalMeshComponent> TargetComponent);
+
 	const TArray<FTransform>& GetLocalRetargetPose() const { return PreviewPoseNode.GetLocalRetargetPose(); };
 	const TArray<FTransform>& GetGlobalRetargetPose() const { return PreviewPoseNode.GetGlobalRetargetPose(); };
 
diff --git a/Engine/Plugins/Animation/IKRig/Source/IKRigEditor/Public/RetargetEditor/IKRetargeterController.h b/Engine/Plugins/Animation/IKRig/Source/IKRigEditor/Public/RetargetEditor/IKRetargeterController.h
--- a/Engine/Plugins/Animation/IKRig/Source/IKRigEditor/Public/RetargetEditor/IKRetargeterController.h
+++ b/Engine/Plugins/Animation/IKRig/Source/IKRigEditor/Public/RetargetEditor/IKRetargeterController.h
@@ -378,7 +378,6 @@ struct FScopedReinitializeIKRetargeter
 	{
 		if (--Controller->ReinitializeScopeCounter == 0)
 		{
-			Controller->GetAsset()->IncrementVersion();
 			Controller->BroadcastNeedsReinitialized();
 		}
 	};
