IK Rig: Fix truncation warnings

#rnx

diff --git a/Engine/Plugins/Animation/IKRig/Source/IKRig/Private/AnimNodes/AnimNode_IKRig.cpp b/Engine/Plugins/Animation/IKRig/Source/IKRig/Private/AnimNodes/AnimNode_IKRig.cpp
--- a/Engine/Plugins/Animation/IKRig/Source/IKRig/Private/AnimNodes/AnimNode_IKRig.cpp
+++ b/Engine/Plugins/Animation/IKRig/Source/IKRig/Private/AnimNodes/AnimNode_IKRig.cpp
@@ -234,7 +234,7 @@ void FAnimNode_IKRig::Update_AnyThread(const FAnimationUpdateContext& Context)
 			FTransform GoalTransform(Goal.FinalBlendedRotation, Goal.FinalBlendedPosition, FVector(DebugScale, DebugScale, DebugScale));
 			GoalTransform = GoalTransform * Context.AnimInstanceProxy->GetComponentTransform();
 			UE_VLOG_OBOX(Context.AnimInstanceProxy->GetAnimInstanceObject(), "IKRig", Display, UnitBox, GoalTransform.ToMatrixWithScale(), FColor::Yellow, TEXT(""));
-			UE_VLOG_LOCATION(Context.AnimInstanceProxy->GetAnimInstanceObject(), "IKRig", Verbose, GoalTransform.GetTranslation(), 0.f, FColor::White, TEXT("%ls"), ToCStr(Goal.Name.ToString()));
+			UE_VLOG_LOCATION(Context.AnimInstanceProxy->GetAnimInstanceObject(), "IKRig", Verbose, GoalTransform.GetTranslation(), 0, FColor::White, TEXT("%ls"), ToCStr(Goal.Name.ToString()));
 		}
 	}
 #endif
@@ -408,15 +408,17 @@ void FAnimNode_IKRig::InitializeProperties(const UObject* InSourceInstance, UCla
 		{
 			UpdateFunctions.Add([&Goal, SourceProperty](const UObject* InSourceInstance)
 			{
-				const double* AlphaValue = SourceProperty->ContainerPtrToValuePtr<double>(InSourceInstance);
-				Goal.PositionAlpha = FMath::Clamp<float>(*AlphaValue, 0.f, 1.f);
+				static_assert(std::is_same_v<decltype(FIKRigGoal::PositionAlpha), float>);
+				const float* AlphaValue = SourceProperty->ContainerPtrToValuePtr<float>(InSourceInstance);
+				Goal.PositionAlpha = FMath::Clamp<float>(static_cast<float>(*AlphaValue), 0.f, 1.f);
 			});
 		}
 		else if (GoalPropStr.StartsWith(AlphaRotPropStr))
 		{
 			UpdateFunctions.Add([&Goal, SourceProperty](const UObject* InSourceInstance)
 			{
-				const double* AlphaValue = SourceProperty->ContainerPtrToValuePtr<double>(InSourceInstance);
+				static_assert(std::is_same_v<decltype(FIKRigGoal::RotationAlpha), float>);
+				const float* AlphaValue = SourceProperty->ContainerPtrToValuePtr<float>(InSourceInstance);
 				Goal.RotationAlpha = FMath::Clamp<float>(*AlphaValue, 0.f, 1.f);
 			});
 		}
@@ -424,6 +426,7 @@ void FAnimNode_IKRig::InitializeProperties(const UObject* InSourceInstance, UCla
 		{
 			UpdateFunctions.Add([&Goal, SourceProperty](const UObject* InSourceInstance)
 			{
+				static_assert(std::is_same_v<decltype(FIKRigGoal::Position), FVector>);
 				Goal.Position = *SourceProperty->ContainerPtrToValuePtr<FVector>(InSourceInstance);
 			});
 		}
@@ -431,6 +434,7 @@ void FAnimNode_IKRig::InitializeProperties(const UObject* InSourceInstance, UCla
 		{
 			UpdateFunctions.Add([&Goal, SourceProperty](const UObject* InSourceInstance)
 			{
+				static_assert(std::is_same_v<decltype(FIKRigGoal::Rotation), FRotator>);
 				Goal.Rotation = *SourceProperty->ContainerPtrToValuePtr<FRotator>(InSourceInstance);
 			});
 		}
diff --git a/Engine/Plugins/Animation/IKRig/Source/IKRig/Private/Retargeter/IKRetargetProcessor.cpp b/Engine/Plugins/Animation/IKRig/Source/IKRig/Private/Retargeter/IKRetargetProcessor.cpp
--- a/Engine/Plugins/Animation/IKRig/Source/IKRig/Private/Retargeter/IKRetargetProcessor.cpp
+++ b/Engine/Plugins/Animation/IKRig/Source/IKRig/Private/Retargeter/IKRetargetProcessor.cpp
@@ -487,12 +487,12 @@ FTransform FChainFK::GetTransformAtParam(const TArray<FTransform>& Transforms, c
 		return Transforms[0];
 	}
 	
-	if (Param < KINDA_SMALL_NUMBER)
+	if (Param < UE_KINDA_SMALL_NUMBER)
 	{
 		return Transforms[0];
 	}
 
-	if (Param > 1.0f - KINDA_SMALL_NUMBER)
+	if (Param > 1.0f - UE_KINDA_SMALL_NUMBER)
 	{
 		return Transforms.Last();
 	}
@@ -537,12 +537,12 @@ bool FChainFK::CalculateBoneParameters(FIKRigLogger& Log)
 	BoneDistances.Add(0.0f);
 	for (int32 i=1; i<InitialGlobalTransforms.Num(); ++i)
 	{
-		TotalChainLength += (InitialGlobalTransforms[i].GetTranslation() - InitialGlobalTransforms[i-1].GetTranslation()).Size();
+		TotalChainLength += static_cast<float>((InitialGlobalTransforms[i].GetTranslation() - InitialGlobalTransforms[i-1].GetTranslation()).Size());
 		BoneDistances.Add(TotalChainLength);
 	}
 
 	// cannot retarget chain if all the bones are sitting directly on each other
-	if (TotalChainLength <= KINDA_SMALL_NUMBER)
+	if (TotalChainLength <= UE_KINDA_SMALL_NUMBER)
 	{
 		Log.LogWarning(LOCTEXT("TinyBoneChain", "IK Retargeter bone chain length is too small to reliably retarget."));
 		return false;
@@ -885,9 +885,9 @@ bool FChainRetargeterIK::InitializeSource(
 	Source.InitialEndRotation = End.GetRotation();
 
 	const FTransform& Start = SourceInitialGlobalPose[Source.StartBoneIndex];
-	const float Length = (Start.GetTranslation() - Source.InitialEndPosition).Size();
+	const float Length = static_cast<float>((Start.GetTranslation() - Source.InitialEndPosition).Size());
 
-	if (Length <= KINDA_SMALL_NUMBER)
+	if (Length <= UE_KINDA_SMALL_NUMBER)
 	{
 		Log.LogWarning(LOCTEXT("SourceZeroLengthIK", "IK Retargeter trying to retarget source bone chain with IK, but it is zero length!"));
     	return false;
@@ -914,7 +914,7 @@ void FChainRetargeterIK::EncodePose(const TArray<FTransform>& InSourceGlobalPose
 	Source.CurrentEndPosition = End;
 	Source.CurrentEndDirectionNormalized = ChainDirection * NormalizedLimbLength;
 	Source.CurrentEndRotation = InSourceGlobalPose[Source.EndBoneIndex].GetRotation();
-	Source.CurrentHeightFromGroundNormalized = (End.Z - Source.InitialEndPosition.Z)  * Source.InvInitialLength;
+	Source.CurrentHeightFromGroundNormalized = static_cast<float>(End.Z - Source.InitialEndPosition.Z)  * Source.InvInitialLength;
 }
 
 bool FChainRetargeterIK::InitializeTarget(
@@ -934,10 +934,10 @@ bool FChainRetargeterIK::InitializeTarget(
 	Target.PrevEndPosition = Last.GetLocation();
 	Target.InitialEndPosition = Last.GetTranslation();
 	Target.InitialEndRotation = Last.GetRotation();
-	Target.InitialLength = (TargetInitialGlobalPose[Target.BoneIndexA].GetTranslation() - Last.GetTranslation()).Size();
+	Target.InitialLength = static_cast<float>((TargetInitialGlobalPose[Target.BoneIndexA].GetTranslation() - Last.GetTranslation()).Size());
 	ResetThisTick = true;
 
-	if (Target.InitialLength <= KINDA_SMALL_NUMBER)
+	if (Target.InitialLength <= UE_KINDA_SMALL_NUMBER)
 	{
 		Log.LogWarning(LOCTEXT("TargetZeroLengthIK", "IK Retargeter trying to retarget target bone chain with IK, but it is zero length!"));
 		return false;
@@ -967,7 +967,7 @@ void FChainRetargeterIK::DecodePose(
 
 	// blend to source rotation
 	const float BlendToSourceRotation = Settings.BlendToSource * Settings.BlendToSourceRotation;
-	if (BlendToSourceRotation > KINDA_SMALL_NUMBER)
+	if (BlendToSourceRotation > UE_KINDA_SMALL_NUMBER)
 	{
 		GoalRotation = FQuat::FastLerp(GoalRotation, Source.CurrentEndRotation, BlendToSourceRotation);
 		GoalRotation.Normalize();
@@ -985,7 +985,7 @@ void FChainRetargeterIK::DecodePose(
 
 	// blend to source location
 	const float BlendToSourceTranslation = Settings.BlendToSource * Settings.BlendToSourceTranslation;
-	if (BlendToSourceTranslation > KINDA_SMALL_NUMBER)
+	if (BlendToSourceTranslation > UE_KINDA_SMALL_NUMBER)
 	{
 		const FVector RootModification = RootRetargeter.Target.RootTranslationDelta * RootRetargeter.Settings.GetAffectIKWeightVector();
 		const FVector Weight = BlendToSourceTranslation * Settings.BlendToSourceWeights;
@@ -1237,11 +1237,11 @@ bool FRootRetargeter::InitializeSource(
 	
 	// record initial root data
 	const FTransform InitialTransform = SourceSkeleton.RetargetGlobalPose[Source.BoneIndex]; 
-	float InitialHeight = InitialTransform.GetTranslation().Z;
+	float InitialHeight = static_cast<float>(InitialTransform.GetTranslation().Z);
 	Source.InitialRotation = InitialTransform.GetRotation();
 
 	// ensure root height is not at origin, this happens if user sets root to ACTUAL skeleton root and not pelvis
-	if (InitialHeight < KINDA_SMALL_NUMBER)
+	if (InitialHeight < UE_KINDA_SMALL_NUMBER)
 	{
 		// warn user and push it up slightly to avoid divide by zero
 		Log.LogError(LOCTEXT("BadRootHeight", "The source retarget root bone is very near the ground plane. This will cause the target to be moved very far. To resolve this, please create a retarget pose with the retarget root at the correct height off the ground."));
@@ -1271,7 +1271,7 @@ bool FRootRetargeter::InitializeTarget(
 	}
 
 	const FTransform TargetInitialTransform = TargetSkeleton.RetargetGlobalPose[Target.BoneIndex];
-	Target.InitialHeight = TargetInitialTransform.GetTranslation().Z;
+	Target.InitialHeight = static_cast<float>(TargetInitialTransform.GetTranslation().Z);
 	Target.InitialRotation = TargetInitialTransform.GetRotation();
 	Target.InitialPosition = TargetInitialTransform.GetTranslation();
 
@@ -1361,7 +1361,7 @@ bool FPoleVectorMatcher::Initialize(
 	const FVector SourcePoleVector = CalculatePoleVector(SourcePoleAxis, SourceIndices, SourceGlobalPose);
 	const FVector TargetPoleVector = CalculatePoleVector(TargetPoleAxis, TargetIndices, TargetGlobalPose);
 	
-	TargetToSourceAngularOffsetAtRefPose = FMath::Acos(FVector::DotProduct(TargetPoleVector, SourcePoleVector));
+	TargetToSourceAngularOffsetAtRefPose = static_cast<float>(FMath::Acos(FVector::DotProduct(TargetPoleVector, SourcePoleVector)));
 
 	// cache indices of bones in chain (and their children recursively) excluding children of the end bone
 	TArray<int32> AllChildrenIndices;
@@ -1389,7 +1389,7 @@ void FPoleVectorMatcher::MatchPoleVector(
 	TArray<FTransform>& OutTargetGlobalPose,
 	FRetargetSkeleton& TargetSkeleton)
 {
-	const bool bIsMatchingPoleVector = Settings.PoleVectorMatching > KINDA_SMALL_NUMBER;
+	const bool bIsMatchingPoleVector = Settings.PoleVectorMatching > UE_KINDA_SMALL_NUMBER;
 	const bool bIsOffsetingPoleVector = !FMath::IsNearlyZero(Settings.PoleVectorOffset);
 	if (!(bIsMatchingPoleVector || bIsOffsetingPoleVector))
 	{
@@ -1421,7 +1421,7 @@ void FPoleVectorMatcher::MatchPoleVector(
 		const FVector SourcePoleVector = CalculatePoleVector(SourcePoleAxis, SourceIndices, SourceGlobalPose);
 		const FVector TargetPoleVector = CalculatePoleVector(TargetPoleAxis, TargetIndices, OutTargetGlobalPose);
 	
-		const float RotateTargetToSource = FMath::Acos(FVector::DotProduct(SourcePoleVector, TargetPoleVector));
+		const float RotateTargetToSource = static_cast<float>(FMath::Acos(FVector::DotProduct(SourcePoleVector, TargetPoleVector)));
 		const float MatchPoleAngle = RotateTargetToSource - (Settings.PoleVectorMaintainOffset ? TargetToSourceAngularOffsetAtRefPose : 0);
 		
 		MatchingRotation = FQuat(TargetChainAxisNorm, MatchPoleAngle);
@@ -1500,7 +1500,7 @@ EAxis::Type FPoleVectorMatcher::GetMostDifferentAxis(
 	for (const EAxis::Type Axis : CardinalAxes)
 	{
 		const FVector AxisVector = Transform.GetUnitAxis(Axis);
-		const float AbsAxisDotNormal = FMath::Abs(FVector::DotProduct(AxisVector, InNormal));
+		const float AbsAxisDotNormal = static_cast<float>(FMath::Abs(FVector::DotProduct(AxisVector, InNormal)));
 		if (AbsAxisDotNormal < MostDifferentDot)
 		{
 			MostDifferentDot = AbsAxisDotNormal;
@@ -2290,7 +2290,7 @@ void UIKRetargetProcessor::RunStrideWarping(const TArray<FTransform>& InTargetGl
 
 		// sideways offset
 		// first determine which side the goal is on originally
-		float GoalSideMultiplier = FVector::DotProduct(ChainRetargeter.Target.InitialEndPosition.GetSafeNormal(), SideOrig);
+		float GoalSideMultiplier = static_cast<float>(FVector::DotProduct(ChainRetargeter.Target.InitialEndPosition.GetSafeNormal(), SideOrig));
 		// push goal by offset in the newly calculated sideways direction
 		WarpedGoalPosition += Side * GlobalSettings.SidewaysOffset * GoalSideMultiplier;
 
diff --git a/Engine/Plugins/Animation/IKRig/Source/IKRig/Private/Retargeter/IKRetargeter.cpp b/Engine/Plugins/Animation/IKRig/Source/IKRig/Private/Retargeter/IKRetargeter.cpp
--- a/Engine/Plugins/Animation/IKRig/Source/IKRig/Private/Retargeter/IKRetargeter.cpp
+++ b/Engine/Plugins/Animation/IKRig/Source/IKRig/Private/Retargeter/IKRetargeter.cpp
@@ -290,7 +290,7 @@ void URetargetRootSettings::Serialize(FArchive& Ar)
 			PRAGMA_DISABLE_DEPRECATION_WARNINGS
 			Settings.ScaleHorizontal = GlobalScaleHorizontal_DEPRECATED;
 			Settings.ScaleVertical = GlobalScaleVertical_DEPRECATED;
-			Settings.BlendToSource = BlendToSource_DEPRECATED.Size();
+			Settings.BlendToSource = static_cast<float>(BlendToSource_DEPRECATED.Size());
 			Settings.TranslationOffset = StaticOffset_DEPRECATED;
 			Settings.RotationOffset = StaticRotationOffset_DEPRECATED;
 			PRAGMA_ENABLE_DEPRECATION_WARNINGS
diff --git a/Engine/Plugins/Animation/IKRig/Source/IKRig/Private/Rig/IKRigSkeleton.cpp b/Engine/Plugins/Animation/IKRig/Source/IKRig/Private/Rig/IKRigSkeleton.cpp
--- a/Engine/Plugins/Animation/IKRig/Source/IKRig/Private/Rig/IKRigSkeleton.cpp
+++ b/Engine/Plugins/Animation/IKRig/Source/IKRig/Private/Rig/IKRigSkeleton.cpp
@@ -445,7 +445,7 @@ bool FIKRigSkeleton::GetMirroredBoneIndices(
 		int32 Index = 0;
 		for (const FTransform& BoneTransform : RefPoseGlobal)
 		{
-			OutDistances[Index] = FVector::Distance(InPoint, BoneTransform.GetLocation());
+			OutDistances[Index] = static_cast<float>(FVector::Distance(InPoint, BoneTransform.GetLocation()));
 			++Index;
 		}
 	};
@@ -474,7 +474,7 @@ bool FIKRigSkeleton::GetMirroredBoneIndices(
 		for (const int32 BoneIndex : InBoneIndices)
 		{
 			FString CurrentNameStr = BoneNames[BoneIndex].ToString().ToLower();
-			float WorstCase = BaseNameStr.Len() + CurrentNameStr.Len();
+			float WorstCase = static_cast<float>(BaseNameStr.Len() + CurrentNameStr.Len());
 			WorstCase = WorstCase < 1.0f ? 1.0f : WorstCase;
 			const float Score = 1.0f - (Algo::LevenshteinDistance(BaseNameStr, CurrentNameStr) / WorstCase);
 			if (Score > HighestScore)
@@ -508,8 +508,8 @@ bool FIKRigSkeleton::GetMirroredBoneIndices(
 		return false;
 	}
 		
-	const float MinSearchRadius = BoundingSphere.W * 0.01f;
-	const float MaxSearchRadius = BoundingSphere.W;
+	const float MinSearchRadius = static_cast<float>(BoundingSphere.W) * 0.01f;
+	const float MaxSearchRadius = static_cast<float>(BoundingSphere.W);
 	const float SearchStep = MinSearchRadius * 2.f;
 
 	// spin through all the bones in the chain:
diff --git a/Engine/Plugins/Animation/IKRig/Source/IKRig/Private/Rig/Solvers/LimbSolver.cpp b/Engine/Plugins/Animation/IKRig/Source/IKRig/Private/Rig/Solvers/LimbSolver.cpp
--- a/Engine/Plugins/Animation/IKRig/Source/IKRig/Private/Rig/Solvers/LimbSolver.cpp
+++ b/Engine/Plugins/Animation/IKRig/Source/IKRig/Private/Rig/Solvers/LimbSolver.cpp
@@ -18,7 +18,7 @@ bool FLimbSolver::Initialize()
 	for (int32 Index = 1; Index < NumLinks(); ++Index)
 	{
 		FLimbLink& Link = Links[Index];
-		Link.Length = FVector::Dist( Link.Location, Links[Index-1].Location);
+		Link.Length = static_cast<float>(FVector::Dist( Link.Location, Links[Index-1].Location));
 		Length += Link.Length;
 	}
 
@@ -237,7 +237,7 @@ namespace AnimationCore
 		// At this point we know the target location is reachable, and we are already aligned with that location.
 		// So the leg is in the right plane.
 		const float a = InOutLinks[1].Length;	// hip to knee
-		const float b = HipToFoot.Length();		// hip to foot
+		const float b = static_cast<float>(HipToFoot.Length());		// hip to foot
 		const float c = InOutLinks[2].Length;	// knee to foot
 		
 		const float Two_ab = 2.f * a * b;
@@ -250,7 +250,7 @@ namespace AnimationCore
 		const FVector ProjKnee = pC + HipToKnee.ProjectOnToNormal(HipToFootDir);
 
 		const FVector ProjKneeToKnee = (pB - ProjKnee);
-		FVector BendDir = ProjKneeToKnee.GetSafeNormal(KINDA_SMALL_NUMBER);
+		FVector BendDir = ProjKneeToKnee.GetSafeNormal(UE_KINDA_SMALL_NUMBER);
 		
 		// If we have a HingeRotationAxis defined, we can cache 'BendDir'
 		// and use it when we can't determine it. (When limb is straight without a bend).
@@ -259,7 +259,7 @@ namespace AnimationCore
 		if ((HingeRotationAxis != FVector::ZeroVector) && (HipToFootDir != FVector::ZeroVector) && !FMath::IsNearlyZero(a))
 		{
 			const FVector HipToKneeDir = HipToKnee / a;
-			const float KneeBendDot = FVector::DotProduct(HipToKneeDir, HipToFootDir);
+			const float KneeBendDot = static_cast<float>(FVector::DotProduct(HipToKneeDir, HipToFootDir));
 
 			FVector& CachedRealBendDir = InOutLinks[1].RealBendDir;
 			FVector& CachedBaseBendDir = InOutLinks[1].BaseBendDir;
@@ -300,7 +300,7 @@ namespace AnimationCore
 			const FVector PlaneNormal = FVector::CrossProduct(AxisX, AxisY);
 
 			// Make sure we have a valid normal (Axes were not coplanar).
-			if (PlaneNormal.SizeSquared() > SMALL_NUMBER)
+			if (PlaneNormal.SizeSquared() > UE_SMALL_NUMBER)
 			{
 				return PlaneNormal.GetUnsafeNormal();
 			}
@@ -325,8 +325,8 @@ namespace AnimationCore
 		const FVector ChildAxisY = FVector::CrossProduct(CurrentLink.LinkAxisZ, ChildAxisX);
 		const FVector ParentAxisX = (ParentLink.Location - CurrentLink.Location).GetSafeNormal();
 	
-		const float ParentCos = FVector::DotProduct(ParentAxisX, ChildAxisX);
-		const float ParentSin = FVector::DotProduct(ParentAxisX, ChildAxisY);
+		const float ParentCos = static_cast<float>(FVector::DotProduct(ParentAxisX, ChildAxisX));
+		const float ParentSin = static_cast<float>(FVector::DotProduct(ParentAxisX, ChildAxisY));
 	
 		const bool bNeedsReorient = (ParentSin < 0.f) || (ParentCos > FMath::Cos(MinRotationAngleRadians));
 	
@@ -362,8 +362,8 @@ namespace AnimationCore
 		const FVector ParentAxisY = FVector::CrossProduct(CurrentLink.LinkAxisZ, ParentAxisX);
 		const FVector ChildAxisX = (ChildLink.Location - CurrentLink.Location).GetSafeNormal();
 	
-		const float ChildCos = FVector::DotProduct(ChildAxisX, ParentAxisX);
-		const float ChildSin = FVector::DotProduct(ChildAxisX, ParentAxisY);
+		const float ChildCos = static_cast<float>(FVector::DotProduct(ChildAxisX, ParentAxisX));
+		const float ChildSin = static_cast<float>(FVector::DotProduct(ChildAxisX, ParentAxisY));
 
 		const float CosMinAngle = FMath::Cos(MinRotationAngleRadians);
 		
@@ -409,7 +409,7 @@ namespace AnimationCore
 			for (int32 Index = NumLinks-1; Index >= 1; Index--)
 			{
 				FVector EndEffectorToParent = InOutLinks[Index].Location - InOutLinks.Last().Location;
-				float ParentDisplacement = FVector::DotProduct(EndEffectorToParent, EndEffectorToTargetDir);
+				float ParentDisplacement = static_cast<float>(FVector::DotProduct(EndEffectorToParent, EndEffectorToTargetDir));
 
 				Displacement = (ParentDisplacement > 0.f) ? FMath::Min(Displacement, ParentDisplacement * ReachStepAlpha) : Displacement;
 			}
@@ -456,7 +456,7 @@ namespace AnimationCore
 			for (int32 Index = 0; Index < NumLinks-1; Index++)
 			{
 				FVector RootToChild = InOutLinks[Index].Location - InOutLinks[0].Location;
-				float ChildDisplacement = FVector::DotProduct(RootToChild, RootToRootTargetDir);
+				float ChildDisplacement = static_cast<float>(FVector::DotProduct(RootToChild, RootToRootTargetDir));
 				
 				Displacement = (ChildDisplacement > 0.f) ? FMath::Min(Displacement, ChildDisplacement * ReachStepAlpha) : Displacement;				
 			}
@@ -492,12 +492,12 @@ namespace AnimationCore
 		const int32 NumLinks = InOutLinks.Num();
 		
 		// Make sure precision is not too small.
-		const float ReachPrecision = FMath::Max(InReachPrecision, KINDA_SMALL_NUMBER);
+		const float ReachPrecision = FMath::Max(InReachPrecision, UE_KINDA_SMALL_NUMBER);
 
 		const FVector RootTargetLocation = InOutLinks[0].Location;
 
 		// Check distance between foot and foot target location
-		float Slop = FVector::Dist(InOutLinks.Last().Location, InTargetLocation);
+		float Slop = static_cast<float>(FVector::Dist(InOutLinks.Last().Location, InTargetLocation));
 		
 		if (Slop > ReachPrecision || bUseAngleLimit)
 		{
@@ -559,7 +559,7 @@ namespace AnimationCore
 					FABRIK_BackwardReach(InOutLinks, RootTargetLocation, ReachStepAlpha, bUseAngleLimit, MinRotationAngleRadians);
 				}
 				
-				Slop = FVector::Dist( InOutLinks.Last().Location, InTargetLocation) + FVector::Dist(InOutLinks[0].Location, RootTargetLocation);
+				Slop = static_cast<float>(FVector::Dist( InOutLinks.Last().Location, InTargetLocation) + FVector::Dist(InOutLinks[0].Location, RootTargetLocation));
 
 				// Abort if we're not getting closer and enter a deadlock.
 				if (Slop > PreviousSlop)
diff --git a/Engine/Plugins/Animation/IKRig/Source/IKRig/Public/Rig/Solvers/PointsToRotation.h b/Engine/Plugins/Animation/IKRig/Source/IKRig/Public/Rig/Solvers/PointsToRotation.h
--- a/Engine/Plugins/Animation/IKRig/Source/IKRig/Public/Rig/Solvers/PointsToRotation.h
+++ b/Engine/Plugins/Animation/IKRig/Source/IKRig/Public/Rig/Solvers/PointsToRotation.h
@@ -34,7 +34,7 @@ static FQuat GetRotationFromDeformedPoints(
 	}
 
 	// average centroids
-	const float InvNumEffectors = 1.0f / static_cast<float>(InInitialPoints.Num());
+	const double InvNumEffectors = 1.0f / static_cast<double>(InInitialPoints.Num());
 	OutInitialCentroid *= InvNumEffectors;
 	OutCurrentCentroid *= InvNumEffectors;
 
@@ -71,9 +71,9 @@ static FQuat GetRotationFromDeformedPoints(
 		FVector RCol1(R.M[1][0], R.M[1][1], R.M[1][2]);
 		FVector RCol2(R.M[2][0], R.M[2][1], R.M[2][2]);
 		FVector Omega = RCol0.Cross(DX) + RCol1.Cross(DY) + RCol2.Cross(DZ);
-		Omega *= 1.0f / (fabs(RCol0.Dot(DX) + RCol1.Dot(DY) + RCol2.Dot(DZ)) + SMALL_NUMBER);
-		const float W = Omega.Size();
-		if (W < SMALL_NUMBER)
+		Omega *= 1.0 / (fabs(RCol0.Dot(DX) + RCol1.Dot(DY) + RCol2.Dot(DZ)) + UE_SMALL_NUMBER);
+		const double W = Omega.Size();
+		if (W < UE_SMALL_NUMBER)
 		{
 			break;
 		}
