IK Rig Editor: Fix truncation warnings

#rnx

diff --git a/Engine/Plugins/Animation/IKRig/Source/IKRigEditor/IKRigEditor.Build.cs b/Engine/Plugins/Animation/IKRig/Source/IKRigEditor/IKRigEditor.Build.cs
--- a/Engine/Plugins/Animation/IKRig/Source/IKRigEditor/IKRigEditor.Build.cs
+++ b/Engine/Plugins/Animation/IKRig/Source/IKRigEditor/IKRigEditor.Build.cs
@@ -8,6 +8,8 @@ namespace UnrealBuildTool.Rules
     {
         public IKRigEditor(ReadOnlyTargetRules Target) : base(Target)
 		{
+			CppCompileWarningSettings.UnsafeTypeCastWarningLevel = WarningLevel.Error;
+			
 			PublicDependencyModuleNames.AddRange(
                 new string[]
                 {
diff --git a/Engine/Plugins/Animation/IKRig/Source/IKRigEditor/Private/RetargetEditor/IKRetargetBatchOperation.cpp b/Engine/Plugins/Animation/IKRig/Source/IKRigEditor/Private/RetargetEditor/IKRetargetBatchOperation.cpp
--- a/Engine/Plugins/Animation/IKRig/Source/IKRigEditor/Private/RetargetEditor/IKRetargetBatchOperation.cpp
+++ b/Engine/Plugins/Animation/IKRig/Source/IKRigEditor/Private/RetargetEditor/IKRetargetBatchOperation.cpp
@@ -910,7 +910,7 @@ void UIKRetargetBatchOperation::RunRetarget(FIKRetargetBatchOperationContext& Co
 	constexpr int32 NumAdditionalProgressFrames = 4;
 	constexpr int32 NumPassesOverAssets = 3;
 	const int32 NumProgressSteps = (NumAssets * NumPassesOverAssets) + NumAdditionalProgressFrames;
-	FScopedSlowTask Progress(NumProgressSteps, LOCTEXT("GatheringBatchRetarget", "Gathering animation assets..."));
+	FScopedSlowTask Progress(static_cast<float>(NumProgressSteps), LOCTEXT("GatheringBatchRetarget", "Gathering animation assets..."));
 	constexpr bool bShowCancelButton = true;
 	Progress.MakeDialog(bShowCancelButton);
 	
diff --git a/Engine/Plugins/Animation/IKRig/Source/IKRigEditor/Private/RetargetEditor/IKRetargetDefaultMode.cpp b/Engine/Plugins/Animation/IKRig/Source/IKRigEditor/Private/RetargetEditor/IKRetargetDefaultMode.cpp
--- a/Engine/Plugins/Animation/IKRig/Source/IKRigEditor/Private/RetargetEditor/IKRetargetDefaultMode.cpp
+++ b/Engine/Plugins/Animation/IKRig/Source/IKRigEditor/Private/RetargetEditor/IKRetargetDefaultMode.cpp
@@ -78,7 +78,7 @@ void FIKRetargetDefaultMode::RenderDebugProxies(FPrimitiveDrawInterface* PDI, co
 
 	UDebugSkelMeshComponent* TargetSkelMesh = Controller->GetSkeletalMeshComponent(ERetargetSourceOrTarget::Target);
 	const FTransform ComponentTransform = TargetSkelMesh->GetComponentTransform();
-	const float ComponentScale = ComponentTransform.GetScale3D().GetMax();
+	const float ComponentScale = static_cast<float>(ComponentTransform.GetScale3D().GetMax());
 
 	const TArray<FName>& SelectedChains = Controller->GetSelectedChains();
 
diff --git a/Engine/Plugins/Animation/IKRig/Source/IKRigEditor/Private/RetargetEditor/IKRetargetEditPoseMode.cpp b/Engine/Plugins/Animation/IKRig/Source/IKRigEditor/Private/RetargetEditor/IKRetargetEditPoseMode.cpp
--- a/Engine/Plugins/Animation/IKRig/Source/IKRigEditor/Private/RetargetEditor/IKRetargetEditPoseMode.cpp
+++ b/Engine/Plugins/Animation/IKRig/Source/IKRigEditor/Private/RetargetEditor/IKRetargetEditPoseMode.cpp
@@ -328,7 +328,7 @@ bool FIKRetargetEditPoseMode::InputDelta(
 		float Scale;
 		FVector Offset;
 		GetEditedComponentScaleAndOffset(Scale,Offset);
-		Scale = FMath::Max(Scale, KINDA_SMALL_NUMBER);
+		Scale = FMath::Max(Scale, UE_KINDA_SMALL_NUMBER);
 		BoneEdit.AccumulatedPositionOffset += InDrag * (1.0f / Scale);
 		
 		// apply translation delta to root
diff --git a/Engine/Plugins/Animation/IKRig/Source/IKRigEditor/Private/RetargetEditor/IKRetargetEditorController.cpp b/Engine/Plugins/Animation/IKRig/Source/IKRigEditor/Private/RetargetEditor/IKRetargetEditorController.cpp
--- a/Engine/Plugins/Animation/IKRig/Source/IKRigEditor/Private/RetargetEditor/IKRetargetEditorController.cpp
+++ b/Engine/Plugins/Animation/IKRig/Source/IKRigEditor/Private/RetargetEditor/IKRetargetEditorController.cpp
@@ -1424,13 +1424,13 @@ void FIKRetargetEditorController::RenderSkeleton(FPrimitiveDrawInterface* PDI, E
 	WorldTransforms.AddUninitialized(NumBones);
 	for (int32 Index=0; Index<NumBones; ++Index)
 	{
-		RequiredBones[Index] = Index;
+		RequiredBones[Index] = static_cast<FBoneIndexType>(Index);
 		WorldTransforms[Index] = MeshComponent->GetBoneTransform(Index, ComponentTransform);
 	}
 	
 	const UIKRetargeter* Asset = AssetController->GetAsset();
 	const float BoneDrawSize = Asset->BoneDrawSize;
-	const float MaxDrawRadius = MeshComponent->Bounds.SphereRadius * 0.01f;
+	const float MaxDrawRadius = static_cast<float>(MeshComponent->Bounds.SphereRadius * 0.01f);
 	const float BoneRadius = FMath::Min(1.0f, MaxDrawRadius) * BoneDrawSize;
 	const bool bIsSelectable = InSourceOrTarget == GetSourceOrTarget();
 	const FLinearColor DefaultColor = bIsSelectable ? GetMutableDefault<UPersonaOptions>()->DefaultBoneColor : GetMutableDefault<UPersonaOptions>()->DisabledBoneColor;
@@ -1596,7 +1596,7 @@ void FIKRetargetEditorController::FixZeroHeightRetargetRoot(ERetargetSourceOrTar
 		if (PromptToFixRootHeight(SourceOrTarget))
 		{
 			// move it up based on the height of the mesh
-			const float FixedHeight = FMath::Abs(SkeletalMesh->GetBounds().GetBoxExtrema(-1).Z);
+			const double FixedHeight = FMath::Abs(SkeletalMesh->GetBounds().GetBoxExtrema(-1).Z);
 			// update the current retarget pose
 			CurrentRetargetPose.SetRootTranslationDelta(FVector(0.f, 0.f,FixedHeight));
 		}
diff --git a/Engine/Plugins/Animation/IKRig/Source/IKRigEditor/Private/RetargetEditor/IKRetargetPoseExporter.cpp b/Engine/Plugins/Animation/IKRig/Source/IKRigEditor/Private/RetargetEditor/IKRetargetPoseExporter.cpp
--- a/Engine/Plugins/Animation/IKRig/Source/IKRigEditor/Private/RetargetEditor/IKRetargetPoseExporter.cpp
+++ b/Engine/Plugins/Animation/IKRig/Source/IKRigEditor/Private/RetargetEditor/IKRetargetPoseExporter.cpp
@@ -300,8 +300,8 @@ FReply FIKRetargetPoseExporter::ImportPoseAsset() const
 		const FTransform& LocalRefTransform = RefPose[BoneIndex];
 		const FTransform& LocalImportedTransform = LocalBoneTransformFromPose[PoseTrackIndex];
 		const FQuat DeltaRotation = LocalRefTransform.GetRotation().Inverse() * LocalImportedTransform.GetRotation();
-		const float DeltaAngle = FMath::RadiansToDegrees(DeltaRotation.GetAngle());
-		constexpr float MinAngleThreshold = 0.05f;
+		const double DeltaAngle = FMath::RadiansToDegrees(DeltaRotation.GetAngle());
+		constexpr double MinAngleThreshold = 0.05;
 		if (DeltaAngle >= MinAngleThreshold)
 		{
 			NewPose.SetDeltaRotationForBone(BoneName, DeltaRotation);
@@ -615,8 +615,8 @@ FReply FIKRetargetPoseExporter::OnImportPoseFromSequence()
 			const FTransform LocalTransformReference = RefPose[BoneIndex];
 			const FQuat DeltaRotation = LocalTransformReference.GetRotation().Inverse() * LocalTransformImported.GetRotation();
 			// only if it's different than the ref pose
-			const float DeltaAngle = FMath::RadiansToDegrees(DeltaRotation.GetAngle());
-			constexpr float MinAngleThreshold = 0.05f;
+			const double DeltaAngle = FMath::RadiansToDegrees(DeltaRotation.GetAngle());
+			constexpr double MinAngleThreshold = 0.05;
 			if (DeltaAngle >= MinAngleThreshold)
 			{
 				ImportedRetargetPose.SetDeltaRotationForBone(BoneName, DeltaRotation);
diff --git a/Engine/Plugins/Animation/IKRig/Source/IKRigEditor/Private/RetargetEditor/IKRetargeterController.cpp b/Engine/Plugins/Animation/IKRig/Source/IKRigEditor/Private/RetargetEditor/IKRetargeterController.cpp
--- a/Engine/Plugins/Animation/IKRig/Source/IKRigEditor/Private/RetargetEditor/IKRetargeterController.cpp
+++ b/Engine/Plugins/Animation/IKRig/Source/IKRigEditor/Private/RetargetEditor/IKRetargeterController.cpp
@@ -506,9 +506,9 @@ void UIKRetargeterController::AutoMapChains(const EAutoMapChainType AutoMapType,
 				for (int32 ChainIndex=0; ChainIndex<SourceChainNames.Num(); ++ChainIndex)
 				{
 					FString SourceNameLowerCase = SourceChainNames[ChainIndex].ToString().ToLower();
-					float WorstCase = TargetNameLowerCase.Len() + SourceNameLowerCase.Len();
+					float WorstCase = static_cast<float>(TargetNameLowerCase.Len() + SourceNameLowerCase.Len());
 					WorstCase = WorstCase < 1.0f ? 1.0f : WorstCase;
-					const float Score = 1.0f - (Algo::LevenshteinDistance(TargetNameLowerCase, SourceNameLowerCase) / WorstCase);
+					const float Score = 1.0f - (static_cast<float>(Algo::LevenshteinDistance(TargetNameLowerCase, SourceNameLowerCase)) / WorstCase);
 					if (Score > HighestScore)
 					{
 						HighestScore = Score;
diff --git a/Engine/Plugins/Animation/IKRig/Source/IKRigEditor/Private/RetargetEditor/IKRetargeterPoseGenerator.cpp b/Engine/Plugins/Animation/IKRig/Source/IKRigEditor/Private/RetargetEditor/IKRetargeterPoseGenerator.cpp
--- a/Engine/Plugins/Animation/IKRig/Source/IKRigEditor/Private/RetargetEditor/IKRetargeterPoseGenerator.cpp
+++ b/Engine/Plugins/Animation/IKRig/Source/IKRigEditor/Private/RetargetEditor/IKRetargeterPoseGenerator.cpp
@@ -303,7 +303,7 @@ void FRetargetAutoPoseGenerator::SnapToGround(const FName BoneToPutOnGround, con
 				continue;
 			}
 			const FTransform& BoneTransform = Skeleton.RetargetGlobalPose[BoneIndex];
-			const float BoneHeight = BoneTransform.GetLocation().Z;
+			const float BoneHeight = static_cast<float>(BoneTransform.GetLocation().Z);
 			if (BoneHeight < HeightOfBoneInRetargetPose)
 			{
 				HeightOfBoneInRetargetPose = BoneHeight;
@@ -315,7 +315,7 @@ void FRetargetAutoPoseGenerator::SnapToGround(const FName BoneToPutOnGround, con
 	{
 		// use the user specified bone
 		BoneToSnapIndex = Skeleton.FindBoneIndexByName(BoneToPutOnGround);
-		HeightOfBoneInRetargetPose = Skeleton.RetargetGlobalPose[BoneToSnapIndex].GetLocation().Z;
+		HeightOfBoneInRetargetPose = static_cast<float>(Skeleton.RetargetGlobalPose[BoneToSnapIndex].GetLocation().Z);
 	}
 
 	if (BoneToSnapIndex == INDEX_NONE)
@@ -327,11 +327,11 @@ void FRetargetAutoPoseGenerator::SnapToGround(const FName BoneToPutOnGround, con
 	// find the height of this bone in the ref pose
 	const FReferenceSkeleton& RefSkeleton = Skeleton.SkeletalMesh->GetRefSkeleton();
 	const FTransform& RefPoseOfBone = FAnimationRuntime::GetComponentSpaceTransformRefPose(RefSkeleton, BoneToSnapIndex);
-	const float HeightInRefPose = RefPoseOfBone.GetLocation().Z;
+	const double HeightInRefPose = RefPoseOfBone.GetLocation().Z;
 
 	// snap it back to original height off the ground
-	const float DeltaHeight = HeightInRefPose - HeightOfBoneInRetargetPose;
-	Controller->SetRootOffsetInRetargetPose(FVector(0.f,0.f,DeltaHeight), SourceOrTarget);
+	const double DeltaHeight = HeightInRefPose - HeightOfBoneInRetargetPose;
+	Controller->SetRootOffsetInRetargetPose(FVector(0,0,DeltaHeight), SourceOrTarget);
 }
 
 bool FRetargetAutoPoseGenerator::GetSpineChain(const ERetargetSourceOrTarget SourceOrTarget, FName& OutChainName) const
@@ -344,9 +344,9 @@ bool FRetargetAutoPoseGenerator::GetSpineChain(const ERetargetSourceOrTarget Sou
 	{
 		FName ChainNameToTest = SourceOrTarget == ERetargetSourceOrTarget::Source ? ChainPair.SourceBoneChainName : ChainPair.TargetBoneChainName;
 		FString ChainNameStr = ChainNameToTest.ToString().ToLower();
-		float WorstCase = ChainNameToFind.Len() + ChainNameStr.Len();
+		float WorstCase = static_cast<float>(ChainNameToFind.Len() + ChainNameStr.Len());
 		WorstCase = WorstCase < 1.0f ? 1.0f : WorstCase;
-		const float Score = 1.0f - (Algo::LevenshteinDistance(ChainNameToFind, ChainNameStr) / WorstCase);
+		const float Score = 1.0f - (static_cast<float>(Algo::LevenshteinDistance(ChainNameToFind, ChainNameStr)) / WorstCase);
 		if (Score > HighestScore && Score > 0.5f)
 		{
 			HighestScore = Score;
diff --git a/Engine/Plugins/Animation/IKRig/Source/IKRigEditor/Private/RetargetEditor/IKRetargeterThumbnailRenderer.cpp b/Engine/Plugins/Animation/IKRig/Source/IKRigEditor/Private/RetargetEditor/IKRetargeterThumbnailRenderer.cpp
--- a/Engine/Plugins/Animation/IKRig/Source/IKRigEditor/Private/RetargetEditor/IKRetargeterThumbnailRenderer.cpp
+++ b/Engine/Plugins/Animation/IKRig/Source/IKRigEditor/Private/RetargetEditor/IKRetargeterThumbnailRenderer.cpp
@@ -48,7 +48,7 @@ void FIKRetargeterThumbnailScene::SetSkeletalMeshes(USkeletalMesh* SourceMesh, U
 			if (bApplyOffset)
 			{
 				const FBoxSphereBounds Bounds = InMesh->GetImportedBounds();
-				const float HalfWidth = (Bounds.BoxExtent.X - Bounds.Origin.X) * 0.5f;
+				const double HalfWidth = (Bounds.BoxExtent.X - Bounds.Origin.X) * 0.5;
 				Offset.X = bIsSourceMesh ? -HalfWidth : HalfWidth;
 			}
 		
@@ -88,7 +88,7 @@ void FIKRetargeterThumbnailScene::GetViewMatrixParameters(
 	const FBoxSphereBounds TotalBounds = TargetBounds + SourceBounds;
 
 	// calc target distance
-	const float HalfBoundSize = TotalBounds.GetBox().GetSize().Size() * 0.5f;
+	const float HalfBoundSize = static_cast<float>(TotalBounds.GetBox().GetSize().Size()) * 0.5f;
 	const float HalfFOVRadians = FMath::DegreesToRadians<float>(InFOVDegrees) * 0.5f;
 	const float TargetDistance = HalfBoundSize / FMath::Tan(HalfFOVRadians);
 
diff --git a/Engine/Plugins/Animation/IKRig/Source/IKRigEditor/Private/RetargetEditor/SRetargetAnimAssetsWindow.cpp b/Engine/Plugins/Animation/IKRig/Source/IKRigEditor/Private/RetargetEditor/SRetargetAnimAssetsWindow.cpp
--- a/Engine/Plugins/Animation/IKRig/Source/IKRigEditor/Private/RetargetEditor/SRetargetAnimAssetsWindow.cpp
+++ b/Engine/Plugins/Animation/IKRig/Source/IKRigEditor/Private/RetargetEditor/SRetargetAnimAssetsWindow.cpp
@@ -526,7 +526,7 @@ void SRetargetPoseViewport::SetSkeletalMesh(USkeletalMesh* InSkeletalMesh, EReta
 	Component->EnablePreview(true, nullptr);
 
 	// translate target sufficiently far from source to avoid them touching on X axis (sideways)
-	const float Offset = SourceComponent->Bounds.GetBox().Max.X + FMath::Abs(TargetComponent->Bounds.GetBox().Min.X);
+	const double Offset = SourceComponent->Bounds.GetBox().Max.X + FMath::Abs(TargetComponent->Bounds.GetBox().Min.X);
 	TargetComponent->SetWorldLocation(FVector(Offset, 0.f, 0.f));
 	
 	// update camera to show both meshes
diff --git a/Engine/Plugins/Animation/IKRig/Source/IKRigEditor/Private/RigEditor/IKRigAutoCharacterizer.cpp b/Engine/Plugins/Animation/IKRig/Source/IKRigEditor/Private/RigEditor/IKRigAutoCharacterizer.cpp
--- a/Engine/Plugins/Animation/IKRig/Source/IKRigEditor/Private/RigEditor/IKRigAutoCharacterizer.cpp
+++ b/Engine/Plugins/Animation/IKRig/Source/IKRigEditor/Private/RigEditor/IKRigAutoCharacterizer.cpp
@@ -237,9 +237,9 @@ void FAbstractHierarchy::Compare(
 	// calculate a simple score with equal weight for number of matching bones and parents
 	// this score represents what percentage of the template is found in the input hierarchy
 	FindBonesNotInOther(OtherHierarchy, OutMissingBones, OutBonesWithDifferentParent);
-	const float NumBonesTotal = CleanBoneNames.Num();
-	const float NumMatchingBones = NumBonesTotal - OutMissingBones.Num();
-	const float NumMatchingParents = NumBonesTotal - OutBonesWithDifferentParent.Num();
+	const float NumBonesTotal = static_cast<float>(CleanBoneNames.Num());
+	const float NumMatchingBones = NumBonesTotal - static_cast<float>(OutMissingBones.Num());
+	const float NumMatchingParents = NumBonesTotal - static_cast<float>(OutBonesWithDifferentParent.Num());
 	const float BoneMatchScore = NumMatchingBones / NumBonesTotal;
 	const float ParentMatchScore = NumMatchingParents / NumBonesTotal;
 	constexpr float InvNumScores = 1.0f / 2.f;
@@ -1753,9 +1753,9 @@ void FAutoCharacterizer::FindClosestNameInArray(
 	for (const FName BoneName : NamesToCheck)
 	{
 		FString CurrentNameStr = BoneName.ToString().ToLower();
-		float WorstCase = NameToMatchStr.Len() + CurrentNameStr.Len();
+		float WorstCase = static_cast<float>(NameToMatchStr.Len() + CurrentNameStr.Len());
 		WorstCase = WorstCase < 1.0f ? 1.0f : WorstCase;
-		const float Score = 1.0f - (Algo::LevenshteinDistance(NameToMatchStr, CurrentNameStr) / WorstCase);
+		const float Score = 1.0f - (static_cast<float>(Algo::LevenshteinDistance(NameToMatchStr, CurrentNameStr)) / WorstCase);
 		if (Score > OutBestScore)
 		{
 			OutBestScore = Score;
diff --git a/Engine/Plugins/Animation/IKRig/Source/IKRigEditor/Private/RigEditor/IKRigEditMode.cpp b/Engine/Plugins/Animation/IKRig/Source/IKRigEditor/Private/RigEditor/IKRigEditMode.cpp
--- a/Engine/Plugins/Animation/IKRig/Source/IKRigEditor/Private/RigEditor/IKRigEditMode.cpp
+++ b/Engine/Plugins/Animation/IKRig/Source/IKRigEditor/Private/RigEditor/IKRigEditMode.cpp
@@ -155,7 +155,7 @@ void FIKRigEditMode::RenderBones(FPrimitiveDrawInterface* PDI)
 	WorldTransforms.AddUninitialized(NumBones);
 	for (int32 Index=0; Index<NumBones; ++Index)
 	{
-		RequiredBones[Index] = Index;
+		RequiredBones[Index] = static_cast<FBoneIndexType>(Index);
 		WorldTransforms[Index] = MeshComponent->GetBoneTransform(Index, ComponentTransform);
 	}
 	
diff --git a/Engine/Plugins/Animation/IKRig/Source/IKRigEditor/Private/RigEditor/IKRigEditorController.cpp b/Engine/Plugins/Animation/IKRig/Source/IKRigEditor/Private/RigEditor/IKRigEditorController.cpp
--- a/Engine/Plugins/Animation/IKRig/Source/IKRigEditor/Private/RigEditor/IKRigEditorController.cpp
+++ b/Engine/Plugins/Animation/IKRig/Source/IKRigEditor/Private/RigEditor/IKRigEditorController.cpp
@@ -297,7 +297,7 @@ EChainSide FRetargetChainAnalyzer::GetSideOfChain(
 			}
 		}
 		
-		return Score > BoneIndices.Num() * 0.5f;
+		return Score > (BoneIndices.Num() / 2);
 	};
 	
 	if (DoMajorityBonesContainText("Right"))
@@ -314,14 +314,14 @@ EChainSide FRetargetChainAnalyzer::GetSideOfChain(
 	//
 	
 	// determine "sidedness" of the chain based on the location of the bones (left, right or center of YZ plane)
-	float AverageXPositionOfChain = 0.f;
+	double AverageXPositionOfChain = 0.f;
 	for (const int32 BoneIndex : BoneIndices)
 	{
 		AverageXPositionOfChain += IKRigSkeleton.RefPoseGlobal[BoneIndex].GetLocation().X;
 	}
-	AverageXPositionOfChain /= static_cast<float>(BoneIndices.Num());
+	AverageXPositionOfChain /= static_cast<double>(BoneIndices.Num());
 
-	constexpr float CenterThresholdDistance = 1.0f;
+	constexpr double CenterThresholdDistance = 1.0f;
 	if (FMath::Abs(AverageXPositionOfChain) < CenterThresholdDistance)
 	{
 		return EChainSide::Center;
diff --git a/Engine/Source/Developer/AnimationWidgets/Public/SAdvancedTransformInputBox.h b/Engine/Source/Developer/AnimationWidgets/Public/SAdvancedTransformInputBox.h
--- a/Engine/Source/Developer/AnimationWidgets/Public/SAdvancedTransformInputBox.h
+++ b/Engine/Source/Developer/AnimationWidgets/Public/SAdvancedTransformInputBox.h
@@ -918,7 +918,7 @@ public:
 
 			LabelWidget = SNew(SComboButton)
 				.IsEnabled(true)
-				.ContentPadding(0)
+				.ContentPadding(0.f)
 				.OnGetMenuContent_Lambda([Labels, RotationRepresentationPtr, OnRotationRepresentationChanged]
 				{
 					FMenuBuilder MenuBuilder(true, nullptr);
@@ -1044,7 +1044,7 @@ public:
 					 "When locked, scales uniformly based on the current xyz scale values so the object maintains its shape in each direction when scaled");
 
 			return SNew(SButton)
-				.ContentPadding(0)
+				.ContentPadding(0.f)
 				.ButtonStyle(FAppStyle::Get(), "SimpleButton")
 				.OnClicked_Lambda([OnScaleLockChanged, LockState]() -> FReply
 				{
@@ -1118,7 +1118,7 @@ public:
 
 			return SNew(SButton)
 				.IsEnabled(true)
-				.ContentPadding(0)
+				.ContentPadding(0.f)
 				.ButtonStyle(FAppStyle::Get(), "SimpleButton")
 				.OnClicked_Lambda([OnIsComponentRelativeChanged, OnGetIsComponentRelative, InComponent]() -> FReply
 				{
