IK Retargeter: Added relative bone length option to the FK translation modes.

diff --git a/Engine/Plugins/Animation/IKRig/Source/IKRig/Private/Retargeter/IKRetargetProcessor.cpp b/Engine/Plugins/Animation/IKRig/Source/IKRig/Private/Retargeter/IKRetargetProcessor.cpp
--- a/Engine/Plugins/Animation/IKRig/Source/IKRig/Private/Retargeter/IKRetargetProcessor.cpp
+++ b/Engine/Plugins/Animation/IKRig/Source/IKRig/Private/Retargeter/IKRetargetProcessor.cpp
@@ -478,7 +478,7 @@ bool FChainFK::Initialize(
 	return CalculateBoneParameters(Log);
 }
 
-FTransform FChainFK::GetTransformAtParam(const TArray<FTransform>& Transforms, const float& Param) const
+FTransform FChainFK::GetTransformAtChainParam(const TArray<FTransform>& Transforms, const float& Param) const
 {
 	check(Transforms.Num() == Params.Num());
 	
@@ -520,6 +520,63 @@ FTransform FChainFK::GetTransformAtParam(const TArray<FTransform>& Transforms, c
 	return FTransform::Identity;
 }
 
+float FChainFK::GetStretchAtParam(
+	const TArray<FTransform>& InitialTransforms,
+	const TArray<FTransform>& CurrentTransforms,
+	const float& Param)
+{
+	check(InitialTransforms.Num() == CurrentTransforms.Num() && InitialTransforms.Num() == Params.Num());
+
+	// chain only has 1 bone, so it can't stretch
+	if (Params.Num() <= 1)
+	{
+		return 1.0f;
+	}
+
+	// start of chain cannot stretch
+	if (Param < KINDA_SMALL_NUMBER)
+	{
+		return 1.0f;
+	}
+	
+	// end of chain always uses last bone
+	if (Param > 1.0f - KINDA_SMALL_NUMBER)
+	{
+		const int32 Last = InitialTransforms.Num() - 1;
+		const int32 Prev = Last - 1;
+		const float InitialLength = (InitialTransforms[Last].GetTranslation() - InitialTransforms[Prev].GetTranslation()).Length();
+		const float CurrentLength = (CurrentTransforms[Last].GetTranslation() - CurrentTransforms[Prev].GetTranslation()).Length();
+		return InitialLength < KINDA_SMALL_NUMBER ? 1.0f : CurrentLength / InitialLength;
+	}
+
+	// find the bone associated with this parameter
+	float ClosestParamDistance = TNumericLimits<float>::Max();
+	int32 ClosestBoneChainIndex = INDEX_NONE;
+	for (int32 ChainIndex=0; ChainIndex<Params.Num(); ++ChainIndex)
+	{
+		float ParamDistance = FMath::Abs(Params[ChainIndex] - Param);
+		if (ParamDistance < ClosestParamDistance)
+		{
+			ClosestParamDistance = ParamDistance;
+			ClosestBoneChainIndex = ChainIndex;
+		}
+	}
+
+	int32 BoneStartIndex = ClosestBoneChainIndex-1;
+	int32 BoneEndIndex = ClosestBoneChainIndex;
+
+	if (ClosestBoneChainIndex == 0)
+	{
+		BoneStartIndex = 0;
+		BoneEndIndex = 1;
+	}
+
+	// calculate the ratio of the current to the initial length of the closest bone
+	const float InitialLength = (InitialTransforms[BoneStartIndex].GetTranslation() - InitialTransforms[BoneEndIndex].GetTranslation()).Length();
+	const float CurrentLength = (CurrentTransforms[BoneStartIndex].GetTranslation() - CurrentTransforms[BoneEndIndex].GetTranslation()).Length();
+	return InitialLength < KINDA_SMALL_NUMBER ? 1.0f : CurrentLength / InitialLength;
+}
+
 bool FChainFK::CalculateBoneParameters(FIKRigLogger& Log)
 {
 	Params.Reset();
@@ -557,6 +614,18 @@ bool FChainFK::CalculateBoneParameters(FIKRigLogger& Log)
 	return true;
 }
 
+float FChainFK::GetChainLength(const TArray<FTransform>& Transforms)
+{
+	float TotalLength = 0.0f;
+	for (int32 ChainIndex = 1; ChainIndex < Transforms.Num(); ++ChainIndex)
+	{
+		const FVector Start = Transforms[ChainIndex - 1].GetTranslation();
+		const FVector End = Transforms[ChainIndex].GetTranslation();
+		TotalLength += FVector::Dist(Start, End);
+	}
+	return TotalLength;
+}
+
 void FChainFK::FillTransformsWithLocalSpaceOfChain(
 	const FRetargetSkeleton& Skeleton,
 	const TArray<FTransform>& InGlobalPose,
@@ -704,14 +773,16 @@ void FChainDecoderFK::DecodePose(
 		switch (Settings.RotationMode)
 		{
 			case ERetargetRotationMode::Interpolated:
+			case ERetargetRotationMode::MatchChain:
+			case ERetargetRotationMode::MatchScaledChain:
 			{
 				// get the initial and current transform of source chain at param
 				// this is the interpolated transform along the chain
 				const float Param = Params[ChainIndex];
-				SourceCurrentTransform = SourceChain.GetTransformAtParam(SourceChain.CurrentGlobalTransforms, Param);
-				SourceInitialTransform = SourceChain.GetTransformAtParam(SourceChain.InitialGlobalTransforms, Param);
+				SourceCurrentTransform = SourceChain.GetTransformAtChainParam(SourceChain.CurrentGlobalTransforms, Param);
+				SourceInitialTransform = SourceChain.GetTransformAtChainParam(SourceChain.InitialGlobalTransforms, Param);
+				break;
 			}
-			break;
 			case ERetargetRotationMode::OneToOne:
 			{
 				if (ChainIndex < NumBonesInSourceChain)
@@ -723,8 +794,8 @@ void FChainDecoderFK::DecodePose(
 					SourceCurrentTransform = SourceChain.CurrentGlobalTransforms.Last();
 					SourceInitialTransform = SourceChain.InitialGlobalTransforms.Last();
 				}
+				break;
 			}
-			break;
 			case ERetargetRotationMode::OneToOneReversed:
 			{
 				if (ChainIndex < TargetStartIndex)
@@ -738,16 +809,16 @@ void FChainDecoderFK::DecodePose(
 					SourceCurrentTransform = SourceChain.CurrentGlobalTransforms[SourceChainIndex];
 					SourceInitialTransform = SourceChain.InitialGlobalTransforms[SourceChainIndex];
 				}
+				break;
 			}
-			break;
 			case ERetargetRotationMode::None:
 			{
-					// in order to induce no rotation on the FK chain, we rotate the chain rigidly from the root of the chain
-					SourceInitialTransform = SourceChain.InitialGlobalTransforms[0];
-					// use the current global space retarget pose as the "current" transform, so chain rotates with parent
-					SourceCurrentTransform = SourceChain.InitialLocalTransforms[0] * SourceChain.ChainParentCurrentGlobalTransform;
+				// in order to induce no rotation on the FK chain, we rotate the chain rigidly from the root of the chain
+				SourceInitialTransform = SourceChain.InitialGlobalTransforms[0];
+				// use the current global space retarget pose as the "current" transform, so chain rotates with parent
+				SourceCurrentTransform = SourceChain.InitialLocalTransforms[0] * SourceChain.ChainParentCurrentGlobalTransform;
+				break;
 			}
-			break;
 			default:
 				checkNoEntry();
 			break;
@@ -770,23 +841,54 @@ void FChainDecoderFK::DecodePose(
 		FVector OutPosition;
 		switch (Settings.TranslationMode)
 		{
-			case ERetargetTranslationMode::None:
+		case ERetargetTranslationMode::None:
+			{
+				const FVector InitialLocalOffset = TargetSkeleton.RetargetLocalPose[BoneIndex].GetTranslation();
+				OutPosition = ParentGlobalTransform.TransformPosition(InitialLocalOffset);
+				break;
+			}
+		case ERetargetTranslationMode::GloballyScaled:
+			{
+				OutPosition = SourceCurrentTransform.GetTranslation() * RootRetargeter.GetGlobalScaleVector();
+				break;
+			}
+		case ERetargetTranslationMode::Absolute:
+			{
+				OutPosition = SourceCurrentTransform.GetTranslation();
+				break;	
+			}
+		case ERetargetTranslationMode::StretchBoneLengthUniformly:
+			{
+				if (ChainIndex==0)
 				{
 					const FVector InitialLocalOffset = TargetSkeleton.RetargetLocalPose[BoneIndex].GetTranslation();
 					OutPosition = ParentGlobalTransform.TransformPosition(InitialLocalOffset);
+					break;
 				}
+				// initial chain length
+				float SourceChainLengthInitial = GetChainLength(SourceChain.InitialGlobalTransforms);
+				float SourceChainLengthCurrent = GetChainLength(SourceChain.CurrentGlobalTransforms);
+				float StretchRatio = SourceChainLengthInitial < KINDA_SMALL_NUMBER ? 1.0f : SourceChainLengthCurrent / SourceChainLengthInitial;
+				// stretch local translation
+				const FVector InitialLocalOffset = TargetSkeleton.RetargetLocalPose[BoneIndex].GetTranslation();
+				OutPosition = ParentGlobalTransform.TransformPosition(InitialLocalOffset * StretchRatio);
 				break;
-			case ERetargetTranslationMode::GloballyScaled:
-				{
-					OutPosition = SourceCurrentTransform.GetTranslation() * RootRetargeter.GetGlobalScaleVector();
-				}
-				break;
-			case ERetargetTranslationMode::Absolute:
-				OutPosition = SourceCurrentTransform.GetTranslation();
+			}
+		case ERetargetTranslationMode::StretchBoneLengthNonUniformly:
+			{
+				const float Param = Params[ChainIndex];
+				float StretchRatio = SourceChain.GetStretchAtParam( SourceChain.InitialGlobalTransforms, SourceChain.CurrentGlobalTransforms, Param);
+				// stretch local translation
+				const FVector InitialLocalOffset = TargetSkeleton.RetargetLocalPose[BoneIndex].GetTranslation();
+				OutPosition = ParentGlobalTransform.TransformPosition(InitialLocalOffset * StretchRatio);
 				break;
+			}
 			default:
-				checkNoEntry();
+			{
+				const FVector InitialLocalOffset = TargetSkeleton.RetargetLocalPose[BoneIndex].GetTranslation();
+				OutPosition = ParentGlobalTransform.TransformPosition(InitialLocalOffset);
 				break;
+			}
 		}
 
 		// calculate output SCALE
@@ -800,6 +902,21 @@ void FChainDecoderFK::DecodePose(
 		InOutGlobalPose[BoneIndex] = CurrentGlobalTransforms[ChainIndex];
 	}
 
+	// apply match chain operation on-top of translated and rotated result
+	if (Settings.RotationMode == ERetargetRotationMode::MatchChain ||
+		Settings.RotationMode == ERetargetRotationMode::MatchScaledChain)
+	{
+		const bool bScaleSourceChain = Settings.RotationMode == ERetargetRotationMode::MatchScaledChain;
+		FVector TargetChainOrigin = InOutGlobalPose[TargetBoneIndices[0]].GetTranslation();
+		MatchChain(bScaleSourceChain, SourceChain, TargetBoneIndices, TargetChainOrigin);
+		
+		// update output pose
+		for (int32 ChainIndex=0; ChainIndex<TargetBoneIndices.Num(); ++ChainIndex)
+		{
+			InOutGlobalPose[TargetBoneIndices[ChainIndex]] = CurrentGlobalTransforms[ChainIndex];
+		}
+	}
+
 	// apply final blending between retarget pose of chain and newly retargeted pose
 	// blend must be done in local space, so we do it in a separate loop after full chain pose is generated
 	const bool bShouldBlendRotation = !FMath::IsNearlyEqual(Settings.RotationAlpha, 1.0f);
@@ -855,6 +972,178 @@ void FChainDecoderFK::InitializeIntermediateParentIndices(
 	Algo::Reverse(IntermediateParentIndices);
 }
 
+void FChainDecoderFK::MatchChain(
+	const bool bScaleSourceChain,
+	const FChainEncoderFK& SourceChain,
+	const TArray<int32>& TargetBoneIndices,
+	const FVector& TargetChainOrigin)
+{
+	// "MatchChain" mode assumes the interpolated rotations (and any translation mode) as a starting point.
+	//
+	// The "spline-IK-like" method used below generates swing rotations, which fix-up the interpolated rotations
+	// such that they align the bone positions to lie on the source chain.
+	//
+	// Typically, spline-ik is not "twist aware" because aligning a joint chain with swing rotations does not twist,
+	// but by using the interpolated rotations/translation as a starting place we have a full twist/bend/stretch.
+	
+	// convert source chain into a linear spline
+	TArray<FVector> SourceSplinePoints;
+	for (const FTransform& SourceGlobalTransform : SourceChain.CurrentGlobalTransforms)
+	{
+		SourceSplinePoints.Add(SourceGlobalTransform.GetTranslation());
+	}
+		
+	// translate chain points to origin of target chain
+	FVector ChainOffset = TargetChainOrigin - SourceSplinePoints[0];
+	for (FVector& SplinePoint : SourceSplinePoints)
+	{
+		SplinePoint += ChainOffset;
+	}
+
+	// scale chain to match target length
+	if (bScaleSourceChain)
+	{
+		float TargetChainLengthCurrent = GetChainLength(CurrentGlobalTransforms);
+		float SourceChainLengthCurrent = GetChainLength(SourceChain.CurrentGlobalTransforms);
+		float ScaleFactor = TargetChainLengthCurrent < KINDA_SMALL_NUMBER ? 1.0f : SourceChainLengthCurrent / TargetChainLengthCurrent;
+		for (FVector& SplinePoint : SourceSplinePoints)
+		{
+			SplinePoint =  TargetChainOrigin + (SplinePoint - TargetChainOrigin) * ScaleFactor;
+		}
+	}
+	
+	// snap orient each bone to lie on the source spline
+	int32 OriginPointIndex = 0;
+	double OriginPointAlpha = 0.0;
+	for (int32 ChainIndex=0; ChainIndex<TargetBoneIndices.Num()-1; ++ChainIndex)
+	{
+		// generate an aim rotation from the current joint vector to the point on the spline
+		const FVector StartBonePosition = CurrentGlobalTransforms[ChainIndex].GetTranslation();
+		const FVector EndBonePosition = CurrentGlobalTransforms[ChainIndex+1].GetTranslation();
+		const FVector EndBoneLocalPosition = CurrentGlobalTransforms[ChainIndex].InverseTransformPosition(EndBonePosition);
+		FVector BoneNorm;
+		double BoneLength;
+		(EndBonePosition - StartBonePosition).ToDirectionAndLength(BoneNorm, BoneLength);
+
+		// get the first point along the spline that is "BoneLength" away from the joint location in a straight line
+		GetPointOnSplineDistanceFromPoint(
+			SourceSplinePoints,
+			OriginPointIndex,
+			OriginPointAlpha,
+			BoneLength,
+			OriginPointIndex,
+			OriginPointAlpha);
+		
+		// convert spline coordinates to euclidean
+		const FVector AimPointOnSpline = GetPointOnSplineFromIndexAndAlpha(SourceSplinePoints, OriginPointIndex, OriginPointAlpha);
+		
+		// generate from/to swing rotation to align child with spline
+		const FVector AimNorm = (AimPointOnSpline - StartBonePosition).GetSafeNormal();
+		const FQuat RotationOffset = FQuat::FindBetweenNormals(BoneNorm, AimNorm);
+
+		// rotate bone to contact spline
+		CurrentGlobalTransforms[ChainIndex].SetRotation( RotationOffset * CurrentGlobalTransforms[ChainIndex].GetRotation());
+
+		// propagate translational offset of end bone to all it's children
+		const FVector NewEndBonePosition = CurrentGlobalTransforms[ChainIndex].TransformPosition(EndBoneLocalPosition);
+		const FVector DeltaTranslation = NewEndBonePosition - EndBonePosition;
+		for (int32 ChildIndex=ChainIndex+1; ChildIndex<TargetBoneIndices.Num(); ++ChildIndex)
+		{
+			CurrentGlobalTransforms[ChildIndex].AddToTranslation(DeltaTranslation);
+		}
+	}
+}
+
+FVector FChainDecoderFK::GetPointOnSplineFromIndexAndAlpha(
+	const TArray<FVector>& InSplinePoints,
+	const int32 InPointIndex,
+	const double InSegmentAlpha)
+{
+	if (!ensure(InPointIndex>=0))
+	{
+		return InSplinePoints[0];
+	}
+	
+	if (InPointIndex == InSplinePoints.Num() - 1)
+	{
+		// extrapolate SegmentAlpha distance beyond the last segment
+		const FVector& LastSegmentStart = InSplinePoints[InSplinePoints.Num() - 2];
+		const FVector& LastSegmentEnd = InSplinePoints.Last();
+		const FVector LastSegmentDirection = (LastSegmentEnd - LastSegmentStart).GetSafeNormal();
+		return LastSegmentEnd + LastSegmentDirection * InSegmentAlpha;
+	}
+
+	// interpolate between point and next point by alpha
+	return FMath::Lerp(InSplinePoints[InPointIndex], InSplinePoints[InPointIndex+1], InSegmentAlpha);
+}
+
+void FChainDecoderFK::GetPointOnSplineDistanceFromPoint(
+	const TArray<FVector>& InSplinePoints,
+	const int32 InOriginPointIndex,
+	const double InOriginPointAlpha,
+	const double InTargetDistanceFromOrigin,
+	int32& OutPointIndex,
+	double& OutPointAlpha)
+{
+	// NOTE: this returns the point on the ray, in RayDirection, that is DistanceFromPoint away from Point
+	// It assumes that there is a point on the ray within TargetDistanceFromPointToRay
+	auto GetPointOnRayDistanceFromPoint = [](
+		const FVector& RayStart,
+		const FVector& RayDirection,
+		const FVector& Point,
+		const double TargetDistanceFromPointToRay)
+	{
+		const FVector RayStartToPoint = Point - RayStart;
+		const double ProjectionLength = FVector::DotProduct(RayStartToPoint, RayDirection);
+
+		// get the point projected onto the line formed by the ray
+		const FVector ProjectedPoint = RayStart + (RayDirection * ProjectionLength);
+		const double DistancePointToRay = (ProjectedPoint - Point).Size();
+
+		// check if point is further away from ray origin than TargetDistanceFromPointToRay (should not happen)
+		if (!ensure(DistancePointToRay < TargetDistanceFromPointToRay))
+		{
+			return (RayStart - Point).GetClampedToMaxSize(TargetDistanceFromPointToRay);
+		}
+	
+		// pythagorean theorem to find distance from projected point to point on ray
+		// NOTE: imagine the right angle triangle formed by Point, ProjectedPoint and the unknown point on the ray
+		const double DistanceFromProjPointToTargetPoint = FMath::Sqrt ((TargetDistanceFromPointToRay * TargetDistanceFromPointToRay) - (DistancePointToRay * DistancePointToRay));
+
+		// calculate the point on the ray that is DistanceFromPoint away from Point
+		return ProjectedPoint + RayDirection * DistanceFromProjPointToTargetPoint;
+	};
+
+	// convert spline coordinates to euclidean
+	const FVector OriginPoint = GetPointOnSplineFromIndexAndAlpha(InSplinePoints, InOriginPointIndex, InOriginPointAlpha);
+	
+	// iterate down the chain until we find a bone that is beyond TargetDistanceFromOrigin away from the origin point
+	for (OutPointIndex = InOriginPointIndex; OutPointIndex < InSplinePoints.Num()-1; ++OutPointIndex)
+	{
+		const FVector& SegmentEnd = InSplinePoints[OutPointIndex+1];
+		const double DistanceOriginToSegmentEnd = (SegmentEnd - OriginPoint).Size();
+		if (DistanceOriginToSegmentEnd < InTargetDistanceFromOrigin)
+		{
+			continue;
+		}
+		
+		const FVector& RayStart = InSplinePoints[OutPointIndex];
+		FVector RayDirection;
+		double SegmentLength;
+		(SegmentEnd - RayStart).ToDirectionAndLength(RayDirection, SegmentLength);
+		const FVector PointOnRay = GetPointOnRayDistanceFromPoint(RayStart, RayDirection, OriginPoint, InTargetDistanceFromOrigin);
+		OutPointAlpha = (PointOnRay - RayStart).Size() / SegmentLength;
+		return;
+	}
+
+	// spline is too short, so extrapolate the last spline segment and find the point on that segment that is DistanceFromOrigin away
+	const FVector& RayStart = InSplinePoints.Last();
+	const FVector RayDirection = (InSplinePoints.Last() - InSplinePoints[InSplinePoints.Num() - 2]).GetSafeNormal();
+	const FVector ExtrapolatedPointOnRay = GetPointOnRayDistanceFromPoint(RayStart, RayDirection, OriginPoint, InTargetDistanceFromOrigin);
+	OutPointAlpha = (ExtrapolatedPointOnRay - RayStart).Size();
+	OutPointIndex = InSplinePoints.Num() - 1;
+}
+
 void FChainDecoderFK::UpdateIntermediateParents(
 	const FTargetSkeleton& TargetSkeleton,
 	TArray<FTransform>& InOutGlobalPose)
@@ -2458,7 +2747,7 @@ FTransform UIKRetargetProcessor::GetGlobalRetargetPoseAtParam(
 {
 	if (const FChainFK* Chain = GetChain(InChainName, SourceOrTarget))
 	{
-		return Chain->GetTransformAtParam(Chain->InitialGlobalTransforms, Param);
+		return Chain->GetTransformAtChainParam(Chain->InitialGlobalTransforms, Param);
 	}
 
 	checkNoEntry();
diff --git a/Engine/Plugins/Animation/IKRig/Source/IKRig/Public/Retargeter/IKRetargetProcessor.h b/Engine/Plugins/Animation/IKRig/Source/IKRig/Public/Retargeter/IKRetargetProcessor.h
--- a/Engine/Plugins/Animation/IKRig/Source/IKRig/Public/Retargeter/IKRetargetProcessor.h
+++ b/Engine/Plugins/Animation/IKRig/Source/IKRig/Public/Retargeter/IKRetargetProcessor.h
@@ -251,7 +251,16 @@ struct FChainFK
 		const TArray<FTransform> &InitialGlobalPose,
 		FIKRigLogger& Log);
 
-	FTransform GetTransformAtParam(const TArray<FTransform>& Transforms, const float& Param) const;
+	FTransform GetTransformAtChainParam(
+		const TArray<FTransform>& Transforms,
+		const float& Param) const;
+	
+	float GetStretchAtParam(
+		const TArray<FTransform>& InitialTransforms,
+		const TArray<FTransform>& CurrentTransforms,
+		const float& Param);
+	
+	static float GetChainLength(const TArray<FTransform>& Transforms);
 
 private:
 	
@@ -307,6 +316,22 @@ struct FChainDecoderFK : public FChainFK
 		TArray<FTransform> &InOutGlobalPose);
 	
 private:
+
+	void MatchChain(
+		const bool bScaleSourceChain,
+		const FChainEncoderFK& SourceChain,
+		const TArray<int32>& TargetBoneIndices,
+		const FVector& TargetChainOrigin);
+
+	static FVector GetPointOnSplineFromIndexAndAlpha(const TArray<FVector>& InSplinePoints, const int32 InPointIndex, const double InSegmentAlpha);
+	
+	static void GetPointOnSplineDistanceFromPoint(
+		const TArray<FVector>& InSplinePoints,
+		const int32 InOriginPointIndex,
+		const double InOriginPointAlpha,
+		const double InTargetDistanceFromOrigin,
+		int32& OutPointIndex,
+		double& OutPointAlpha);
 	
 	void UpdateIntermediateParents(
 		const FTargetSkeleton& TargetSkeleton,
diff --git a/Engine/Plugins/Animation/IKRig/Source/IKRig/Public/Retargeter/IKRetargetSettings.h b/Engine/Plugins/Animation/IKRig/Source/IKRig/Public/Retargeter/IKRetargetSettings.h
--- a/Engine/Plugins/Animation/IKRig/Source/IKRig/Public/Retargeter/IKRetargetSettings.h
+++ b/Engine/Plugins/Animation/IKRig/Source/IKRig/Public/Retargeter/IKRetargetSettings.h
@@ -10,18 +10,21 @@ class URetargetChainSettings;
 UENUM(BlueprintType)
 enum class ERetargetTranslationMode : uint8
 {
-	None			UMETA(DisplayName = "None"),
-	GloballyScaled	UMETA(DisplayName = "Globally Scaled"),
-	Absolute		UMETA(DisplayName = "Absolute"),
+	None							UMETA(DisplayName = "None"),
+	GloballyScaled					UMETA(DisplayName = "Globally Scaled"),
+	Absolute						UMETA(DisplayName = "Absolute"),
+	StretchBoneLengthUniformly		UMETA(DisplayName = "Stretch Bone Length Uniformly"),
+	StretchBoneLengthNonUniformly	UMETA(DisplayName = "Stretch Bone Length Non-Uniformly"),
 };
 
 UENUM(BlueprintType)
 enum class ERetargetRotationMode : uint8
 {
-	
 	Interpolated		UMETA(DisplayName = "Interpolated"),
 	OneToOne			UMETA(DisplayName = "One to One"),
 	OneToOneReversed	UMETA(DisplayName = "One to One Reversed"),
+	MatchChain			UMETA(DisplayName = "Match Chain"),
+	MatchScaledChain	UMETA(DisplayName = "Match Scaled Chain"),
 	None				UMETA(DisplayName = "None"),
 };
 
@@ -68,6 +71,8 @@ struct IKRIG_API FTargetChainFKSettings
 	* Interpolated: Source and target chains are normalized by their length, then each target bone rotation is generated by finding the rotation at the same normalized distance on the source chain and interpolating between the neighboring bones.
 	* One to One: Each target bone rotation is copied from the equivalent bone in the source chain, based on the order in the chain, starting at the root of the chain. If the target chain has more bones than the source, the extra bones will remain at their reference pose.
 	* One to One Reversed: Same as One-to-One, but starting from the tip of the chain.
+	* Match Chain: Uses a Spline-IK type method to orient the target joint chain to exactly match the shape of the source chain.
+	* Match Scaled Chain: Same as Match Chain, but scales the source chain to the same length as the target chain before running the match procedure.
 	* None: The rotation of each target bone in the chain is left at the reference pose. */
 	UPROPERTY(EditAnywhere, BlueprintReadWrite, Category = "FK Adjustments")
 	ERetargetRotationMode RotationMode = ERetargetRotationMode::Interpolated;
@@ -81,7 +86,10 @@ struct IKRIG_API FTargetChainFKSettings
 	/** Determines how translation is copied from the source chain to the target chain. Default is None.
 	* None: Translation of target bones are left unmodified from the retarget pose.
 	* Globally Scaled: Translation of target bone is set to the source bone offset multiplied by the global scale of the skeleton (determined by the relative height difference between retarget root bones).
-	* Absolute: Translation of target bone is set to the absolute position of the source bone. */
+	* Absolute: Translation of target bone is set to the absolute position of the source bone.
+	* Stretch Bone Length Uniformly: squashes or stretches the local translation of each bone so that the length of the target chain matches the relative length of the source chain. Accounts for differences in initial chain lengths.
+	* Stretch Bone Length Non-Uniformly: same as "Stretch Bone Length Uniformly" but treats each bone independently to allow for chains that stretch non-uniformly along their length (ie the base may stretch more than the tip or vice versa).   
+	 */
 	UPROPERTY(EditAnywhere, BlueprintReadWrite, Category = "FK Adjustments")
 	ERetargetTranslationMode TranslationMode = ERetargetTranslationMode::None;
 
