Fixing incremental build issue from IKRetargetProcessor.cpp

diff --git a/Engine/Plugins/Animation/IKRig/Source/IKRig/Private/Retargeter/IKRetargetProcessor.cpp b/Engine/Plugins/Animation/IKRig/Source/IKRig/Private/Retargeter/IKRetargetProcessor.cpp
--- a/Engine/Plugins/Animation/IKRig/Source/IKRig/Private/Retargeter/IKRetargetProcessor.cpp
+++ b/Engine/Plugins/Animation/IKRig/Source/IKRig/Private/Retargeter/IKRetargetProcessor.cpp
@@ -523,7 +523,7 @@ FTransform FChainFK::GetTransformAtChainParam(const TArray<FTransform>& Transfor
 float FChainFK::GetStretchAtParam(
 	const TArray<FTransform>& InitialTransforms,
 	const TArray<FTransform>& CurrentTransforms,
-	const float& Param)
+	const double& Param)
 {
 	check(InitialTransforms.Num() == CurrentTransforms.Num() && InitialTransforms.Num() == Params.Num());
 
@@ -534,23 +534,23 @@ float FChainFK::GetStretchAtParam(
 	}
 
 	// start of chain cannot stretch
-	if (Param < KINDA_SMALL_NUMBER)
+	if (Param < UE_KINDA_SMALL_NUMBER)
 	{
 		return 1.0f;
 	}
 	
 	// end of chain always uses last bone
-	if (Param > 1.0f - KINDA_SMALL_NUMBER)
+	if (Param > 1.0 - UE_KINDA_SMALL_NUMBER)
 	{
 		const int32 Last = InitialTransforms.Num() - 1;
 		const int32 Prev = Last - 1;
-		const float InitialLength = (InitialTransforms[Last].GetTranslation() - InitialTransforms[Prev].GetTranslation()).Length();
-		const float CurrentLength = (CurrentTransforms[Last].GetTranslation() - CurrentTransforms[Prev].GetTranslation()).Length();
-		return InitialLength < KINDA_SMALL_NUMBER ? 1.0f : CurrentLength / InitialLength;
+		const double InitialLength = (InitialTransforms[Last].GetTranslation() - InitialTransforms[Prev].GetTranslation()).Length();
+		const double CurrentLength = (CurrentTransforms[Last].GetTranslation() - CurrentTransforms[Prev].GetTranslation()).Length();
+		return InitialLength < UE_KINDA_SMALL_NUMBER ? 1.0f : CurrentLength / InitialLength;
 	}
 
 	// find the bone associated with this parameter
-	float ClosestParamDistance = TNumericLimits<float>::Max();
+	double ClosestParamDistance = TNumericLimits<double>::Max();
 	int32 ClosestBoneChainIndex = INDEX_NONE;
 	for (int32 ChainIndex=0; ChainIndex<Params.Num(); ++ChainIndex)
 	{
@@ -572,9 +572,9 @@ float FChainFK::GetStretchAtParam(
 	}
 
 	// calculate the ratio of the current to the initial length of the closest bone
-	const float InitialLength = (InitialTransforms[BoneStartIndex].GetTranslation() - InitialTransforms[BoneEndIndex].GetTranslation()).Length();
-	const float CurrentLength = (CurrentTransforms[BoneStartIndex].GetTranslation() - CurrentTransforms[BoneEndIndex].GetTranslation()).Length();
-	return InitialLength < KINDA_SMALL_NUMBER ? 1.0f : CurrentLength / InitialLength;
+	const double InitialLength = (InitialTransforms[BoneStartIndex].GetTranslation() - InitialTransforms[BoneEndIndex].GetTranslation()).Length();
+	const double CurrentLength = (CurrentTransforms[BoneStartIndex].GetTranslation() - CurrentTransforms[BoneEndIndex].GetTranslation()).Length();
+	return InitialLength < UE_KINDA_SMALL_NUMBER ? 1.0f : CurrentLength / InitialLength;
 }
 
 bool FChainFK::CalculateBoneParameters(FIKRigLogger& Log)
@@ -614,9 +614,9 @@ bool FChainFK::CalculateBoneParameters(FIKRigLogger& Log)
 	return true;
 }
 
-float FChainFK::GetChainLength(const TArray<FTransform>& Transforms)
+double FChainFK::GetChainLength(const TArray<FTransform>& Transforms)
 {
-	float TotalLength = 0.0f;
+	double TotalLength = 0.0f;
 	for (int32 ChainIndex = 1; ChainIndex < Transforms.Num(); ++ChainIndex)
 	{
 		const FVector Start = Transforms[ChainIndex - 1].GetTranslation();
@@ -866,9 +866,9 @@ void FChainDecoderFK::DecodePose(
 					break;
 				}
 				// initial chain length
-				float SourceChainLengthInitial = GetChainLength(SourceChain.InitialGlobalTransforms);
-				float SourceChainLengthCurrent = GetChainLength(SourceChain.CurrentGlobalTransforms);
-				float StretchRatio = SourceChainLengthInitial < KINDA_SMALL_NUMBER ? 1.0f : SourceChainLengthCurrent / SourceChainLengthInitial;
+				double SourceChainLengthInitial = GetChainLength(SourceChain.InitialGlobalTransforms);
+				double SourceChainLengthCurrent = GetChainLength(SourceChain.CurrentGlobalTransforms);
+				double StretchRatio = SourceChainLengthInitial < UE_KINDA_SMALL_NUMBER ? 1.0f : SourceChainLengthCurrent / SourceChainLengthInitial;
 				// stretch local translation
 				const FVector InitialLocalOffset = TargetSkeleton.RetargetLocalPose[BoneIndex].GetTranslation();
 				OutPosition = ParentGlobalTransform.TransformPosition(InitialLocalOffset * StretchRatio);
@@ -876,7 +876,7 @@ void FChainDecoderFK::DecodePose(
 			}
 		case ERetargetTranslationMode::StretchBoneLengthNonUniformly:
 			{
-				const float Param = Params[ChainIndex];
+				const double Param = Params[ChainIndex];
 				float StretchRatio = SourceChain.GetStretchAtParam( SourceChain.InitialGlobalTransforms, SourceChain.CurrentGlobalTransforms, Param);
 				// stretch local translation
 				const FVector InitialLocalOffset = TargetSkeleton.RetargetLocalPose[BoneIndex].GetTranslation();
@@ -1003,9 +1003,9 @@ void FChainDecoderFK::MatchChain(
 	// scale chain to match target length
 	if (bScaleSourceChain)
 	{
-		float TargetChainLengthCurrent = GetChainLength(CurrentGlobalTransforms);
-		float SourceChainLengthCurrent = GetChainLength(SourceChain.CurrentGlobalTransforms);
-		float ScaleFactor = TargetChainLengthCurrent < KINDA_SMALL_NUMBER ? 1.0f : SourceChainLengthCurrent / TargetChainLengthCurrent;
+		double TargetChainLengthCurrent = GetChainLength(CurrentGlobalTransforms);
+		double SourceChainLengthCurrent = GetChainLength(SourceChain.CurrentGlobalTransforms);
+		double ScaleFactor = TargetChainLengthCurrent < UE_KINDA_SMALL_NUMBER ? 1.0f : SourceChainLengthCurrent / TargetChainLengthCurrent;
 		for (FVector& SplinePoint : SourceSplinePoints)
 		{
 			SplinePoint =  TargetChainOrigin + (SplinePoint - TargetChainOrigin) * ScaleFactor;
diff --git a/Engine/Plugins/Animation/IKRig/Source/IKRig/Public/Retargeter/IKRetargetProcessor.h b/Engine/Plugins/Animation/IKRig/Source/IKRig/Public/Retargeter/IKRetargetProcessor.h
--- a/Engine/Plugins/Animation/IKRig/Source/IKRig/Public/Retargeter/IKRetargetProcessor.h
+++ b/Engine/Plugins/Animation/IKRig/Source/IKRig/Public/Retargeter/IKRetargetProcessor.h
@@ -258,9 +258,9 @@ struct FChainFK
 	float GetStretchAtParam(
 		const TArray<FTransform>& InitialTransforms,
 		const TArray<FTransform>& CurrentTransforms,
-		const float& Param);
+		const double& Param);
 	
-	static float GetChainLength(const TArray<FTransform>& Transforms);
+	static double GetChainLength(const TArray<FTransform>& Transforms);
 
 private:
 	
