IK Rig: Fix truncation warnings (part 2)

#rnx

diff --git a/Engine/Plugins/Animation/IKRig/Source/IKRig/IKRig.Build.cs b/Engine/Plugins/Animation/IKRig/Source/IKRig/IKRig.Build.cs
--- a/Engine/Plugins/Animation/IKRig/Source/IKRig/IKRig.Build.cs
+++ b/Engine/Plugins/Animation/IKRig/Source/IKRig/IKRig.Build.cs
@@ -7,6 +7,8 @@ public class IKRig : ModuleRules
 	public IKRig(ReadOnlyTargetRules Target) : base(Target)
 	{
 		PCHUsage = ModuleRules.PCHUsageMode.UseExplicitOrSharedPCHs;
+
+		CppCompileWarningSettings.UnsafeTypeCastWarningLevel = WarningLevel.Error;
 		
 		PrivateDependencyModuleNames.AddRange(
 			new string[]
diff --git a/Engine/Plugins/Animation/IKRig/Source/IKRig/Private/Retargeter/IKRetargetProcessor.cpp b/Engine/Plugins/Animation/IKRig/Source/IKRig/Private/Retargeter/IKRetargetProcessor.cpp
--- a/Engine/Plugins/Animation/IKRig/Source/IKRig/Private/Retargeter/IKRetargetProcessor.cpp
+++ b/Engine/Plugins/Animation/IKRig/Source/IKRig/Private/Retargeter/IKRetargetProcessor.cpp
@@ -546,7 +546,6 @@ float FChainFK::GetStretchAtParam(
 		const int32 Prev = Last - 1;
 		const double InitialLength = (InitialTransforms[Last].GetTranslation() - InitialTransforms[Prev].GetTranslation()).Length();
 		const double CurrentLength = (CurrentTransforms[Last].GetTranslation() - CurrentTransforms[Prev].GetTranslation()).Length();
-		return InitialLength < UE_KINDA_SMALL_NUMBER ? 1.0f : CurrentLength / InitialLength;
 	}
 
 	// find the bone associated with this parameter
@@ -554,7 +553,7 @@ float FChainFK::GetStretchAtParam(
 	int32 ClosestBoneChainIndex = INDEX_NONE;
 	for (int32 ChainIndex=0; ChainIndex<Params.Num(); ++ChainIndex)
 	{
-		float ParamDistance = FMath::Abs(Params[ChainIndex] - Param);
+		double ParamDistance = FMath::Abs(Params[ChainIndex] - Param);
 		if (ParamDistance < ClosestParamDistance)
 		{
 			ClosestParamDistance = ParamDistance;
@@ -574,7 +573,7 @@ float FChainFK::GetStretchAtParam(
 	// calculate the ratio of the current to the initial length of the closest bone
 	const double InitialLength = (InitialTransforms[BoneStartIndex].GetTranslation() - InitialTransforms[BoneEndIndex].GetTranslation()).Length();
 	const double CurrentLength = (CurrentTransforms[BoneStartIndex].GetTranslation() - CurrentTransforms[BoneEndIndex].GetTranslation()).Length();
-	return InitialLength < UE_KINDA_SMALL_NUMBER ? 1.0f : CurrentLength / InitialLength;
+	return static_cast<float>(InitialLength < UE_KINDA_SMALL_NUMBER ? 1.0 : CurrentLength / InitialLength);
 }
 
 bool FChainFK::CalculateBoneParameters(FIKRigLogger& Log)
@@ -621,7 +620,7 @@ double FChainFK::GetChainLength(const TArray<FTransform>& Transforms)
 	{
 		const FVector Start = Transforms[ChainIndex - 1].GetTranslation();
 		const FVector End = Transforms[ChainIndex].GetTranslation();
-		TotalLength += FVector::Dist(Start, End);
+		TotalLength += static_cast<float>(FVector::Dist(Start, End));
 	}
 	return TotalLength;
 }
diff --git a/Engine/Plugins/Animation/IKRig/Source/IKRig/Private/Rig/IKRigSkeleton.cpp b/Engine/Plugins/Animation/IKRig/Source/IKRig/Private/Rig/IKRigSkeleton.cpp
--- a/Engine/Plugins/Animation/IKRig/Source/IKRig/Private/Rig/IKRigSkeleton.cpp
+++ b/Engine/Plugins/Animation/IKRig/Source/IKRig/Private/Rig/IKRigSkeleton.cpp
@@ -476,7 +476,7 @@ bool FIKRigSkeleton::GetMirroredBoneIndices(
 			FString CurrentNameStr = BoneNames[BoneIndex].ToString().ToLower();
 			float WorstCase = static_cast<float>(BaseNameStr.Len() + CurrentNameStr.Len());
 			WorstCase = WorstCase < 1.0f ? 1.0f : WorstCase;
-			const float Score = 1.0f - (Algo::LevenshteinDistance(BaseNameStr, CurrentNameStr) / WorstCase);
+			const float Score = 1.0f - (static_cast<float>(Algo::LevenshteinDistance(BaseNameStr, CurrentNameStr)) / WorstCase);
 			if (Score > HighestScore)
 			{
 				HighestScore = Score;
diff --git a/Engine/Source/Runtime/Engine/Public/Animation/NamedValueArray.h b/Engine/Source/Runtime/Engine/Public/Animation/NamedValueArray.h
--- a/Engine/Source/Runtime/Engine/Public/Animation/NamedValueArray.h
+++ b/Engine/Source/Runtime/Engine/Public/Animation/NamedValueArray.h
@@ -351,7 +351,7 @@ struct FNamedValueArrayUtils
 		{
 			// Reached end of ValueArray0 with remaining in ValueArray1, we can just copy the remainder of ValueArray1
 			const int32 NumResults = InOutValueArray0.Elements.Num();
-			const int32 NumNewElements = (ElementEndPtr1 - ElementPtr1);
+			const int32 NumNewElements = static_cast<int32>(ElementEndPtr1 - ElementPtr1);
 			InOutValueArray0.Elements.Reserve(NumResults + NumNewElements);
 			InOutValueArray0.Elements.AddUninitialized(NumNewElements);
 
