IK Retargeter: Fixes bug where source scale factor was not respected by PinBoneOp

diff --git a/Engine/Plugins/Animation/IKRig/Source/IKRig/Private/Retargeter/RetargetOps/PinBoneOp.cpp b/Engine/Plugins/Animation/IKRig/Source/IKRig/Private/Retargeter/RetargetOps/PinBoneOp.cpp
--- a/Engine/Plugins/Animation/IKRig/Source/IKRig/Private/Retargeter/RetargetOps/PinBoneOp.cpp
+++ b/Engine/Plugins/Animation/IKRig/Source/IKRig/Private/Retargeter/RetargetOps/PinBoneOp.cpp
@@ -38,6 +38,39 @@ const FTransform& FPinBoneData::GetParentTransform(int32 BoneIndex, const FRetar
 	return Pose[ParentIndex];
 }
 
+void FPinBoneData::CachePinData(ERetargetSourceOrTarget SkeletonToPinTo, const UIKRetargetProcessor* Processor)
+{
+	SourceScaleFactorUsedForCachedData = Processor->GetGlobalSettings().SourceScaleFactor;
+	
+	auto ScaleSourceGlobalTransform = [this](FTransform& InOutTransformToScale)
+    {
+    	InOutTransformToScale.SetTranslation( InOutTransformToScale.GetTranslation() * SourceScaleFactorUsedForCachedData );
+    };
+
+    // get skeletons we are copying from/to
+    const FRetargetSkeleton& SkeletonToCopyFrom = Processor->GetSkeleton(SkeletonToPinTo);
+    const FRetargetSkeleton& SkeletonToCopyTo = Processor->GetSkeleton(ERetargetSourceOrTarget::Target);
+
+	// get relevant transforms from the retarget poses
+    FTransform BoneToCopyFromRefPoseGlobal = SkeletonToCopyFrom.RetargetGlobalPose[BoneToCopyFrom.BoneIndex];
+    FTransform BoneToCopyToRefPoseGlobal = SkeletonToCopyTo.RetargetGlobalPose[BoneToCopyTo.BoneIndex];
+	FTransform BoneToCopyFromParentRefPoseGlobal = GetParentTransform(BoneToCopyFrom.BoneIndex, SkeletonToCopyFrom, SkeletonToCopyFrom.RetargetGlobalPose);
+	FTransform BoneToCopyToParentRefPoseGlobal = GetParentTransform(BoneToCopyFrom.BoneIndex, SkeletonToCopyFrom, SkeletonToCopyFrom.RetargetGlobalPose);
+	
+	// scale transforms according to the source scale factor
+    ScaleSourceGlobalTransform(BoneToCopyFromRefPoseGlobal);
+    ScaleSourceGlobalTransform(BoneToCopyToRefPoseGlobal);
+	ScaleSourceGlobalTransform(BoneToCopyFromParentRefPoseGlobal);
+	ScaleSourceGlobalTransform(BoneToCopyToParentRefPoseGlobal);
+
+	// cache offset from BoneToCopyFrom to BoneToCopyTo in ref pose
+    OffsetFromBoneToCopyFromInRefPose = BoneToCopyFromRefPoseGlobal.GetRelativeTransform(BoneToCopyToRefPoseGlobal);
+	// cache local transform of BoneToCopyFrom in ref pose
+    LocalRefPoseBoneToCopyFrom = BoneToCopyFromParentRefPoseGlobal.GetRelativeTransform(BoneToCopyFromRefPoseGlobal);
+	// cache local transform of BoneToCopyTo in ref pose
+    LocalRefPoseBoneToCopyTo = BoneToCopyToParentRefPoseGlobal.GetRelativeTransform(BoneToCopyToRefPoseGlobal);
+}
+
 bool UPinBoneOp::Initialize(
 	const UIKRetargetProcessor* Processor,
 	const FRetargetSkeleton& SourceSkeleton,
@@ -76,25 +109,8 @@ bool UPinBoneOp::Initialize(
 				FText::FromName(PinData.BoneToCopyFrom.BoneName)));
 		}
 
-		if (!(bFoundBoneToPin && bFoundBoneToPinTo))
-		{
-			continue;
-		}
-
-		// cache offsets
-		const TArray<FTransform>& PoseToCopyFrom = PinTo == ERetargetSourceOrTarget::Source ? SourceSkeleton.RetargetGlobalPose : TargetSkeleton.RetargetGlobalPose;
-		const FRetargetSkeleton& SkeletonToCopyFrom = Processor->GetSkeleton(PinTo);
-		const FRetargetSkeleton& SkeletonToCopyTo = Processor->GetSkeleton(ERetargetSourceOrTarget::Target);
-		const FTransform& BoneToCopyFromRefPoseGlobal = SkeletonToCopyFrom.RetargetGlobalPose[PinData.BoneToCopyFrom.BoneIndex];
-		const FTransform& BoneToCopyToRefPoseGlobal = SkeletonToCopyTo.RetargetGlobalPose[PinData.BoneToCopyTo.BoneIndex];
-		// cache offset from BoneToCopyFrom to BoneToCopyTo in ref pose
-		PinData.OffsetFromBoneToCopyFromInRefPose = BoneToCopyFromRefPoseGlobal.GetRelativeTransform(BoneToCopyToRefPoseGlobal);
-		// cache local transform of BoneToCopyFrom in ref pose
-		const FTransform& BoneToCopyFromParentRefPoseGlobal = FPinBoneData::GetParentTransform(PinData.BoneToCopyFrom.BoneIndex, SkeletonToCopyFrom, PoseToCopyFrom);
-		PinData.LocalRefPoseBoneToCopyFrom = BoneToCopyFromParentRefPoseGlobal.GetRelativeTransform(BoneToCopyFromRefPoseGlobal);
-		// cache local transform of BoneToCopyTo in ref pose
-		const FTransform& BoneToCopyToParentRefPoseGlobal = FPinBoneData::GetParentTransform(PinData.BoneToCopyFrom.BoneIndex, SkeletonToCopyFrom, PoseToCopyFrom);
-		PinData.LocalRefPoseBoneToCopyTo = BoneToCopyToParentRefPoseGlobal.GetRelativeTransform(BoneToCopyToRefPoseGlobal);
+		// force reinitialization
+		PinData.SourceScaleFactorUsedForCachedData = -1.0f;
 	}
 
 #if WITH_EDITOR
@@ -123,7 +139,7 @@ void UPinBoneOp::Run(
 {
 	const TArray<FTransform>& PoseToCopyFrom = PinTo == ERetargetSourceOrTarget::Source ? InSourceGlobalPose : OutTargetGlobalPose;
 	
-	for (const FPinBoneData& PinData : BonesToPin)
+	for (FPinBoneData& PinData : BonesToPin)
 	{
 		bool bValidBoneToCopyFrom = PoseToCopyFrom.IsValidIndex(PinData.BoneToCopyFrom.BoneIndex);
 		bool bValidBoneToCopyTo = OutTargetGlobalPose.IsValidIndex(PinData.BoneToCopyTo.BoneIndex);
@@ -132,6 +148,12 @@ void UPinBoneOp::Run(
 			continue; // disabled or not successfully initialized
 		}
 
+		// recache the offsets if the source is scaled differently
+		if (!FMath::IsNearlyEqual(PinData.SourceScaleFactorUsedForCachedData, Processor->GetGlobalSettings().SourceScaleFactor))
+		{
+			PinData.CachePinData(PinTo, Processor);
+		}
+
 		// calculate new transform for bone to pin
 		const FTransform NewTransform = GetNewBoneTransform(PinData, Processor, InSourceGlobalPose, OutTargetGlobalPose);
 		// apply static local and global offsets
diff --git a/Engine/Plugins/Animation/IKRig/Source/IKRig/Public/Retargeter/RetargetOps/PinBoneOp.h b/Engine/Plugins/Animation/IKRig/Source/IKRig/Public/Retargeter/RetargetOps/PinBoneOp.h
--- a/Engine/Plugins/Animation/IKRig/Source/IKRig/Public/Retargeter/RetargetOps/PinBoneOp.h
+++ b/Engine/Plugins/Animation/IKRig/Source/IKRig/Public/Retargeter/RetargetOps/PinBoneOp.h
@@ -27,6 +27,8 @@ struct FPinBoneData
 	
 	static const FTransform& GetParentTransform(int32 BoneIndex, const FRetargetSkeleton& Skeleton, const TArray<FTransform>& Pose);
 
+	void CachePinData(ERetargetSourceOrTarget SkeletonToPinTo, const UIKRetargetProcessor* Processor);
+
 	// The bone to be affected by this op. Will have its transform modified to match the BoneToCopyFrom
 	UPROPERTY(EditAnywhere, Category=Settings)
 	FBoneReference BoneToCopyTo;
@@ -34,6 +36,8 @@ struct FPinBoneData
 	UPROPERTY(EditAnywhere, Category=Settings)
 	FBoneReference BoneToCopyFrom;
 
+	// the scale factor used to calculate all the cached offset (lazy evaluated when scale factor is modified)
+	float SourceScaleFactorUsedForCachedData;
 	// the local transform of the BoneToCopyTo in the ref pose
 	FTransform LocalRefPoseBoneToCopyTo;
 	// the local transform of the BoneToCopyTo in the ref pose
