IK Retargter: Fixed bug auto-retarget window filtering anim blueprints incorrectly

diff --git a/Engine/Plugins/Animation/IKRig/Source/IKRigEditor/Private/RetargetEditor/SRetargetAnimAssetsWindow.cpp b/Engine/Plugins/Animation/IKRig/Source/IKRigEditor/Private/RetargetEditor/SRetargetAnimAssetsWindow.cpp
--- a/Engine/Plugins/Animation/IKRig/Source/IKRigEditor/Private/RetargetEditor/SRetargetAnimAssetsWindow.cpp
+++ b/Engine/Plugins/Animation/IKRig/Source/IKRigEditor/Private/RetargetEditor/SRetargetAnimAssetsWindow.cpp
@@ -895,7 +895,7 @@ bool SRetargetExporterAssetBrowser::OnShouldFilterAsset(const FAssetData& AssetD
 			return true;
 		}
 
-		return DesiredSkeleton->IsCompatibleForEditor(Result.GetValue());
+		return !DesiredSkeleton->IsCompatibleForEditor(Result.GetValue());
 	}
 
 	return !DesiredSkeleton->IsCompatibleForEditor(AssetData);
