IKRig - fix for reading AlphaValue correctly in FAnimNode_IKRig::InitializeProperties

diff --git a/Engine/Plugins/Animation/IKRig/Source/IKRig/Private/AnimNodes/AnimNode_IKRig.cpp b/Engine/Plugins/Animation/IKRig/Source/IKRig/Private/AnimNodes/AnimNode_IKRig.cpp
--- a/Engine/Plugins/Animation/IKRig/Source/IKRig/Private/AnimNodes/AnimNode_IKRig.cpp
+++ b/Engine/Plugins/Animation/IKRig/Source/IKRig/Private/AnimNodes/AnimNode_IKRig.cpp
@@ -409,7 +409,7 @@ void FAnimNode_IKRig::InitializeProperties(const UObject* InSourceInstance, UCla
 			UpdateFunctions.Add([&Goal, SourceProperty](const UObject* InSourceInstance)
 			{
 				static_assert(std::is_same_v<decltype(FIKRigGoal::PositionAlpha), float>);
-				const float* AlphaValue = SourceProperty->ContainerPtrToValuePtr<float>(InSourceInstance);
+				const double* AlphaValue = SourceProperty->ContainerPtrToValuePtr<double>(InSourceInstance);
 				Goal.PositionAlpha = FMath::Clamp<float>(static_cast<float>(*AlphaValue), 0.f, 1.f);
 			});
 		}
@@ -418,8 +418,8 @@ void FAnimNode_IKRig::InitializeProperties(const UObject* InSourceInstance, UCla
 			UpdateFunctions.Add([&Goal, SourceProperty](const UObject* InSourceInstance)
 			{
 				static_assert(std::is_same_v<decltype(FIKRigGoal::RotationAlpha), float>);
-				const float* AlphaValue = SourceProperty->ContainerPtrToValuePtr<float>(InSourceInstance);
-				Goal.RotationAlpha = FMath::Clamp<float>(*AlphaValue, 0.f, 1.f);
+				const double* AlphaValue = SourceProperty->ContainerPtrToValuePtr<double>(InSourceInstance);
+				Goal.RotationAlpha = FMath::Clamp<float>(static_cast<float>(*AlphaValue), 0.f, 1.f);
 			});
 		}
 		else if (GoalPropStr.StartsWith(PositionPropStr))
