IK Retargeter: PinBoneOp reordered properties and adjusted tooltips

diff --git a/Engine/Plugins/Animation/IKRig/Source/IKRig/Public/Retargeter/RetargetOps/PinBoneOp.h b/Engine/Plugins/Animation/IKRig/Source/IKRig/Public/Retargeter/RetargetOps/PinBoneOp.h
--- a/Engine/Plugins/Animation/IKRig/Source/IKRig/Public/Retargeter/RetargetOps/PinBoneOp.h
+++ b/Engine/Plugins/Animation/IKRig/Source/IKRig/Public/Retargeter/RetargetOps/PinBoneOp.h
@@ -29,12 +29,12 @@ struct FPinBoneData
 
 	void CachePinData(ERetargetSourceOrTarget SkeletonToPinTo, const UIKRetargetProcessor* Processor);
 
-	// The bone to be affected by this op. Will have its transform modified to match the BoneToCopyFrom
-	UPROPERTY(EditAnywhere, Category=Settings)
-	FBoneReference BoneToCopyTo;
-	// The bone to copy transform from.
+	// The bone to copy FROM.
 	UPROPERTY(EditAnywhere, Category=Settings)
 	FBoneReference BoneToCopyFrom;
+	// The bone to copy TO. Will have its transform pinned to BoneToCopyFrom
+	UPROPERTY(EditAnywhere, Category=Settings)
+	FBoneReference BoneToCopyTo;
 
 	// the scale factor used to calculate all the cached offset (lazy evaluated when scale factor is modified)
 	float SourceScaleFactorUsedForCachedData;
@@ -115,7 +115,7 @@ public:
 	UPROPERTY(EditAnywhere, Category=Settings)
 	TArray<FPinBoneData> BonesToPin;
 
-	// Which skeleton to use to copy animation from (usually the target).
+	// Which skeleton to copy animation from.
 	UPROPERTY(EditAnywhere, Category=Settings, meta=(DisplayName="Skeleton To Copy From"))
 	ERetargetSourceOrTarget PinTo = ERetargetSourceOrTarget::Target;
 
