IKRig - keeping Goal.FinalBlendedRotation normalization errors within tolerance in UIKRigProcessor::ResolveFinalGoalTransforms

diff --git a/Engine/Plugins/Animation/IKRig/Source/IKRig/Private/Rig/IKRigProcessor.cpp b/Engine/Plugins/Animation/IKRig/Source/IKRig/Private/Rig/IKRigProcessor.cpp
--- a/Engine/Plugins/Animation/IKRig/Source/IKRig/Private/Rig/IKRigProcessor.cpp
+++ b/Engine/Plugins/Animation/IKRig/Source/IKRig/Private/Rig/IKRigProcessor.cpp
@@ -484,7 +484,7 @@ void UIKRigProcessor::ResolveFinalGoalTransforms(const FTransform& WorldToCompon
             ComponentSpaceGoalPosition,
             Goal.PositionAlpha);
 		
-		Goal.FinalBlendedRotation = FQuat::FastLerp(
+		Goal.FinalBlendedRotation = FQuat::Slerp(
             InputPoseBoneTransform.GetRotation(),
             ComponentSpaceGoalRotation,
             Goal.RotationAlpha);
