IK Retargeter: fixes issue in retarget exporter if source anim sequence had an FK control rig with different bone names than target sequence.

diff --git a/Engine/Plugins/Animation/IKRig/Source/IKRigEditor/Private/RetargetEditor/IKRetargetBatchOperation.cpp b/Engine/Plugins/Animation/IKRig/Source/IKRigEditor/Private/RetargetEditor/IKRetargetBatchOperation.cpp
--- a/Engine/Plugins/Animation/IKRig/Source/IKRigEditor/Private/RetargetEditor/IKRetargetBatchOperation.cpp
+++ b/Engine/Plugins/Animation/IKRig/Source/IKRigEditor/Private/RetargetEditor/IKRetargetBatchOperation.cpp
@@ -384,11 +384,14 @@ void UIKRetargetBatchOperation::ConvertAnimation(
 		FString AssetName = TargetSequence->GetName();
 		Progress.EnterProgressFrame(1.f, FText::Format(LOCTEXT("RunningBatchRetarget", "Retargeting animation asset: {0}"), FText::FromString(AssetName)));
 
-		// remove all keys from the destination animation sequence
+		// update the sequence to use the new skeleton
 		IAnimationDataController& TargetSeqController = TargetSequence->GetController();
 		constexpr bool bShouldTransact = false;
 		TargetSeqController.OpenBracket(FText::FromString("Generating Retargeted Animation Data"), bShouldTransact);
-		TargetSeqController.RemoveAllBoneTracks(bShouldTransact);
+		TargetSeqController.NotifyPopulated(); // must set bIsPopulated=true otherwise "UpdateWithSkeleton" will early out
+		// this removes all bone tracks and reinitializes the FKControlRig to use the new skeleton.
+		// NOTE: we do NOT ResetModel() because we want to keep curves and attributes from the source
+		TargetSeqController.UpdateWithSkeleton(TargetSkeleton.SkeletalMesh->GetSkeleton(), bShouldTransact);
 
 		// number of frames in this animation
 		const int32 NumFrames = SourceSequence->GetNumberOfSampledKeys();
