IK Retargeter: Adding input pin to retarget pose from mesh anim node and added CopyPose from source support.

diff --git a/Engine/Plugins/Animation/IKRig/Source/IKRig/Private/AnimNodes/AnimNode_RetargetPoseFromMesh.cpp b/Engine/Plugins/Animation/IKRig/Source/IKRig/Private/AnimNodes/AnimNode_RetargetPoseFromMesh.cpp
--- a/Engine/Plugins/Animation/IKRig/Source/IKRig/Private/AnimNodes/AnimNode_RetargetPoseFromMesh.cpp
+++ b/Engine/Plugins/Animation/IKRig/Source/IKRig/Private/AnimNodes/AnimNode_RetargetPoseFromMesh.cpp
@@ -2,6 +2,7 @@
 
 #include "AnimNodes/AnimNode_RetargetPoseFromMesh.h"
 
+#include "IKRigObjectVersion.h"
 #include "Animation/AnimCurveUtils.h"
 #include "Animation/AnimInstance.h"
 #include "Animation/AnimInstanceProxy.h"
@@ -16,91 +17,257 @@
 void FAnimNode_RetargetPoseFromMesh::Initialize_AnyThread(const FAnimationInitializeContext& Context)
 {
 	DECLARE_SCOPE_HIERARCHICAL_COUNTER_ANIMNODE(Initialize_AnyThread)
-    FAnimNode_Base::Initialize_AnyThread(Context);
 
-	// Initial update of the node, so we dont have a frame-delay on setup
+	FAnimNode_Base::Initialize_AnyThread(Context);
+
 	GetEvaluateGraphExposedInputs().Execute(Context);
 
-	// create the processor
-	CreateRetargetProcessorIfNeeded(Context.AnimInstanceProxy->GetSkelMeshComponent());
+	if (RetargetFrom == ERetargetSourceMode::SourcePosePin)
+	{
+		Source.Initialize(Context);
+	}
 }
 
 void FAnimNode_RetargetPoseFromMesh::CacheBones_AnyThread(const FAnimationCacheBonesContext& Context)
 {
 	DECLARE_SCOPE_HIERARCHICAL_COUNTER_ANIMNODE(CacheBones_AnyThread)
 
-	const FBoneContainer& RequiredBones = Context.AnimInstanceProxy->GetRequiredBones();
-	if (!RequiredBones.IsValid())
+	FAnimNode_Base::CacheBones_AnyThread(Context);
+
+	if (RetargetFrom == ERetargetSourceMode::SourcePosePin)
 	{
-		return;
+		Source.CacheBones(Context);
 	}
-	
-	if (!Context.AnimInstanceProxy->GetSkelMeshComponent()->GetSkeletalMeshAsset())
-	{
-		return;
-	}
-
-	// rebuild mapping
-	RequiredToTargetBoneMapping.Reset();
 
-	const FReferenceSkeleton& RefSkeleton = RequiredBones.GetReferenceSkeleton();
-	const FReferenceSkeleton& TargetSkeleton = Context.AnimInstanceProxy->GetSkelMeshComponent()->GetSkeletalMeshAsset()->GetRefSkeleton();
-	const TArray<FBoneIndexType>& RequiredBonesArray = RequiredBones.GetBoneIndicesArray();
-	for (int32 Index = 0; Index < RequiredBonesArray.Num(); ++Index)
+	// rebuild mapping of compact index to full skeleton bone index
+	CompactToTargetBoneIndexMap.Reset();
+	const TArray<FBoneIndexType>& RequiredBonesArray = Context.AnimInstanceProxy->GetRequiredBones().GetBoneIndicesArray();
+	for (FBoneIndexType CompactIndex = 0; CompactIndex < RequiredBonesArray.Num(); ++CompactIndex)
 	{
-		const FBoneIndexType ReqBoneIndex = RequiredBonesArray[Index]; 
-		if (ReqBoneIndex != INDEX_NONE)
-		{
-			const FName Name = RefSkeleton.GetBoneName(ReqBoneIndex);
-			const int32 TargetBoneIndex = TargetSkeleton.FindBoneIndex(Name);
-			if (TargetBoneIndex != INDEX_NONE)
-			{
-				// store require bone to target bone indices
-				RequiredToTargetBoneMapping.Emplace(Index, TargetBoneIndex);
-			}
-		}
+		CompactToTargetBoneIndexMap.Emplace(CompactIndex, RequiredBonesArray[CompactIndex]);
 	}
 }
 
 void FAnimNode_RetargetPoseFromMesh::Update_AnyThread(const FAnimationUpdateContext& Context)
 {
 	DECLARE_SCOPE_HIERARCHICAL_COUNTER_ANIMNODE(Update_AnyThread)
+
 	FAnimNode_Base::Update_AnyThread(Context);
-    // this introduces a frame of latency in setting the pin-driven source component,
+
+	// this introduces a frame of latency in setting the pin-driven source component,
     // but we cannot do the work to extract transforms on a worker thread as it is not thread safe.
     GetEvaluateGraphExposedInputs().Execute(Context);
-	DeltaTime += Context.GetDeltaTime();
+
+	// delta time stored here and passed to retargeter
+	DeltaTime = Context.GetDeltaTime();
+
+	if (RetargetFrom == ERetargetSourceMode::SourcePosePin)
+	{
+		// update pose input
+		Source.Update(Context);
+	}
 }
 
 void FAnimNode_RetargetPoseFromMesh::Evaluate_AnyThread(FPoseContext& Output)
 {
 	DECLARE_SCOPE_HIERARCHICAL_COUNTER_ANIMNODE(Evaluate_AnyThread)
-
 	SCOPE_CYCLE_COUNTER(STAT_IKRetarget);
 
-	if (!(IKRetargeterAsset && Processor && SourceMeshComponent.IsValid() && IsLODEnabled(Output.AnimInstanceProxy)))
+	auto IsReadyToRun = [this, &Output]()
 	{
-		Output.ResetToRefPose();
-		return;
-	}
+		if (!IKRetargeterAsset)
+		{
+			return false; // must have retarget asset
+		}
 
-	// it's possible in editor to have anim instances initialized before PreUpdate() is called
-	// which results in trying to run the retargeter without an source pose to copy from
-	const bool bSourcePoseCopied = !SourceMeshComponentSpaceBoneTransforms.IsEmpty();
+		if (!Processor)
+		{
+			return false; // must have processor
+		}
 
-	// ensure processor was initialized with the currently used assets (source/target meshes and retarget asset)
-	// if processor is not ready this tick, it will be next tick as this state will trigger re-initialization
-	const TObjectPtr<USkeletalMesh> SourceMesh = SourceMeshComponent->GetSkeletalMeshAsset();
-	const TObjectPtr<USkeletalMesh> TargetMesh = Output.AnimInstanceProxy->GetSkelMeshComponent()->GetSkeletalMeshAsset();
-	const bool bIsProcessorReady = Processor->WasInitializedWithTheseAssets(SourceMesh, TargetMesh, IKRetargeterAsset);
+		const bool bCopyingFromOtherComponent = RetargetFrom != ERetargetSourceMode::SourcePosePin;
+		if (bCopyingFromOtherComponent && !SourceMeshComponent.IsValid())
+		{
+			return false; // if copying from another mesh, must have valid pointer to it
+		}
+
+		if (bCopyingFromOtherComponent && PoseToRetargetFromComponentSpace.IsEmpty())
+		{
+			return false; // if copying from another mesh, pose must be copied already from the PreUpdate()
+		}
+
+		if (!bCopyingFromOtherComponent && !Source.GetLinkNode())
+		{
+			return false; // if copying from input pin, must be connected
+		}
+		
+		// ensure processor was initialized with the currently used assets (source/target meshes and retarget asset)
+		// if processor is not ready this tick, it will be next tick as this state will trigger re-initialization
+		// NOTE: do in-line reinitialization here once all runtime UObjects are removed (can't right now because UObject creation not allowed off the main thread)
+		const USkeletalMesh* TargetMesh = Output.AnimInstanceProxy->GetSkelMeshComponent()->GetSkeletalMeshAsset();
+		const USkeletalMesh* SourceMesh = bCopyingFromOtherComponent ? SourceMeshComponent->GetSkeletalMeshAsset() : TargetMesh;
+		const bool bIsProcessorReady = Processor->WasInitializedWithTheseAssets(SourceMesh, TargetMesh, IKRetargeterAsset);
+		if (!bIsProcessorReady)
+		{
+			return false;
+		}
+
+		if (!IsLODEnabled(Output.AnimInstanceProxy))
+		{
+			return false; // LOD'd out
+		}
+
+		// phew! ready to run
+		return true;
+	};
 
-	// if not ready to run, skip retarget and output the ref pose
-	if (!(bIsProcessorReady && bSourcePoseCopied))
+	auto FillPoseToRetargetFromInputPin = [this, &Output]()
+	{
+		// NOTE: in this case, we are copying the input pose from the anim graph pin, so the source and the target skeleton are the same!
+		// Retargeting between the same mesh allows retarget operations to modify or reinterpret a pose in some way.
+		// But the input pose is local and may be compacted, so we need to reconstruct the full component space pose to pass to the retargeter...
+
+		// start with the full ref pose
+		const USkeletalMesh* TargetMesh = Output.AnimInstanceProxy->GetSkelMeshComponent()->GetSkeletalMeshAsset();
+		const FReferenceSkeleton& TargetRefSkeleton = TargetMesh->GetRefSkeleton();
+		const TArray<FTransform>& TargetLocalRefPose = TargetRefSkeleton.GetRefBonePose();;
+
+		// start with full local ref pose
+		InputLocalTransforms.Reset(TargetLocalRefPose.Num());
+		InputLocalTransforms = TargetRefSkeleton.GetRefBonePose();
+		
+		// use the input pose from the anim graph as the pose to retarget from
+		if (Source.GetLinkNode())
+		{
+			// evaluate the input pose
+			Source.Evaluate(Output);
+
+			// overwrite required bones with input pose from the anim graph
+			for (const TPair<int32, int32>& Pair : CompactToTargetBoneIndexMap)
+			{
+				const FCompactPoseBoneIndex CompactBoneIndex(Pair.Key);
+				if (Output.Pose.IsValidIndex(CompactBoneIndex))
+				{
+					const int32 TargetBoneIndex = Pair.Value;
+					InputLocalTransforms[TargetBoneIndex] = Output.Pose[CompactBoneIndex];
+				}
+			}
+		}
+		
+		// convert to component space
+		FAnimationRuntime::FillUpComponentSpaceTransforms(TargetRefSkeleton, InputLocalTransforms, PoseToRetargetFromComponentSpace);
+
+		// update speed values from the main anim graph curves
+		{
+			SpeedValuesFromCurves.Reset();
+			const TArray<TObjectPtr<URetargetChainSettings>>& ChainSettings = IKRetargeterAsset->GetAllChainSettings();
+			for (const TObjectPtr<URetargetChainSettings>& ChainSetting : ChainSettings)
+			{
+				FName CurveName = ChainSetting->Settings.SpeedPlanting.SpeedCurveName;
+				if (CurveName == NAME_None)
+				{
+					continue;
+				}
+
+				// value will be negative if the curve was not found (retargeter ignores negative speeds)
+				bool bFoundCurve = false;
+				float CurveValue = Output.Curve.Get(CurveName, bFoundCurve, -1.0f);
+				SpeedValuesFromCurves.Add(CurveName, CurveValue);
+			}
+		}
+	};
+	
+	auto ApplyRetargetedPose = [this, &Output](const TArray<FTransform>& RetargetedPose)
+	{
+		// convert pose to local space and apply to output
+		FCSPose<FCompactPose> ComponentPose;
+		ComponentPose.InitPose(Output.Pose);
+		const FCompactPose& CompactPose = ComponentPose.GetPose();
+		for (const TPair<int32, int32>& Pair : CompactToTargetBoneIndexMap)
+		{
+			const FCompactPoseBoneIndex CompactBoneIndex(Pair.Key);
+			if (CompactPose.IsValidIndex(CompactBoneIndex))
+			{
+				const int32 TargetBoneIndex = Pair.Value;
+				ComponentPose.SetComponentSpaceTransform(CompactBoneIndex, RetargetedPose[TargetBoneIndex]);
+			}
+		}
+
+		// convert to local space
+		FCSPose<FCompactPose>::ConvertComponentPosesToLocalPoses(ComponentPose, Output.Pose);
+
+		// once converted back to local space, we copy scale values back
+		// (retargeter strips scale values and deals with translation only in component space)
+		const TObjectPtr<USkeletalMesh> TargetMesh = Output.AnimInstanceProxy->GetSkelMeshComponent()->GetSkeletalMeshAsset();
+		const TArray<FTransform>& RefPose = TargetMesh->GetRefSkeleton().GetRefBonePose();
+		for (const TPair<int32, int32>& Pair : CompactToTargetBoneIndexMap)
+		{
+			const FCompactPoseBoneIndex CompactBoneIndex(Pair.Key);
+			if (Output.Pose.IsValidIndex(CompactBoneIndex))
+			{
+				Output.Pose[CompactBoneIndex].SetScale3D(RefPose[Pair.Value].GetScale3D());
+			}
+		}
+	};
+
+	// remap curves for CurveRemapOp if one is present in the RetargetOp stack
+	auto CopyAndRemapCurvesFromSourceToTarget = [this](FBlendedCurve& OutputCurves)
+	{
+		if (!bCopyCurves)
+		{
+			return;
+		}
+	
+		if (!(Processor && Processor->IsInitialized()))
+		{
+			return;
+		}
+
+		// copy curves over to same name on target (if it exists)
+		OutputCurves.CopyFrom(SourceCurves);
+	
+		// now remap curves that have different names using the CurveRemapOp data (if there are any)
+		FBlendedCurve RemapedCurves;
+		const TArray<TObjectPtr<URetargetOpBase>>& RetargetOps = Processor->GetRetargetOps();
+		for (URetargetOpBase* RetargetOp : RetargetOps)
+		{
+			UCurveRemapOp* CurveRemapOp = Cast<UCurveRemapOp>(RetargetOp);
+			if (!CurveRemapOp)
+			{
+				continue;
+			}
+
+			if (!CurveRemapOp->bIsEnabled)
+			{
+				continue;
+			}
+		
+			for (const FCurveRemapPair& CurveToRemap : CurveRemapOp->CurvesToRemap)
+			{
+				bool bOutIsValid = false;
+				const float SourceValue = SourceCurves.Get(CurveToRemap.SourceCurve, bOutIsValid);
+				if (bOutIsValid)
+				{
+					RemapedCurves.Add(CurveToRemap.TargetCurve, SourceValue);
+				}
+			}
+		}
+		OutputCurves.Combine(RemapedCurves);
+	};
+
+	// check the configuration of the node and output the reference pose unless it's ready to retarget
+	if (!IsReadyToRun())
 	{
 		Output.ResetToRefPose();
 		return;
 	}
 
+	// if we're retargeting the connected input pose from the anim graph, then evaluate it and get the pose to retarget
+	if (RetargetFrom == ERetargetSourceMode::SourcePosePin)
+	{
+		FillPoseToRetargetFromInputPin(); // (also gets speed curves from anim graph)
+	}
+
 	#if WITH_EDITOR
 	if (GIsEditor)
 	{
@@ -116,40 +283,13 @@ void FAnimNode_RetargetPoseFromMesh::Evaluate_AnyThread(FPoseContext& Output)
 
 	// run the retargeter
 	const TArray<FTransform>& RetargetedPose = Processor->RunRetargeter(
-		SourceMeshComponentSpaceBoneTransforms,
+		PoseToRetargetFromComponentSpace,
 		SpeedValuesFromCurves,
 		DeltaTime,
 		RetargetProfileToUse);
-	DeltaTime = 0.0f;
-	
-	// convert pose to local space and apply to output
-	FCSPose<FCompactPose> ComponentPose;
-	ComponentPose.InitPose(Output.Pose);
-	const FCompactPose& CompactPose = ComponentPose.GetPose();
-	for (const TPair<int32, int32>& Pair : RequiredToTargetBoneMapping)
-	{
-		const FCompactPoseBoneIndex CompactBoneIndex(Pair.Key);
-		if (CompactPose.IsValidIndex(CompactBoneIndex))
-		{
-			const int32 TargetBoneIndex = Pair.Value;
-			ComponentPose.SetComponentSpaceTransform(CompactBoneIndex, RetargetedPose[TargetBoneIndex]);
-		}
-	}
 
-	// convert to local space
-	FCSPose<FCompactPose>::ConvertComponentPosesToLocalPoses(ComponentPose, Output.Pose);
-
-	// once converted back to local space, we copy scale values back
-	// (retargeter strips scale values and deals with translation only in component space)
-	const TArray<FTransform>& RefPose = TargetMesh->GetRefSkeleton().GetRefBonePose();
-	for (const TPair<int32, int32>& Pair : RequiredToTargetBoneMapping)
-	{
-		const FCompactPoseBoneIndex CompactBoneIndex(Pair.Key);
-		if (Output.Pose.IsValidIndex(CompactBoneIndex))
-		{
-			Output.Pose[CompactBoneIndex].SetScale3D(RefPose[Pair.Value].GetScale3D());
-		}
-	}
+	// apply the retargeted pose to the output pin
+	ApplyRetargetedPose(RetargetedPose);
 
 	// copy and/or remap curves from the source to the target skeletal mesh
 	CopyAndRemapCurvesFromSourceToTarget(Output.Curve);
@@ -158,118 +298,132 @@ void FAnimNode_RetargetPoseFromMesh::Evaluate_AnyThread(FPoseContext& Output)
 void FAnimNode_RetargetPoseFromMesh::PreUpdate(const UAnimInstance* InAnimInstance)
 {
 	DECLARE_SCOPE_HIERARCHICAL_COUNTER_ANIMNODE(PreUpdate)
-	
-	if (!IKRetargeterAsset)
-	{
-		return;
-	}
-	
-	if (!Processor)
-	{
-		Processor = NewObject<UIKRetargetProcessor>(InAnimInstance->GetOwningComponent());	
-	}
 
-	const TObjectPtr<USkeletalMeshComponent> TargetMeshComponent = InAnimInstance->GetSkelMeshComponent();
-	if (InitializeRetargetProcessor(TargetMeshComponent))
+	auto GetPoseToRetargetFromSourceMesh = [this](const USkeletalMeshComponent* TargetMeshComponent)
 	{
-		CopyBoneTransformsFromSource(TargetMeshComponent);
-
-		if(bCopyCurves)
+		// skip if copying from a source pin
+		const bool bCopyingFromOtherComponent = RetargetFrom != ERetargetSourceMode::SourcePosePin;
+		if (!bCopyingFromOtherComponent)
 		{
-			if (const UAnimInstance* SourceAnimInstance = SourceMeshComponent->GetAnimInstance())
-			{
-				// Potential optimization/tradeoff: If we stored the curve results on the mesh component in non-editor scenarios, this would be
-				// much faster (but take more memory). As it is, we need to translate the map stored on the anim instance.
-				const TMap<FName, float>& AnimCurveList = SourceAnimInstance->GetAnimationCurveList(EAnimCurveType::AttributeCurve);
-				UE::Anim::FCurveUtils::BuildUnsorted(SourceCurves, AnimCurveList);
-			}
-			else
-			{
-				SourceCurves.Empty();
-			}
-			
-			UpdateSpeedValuesFromCurves();
+			return;
 		}
-	}
-}
 
-void FAnimNode_RetargetPoseFromMesh::CreateRetargetProcessorIfNeeded(UObject* Outer)
-{
-	if (!Processor && IsInGameThread())
-	{
-		Processor = NewObject<UIKRetargetProcessor>(Outer);	
-	}
-}
+		// source mesh not connected or not found
+		if (!SourceMeshComponent.IsValid())
+		{
+			return;
+		}
 
-void FAnimNode_RetargetPoseFromMesh::CopyAndRemapCurvesFromSourceToTarget(FBlendedCurve& OutputCurves) const
-{
-	if (!bCopyCurves)
-	{
-		return;
-	}
+		// if our source is running under leader-pose, then get bone data from there
+		USkeletalMeshComponent* ComponentToCopyFrom =  SourceMeshComponent.Get();
+		if(USkeletalMeshComponent* LeaderPoseComponent = Cast<USkeletalMeshComponent>(SourceMeshComponent->LeaderPoseComponent.Get()))
+		{
+			ComponentToCopyFrom = LeaderPoseComponent;
+		}
 	
-	if (!(Processor && Processor->IsInitialized()))
-	{
-		return;
-	}
-
-	// copy curves over to same name on target (if it exists)
-	OutputCurves.CopyFrom(SourceCurves);
+		// skip copying pose when component is no longer ticking
+		if (!ComponentToCopyFrom->IsRegistered())
+		{
+			return; 
+		}
 	
-	// now remap curves that have different names using the CurveRemapOp data (if there are any)
-	FBlendedCurve RemapedCurves;
-	const TArray<TObjectPtr<URetargetOpBase>>& RetargetOps = Processor->GetRetargetOps();
-	for (URetargetOpBase* RetargetOp : RetargetOps)
+		const bool bUROInSync =
+			ComponentToCopyFrom->ShouldUseUpdateRateOptimizations() &&
+			ComponentToCopyFrom->AnimUpdateRateParams != nullptr &&
+			SourceMeshComponent->AnimUpdateRateParams == TargetMeshComponent->AnimUpdateRateParams;
+		const bool bUsingExternalInterpolation = ComponentToCopyFrom->IsUsingExternalInterpolation();
+		const TArray<FTransform>& CachedComponentSpaceTransforms = ComponentToCopyFrom->GetCachedComponentSpaceTransforms();
+		const bool bArraySizesMatch = CachedComponentSpaceTransforms.Num() == ComponentToCopyFrom->GetComponentSpaceTransforms().Num();
+
+		// copy source array from the appropriate location
+		PoseToRetargetFromComponentSpace.Reset();
+		if ((bUROInSync || bUsingExternalInterpolation) && bArraySizesMatch)
+		{
+			// copy from source's cache
+			PoseToRetargetFromComponentSpace.Append(CachedComponentSpaceTransforms);
+		}
+		else
+		{
+			// copy directly
+			PoseToRetargetFromComponentSpace.Append(ComponentToCopyFrom->GetComponentSpaceTransforms());
+		}
+	
+		// strip all scale out of the pose values, the translation of a component-space pose has incorporated scale values
+		for (FTransform& Transform : PoseToRetargetFromComponentSpace)
+		{
+			Transform.SetScale3D(FVector::OneVector);
+		}
+	};
+
+	auto CopyCurvesFromSourceMesh = [this]()
 	{
-		UCurveRemapOp* CurveRemapOp = Cast<UCurveRemapOp>(RetargetOp);
-		if (!CurveRemapOp)
+		if(!bCopyCurves)
+		{
+			return;
+		}
+
+		const UAnimInstance* SourceAnimInstance = SourceMeshComponent->GetAnimInstance();
+		if (!SourceAnimInstance)
 		{
-			continue;
+			return;
 		}
+		
+		// Potential optimization/tradeoff: If we stored the curve results on the mesh component in non-editor scenarios, this would be
+		// much faster (but take more memory). As it is, we need to translate the map stored on the anim instance.
+		const TMap<FName, float>& AnimCurveList = SourceAnimInstance->GetAnimationCurveList(EAnimCurveType::AttributeCurve);
+		SourceCurves.Empty();
+		UE::Anim::FCurveUtils::BuildUnsorted(SourceCurves, AnimCurveList);
+	};
+
+	auto UpdateSpeedFromSourceMeshCurves = [this]()
+	{
+		SpeedValuesFromCurves.Reset();
 
-		if (!CurveRemapOp->bIsEnabled)
+		const UAnimInstance* SourceAnimInstance = SourceMeshComponent->GetAnimInstance();
+		if (!SourceAnimInstance)
 		{
-			continue;
+			return;
 		}
 		
-		for (const FCurveRemapPair& CurveToRemap : CurveRemapOp->CurvesToRemap)
+		const TMap<FName, float>& AnimCurveList = SourceAnimInstance->GetAnimationCurveList(EAnimCurveType::AttributeCurve);
+		const TArray<TObjectPtr<URetargetChainSettings>>& ChainSettings = IKRetargeterAsset->GetAllChainSettings();
+		for (const TObjectPtr<URetargetChainSettings>& ChainSetting : ChainSettings)
 		{
-			bool bOutIsValid = false;
-			const float SourceValue = SourceCurves.Get(CurveToRemap.SourceCurve, bOutIsValid);
-			if (bOutIsValid)
+			FName CurveName = ChainSetting->Settings.SpeedPlanting.SpeedCurveName;
+			if (CurveName == NAME_None)
 			{
-				RemapedCurves.Add(CurveToRemap.TargetCurve, SourceValue);
+				continue;
 			}
+			
+			const float* SourceValue = AnimCurveList.Find(CurveName);
+			constexpr float NotFoundValue = -1.0f; // value will be negative if the curve was not found (retargeter ignores negative speeds)
+			const float CurveValue = SourceValue ? *SourceValue : NotFoundValue;
+			SpeedValuesFromCurves.Add(CurveName, CurveValue);
 		}
-	}
-	OutputCurves.Combine(RemapedCurves);
-}
+	};
 
-void FAnimNode_RetargetPoseFromMesh::UpdateSpeedValuesFromCurves()
-{
-	SpeedValuesFromCurves.Reset();
-	
-	if (!IKRetargeterAsset)
+	// initialization of the retarget processor must be on the main thread
+	// NOTE: once UObject creation has been refactored out of IK Rig and IK Retargeter we can do this in-line off the main thread
+	USkeletalMeshComponent* TargetMeshComponent = InAnimInstance->GetSkelMeshComponent();
+	if (!TryInitializeProcessor(TargetMeshComponent))
 	{
 		return;
 	}
 
-	const TArray<TObjectPtr<URetargetChainSettings>>& ChainSettings = IKRetargeterAsset->GetAllChainSettings();
-	for (const TObjectPtr<URetargetChainSettings>& ChainSetting : ChainSettings)
+	// copy all the data from the source component
+	if (RetargetFrom != ERetargetSourceMode::SourcePosePin)
 	{
-		FName CurveName = ChainSetting->Settings.SpeedPlanting.SpeedCurveName;
-		if (CurveName == NAME_None)
-		{
-			continue;
-		}
+		GetPoseToRetargetFromSourceMesh(TargetMeshComponent);
+		CopyCurvesFromSourceMesh();
+		UpdateSpeedFromSourceMeshCurves();
+	}
+}
 
-		// value will be negative if the curve was not found (retargeter ignores negative speeds)
-		bool bCurveValid = false;
-		float CurveValue = SourceCurves.Get(CurveName, bCurveValid, -1.0f);
-		if (bCurveValid)
-		{
-			SpeedValuesFromCurves.Add(CurveName, CurveValue);
-		}
+void FAnimNode_RetargetPoseFromMesh::CreateRetargetProcessorIfNeeded(UObject* Outer)
+{
+	if (!Processor && IsInGameThread())
+	{
+		Processor = NewObject<UIKRetargetProcessor>(Outer);	
 	}
 }
 
@@ -278,44 +432,57 @@ UIKRetargetProcessor* FAnimNode_RetargetPoseFromMesh::GetRetargetProcessor() con
 	return Processor;
 }
 
-bool FAnimNode_RetargetPoseFromMesh::InitializeRetargetProcessor(const TObjectPtr<USkeletalMeshComponent> TargetMeshComponent)
+bool FAnimNode_RetargetPoseFromMesh::TryInitializeProcessor(USkeletalMeshComponent* TargetMeshComponent)
 {
+	if (!ensure(IsInGameThread()))
+	{
+		// NOTE: remove this when processor no longer needs game thread to initialize
+		return false;
+	}
+	
+	if (!Processor)
+	{
+		Processor = NewObject<UIKRetargetProcessor>(TargetMeshComponent);	
+	}
+	
 	// has user supplied a retargeter asset?
 	if (!IKRetargeterAsset)
 	{
 		return false;
 	}
 	
-	// if user hasn't explicitly connected a source mesh, optionally use the parent mesh component (if there is one) 
-	if (!SourceMeshComponent.IsValid() && bUseAttachedParent)
+	// if user hasn't explicitly connected a source mesh, optionally use the parent mesh component (if there is one)
+	if (RetargetFrom == ERetargetSourceMode::ParentSkeletalMeshComponent && !bSearchedForParentComponent)
 	{
-		// Walk up the attachment chain until we find a skeletal mesh component
-		USkeletalMeshComponent* ParentComponent = nullptr;
+		bSearchedForParentComponent = true;
+		SourceMeshComponent.Reset();
+		
+		// walk up component hierarchy until we find a skeletal mesh component
 		for (USceneComponent* AttachParentComp = TargetMeshComponent->GetAttachParent(); AttachParentComp != nullptr; AttachParentComp = AttachParentComp->GetAttachParent())
 		{
-			ParentComponent = Cast<USkeletalMeshComponent>(AttachParentComp);
-			if (ParentComponent)
+			SourceMeshComponent = Cast<USkeletalMeshComponent>(AttachParentComp);
+			if (SourceMeshComponent.IsValid())
 			{
 				break;
 			}
 		}
-
-		if (ParentComponent)
-		{
-			SourceMeshComponent = ParentComponent;
-		}
+	}
+	else
+	{
+		bSearchedForParentComponent = false;
 	}
 	
 	// has a source mesh been plugged in or found?
-	if (!SourceMeshComponent.IsValid())
+	const bool bCopyingFromOtherComponent = RetargetFrom != ERetargetSourceMode::SourcePosePin;
+	if (bCopyingFromOtherComponent && !SourceMeshComponent.IsValid())
 	{
 		return false; // can't do anything if we don't have a source mesh component
 	}
 
 	// check that both a source and target mesh exist
-	const TObjectPtr<USkeletalMesh> SourceMesh = SourceMeshComponent->GetSkeletalMeshAsset();
-	const TObjectPtr<USkeletalMesh> TargetMesh = TargetMeshComponent->GetSkeletalMeshAsset();
-	if (!SourceMesh || !TargetMesh)
+	const USkeletalMesh* TargetMesh = TargetMeshComponent->GetSkeletalMeshAsset();
+	const USkeletalMesh* SourceMesh = bCopyingFromOtherComponent ? SourceMeshComponent->GetSkeletalMeshAsset() : TargetMesh;
+	if (!(SourceMesh && TargetMesh))
 	{
 		return false; // cannot initialize if components are missing skeletal mesh references
 	}
@@ -333,64 +500,6 @@ bool FAnimNode_RetargetPoseFromMesh::InitializeRetargetProcessor(const TObjectPt
 	return Processor->IsInitialized();
 }
 
-void FAnimNode_RetargetPoseFromMesh::CopyBoneTransformsFromSource(USkeletalMeshComponent* TargetMeshComponent)
-{
-	// get the mesh component to use as the source
-	const TObjectPtr<USkeletalMeshComponent> ComponentToCopyFrom =  GetComponentToCopyPoseFrom();
-
-	// this should not happen as we're guaranteed to be initialized at this stage
-	// but just in case component is lost after initialization, we avoid a crash
-	if (!ComponentToCopyFrom)
-	{
-		return; 
-	}
-	
-	// skip copying pose when component is no longer ticking
-	if (!ComponentToCopyFrom->IsRegistered())
-	{
-		return; 
-	}
-	
-	const bool bUROInSync =
-		ComponentToCopyFrom->ShouldUseUpdateRateOptimizations() &&
-		ComponentToCopyFrom->AnimUpdateRateParams != nullptr &&
-		SourceMeshComponent->AnimUpdateRateParams == TargetMeshComponent->AnimUpdateRateParams;
-	const bool bUsingExternalInterpolation = ComponentToCopyFrom->IsUsingExternalInterpolation();
-	const TArray<FTransform>& CachedComponentSpaceTransforms = ComponentToCopyFrom->GetCachedComponentSpaceTransforms();
-	const bool bArraySizesMatch = CachedComponentSpaceTransforms.Num() == ComponentToCopyFrom->GetComponentSpaceTransforms().Num();
-
-	// copy source array from the appropriate location
-	SourceMeshComponentSpaceBoneTransforms.Reset();
-	if ((bUROInSync || bUsingExternalInterpolation) && bArraySizesMatch)
-	{
-		SourceMeshComponentSpaceBoneTransforms.Append(CachedComponentSpaceTransforms); // copy from source's cache
-	}
-	else
-	{
-		SourceMeshComponentSpaceBoneTransforms.Append(ComponentToCopyFrom->GetComponentSpaceTransforms()); // copy directly
-	}
-	
-	// strip all scale out of the pose values, the translation of a component-space pose has incorporated scale values
-	for (FTransform& Transform : SourceMeshComponentSpaceBoneTransforms)
-	{
-		Transform.SetScale3D(FVector::OneVector);
-	}
-}
-
-TObjectPtr<USkeletalMeshComponent> FAnimNode_RetargetPoseFromMesh::GetComponentToCopyPoseFrom() const
-{
-	// if our source is running under leader-pose, then get bone data from there
-	if (SourceMeshComponent.IsValid())
-	{
-		if(USkeletalMeshComponent* LeaderPoseComponent = Cast<USkeletalMeshComponent>(SourceMeshComponent->LeaderPoseComponent.Get()))
-		{
-			return LeaderPoseComponent;
-		}
-	}
-	
-	return SourceMeshComponent.Get();
-}
-
 FRetargetProfile FAnimNode_RetargetPoseFromMesh::GetMergedRetargetProfile(bool bForceIKOff) const
 {
 	// collect settings to retarget with starting with asset settings and overriding with custom profile
@@ -407,3 +516,25 @@ FRetargetProfile FAnimNode_RetargetPoseFromMesh::GetMergedRetargetProfile(bool b
 	return Profile;
 }
 
+bool FAnimNode_RetargetPoseFromMesh::Serialize(FArchive& Ar)
+{
+	Ar.UsingCustomVersion(FIKRigObjectVersion::GUID);
+	// doesn't actually serialize, just write the custom version for PostSerialize
+	return false;
+}
+
+void FAnimNode_RetargetPoseFromMesh::PostSerialize(const FArchive& Ar)
+{
+#if WITH_EDITORONLY_DATA
+	
+	if (Ar.IsLoading())
+	{
+		if (Ar.CustomVer(FIKRigObjectVersion::GUID) < FIKRigObjectVersion::UseAttachedParentDeprecated)
+		{
+			RetargetFrom = bUseAttachedParent_DEPRECATED ? ERetargetSourceMode::ParentSkeletalMeshComponent : ERetargetSourceMode::CustomSkeletalMeshComponent;
+		}
+	}
+	
+#endif // WITH_EDITORONLY_DATA
+}
+
diff --git a/Engine/Plugins/Animation/IKRig/Source/IKRig/Private/Retargeter/IKRetargetProcessor.cpp b/Engine/Plugins/Animation/IKRig/Source/IKRig/Private/Retargeter/IKRetargetProcessor.cpp
--- a/Engine/Plugins/Animation/IKRig/Source/IKRig/Private/Retargeter/IKRetargetProcessor.cpp
+++ b/Engine/Plugins/Animation/IKRig/Source/IKRig/Private/Retargeter/IKRetargetProcessor.cpp
@@ -13,7 +13,6 @@
 
 #include "Engine/SkeletalMesh.h"
 #include "Kismet/KismetMathLibrary.h"
-#include "UObject/UnrealTypePrivate.h"
 
 #include UE_INLINE_GENERATED_CPP_BY_NAME(IKRetargetProcessor)
 
@@ -23,7 +22,7 @@
 #define RETARGETSKELETON_INVALID_BRANCH_INDEX -2
 
 void FRetargetSkeleton::Initialize(
-	USkeletalMesh* InSkeletalMesh,
+	const USkeletalMesh* InSkeletalMesh,
 	const TArray<FBoneChain>& BoneChains,
 	const FName InRetargetPoseName,
 	const FIKRetargetPose* RetargetPose,
@@ -92,7 +91,7 @@ void FRetargetSkeleton::GenerateRetargetPose(
 	// copy local pose to global
 	RetargetGlobalPose = RetargetLocalPose;
 	// convert to global space
-	UpdateGlobalTransformsBelowBone(-1, RetargetLocalPose, RetargetGlobalPose);
+	UpdateGlobalTransformsBelowBone(INDEX_NONE, RetargetLocalPose, RetargetGlobalPose);
 
 	// strip scale (done AFTER generating global pose so that scales are baked into translation)
 	for (int32 BoneIndex=0; BoneIndex<BoneNames.Num(); ++BoneIndex)
@@ -133,7 +132,7 @@ void FRetargetSkeleton::GenerateRetargetPose(
 	}
 
 	// update global transforms based on local pose modified by the retarget pose offsets
-	UpdateGlobalTransformsBelowBone(-1, RetargetLocalPose, RetargetGlobalPose);
+	UpdateGlobalTransformsBelowBone(INDEX_NONE, RetargetLocalPose, RetargetGlobalPose);
 
 	// scale the retarget pose
 	if (!FMath::IsNearlyEqual(ScaleFactor, 1.0f))
@@ -145,7 +144,7 @@ void FRetargetSkeleton::GenerateRetargetPose(
 		}
 
 		// update the local transforms
-		UpdateLocalTransformsBelowBone(-1, RetargetLocalPose, RetargetGlobalPose);
+		UpdateLocalTransformsBelowBone(INDEX_NONE, RetargetLocalPose, RetargetGlobalPose);
 	}
 }
 
@@ -315,7 +314,7 @@ void FRetargetSkeleton::GetChildrenIndicesRecursive(const int32 BoneIndex, TArra
 	}
 }
 
-bool FRetargetSkeleton::IsParentOfChild(const int32 PotentialParentIndex, const int32 ChildBoneIndex) const
+bool FRetargetSkeleton::IsParentOf(const int32 PotentialParentIndex, const int32 ChildBoneIndex) const
 {
 	int32 ParentIndex = GetParentIndex(ChildBoneIndex);
 	while (ParentIndex != INDEX_NONE)
@@ -342,7 +341,7 @@ int32 FRetargetSkeleton::GetParentIndex(const int32 BoneIndex) const
 }
 
 void FTargetSkeleton::Initialize(
-	USkeletalMesh* InSkeletalMesh,
+	const USkeletalMesh* InSkeletalMesh,
 	const TArray<FBoneChain>& BoneChains,
 	const FName InRetargetPoseName,
 	const FIKRetargetPose* RetargetPose,
@@ -358,6 +357,9 @@ void FTargetSkeleton::Initialize(
 	// these flags are actually set later in init phase when bone chains are mapped together
 	IsBoneRetargeted.Init(false, BoneNames.Num());
 
+	// initialize pose for base local pose (may be overridden by pose copier, but uses retarget pose by default)
+	InputLocalPose = RetargetLocalPose;
+	
 	// initialize storage for output pose (the result of the retargeting)
 	OutputGlobalPose = RetargetGlobalPose;
 }
@@ -369,7 +371,7 @@ void FTargetSkeleton::Reset()
 	IsBoneRetargeted.Reset();
 }
 
-void FTargetSkeleton::UpdateGlobalTransformsAllNonRetargetedBones(TArray<FTransform>& InOutGlobalPose)
+void FTargetSkeleton::UpdateGlobalTransformsAllNonRetargetedBones(const TArray<FTransform>& LocalPose, TArray<FTransform>& InOutGlobalPose)
 {
 	check(IsBoneRetargeted.Num() == InOutGlobalPose.Num());
 	
@@ -377,7 +379,7 @@ void FTargetSkeleton::UpdateGlobalTransformsAllNonRetargetedBones(TArray<FTransf
 	{
 		if (!IsBoneRetargeted[BoneIndex])
 		{
-			UpdateGlobalTransformOfSingleBone(BoneIndex, RetargetLocalPose, InOutGlobalPose);
+			UpdateGlobalTransformOfSingleBone(BoneIndex, LocalPose, InOutGlobalPose);
 		}
 	}
 }
@@ -1149,7 +1151,7 @@ void FChainDecoderFK::UpdateIntermediateParents(
 {
 	for (const int32& ParentIndex : IntermediateParentIndices)
 	{
-		TargetSkeleton.UpdateGlobalTransformOfSingleBone(ParentIndex, TargetSkeleton.RetargetLocalPose, InOutGlobalPose);
+		TargetSkeleton.UpdateGlobalTransformOfSingleBone(ParentIndex, TargetSkeleton.InputLocalPose, InOutGlobalPose);
 	}
 }
 
@@ -1507,6 +1509,59 @@ bool FRetargetChainPairIK::Initialize(
 	return true;
 }
 
+void FBasePoseCopier::CopyPoseFromSource(
+		const FName InRootBoneToCopy,
+		const FRetargetSkeleton& InSourceSkeleton,
+		const FRetargetSkeleton& InTargetSkeleton,
+		TArray<FTransform>& InOutSourcePose,
+		TArray<FTransform>& OutTargetPose)
+{
+	// check if we need to rebuild the cached bone mapping
+	if (SourceToTargetBoneIndexMap.IsEmpty() || InRootBoneToCopy != RootBoneToCopy)
+	{
+		// update the mapping from source-to-target bones (by name)
+		SourceToTargetBoneIndexMap.Reset();
+		RootBoneToCopy = InRootBoneToCopy;
+		RootBoneToCopyIndex = RootBoneToCopy == NAME_None ? INDEX_NONE : InSourceSkeleton.FindBoneIndexByName(RootBoneToCopy);
+		for (int32 SourceBoneIndex=0; SourceBoneIndex<InSourceSkeleton.BoneNames.Num(); ++SourceBoneIndex)
+		{
+			// filter out bones above the RootBoneToCopy
+			if (RootBoneToCopyIndex != INDEX_NONE &&
+				SourceBoneIndex != RootBoneToCopyIndex &&
+				!InSourceSkeleton.IsParentOf(RootBoneToCopyIndex, SourceBoneIndex))
+			{
+				continue;
+			}
+		
+			// filter out bones not in the target skeleton
+			const int32 TargetBoneIndex = InTargetSkeleton.FindBoneIndexByName(InTargetSkeleton.BoneNames[SourceBoneIndex]);
+			if (TargetBoneIndex == INDEX_NONE)
+			{
+				continue;
+			}
+		
+			// store map of source/target bones in common
+			SourceToTargetBoneIndexMap.Emplace(SourceBoneIndex, TargetBoneIndex);
+		}
+	}
+
+	// copy source bones directly to the target (no retargeting)
+	for (const TPair<int32, int32>& Pair : SourceToTargetBoneIndexMap)
+	{
+		OutTargetPose[Pair.Key] = InOutSourcePose[Pair.Value];
+	}
+
+	// reset source bones above the RootBoneToCopy to retarget pose
+	if (InRootBoneToCopy != NAME_None)
+	{
+		InOutSourcePose = InSourceSkeleton.RetargetGlobalPose;
+		for (const TPair<int32, int32>& Pair : SourceToTargetBoneIndexMap)
+		{
+			InOutSourcePose[Pair.Value] = OutTargetPose[Pair.Key];
+		}
+	}
+}
+
 bool FRootRetargeter::InitializeSource(
 	const FName SourceRootBoneName,
 	const FRetargetSkeleton& SourceSkeleton,
@@ -1815,9 +1870,9 @@ UIKRetargetProcessor::UIKRetargetProcessor()
 }
 
 void UIKRetargetProcessor::Initialize(
-		USkeletalMesh* SourceSkeletalMesh,
-		USkeletalMesh* TargetSkeletalMesh,
-		UIKRetargeter* InRetargeterAsset,
+		const USkeletalMesh* SourceSkeletalMesh,
+		const USkeletalMesh* TargetSkeletalMesh,
+		const UIKRetargeter* InRetargeterAsset,
 		const FRetargetProfile& Settings,
 		const bool bSuppressWarnings)
 {
@@ -2267,15 +2322,15 @@ TArray<FTransform>&  UIKRetargetProcessor::RunRetargeter(
 		}
 	}
 	
-	// start from retarget pose
-	TargetSkeleton.OutputGlobalPose = TargetSkeleton.RetargetGlobalPose;
+	// generate the pose to start retargeting from
+	GenerateBasePoses(InSourceGlobalPose);
 
 	// ROOT retargeting
 	if (GlobalSettings.bEnableRoot && bRootsInitialized)
 	{
 		RunRootRetarget(InSourceGlobalPose, TargetSkeleton.OutputGlobalPose);
 		// update global transforms below root
-		TargetSkeleton.UpdateGlobalTransformsBelowBone(RootRetargeter.Target.BoneIndex, TargetSkeleton.RetargetLocalPose, TargetSkeleton.OutputGlobalPose);
+		TargetSkeleton.UpdateGlobalTransformsBelowBone(RootRetargeter.Target.BoneIndex, TargetSkeleton.InputLocalPose, TargetSkeleton.OutputGlobalPose);
 	}
 	
 	// FK CHAIN retargeting
@@ -2283,7 +2338,7 @@ TArray<FTransform>&  UIKRetargetProcessor::RunRetargeter(
 	{
 		RunFKRetarget(InSourceGlobalPose, TargetSkeleton.OutputGlobalPose);
 		// update all the bones that are not controlled by FK chains or root
-		TargetSkeleton.UpdateGlobalTransformsAllNonRetargetedBones(TargetSkeleton.OutputGlobalPose);
+		TargetSkeleton.UpdateGlobalTransformsAllNonRetargetedBones(TargetSkeleton.InputLocalPose,TargetSkeleton.OutputGlobalPose);
 	}
 	
 	// IK CHAIN retargeting
@@ -2307,6 +2362,29 @@ TArray<FTransform>&  UIKRetargetProcessor::RunRetargeter(
 	return TargetSkeleton.OutputGlobalPose;
 }
 
+void UIKRetargetProcessor::GenerateBasePoses(TArray<FTransform>& InSourceGlobalPose)
+{
+	// start from the retarget pose
+	TargetSkeleton.OutputGlobalPose = TargetSkeleton.RetargetGlobalPose;
+	TargetSkeleton.InputLocalPose = TargetSkeleton.RetargetLocalPose;
+
+	// partially override with input pose
+	if (GlobalSettings.bCopyBasePose)
+	{
+		// do a straight copy of bone transforms from the source (no retargeting!)
+		BasePoseCopier.CopyPoseFromSource(
+			GlobalSettings.CopyBasePoseRoot,
+			SourceSkeleton,
+			TargetSkeleton,
+			InSourceGlobalPose,
+			TargetSkeleton.OutputGlobalPose);
+		
+		// update input local pose (used when updating intermediate bones)
+		const int32 ParentOfCopyRootIndex = TargetSkeleton.GetParentIndex(BasePoseCopier.GetIndexOfRootBoneToCopy());
+		TargetSkeleton.UpdateLocalTransformsBelowBone(ParentOfCopyRootIndex, TargetSkeleton.InputLocalPose, TargetSkeleton.OutputGlobalPose );
+	}
+}
+
 void UIKRetargetProcessor::RunRootRetarget(
 	const TArray<FTransform>& InGlobalTransforms,
     TArray<FTransform>& OutGlobalTransforms)
@@ -2615,9 +2693,9 @@ void UIKRetargetProcessor::ResetPlanting()
 }
 
 bool UIKRetargetProcessor::WasInitializedWithTheseAssets(
-	const TObjectPtr<USkeletalMesh> InSourceMesh,
-	const TObjectPtr<USkeletalMesh> InTargetMesh,
-	const TObjectPtr<UIKRetargeter> InRetargetAsset)
+	const USkeletalMesh* InSourceMesh,
+	const USkeletalMesh* InTargetMesh,
+	const UIKRetargeter* InRetargetAsset) const
 {
 	// not initialized at all
 	if (!IsInitialized())
diff --git a/Engine/Plugins/Animation/IKRig/Source/IKRig/Private/Retargeter/RetargetOps/PinBoneOp.cpp b/Engine/Plugins/Animation/IKRig/Source/IKRig/Private/Retargeter/RetargetOps/PinBoneOp.cpp
--- a/Engine/Plugins/Animation/IKRig/Source/IKRig/Private/Retargeter/RetargetOps/PinBoneOp.cpp
+++ b/Engine/Plugins/Animation/IKRig/Source/IKRig/Private/Retargeter/RetargetOps/PinBoneOp.cpp
@@ -249,12 +249,14 @@ USkeleton* UPinBoneOp::GetSkeleton(bool& bInvalidSkeletonIsError, const IPropert
 		FName PropertyName = PropertyHandle->GetProperty()->GetFName();
 		if (PropertyName == GET_MEMBER_NAME_CHECKED(FPinBoneData, BoneToCopyTo))
 		{
-			return RunningInstance->PinTo == ERetargetSourceOrTarget::Source ? RunningInstance->TargetSkeletonAsset : RunningInstance->SourceSkeletonAsset;
+			const USkeleton* Skeleton = RunningInstance->PinTo == ERetargetSourceOrTarget::Source ? RunningInstance->TargetSkeletonAsset : RunningInstance->SourceSkeletonAsset;
+			return const_cast<USkeleton*>(Skeleton);
 		}
 
 		if (PropertyName == GET_MEMBER_NAME_CHECKED(FPinBoneData, BoneToCopyFrom))
 		{
-			return RunningInstance->PinTo == ERetargetSourceOrTarget::Target ? RunningInstance->TargetSkeletonAsset : RunningInstance->SourceSkeletonAsset;
+			const USkeleton* Skeleton = RunningInstance->PinTo == ERetargetSourceOrTarget::Target ? RunningInstance->TargetSkeletonAsset : RunningInstance->SourceSkeletonAsset;
+			return const_cast<USkeleton*>(Skeleton);
 		}
 	}
 #endif
diff --git a/Engine/Plugins/Animation/IKRig/Source/IKRig/Public/AnimNodes/AnimNode_RetargetPoseFromMesh.h b/Engine/Plugins/Animation/IKRig/Source/IKRig/Public/AnimNodes/AnimNode_RetargetPoseFromMesh.h
--- a/Engine/Plugins/Animation/IKRig/Source/IKRig/Public/AnimNodes/AnimNode_RetargetPoseFromMesh.h
+++ b/Engine/Plugins/Animation/IKRig/Source/IKRig/Public/AnimNodes/AnimNode_RetargetPoseFromMesh.h
@@ -13,51 +13,62 @@
 
 DECLARE_CYCLE_STAT(TEXT("IK Retarget"), STAT_IKRetarget, STATGROUP_Anim);
 
+UENUM(BlueprintType)
+enum class ERetargetSourceMode : uint8
+{
+	ParentSkeletalMeshComponent,
+	CustomSkeletalMeshComponent,
+	SourcePosePin
+};
+
 USTRUCT(BlueprintInternalUseOnly)
 struct IKRIG_API FAnimNode_RetargetPoseFromMesh : public FAnimNode_Base
 {
 	GENERATED_BODY()
+	
+	// Input pose to be modified by the retargeter when using "Source Pose Pin" mode as the Input Pose Mode.
+	UPROPERTY(EditAnywhere, Category = Links)
+	FPoseLink Source;
+	
+	// Specify where to get the source pose to retarget from. Can be from the anim graph, or a different skeletal mesh component.
+	UPROPERTY(EditAnywhere, BlueprintReadWrite, Category = Settings, meta=(NeverAsPin))
+	ERetargetSourceMode RetargetFrom;
 
-	/** The Skeletal Mesh Component to retarget animation from. Assumed to be animated and tick BEFORE this anim instance.*/
-	UPROPERTY(BlueprintReadWrite, transient, Category=Settings, meta=(PinShownByDefault))
+	// The Skeletal Mesh Component to retarget animation from. Assumed to be animated and tick BEFORE this anim instance.
+	UPROPERTY(BlueprintReadWrite, transient, Category = Settings, meta=(PinShownByDefault))
 	TWeakObjectPtr<USkeletalMeshComponent> SourceMeshComponent = nullptr;
-
-	/* If SourceMeshComponent is not valid, and if this is true, it will look for attached parent as a source */
-	UPROPERTY(EditAnywhere, BlueprintReadWrite, Category=Settings, meta=(NeverAsPin))
-	bool bUseAttachedParent = true;
 	
-	/** Retarget asset to use. Must define a Source and Target IK Rig compatible with the SourceMeshComponent and current anim instance.*/
+	// Retarget asset to use. Must define a Source and Target IK Rig compatible with the SourceMeshComponent and current anim instance.
 	UPROPERTY(EditAnywhere, BlueprintReadWrite, Category = Settings, meta=(PinHiddenByDefault))
 	TObjectPtr<UIKRetargeter> IKRetargeterAsset = nullptr;
 
-	/** connect a custom retarget profile to modify the retargeter's settings at runtime.*/
-	UPROPERTY(EditAnywhere, BlueprintReadWrite, Category=Settings, meta=(PinHiddenByDefault))
+	// Connect a custom retarget profile to modify the retargeter's settings at runtime.
+	UPROPERTY(EditAnywhere, BlueprintReadWrite, Category = Settings, meta=(PinHiddenByDefault))
 	FRetargetProfile CustomRetargetProfile;
-
-	/* Toggle whether to print warnings about missing or incorrectly configured retarget configurations. */
-	UPROPERTY(EditAnywhere, BlueprintReadWrite, Category=Debug, meta = (NeverAsPin))
-	bool bSuppressWarnings = false;
 	
-	/* Copy curves from SouceMeshComponent. This will copy any curves the source/target Skeleton have in common. */
+	// Copy curves from the Source Mesh Component. This will copy any curves the source/target Skeleton have in common.
 	UPROPERTY(EditAnywhere, BlueprintReadWrite, Category = Settings, meta = (NeverAsPin))
 	bool bCopyCurves = true;
-
-	/*
-	* Max LOD that this node is allowed to run.
-	* For example if you have LODThreshold at 2, it will run until LOD 2 (based on 0 index) when the component LOD becomes 3, it will stop update/evaluate
-	* A value of -1 forces the node to execute at all LOD levels.
-	*/
+	
+	// Max LOD that this node is allowed to run.
+	// For example if you have LODThreshold at 2, it will run until LOD 2 (based on 0 index) when the component LOD becomes 3, it will stop update/evaluate
+	// A value of -1 forces the node to execute at all LOD levels.
 	UPROPERTY(EditAnywhere, BlueprintReadWrite, Category = Performance, meta = (DisplayName = "LOD Threshold"))
 	int32 LODThreshold = -1;
-
-	/*
-	* Max LOD that IK is allowed to run.
-	* For example if you have LODThresholdForIK at 2, it will skip the IK pass on LODs 3 and greater.
-	* This only disables IK and does not affect the Root or FK passes.
-	* A value of -1 forces the node to execute at all LOD levels.
-	*/
+	
+	// Max LOD that IK is allowed to run.
+	// For example if you have LODThresholdForIK at 2, it will skip the IK pass on LODs 3 and greater.
+	// This only disables IK and does not affect the Root or FK passes.
+	// A value of -1 forces the node to execute at all LOD levels. 
 	UPROPERTY(EditAnywhere, BlueprintReadWrite, Category = Performance, meta = (DisplayName = "IK LOD Threshold"))
 	int32 LODThresholdForIK = -1;
+
+	// Toggle whether to print warnings about missing or incorrectly configured retarget configurations.
+	UPROPERTY(EditAnywhere, BlueprintReadWrite, Category = Debug, meta = (NeverAsPin))
+	bool bSuppressWarnings = false;
+
+	UPROPERTY()
+	bool bUseAttachedParent_DEPRECATED = true;
 	
 	// FAnimNode_Base interface
 	virtual void Initialize_AnyThread(const FAnimationInitializeContext& Context) override;
@@ -69,42 +80,54 @@ struct IKRIG_API FAnimNode_RetargetPoseFromMesh : public FAnimNode_Base
 	virtual int32 GetLODThreshold() const override { return LODThreshold; }
 	// End of FAnimNode_Base interface
 
-	/** Access to the runtime processor */
-	void CreateRetargetProcessorIfNeeded(UObject* Outer);
+	// access to the runtime processor
 	UIKRetargetProcessor* GetRetargetProcessor() const;
+	void CreateRetargetProcessorIfNeeded(UObject* Outer);
+
+	// returns true if processor is setup and ready to go, false otherwise
+	bool TryInitializeProcessor(USkeletalMeshComponent* TargetMeshComponent);
 
-	/** returns true if processor is setup and ready to go, false otherwise */
-	bool InitializeRetargetProcessor(const TObjectPtr<USkeletalMeshComponent> TargetMeshComponent);
+	// load deprecated properties
+	bool Serialize(FArchive& Ar);
+	void PostSerialize(const FArchive& Ar);
 
 private:
 	
-	/** copies the source mesh pose (on main thread) */
-	void CopyBoneTransformsFromSource(USkeletalMeshComponent* TargetMeshComponent);
-	/** indirection to account for Leader Pose Component setup */
-	TObjectPtr<USkeletalMeshComponent> GetComponentToCopyPoseFrom() const;
 	// get a retarget profile that includes settings from the asset, plus any custom profile overrides
 	FRetargetProfile GetMergedRetargetProfile(bool bEnableIK) const;
 
-	/** the runtime processor used to run the retarget and generate new poses */
+	// the runtime processor used to run the retarget and generate new poses
 	UPROPERTY(Transient)
 	TObjectPtr<UIKRetargetProcessor> Processor = nullptr;
 
 	// cached transforms, copied on the game thread
-	TArray<FTransform> SourceMeshComponentSpaceBoneTransforms;
+	TArray<FTransform> PoseToRetargetFromComponentSpace;
+
+	// reusable buffer when retargeting a pose from the anim graph pose pin
+	TArray<FTransform> InputLocalTransforms;
 
-	// mapping from required bones to actual bones within the target skeleton
-	TArray< TPair<int32, int32> > RequiredToTargetBoneMapping;
+	// mapping from required compact bone indices to target mesh bones expected by the retargeter
+	TArray< TPair<int32, int32> > CompactToTargetBoneIndexMap;
 
 	// cached curves, copied on the game thread in PreUpdate()
 	FBlendedHeapCurve SourceCurves;
-
-	// remap curves for CurveRemapOp if one is present in the RetargetOp stack
-	void CopyAndRemapCurvesFromSourceToTarget(FBlendedCurve& OutputCurves) const;
 	
-	// update map of curve values containing speeds used for IK planting
-	void UpdateSpeedValuesFromCurves();
 	// map of curve names to values, passed to retargeter for IK planting (copied from source mesh)
 	TMap<FName, float> SpeedValuesFromCurves;
-	// the accumulated delta time this tick
+	
+	// the delta time this tick
 	float DeltaTime;
+
+	// used to determine when to look for a parent component (prevents constantly searching component hierarchy)
+	bool bSearchedForParentComponent = false;
+};
+
+template<>
+struct TStructOpsTypeTraits<FAnimNode_RetargetPoseFromMesh> : public TStructOpsTypeTraitsBase2<FAnimNode_RetargetPoseFromMesh>
+{
+	enum
+	{
+		WithSerializer = true,
+		WithPostSerialize = true,
+	};
 };
diff --git a/Engine/Plugins/Animation/IKRig/Source/IKRig/Public/IKRigObjectVersion.h b/Engine/Plugins/Animation/IKRig/Source/IKRig/Public/IKRigObjectVersion.h
--- a/Engine/Plugins/Animation/IKRig/Source/IKRig/Public/IKRigObjectVersion.h
+++ b/Engine/Plugins/Animation/IKRig/Source/IKRig/Public/IKRigObjectVersion.h
@@ -24,6 +24,9 @@ struct IKRIG_API FIKRigObjectVersion
 		// Pin Bone "Maintain Offset" converted from bool to enum
 		PinBoneTypeAndOffsetsUpgraded,
 
+		// we switched "Use Attach Parent" from a bool to an enum on retarget pose from mesh node
+		UseAttachedParentDeprecated,
+
 		// -----<new versions can be added above this line>-------------------------------------------------
 		VersionPlusOne,
 		LatestVersion = VersionPlusOne - 1
diff --git a/Engine/Plugins/Animation/IKRig/Source/IKRig/Public/Retargeter/IKRetargetProcessor.h b/Engine/Plugins/Animation/IKRig/Source/IKRig/Public/Retargeter/IKRetargetProcessor.h
--- a/Engine/Plugins/Animation/IKRig/Source/IKRig/Public/Retargeter/IKRetargetProcessor.h
+++ b/Engine/Plugins/Animation/IKRig/Source/IKRig/Public/Retargeter/IKRetargetProcessor.h
@@ -35,11 +35,11 @@ struct IKRIG_API FRetargetSkeleton
 	int32 RetargetPoseVersion;				// the version of the retarget pose this was initialized with (transient)
 	float RetargetPoseScale;				// the scale of the retarget pose this was initialized with (transient)
 	
-	USkeletalMesh* SkeletalMesh;			// the skeletal mesh this was initialized with
+	const USkeletalMesh* SkeletalMesh;		// the skeletal mesh this was initialized with
 	TArray<FName> ChainThatContainsBone;	// record which chain is actually controlling each bone
 
 	void Initialize(
-		USkeletalMesh* InSkeletalMesh,
+		const USkeletalMesh* InSkeletalMesh,
 		const TArray<FBoneChain>& BoneChains,
 		const FName InRetargetPoseName,
 		const FIKRetargetPose* RetargetPose,
@@ -88,7 +88,7 @@ struct IKRIG_API FRetargetSkeleton
 
 	void GetChildrenIndicesRecursive(const int32 BoneIndex, TArray<int32>& OutChildren) const;
 	
-	bool IsParentOfChild(const int32 PotentialParentIndex, const int32 ChildBoneIndex) const;
+	bool IsParentOf(const int32 PotentialParentIndex, const int32 ChildBoneIndex) const;
 
 private:
 	
@@ -100,12 +100,13 @@ private:
 struct FTargetSkeleton : public FRetargetSkeleton
 {
 	TArray<FTransform> OutputGlobalPose;
+	TArray<FTransform> InputLocalPose;
 	// true for bones that are in a target chain that is ALSO mapped to a source chain
 	// ie, bones that are actually posed based on a mapped source chain
 	TArray<bool> IsBoneRetargeted;
 
 	void Initialize(
-		USkeletalMesh* InSkeletalMesh,
+		const USkeletalMesh* InSkeletalMesh,
 		const TArray<FBoneChain>& BoneChains,
 		const FName InRetargetPoseName,
 		const FIKRetargetPose* RetargetPose,
@@ -115,7 +116,7 @@ struct FTargetSkeleton : public FRetargetSkeleton
 
 	void SetBoneIsRetargeted(const int32 BoneIndex, const bool IsRetargeted);
 
-	void UpdateGlobalTransformsAllNonRetargetedBones(TArray<FTransform>& InOutGlobalPose);
+	void UpdateGlobalTransformsAllNonRetargetedBones(const TArray<FTransform>& LocalPose, TArray<FTransform>& InOutGlobalPose);
 };
 
 /** resolving an FBoneChain to an actual skeleton, used to validate compatibility and get all chain indices */
@@ -487,6 +488,28 @@ enum class ERetargetBoneSpace : uint8
 	Local
 };
 
+// the system to do a straight copy of bone transform with matching names (does not do any retargeting!)
+struct FBasePoseCopier
+{
+	void CopyPoseFromSource(
+		const FName InRootBoneToCopy,
+		const FRetargetSkeleton& InSourceSkeleton,
+		const FRetargetSkeleton& InTargetSkeleton,
+		TArray<FTransform>& InOutSourcePose,
+		TArray<FTransform>& OutTargetPose);
+
+	int32 GetIndexOfRootBoneToCopy() const { return RootBoneToCopyIndex; };
+
+private:
+	
+	// mapping from source bones to target bone indices
+	// NOTE: this map will only contain bones with the same name in both skeletons that are below the RootBoneToCopy (if specified)
+	TArray< TPair<int32, int32> > SourceToTargetBoneIndexMap;
+	// the bone to start copying from (includes all children recursive)
+	FName RootBoneToCopy = NAME_None;
+	int32 RootBoneToCopyIndex = INDEX_NONE;
+};
+
 /** The runtime processor that converts an input pose from a source skeleton into an output pose on a target skeleton.
  * To use:
  * 1. Initialize a processor with a Source/Target skeletal mesh and a UIKRetargeter asset.
@@ -511,9 +534,9 @@ public:
 	* @warning - Initialization does a lot of validation and can fail for many reasons. Check bIsLoadedAndValid afterwards.
 	*/
 	void Initialize(
-		USkeletalMesh *SourceSkeleton,
-		USkeletalMesh *TargetSkeleton,
-		UIKRetargeter* InRetargeterAsset,
+		const USkeletalMesh *SourceSkeleton,
+		const USkeletalMesh *TargetSkeleton,
+		const UIKRetargeter* InRetargeterAsset,
 		const FRetargetProfile& Settings,
 		const bool bSuppressWarnings=false);
 
@@ -522,14 +545,14 @@ public:
 	* @param InSourceGlobalPose -  is the source mesh input pose in Component/Global space
 	* @param SpeedValuesFromCurves - the speed of each curve used by speed planting (blended in anim graph)
 	* @param DeltaTime -  time since last tick in seconds (used by speed planting)
-	* @param Settings -  the retarget profile to use for this update
+	* @param Profile -  the retarget profile to use for this update
 	* @return The retargeted Component/Global space pose for the target skeleton
 	*/
 	TArray<FTransform>& RunRetargeter(
 		TArray<FTransform>& InSourceGlobalPose,
 		const TMap<FName, float>& SpeedValuesFromCurves,
 		const float DeltaTime,
-		const FRetargetProfile& Settings);
+		const FRetargetProfile& Profile);
 	
 	/** Apply the settings stored in the IKRig asset. */
 	void CopyIKRigSettingsFromAsset();
@@ -549,9 +572,9 @@ public:
 
 	/** Get whether this processor was initialized with these skeletal meshes and retarget asset*/
 	bool WasInitializedWithTheseAssets(
-		const TObjectPtr<USkeletalMesh> InSourceMesh,
-		const TObjectPtr<USkeletalMesh> InTargetMesh,
-		const TObjectPtr<UIKRetargeter> InRetargetAsset);
+		const USkeletalMesh* InSourceMesh,
+		const USkeletalMesh* InTargetMesh,
+		const UIKRetargeter* InRetargetAsset) const;
 
 	/** Get the currently running IK Rig processor for the target */
 	UIKRigProcessor* GetTargetIKRigProcessor() const { return IKRigProcessor; };
@@ -622,7 +645,7 @@ private:
 	bool bIKRigInitialized = false;
 
 	/** The source asset this processor was initialized with. */
-	UIKRetargeter* RetargeterAsset = nullptr;
+	const UIKRetargeter* RetargeterAsset = nullptr;
 
 	/** The internal data structures used to represent the SOURCE skeleton / pose during retargeter.*/
 	FRetargetSkeleton SourceSkeleton;
@@ -643,6 +666,9 @@ private:
 	/** The Source/Target pair of Root Bones retargeted with scaled translation */
 	FRootRetargeter RootRetargeter;
 
+	/** The system to do a straight copy of source-to-target bone transforms based on name */
+	FBasePoseCopier BasePoseCopier;
+
 	/** The currently used global settings (driven either by source asset or a profile) */
 	FRetargetGlobalSettings GlobalSettings;
 
@@ -668,6 +694,9 @@ private:
 
 	// Initialize the retarget op stack
 	bool InitializeOpStack(const TArray<TObjectPtr<URetargetOpBase>>& OpStackFromAsset);
+
+	/** Internal retarget phase that does simple bone-to-bone copying from source*/
+	void GenerateBasePoses(TArray<FTransform>& InSourceGlobalPose);
 	
 	/** Internal retarget phase for the root. */
 	void RunRootRetarget(const TArray<FTransform>& InGlobalTransforms, TArray<FTransform>& OutGlobalTransforms);
diff --git a/Engine/Plugins/Animation/IKRig/Source/IKRig/Public/Retargeter/IKRetargetSettings.h b/Engine/Plugins/Animation/IKRig/Source/IKRig/Public/Retargeter/IKRetargetSettings.h
--- a/Engine/Plugins/Animation/IKRig/Source/IKRig/Public/Retargeter/IKRetargetSettings.h
+++ b/Engine/Plugins/Animation/IKRig/Source/IKRig/Public/Retargeter/IKRetargetSettings.h
@@ -318,6 +318,16 @@ struct IKRIG_API FRetargetGlobalSettings
 	UPROPERTY(EditAnywhere, BlueprintReadWrite, Category = RetargetPhases)
 	bool bEnablePost = true;
 
+	/** When true, will copy all the source bone transforms with matching names to use as a base pose. This can be useful for partial retargeting.
+	 * NOTE: no retargeting is applied to the bone transforms, they are assumed to be fully compatible between source/target skeletons.
+	 */
+	UPROPERTY(EditAnywhere, BlueprintReadWrite, Category = RetargetPhases)
+	bool bCopyBasePose = false;
+
+	/** Filters the bones to copy when using "Copy Base Pose". If specified, will only copy all children of the specified bone (inclusive).*/
+	UPROPERTY(EditAnywhere, BlueprintReadWrite, Category = RetargetPhases, meta = (EditCondition = "bCopyBasePose"))
+	FName CopyBasePoseRoot = NAME_None;
+
 	/** Range 0.01 to +inf. Default 1. Scales the incoming source pose. Affects entire skeleton and all IK goals.*/
 	UPROPERTY(EditAnywhere, BlueprintReadWrite, Category = Scale, meta = (ClampMin = "0.01", UIMin = "0.01", UIMax = "10.0"))
 	float SourceScaleFactor = 1.0f;
diff --git a/Engine/Plugins/Animation/IKRig/Source/IKRig/Public/Retargeter/IKRetargeter.h b/Engine/Plugins/Animation/IKRig/Source/IKRig/Public/Retargeter/IKRetargeter.h
--- a/Engine/Plugins/Animation/IKRig/Source/IKRig/Public/Retargeter/IKRetargeter.h
+++ b/Engine/Plugins/Animation/IKRig/Source/IKRig/Public/Retargeter/IKRetargeter.h
@@ -261,7 +261,7 @@ public:
 	// get current version of the data (to compare against running processor instances)
 	int32 GetVersion() const { return Version; };
 	// do this after any edit that would require running instance to reinitialize
-	void IncrementVersion() { ++Version; };
+	void IncrementVersion() const { ++Version; };
 
 	// BLUEPRINT GETTERS 
 
@@ -384,7 +384,7 @@ private:
 	
 	// incremented by any edits that require re-initialization
 	UPROPERTY(Transient)
-	int32 Version = INDEX_NONE;
+	mutable int32 Version = INDEX_NONE;
 
 	// The rig to copy animation FROM.
 	UPROPERTY(EditAnywhere, Category = Source)
diff --git a/Engine/Plugins/Animation/IKRig/Source/IKRig/Public/Retargeter/RetargetOps/PinBoneOp.h b/Engine/Plugins/Animation/IKRig/Source/IKRig/Public/Retargeter/RetargetOps/PinBoneOp.h
--- a/Engine/Plugins/Animation/IKRig/Source/IKRig/Public/Retargeter/RetargetOps/PinBoneOp.h
+++ b/Engine/Plugins/Animation/IKRig/Source/IKRig/Public/Retargeter/RetargetOps/PinBoneOp.h
@@ -153,8 +153,8 @@ public:
 	FText Message;
 
 private:
-	USkeleton* SourceSkeletonAsset;
-	USkeleton* TargetSkeletonAsset;
+	const USkeleton* SourceSkeletonAsset;
+	const USkeleton* TargetSkeletonAsset;
 #endif
 
 private:
diff --git a/Engine/Plugins/Animation/IKRig/Source/IKRigDeveloper/Private/AnimGraphNode_RetargetPoseFromMesh.cpp b/Engine/Plugins/Animation/IKRig/Source/IKRigDeveloper/Private/AnimGraphNode_RetargetPoseFromMesh.cpp
--- a/Engine/Plugins/Animation/IKRig/Source/IKRigDeveloper/Private/AnimGraphNode_RetargetPoseFromMesh.cpp
+++ b/Engine/Plugins/Animation/IKRig/Source/IKRigDeveloper/Private/AnimGraphNode_RetargetPoseFromMesh.cpp
@@ -15,6 +15,11 @@ void UAnimGraphNode_RetargetPoseFromMesh::Draw(FPrimitiveDrawInterface* PDI, USk
 
 FText UAnimGraphNode_RetargetPoseFromMesh::GetNodeTitle(ENodeTitleType::Type TitleType) const
 {
+	if (Node.RetargetFrom == ERetargetSourceMode::SourcePosePin)
+	{
+		return LOCTEXT("AnimGraphNode_IKRetargeter_Title_FromInput", "Retarget Input Pose");
+	}
+	
 	return LOCTEXT("AnimGraphNode_IKRetargeter_Title", "Retarget Pose From Mesh");
 }
 
@@ -31,6 +36,18 @@ FEditorModeID UAnimGraphNode_RetargetPoseFromMesh::GetEditorMode() const
 void UAnimGraphNode_RetargetPoseFromMesh::CustomizePinData(UEdGraphPin* Pin, FName SourcePropertyName, int32 ArrayIndex) const
 {
 	Super::CustomizePinData(Pin, SourcePropertyName, ArrayIndex);
+
+	if (Pin->PinName == GET_MEMBER_NAME_STRING_CHECKED(FAnimNode_RetargetPoseFromMesh, SourceMeshComponent))
+	{
+		Pin->bHidden = Node.RetargetFrom != ERetargetSourceMode::CustomSkeletalMeshComponent;
+	}
+
+	if (Pin->PinName == GET_MEMBER_NAME_STRING_CHECKED(FAnimNode_RetargetPoseFromMesh, Source))
+	{
+		const bool bCopyingFromOtherComponent = Node.RetargetFrom != ERetargetSourceMode::SourcePosePin;
+		Pin->bHidden = bCopyingFromOtherComponent;
+		Pin->bNotConnectable = bCopyingFromOtherComponent;
+	}
 }
 
 UObject* UAnimGraphNode_RetargetPoseFromMesh::GetJumpTargetForDoubleClick() const
@@ -41,7 +58,7 @@ UObject* UAnimGraphNode_RetargetPoseFromMesh::GetJumpTargetForDoubleClick() cons
 void UAnimGraphNode_RetargetPoseFromMesh::PostEditChangeProperty(struct FPropertyChangedEvent& PropertyChangedEvent)
 {
 	const FName PropertyName = (PropertyChangedEvent.Property ? PropertyChangedEvent.Property->GetFName() : NAME_None);
-	if ((PropertyName == GET_MEMBER_NAME_STRING_CHECKED(FAnimNode_RetargetPoseFromMesh, bUseAttachedParent)))
+	if ((PropertyName == GET_MEMBER_NAME_STRING_CHECKED(FAnimNode_RetargetPoseFromMesh, RetargetFrom)))
 	{
 		ReconstructNode();
 	}
@@ -52,13 +69,13 @@ void UAnimGraphNode_RetargetPoseFromMesh::ValidateAnimNodeDuringCompilation(USke
 	Super::ValidateAnimNodeDuringCompilation(ForSkeleton, MessageLog);
 
 	// validate source mesh component is not null
-	if (!Node.bUseAttachedParent)
+	if (Node.RetargetFrom == ERetargetSourceMode::CustomSkeletalMeshComponent)
 	{
 		const bool bIsLinked = IsPinExposedAndLinked(GET_MEMBER_NAME_STRING_CHECKED(FAnimNode_RetargetPoseFromMesh, SourceMeshComponent));
 		const bool bIsBound = IsPinExposedAndBound(GET_MEMBER_NAME_STRING_CHECKED(FAnimNode_RetargetPoseFromMesh, SourceMeshComponent));
 		if (!(bIsLinked || bIsBound))
 		{
-			MessageLog.Error(TEXT("@@ is missing a Source Skeletal Mesh Component reference."), this);
+			MessageLog.Error(TEXT("@@ requires a source Skeletal Mesh Component to be plugged in."), this);
 			return;
 		}
 	}
@@ -79,14 +96,14 @@ void UAnimGraphNode_RetargetPoseFromMesh::ValidateAnimNodeDuringCompilation(USke
 	const UIKRigDefinition* SourceIKRig = Node.IKRetargeterAsset->GetIKRig(ERetargetSourceOrTarget::Source);
 	if (!SourceIKRig)
 	{
-		MessageLog.Warning(TEXT("@@ has IK Retargeter that is missing a source IK Rig asset."), this);
+		MessageLog.Warning(TEXT("@@ has an IK Retargeter that is missing a source IK Rig asset."), this);
 	}
 
 	// validate TARGET IK Rig asset has been assigned
 	const UIKRigDefinition* TargetIKRig = Node.IKRetargeterAsset->GetIKRig(ERetargetSourceOrTarget::Target);
 	if (!TargetIKRig)
 	{
-		MessageLog.Warning(TEXT("@@ has IK Retargeter that is missing a target IK Rig asset."), this);
+		MessageLog.Warning(TEXT("@@ has an IK Retargeter that is missing a target IK Rig asset."), this);
 	}
 
 	if (!(SourceIKRig && TargetIKRig))
diff --git a/Engine/Plugins/Animation/IKRig/Source/IKRigDeveloper/Public/AnimGraphNode_RetargetPoseFromMesh.h b/Engine/Plugins/Animation/IKRig/Source/IKRigDeveloper/Public/AnimGraphNode_RetargetPoseFromMesh.h
--- a/Engine/Plugins/Animation/IKRig/Source/IKRigDeveloper/Public/AnimGraphNode_RetargetPoseFromMesh.h
+++ b/Engine/Plugins/Animation/IKRig/Source/IKRigDeveloper/Public/AnimGraphNode_RetargetPoseFromMesh.h
@@ -11,7 +11,8 @@
 class FPrimitiveDrawInterface;
 class USkeletalMeshComponent;
 
-// Editor node for IKRig 
+// Facilitates runtime retargeting of an input pose using a user-specified IK Retargeter asset
+// Supports input poses from the animation graph, or from other skeletal mesh components
 UCLASS()
 class IKRIGDEVELOPER_API UAnimGraphNode_RetargetPoseFromMesh : public UAnimGraphNode_Base
 {
diff --git a/Engine/Plugins/Animation/IKRig/Source/IKRigEditor/Private/RetargetEditor/IKRetargetAnimInstance.cpp b/Engine/Plugins/Animation/IKRig/Source/IKRigEditor/Private/RetargetEditor/IKRetargetAnimInstance.cpp
--- a/Engine/Plugins/Animation/IKRig/Source/IKRigEditor/Private/RetargetEditor/IKRetargetAnimInstance.cpp
+++ b/Engine/Plugins/Animation/IKRig/Source/IKRigEditor/Private/RetargetEditor/IKRetargetAnimInstance.cpp
@@ -159,10 +159,10 @@ UIKRetargetProcessor* UIKRetargetAnimInstance::GetRetargetProcessor()
 	return RetargetNode.GetRetargetProcessor();
 }
 
-void UIKRetargetAnimInstance::ForceInitializeProcessor(const TObjectPtr<USkeletalMeshComponent> TargetComponent)
+void UIKRetargetAnimInstance::ForceInitializeProcessor(USkeletalMeshComponent* TargetComponent)
 {
 	RetargetNode.IKRetargeterAsset->IncrementVersion();
-	RetargetNode.InitializeRetargetProcessor(TargetComponent);
+	RetargetNode.TryInitializeProcessor(TargetComponent);
 }
 
 FAnimInstanceProxy* UIKRetargetAnimInstance::CreateAnimInstanceProxy()
diff --git a/Engine/Plugins/Animation/IKRig/Source/IKRigEditor/Private/RetargetEditor/IKRetargetAnimInstanceProxy.cpp b/Engine/Plugins/Animation/IKRig/Source/IKRigEditor/Private/RetargetEditor/IKRetargetAnimInstanceProxy.cpp
--- a/Engine/Plugins/Animation/IKRig/Source/IKRigEditor/Private/RetargetEditor/IKRetargetAnimInstanceProxy.cpp
+++ b/Engine/Plugins/Animation/IKRig/Source/IKRigEditor/Private/RetargetEditor/IKRetargetAnimInstanceProxy.cpp
@@ -112,7 +112,7 @@ void FIKRetargetAnimInstanceProxy::ConfigureAnimInstance(
 	if (SourceOrTarget == ERetargetSourceOrTarget::Target)
 	{
 		RetargetNode->IKRetargeterAsset = InIKRetargetAsset;
-		RetargetNode->bUseAttachedParent = false;
+		RetargetNode->RetargetFrom = ERetargetSourceMode::CustomSkeletalMeshComponent;
 		RetargetNode->SourceMeshComponent = InSourceMeshComponent;
 		if (UIKRetargetProcessor* Processor = RetargetNode->GetRetargetProcessor())
 		{
diff --git a/Engine/Plugins/Animation/IKRig/Source/IKRigEditor/Private/RetargetEditor/IKRetargetBatchOperation.cpp b/Engine/Plugins/Animation/IKRig/Source/IKRigEditor/Private/RetargetEditor/IKRetargetBatchOperation.cpp
--- a/Engine/Plugins/Animation/IKRig/Source/IKRigEditor/Private/RetargetEditor/IKRetargetBatchOperation.cpp
+++ b/Engine/Plugins/Animation/IKRig/Source/IKRigEditor/Private/RetargetEditor/IKRetargetBatchOperation.cpp
@@ -391,7 +391,7 @@ void UIKRetargetBatchOperation::ConvertAnimation(
 		TargetSeqController.NotifyPopulated(); // must set bIsPopulated=true otherwise "UpdateWithSkeleton" will early out
 		// this removes all bone tracks and reinitializes the FKControlRig to use the new skeleton.
 		// NOTE: we do NOT ResetModel() because we want to keep curves and attributes from the source
-		TargetSeqController.UpdateWithSkeleton(TargetSkeleton.SkeletalMesh->GetSkeleton(), bShouldTransact);
+		TargetSeqController.UpdateWithSkeleton(const_cast<USkeleton*>(TargetSkeleton.SkeletalMesh->GetSkeleton()), bShouldTransact);
 
 		// number of frames in this animation
 		const int32 NumFrames = SourceSequence->GetNumberOfSampledKeys();
@@ -406,7 +406,7 @@ void UIKRetargetBatchOperation::ConvertAnimation(
 
 		// ensure we evaluate the source animation using the skeletal mesh proportions that were evaluated in the viewport
 		FAnimPoseEvaluationOptions EvaluationOptions = FAnimPoseEvaluationOptions();
-		EvaluationOptions.OptionalSkeletalMesh = SourceSkeleton.SkeletalMesh;
+		EvaluationOptions.OptionalSkeletalMesh = const_cast<USkeletalMesh*>(SourceSkeleton.SkeletalMesh);
 		// ensure WYSIWYG with editor by ensuring the same root motion is applied to the pose, not to the component
 		EvaluationOptions.bExtractRootMotion = false;
 		EvaluationOptions.bIncorporateRootMotionIntoPose = true;
diff --git a/Engine/Plugins/Animation/IKRig/Source/IKRigEditor/Private/RetargetEditor/IKRetargetDetails.cpp b/Engine/Plugins/Animation/IKRig/Source/IKRigEditor/Private/RetargetEditor/IKRetargetDetails.cpp
--- a/Engine/Plugins/Animation/IKRig/Source/IKRigEditor/Private/RetargetEditor/IKRetargetDetails.cpp
+++ b/Engine/Plugins/Animation/IKRig/Source/IKRigEditor/Private/RetargetEditor/IKRetargetDetails.cpp
@@ -1381,7 +1381,6 @@ void FRetargetGlobalSettingsCustomization::CustomizeDetails(IDetailLayoutBuilder
 
 	auto GetPropertyHandle = [&DetailBuilder, &StructPropertyPath](const FString& PropertyName)
 	{
-		const FString EnableWarpingPropertyName = GET_MEMBER_NAME_STRING_CHECKED(FRetargetGlobalSettings, bWarping);
 		const FString PropertyPath = FString::Printf(TEXT("%s.%s"), *StructPropertyPath, *PropertyName);
 		return DetailBuilder.GetProperty(*PropertyPath);
 	};
@@ -1401,6 +1400,11 @@ void FRetargetGlobalSettingsCustomization::CustomizeDetails(IDetailLayoutBuilder
 	AlphaGroup.AddPropertyRow(GetPropertyHandle(GET_MEMBER_NAME_STRING_CHECKED(FRetargetGlobalSettings, bEnableFK)));
 	AlphaGroup.AddPropertyRow(GetPropertyHandle(GET_MEMBER_NAME_STRING_CHECKED(FRetargetGlobalSettings, bEnableIK)));
 	AlphaGroup.AddPropertyRow(GetPropertyHandle(GET_MEMBER_NAME_STRING_CHECKED(FRetargetGlobalSettings, bEnablePost)));
+
+	// copy pose group
+	IDetailGroup& CopyPoseGroup = SettingsCategory.AddGroup("CopyPose", LOCTEXT("CopyPose_Label", "Copy Pose from Source"), false, true);
+	CopyPoseGroup.AddPropertyRow(GetPropertyHandle(GET_MEMBER_NAME_STRING_CHECKED(FRetargetGlobalSettings, bCopyBasePose)));
+	CopyPoseGroup.AddPropertyRow(GetPropertyHandle(GET_MEMBER_NAME_STRING_CHECKED(FRetargetGlobalSettings, CopyBasePoseRoot)));
 	
 	// stride warping group
 	const FText WarpingTitleLabel = LOCTEXT("Warping_Label", "Stride Warping");
@@ -1409,17 +1413,16 @@ void FRetargetGlobalSettingsCustomization::CustomizeDetails(IDetailLayoutBuilder
 	// header with "enable" checkbox that disables whole group
 	const TSharedRef<IPropertyHandle> EnableWarpingPropertyHandle = GetPropertyHandle(GET_MEMBER_NAME_STRING_CHECKED(FRetargetGlobalSettings, bWarping));
 	WarpingGroup.HeaderProperty(EnableWarpingPropertyHandle).DisplayName(WarpingTitleLabel);
-	
-	const TAttribute<bool> EditCondition = TAttribute<bool>::Create([EnableWarpingPropertyHandle]()
+
+	// add all the warping parameters
+	const TAttribute<bool> WarpingEditCondition = TAttribute<bool>::Create([EnableWarpingPropertyHandle]()
 	{
 		bool bIsEnabledPropertyTrue = true;
 		EnableWarpingPropertyHandle->GetValue(bIsEnabledPropertyTrue);
 		return bIsEnabledPropertyTrue;
 	});
-	
-	// add all the warping parameters
-	WarpingGroup.AddPropertyRow(GetPropertyHandle(GET_MEMBER_NAME_STRING_CHECKED(FRetargetGlobalSettings, ForwardDirection))).EditCondition(EditCondition,nullptr);
-	WarpingGroup.AddPropertyRow(GetPropertyHandle(GET_MEMBER_NAME_STRING_CHECKED(FRetargetGlobalSettings, DirectionSource))).EditCondition(EditCondition,nullptr);
+	WarpingGroup.AddPropertyRow(GetPropertyHandle(GET_MEMBER_NAME_STRING_CHECKED(FRetargetGlobalSettings, ForwardDirection))).EditCondition(WarpingEditCondition,nullptr);
+	WarpingGroup.AddPropertyRow(GetPropertyHandle(GET_MEMBER_NAME_STRING_CHECKED(FRetargetGlobalSettings, DirectionSource))).EditCondition(WarpingEditCondition,nullptr);
 
 	// get target chain options for combobox
 	TargetChainOptions.Reset();
@@ -1472,9 +1475,9 @@ void FRetargetGlobalSettingsCustomization::CustomizeDetails(IDetailLayoutBuilder
 		]
 	];
 	
-	WarpingGroup.AddPropertyRow(GetPropertyHandle(GET_MEMBER_NAME_STRING_CHECKED(FRetargetGlobalSettings, WarpForwards))).EditCondition(EditCondition,nullptr);
-	WarpingGroup.AddPropertyRow(GetPropertyHandle(GET_MEMBER_NAME_STRING_CHECKED(FRetargetGlobalSettings, WarpSplay))).EditCondition(EditCondition,nullptr);
-	WarpingGroup.AddPropertyRow(GetPropertyHandle(GET_MEMBER_NAME_STRING_CHECKED(FRetargetGlobalSettings, SidewaysOffset))).EditCondition(EditCondition,nullptr);
+	WarpingGroup.AddPropertyRow(GetPropertyHandle(GET_MEMBER_NAME_STRING_CHECKED(FRetargetGlobalSettings, WarpForwards))).EditCondition(WarpingEditCondition,nullptr);
+	WarpingGroup.AddPropertyRow(GetPropertyHandle(GET_MEMBER_NAME_STRING_CHECKED(FRetargetGlobalSettings, WarpSplay))).EditCondition(WarpingEditCondition,nullptr);
+	WarpingGroup.AddPropertyRow(GetPropertyHandle(GET_MEMBER_NAME_STRING_CHECKED(FRetargetGlobalSettings, SidewaysOffset))).EditCondition(WarpingEditCondition,nullptr);
 }
 
 void FRetargetOpStackCustomization::CustomizeDetails(IDetailLayoutBuilder& DetailBuilder)
diff --git a/Engine/Plugins/Animation/IKRig/Source/IKRigEditor/Private/RetargetEditor/IKRetargetPoseExporter.cpp b/Engine/Plugins/Animation/IKRig/Source/IKRigEditor/Private/RetargetEditor/IKRetargetPoseExporter.cpp
--- a/Engine/Plugins/Animation/IKRig/Source/IKRigEditor/Private/RetargetEditor/IKRetargetPoseExporter.cpp
+++ b/Engine/Plugins/Animation/IKRig/Source/IKRigEditor/Private/RetargetEditor/IKRetargetPoseExporter.cpp
@@ -568,7 +568,7 @@ FReply FIKRetargetPoseExporter::OnImportPoseFromSequence()
 	FIKRetargetPose& ImportedRetargetPose = ControllerPtr->AssetController->GetRetargetPoses(SourceOrTarget)[NewPose];
 	
 	// get all imported bone transforms and record them in the retarget pose
-	FReferenceSkeleton& RefSkeleton = Mesh->GetRefSkeleton();
+	const FReferenceSkeleton& RefSkeleton = Mesh->GetRefSkeleton();
 	const TArray<FTransform>& RefPose = RefSkeleton.GetRefBonePose();
 	int32 NumBones = RefSkeleton.GetNum();
 	const FName RootBoneName = ControllerPtr->AssetController->GetRetargetRootBone(ControllerPtr->GetSourceOrTarget());
diff --git a/Engine/Plugins/Animation/IKRig/Source/IKRigEditor/Private/RetargetEditor/IKRetargeterPoseGenerator.cpp b/Engine/Plugins/Animation/IKRig/Source/IKRigEditor/Private/RetargetEditor/IKRetargeterPoseGenerator.cpp
--- a/Engine/Plugins/Animation/IKRig/Source/IKRigEditor/Private/RetargetEditor/IKRetargeterPoseGenerator.cpp
+++ b/Engine/Plugins/Animation/IKRig/Source/IKRigEditor/Private/RetargetEditor/IKRetargeterPoseGenerator.cpp
@@ -506,7 +506,7 @@ void FRetargetAutoPoseGenerator::GetRefPoseVerticesAndWeights(
 	OutPositions.Reset();
 	
 	const FRetargetSkeleton& Skeleton = Processor->GetSkeleton(SourceOrTarget);
-	USkeletalMesh* Mesh = Skeleton.SkeletalMesh;
+	const USkeletalMesh* Mesh = Skeleton.SkeletalMesh;
 	constexpr int32 LODIndex = 0;
 	const FMeshDescription *MeshDescription = Mesh->GetMeshDescription(LODIndex);
 	if (!MeshDescription || MeshDescription->IsEmpty())
diff --git a/Engine/Plugins/Animation/IKRig/Source/IKRigEditor/Public/RetargetEditor/IKRetargetAnimInstance.h b/Engine/Plugins/Animation/IKRig/Source/IKRigEditor/Public/RetargetEditor/IKRetargetAnimInstance.h
--- a/Engine/Plugins/Animation/IKRig/Source/IKRigEditor/Public/RetargetEditor/IKRetargetAnimInstance.h
+++ b/Engine/Plugins/Animation/IKRig/Source/IKRigEditor/Public/RetargetEditor/IKRetargetAnimInstance.h
@@ -62,7 +62,7 @@ public:
 
 	UIKRetargetProcessor* GetRetargetProcessor();
 
-	void ForceInitializeProcessor(const TObjectPtr<USkeletalMeshComponent> TargetComponent);
+	void ForceInitializeProcessor(USkeletalMeshComponent* TargetComponent);
 
 	const TArray<FTransform>& GetLocalRetargetPose() const { return PreviewPoseNode.GetLocalRetargetPose(); };
 	const TArray<FTransform>& GetGlobalRetargetPose() const { return PreviewPoseNode.GetGlobalRetargetPose(); };
diff --git a/Engine/Plugins/Animation/PerformanceCaptureCore/Source/PerformanceCaptureCore/Private/RetargetInstanceProxy.cpp b/Engine/Plugins/Animation/PerformanceCaptureCore/Source/PerformanceCaptureCore/Private/RetargetInstanceProxy.cpp
--- a/Engine/Plugins/Animation/PerformanceCaptureCore/Source/PerformanceCaptureCore/Private/RetargetInstanceProxy.cpp
+++ b/Engine/Plugins/Animation/PerformanceCaptureCore/Source/PerformanceCaptureCore/Private/RetargetInstanceProxy.cpp
@@ -70,7 +70,7 @@ void FRetargetAnimInstanceProxy::UpdateAnimationNode(const FAnimationUpdateConte
 void FRetargetAnimInstanceProxy::ConfigureAnimInstanceProxy(UIKRetargeter* InIKRetargetAsset, TWeakObjectPtr<USkeletalMeshComponent> InSourceMeshComponent, FRetargetProfile &InRetargetProfile)
 {
 	RetargetNode->IKRetargeterAsset = InIKRetargetAsset;
-	RetargetNode->bUseAttachedParent = false;
+	RetargetNode->RetargetFrom = ERetargetSourceMode::CustomSkeletalMeshComponent;
 	RetargetNode->SourceMeshComponent = InSourceMeshComponent;
 	RetargetNode->CustomRetargetProfile = InRetargetProfile;
 	
