Using hard instead of soft ref for IKRig so they are not always sync loaded.
#rnx

diff --git a/Engine/Plugins/Animation/IKRig/Source/IKRig/Public/Retargeter/IKRetargeter.h b/Engine/Plugins/Animation/IKRig/Source/IKRig/Public/Retargeter/IKRetargeter.h
--- a/Engine/Plugins/Animation/IKRig/Source/IKRig/Public/Retargeter/IKRetargeter.h
+++ b/Engine/Plugins/Animation/IKRig/Source/IKRig/Public/Retargeter/IKRetargeter.h
@@ -388,7 +388,7 @@ private:
 
 	// The rig to copy animation FROM.
 	UPROPERTY(EditAnywhere, Category = Source)
-	TSoftObjectPtr<UIKRigDefinition> SourceIKRigAsset = nullptr;
+	TObjectPtr<UIKRigDefinition> SourceIKRigAsset = nullptr;
 
 #if WITH_EDITORONLY_DATA
 	// Optional. Override the Skeletal Mesh to copy animation from. Uses the preview mesh from the Source IK Rig asset by default. 
@@ -398,7 +398,7 @@ private:
 	
 	// The rig to copy animation TO.
 	UPROPERTY(EditAnywhere, Category = Target)
-	TSoftObjectPtr<UIKRigDefinition> TargetIKRigAsset = nullptr;
+	TObjectPtr<UIKRigDefinition> TargetIKRigAsset = nullptr;
 
 #if WITH_EDITORONLY_DATA
 	// Optional. Override the Skeletal Mesh to preview the retarget on. Uses the preview mesh from the Target IK Rig asset by default. 
