[IK Rig] Renaming FBIK solver cpp

diff --git a/Engine/Plugins/Animation/IKRig/Source/IKRig/Private/Rig/Solvers/IKRig_FBIKSolver.cpp b/Engine/Plugins/Animation/IKRig/Source/IKRig/Private/Rig/Solvers/IKRig_FBIKSolver.cpp
new file mode 100644
--- /dev/null
+++ b/Engine/Plugins/Animation/IKRig/Source/IKRig/Private/Rig/Solvers/IKRig_FBIKSolver.cpp
@@ -0,0 +1,402 @@
+// Copyright Epic Games, Inc. All Rights Reserved.
+
+#include "Rig/Solvers/IKRig_FBIKSolver.h"
+
+#include "Rig/IKRigDataTypes.h"
+#include "Rig/IKRigSkeleton.h"
+
+#include UE_INLINE_GENERATED_CPP_BY_NAME(IKRig_FBIKSolver)
+
+#define LOCTEXT_NAMESPACE "UIKRigFBIKSolver"
+
+void UIKRigFBIKSolver::Initialize(const FIKRigSkeleton& InSkeleton)
+{	
+	// check how many effectors are assigned to a bone
+	int NumEffectors = 0;
+	for (const UIKRig_FBIKEffector* Effector : Effectors)
+	{
+		if (InSkeleton.GetBoneIndexFromName(Effector->BoneName) != INDEX_NONE)
+		{
+			++NumEffectors; // bone is set and exists!
+		}
+	}
+
+	// validate inputs are ready to be initialized
+	const bool bHasEffectors = NumEffectors > 0;
+	const bool bRootIsAssigned = RootBone != NAME_None;
+	if (!(bHasEffectors && bRootIsAssigned))
+	{
+		return; // not setup yet
+	}
+
+	// reset all internal data
+	Solver.Reset();
+
+	// create bones
+	for (int BoneIndex = 0; BoneIndex < InSkeleton.BoneNames.Num(); ++BoneIndex)
+	{
+		const FName& Name = InSkeleton.BoneNames[BoneIndex];
+
+		// get the parent bone solver index
+		const int32 ParentIndex = InSkeleton.GetParentIndexThatIsNotExcluded(BoneIndex);
+		const FTransform OrigTransform = InSkeleton.RefPoseGlobal[BoneIndex];
+		const FVector InOrigPosition = OrigTransform.GetLocation();
+		const FQuat InOrigRotation = OrigTransform.GetRotation();
+		const bool bIsRoot = Name == RootBone;
+		Solver.AddBone(Name, ParentIndex, InOrigPosition, InOrigRotation, bIsRoot);
+	}
+
+	// create effectors
+	for (UIKRig_FBIKEffector* Effector : Effectors)
+	{
+		Effector->IndexInSolver = Solver.AddEffector(Effector->BoneName);
+	}
+		
+	// initialize solver
+	Solver.Initialize();
+}
+
+void UIKRigFBIKSolver::Solve(FIKRigSkeleton& IKRigSkeleton, const FIKRigGoalContainer& Goals)
+{
+	if (!Solver.IsReadyToSimulate())
+	{
+		return;
+	}
+
+	if (Solver.GetNumBones() != IKRigSkeleton.BoneNames.Num())
+	{
+		return;
+	}
+
+	TArray<FTransform>& InOutTransforms = IKRigSkeleton.CurrentPoseGlobal;
+	
+	// set bones to input pose
+	for(int32 BoneIndex = 0; BoneIndex < Solver.GetNumBones(); BoneIndex++)
+	{
+		Solver.SetBoneTransform(BoneIndex, InOutTransforms[BoneIndex]);
+	}
+
+	// update bone settings
+	for (const UIKRig_FBIKBoneSettings* BoneSetting : BoneSettings)
+	{
+		const int32 BoneIndex = Solver.GetBoneIndex(BoneSetting->Bone);
+		if (PBIK::FBoneSettings* InternalSettings = Solver.GetBoneSettings(BoneIndex))
+		{
+			BoneSetting->CopyToCoreStruct(*InternalSettings);
+		}
+	}
+
+	// update effectors
+	for (const UIKRig_FBIKEffector* Effector : Effectors)
+	{
+		if (Effector->IndexInSolver < 0)
+		{
+			continue;
+		}
+		
+		const FIKRigGoal* Goal = Goals.FindGoalByName(Effector->GoalName);
+		if (!Goal)
+		{
+			return;
+		}
+
+		PBIK::FEffectorSettings Settings;
+		Settings.PositionAlpha = 1.0f; // this is constant because IKRig manages offset alphas itself
+		Settings.RotationAlpha = 1.0f; // this is constant because IKRig manages offset alphas itself
+		Settings.StrengthAlpha = Effector->StrengthAlpha;
+		Settings.ChainDepth = Effector->ChainDepth;
+		Settings.PullChainAlpha = Effector->PullChainAlpha;
+		Settings.PinRotation = Effector->PinRotation;
+		
+		Solver.SetEffectorGoal(
+			Effector->IndexInSolver,
+			Goal->FinalBlendedPosition,
+			Goal->FinalBlendedRotation,
+			Settings);
+	}
+
+	// update settings
+	FPBIKSolverSettings Settings;
+	Settings.Iterations = Iterations;
+	Settings.SubIterations = SubIterations;
+	Settings.MassMultiplier = MassMultiplier;
+	Settings.bAllowStretch = bAllowStretch;
+	Settings.RootBehavior = RootBehavior;
+	Settings.PrePullRootSettings = PrePullRootSettings;
+	Settings.GlobalPullChainAlpha = PullChainAlpha;
+	Settings.MaxAngle = MaxAngle;
+	Settings.OverRelaxation = OverRelaxation;
+	Settings.bStartSolveFromInputPose_DEPRECATED = bStartSolveFromInputPose_DEPRECATED;
+
+	// solve
+	Solver.Solve(Settings);
+
+	// copy transforms back
+	for(int32 BoneIndex = 0; BoneIndex < Solver.GetNumBones(); BoneIndex++)
+	{
+		Solver.GetBoneGlobalTransform(BoneIndex, InOutTransforms[BoneIndex]);
+	}
+}
+
+void UIKRigFBIKSolver::GetBonesWithSettings(TSet<FName>& OutBonesWithSettings) const
+{
+	for (UIKRig_FBIKBoneSettings* BoneSetting : BoneSettings)
+	{
+		if (BoneSetting)
+		{
+			OutBonesWithSettings.Add(BoneSetting->Bone);
+		}
+	}
+}
+
+void UIKRigFBIKSolver::UpdateSolverSettings(UIKRigSolver* InSettings)
+{
+	if(UIKRigFBIKSolver* Settings = Cast<UIKRigFBIKSolver>(InSettings))
+	{
+		Iterations = Settings->Iterations;
+		SubIterations = Settings->SubIterations;
+		MassMultiplier = Settings->MassMultiplier;
+		bAllowStretch = Settings->bAllowStretch;
+		RootBehavior = Settings->RootBehavior;
+		PrePullRootSettings = Settings->PrePullRootSettings;
+		PullChainAlpha = Settings->PullChainAlpha;
+		MaxAngle = Settings->MaxAngle;
+		OverRelaxation = Settings->OverRelaxation;
+		bStartSolveFromInputPose_DEPRECATED = Settings->bStartSolveFromInputPose_DEPRECATED;
+
+		// copy effector settings
+		for (const UIKRig_FBIKEffector* InEffector : Settings->Effectors)
+		{
+			for (UIKRig_FBIKEffector* Effector : Effectors)
+			{
+				if (Effector && Effector->GoalName == InEffector->GoalName)
+				{
+					Effector->CopySettings(InEffector);
+					break;
+				}
+			}
+		}
+
+		// copy bone settings
+		for (const UIKRig_FBIKBoneSettings* InBoneSetting : Settings->BoneSettings)
+		{
+			for (UIKRig_FBIKBoneSettings* BoneSetting : BoneSettings)
+			{
+				if (BoneSetting && BoneSetting->Bone == InBoneSetting->Bone)
+				{
+					BoneSetting->CopySettings(InBoneSetting);
+					break;	
+				}
+			}
+		}
+	}
+}
+
+void UIKRigFBIKSolver::RemoveGoal(const FName& GoalName)
+{
+	// ensure goal even exists in this solver
+	const int32 GoalIndex = GetIndexOfGoal(GoalName);
+	if (GoalIndex == INDEX_NONE)
+	{
+		return;
+	}
+
+	// remove it
+	Effectors.RemoveAt(GoalIndex);
+}
+
+bool UIKRigFBIKSolver::IsGoalConnected(const FName& GoalName) const
+{
+	return GetIndexOfGoal(GoalName) != INDEX_NONE;
+}
+
+#if WITH_EDITOR
+
+FText UIKRigFBIKSolver::GetNiceName() const
+{
+	return FText(LOCTEXT("SolverName", "Full Body IK"));
+}
+
+bool UIKRigFBIKSolver::GetWarningMessage(FText& OutWarningMessage) const
+{
+	if (RootBone == NAME_None)
+	{
+		OutWarningMessage = LOCTEXT("MissingRoot", "Missing root bone.");
+		return true;
+	}
+
+	if (Effectors.IsEmpty())
+	{
+		OutWarningMessage = LOCTEXT("MissingGoal", "Missing goals.");
+		return true;
+	}
+	
+	return false;
+}
+
+void UIKRigFBIKSolver::AddGoal(const UIKRigEffectorGoal& NewGoal)
+{
+	UIKRig_FBIKEffector* NewEffector = NewObject<UIKRig_FBIKEffector>(this, UIKRig_FBIKEffector::StaticClass());
+	NewEffector->GoalName = NewGoal.GoalName;
+	NewEffector->BoneName = NewGoal.BoneName;
+	Effectors.Add(NewEffector);
+}
+
+void UIKRigFBIKSolver::RenameGoal(const FName& OldName, const FName& NewName)
+{
+	// ensure goal even exists in this solver
+	const int32 GoalIndex = GetIndexOfGoal(OldName);
+	if (GoalIndex == INDEX_NONE)
+	{
+		return;
+	}
+
+	// rename
+	Effectors[GoalIndex]->Modify();
+	Effectors[GoalIndex]->GoalName = NewName;
+}
+
+void UIKRigFBIKSolver::SetGoalBone(const FName& GoalName, const FName& NewBoneName)
+{
+	// ensure goal even exists in this solver
+	const int32 GoalIndex = GetIndexOfGoal(GoalName);
+	if (GoalIndex == INDEX_NONE)
+	{
+		return;
+	}
+
+	// rename
+	Effectors[GoalIndex]->Modify();
+	Effectors[GoalIndex]->BoneName = NewBoneName;
+}
+
+void UIKRigFBIKSolver::SetRootBone(const FName& RootBoneName)
+{
+	RootBone = RootBoneName;
+}
+
+UObject* UIKRigFBIKSolver::GetGoalSettings(const FName& GoalName) const
+{
+	const int32 GoalIndex = GetIndexOfGoal(GoalName);
+	if (GoalIndex == INDEX_NONE)
+	{
+		return nullptr;
+	}
+
+	return Effectors[GoalIndex];
+}
+
+void UIKRigFBIKSolver::AddBoneSetting(const FName& BoneName)
+{
+	if (GetBoneSetting(BoneName))
+	{
+		return; // already have settings on this bone
+	}
+
+	UIKRig_FBIKBoneSettings* NewBoneSettings = NewObject<UIKRig_FBIKBoneSettings>(this, UIKRig_FBIKBoneSettings::StaticClass());
+	NewBoneSettings->Bone = BoneName;
+	BoneSettings.Add(NewBoneSettings);
+}
+
+void UIKRigFBIKSolver::RemoveBoneSetting(const FName& BoneName)
+{
+	UIKRig_FBIKBoneSettings* BoneSettingToRemove = nullptr; 
+	for (UIKRig_FBIKBoneSettings* BoneSetting : BoneSettings)
+	{
+		if (BoneSetting->Bone == BoneName)
+		{
+			BoneSettingToRemove = BoneSetting;
+			break; // can only be one with this name
+		}
+	}
+
+	if (BoneSettingToRemove)
+	{
+		BoneSettings.Remove(BoneSettingToRemove);
+	}
+}
+
+UObject* UIKRigFBIKSolver::GetBoneSetting(const FName& BoneName) const
+{
+	for (UIKRig_FBIKBoneSettings* BoneSetting : BoneSettings)
+	{
+		if (BoneSetting && BoneSetting->Bone == BoneName)
+		{
+			return BoneSetting;
+		}
+	}
+	
+	return nullptr;
+}
+
+void UIKRigFBIKSolver::DrawBoneSettings(
+	const FName& BoneName,
+	const FIKRigSkeleton& IKRigSkeleton,
+	FPrimitiveDrawInterface* PDI) const
+{
+	
+}
+
+bool UIKRigFBIKSolver::IsBoneAffectedBySolver(const FName& BoneName, const FIKRigSkeleton& IKRigSkeleton) const
+{
+	// nothing is affected by solver without a root bone assigned or at least 1 effector
+	if (RootBone == NAME_None || Effectors.IsEmpty())
+	{
+		return false;
+	}
+
+	// has to be BELOW root
+	if (!IKRigSkeleton.IsBoneInDirectLineage(BoneName, RootBone))
+	{
+		return false;
+	}
+
+	// has to be ABOVE an effector
+	for (UIKRig_FBIKEffector* Effector : Effectors)
+	{
+		if (IKRigSkeleton.IsBoneInDirectLineage(Effector->BoneName, BoneName))
+		{
+			return true;
+		}
+	}
+	
+	return false;
+}
+
+void UIKRigFBIKSolver::PostLoad()
+{
+	Super::PostLoad();
+	
+	// patch for loading old effectors, we blow them away
+	bool bHasNullEffector = false;
+	for (const UIKRig_FBIKEffector* Effector : Effectors)
+	{
+		if (!Effector)
+		{
+			bHasNullEffector = true;
+		}
+	}
+
+	if (bHasNullEffector)
+	{
+		Effectors.Empty();
+	}
+}
+
+#endif
+
+int32 UIKRigFBIKSolver::GetIndexOfGoal(const FName& OldName) const
+{
+	for (int32 i=0; i<Effectors.Num(); ++i)
+	{
+		if (Effectors[i]->GoalName == OldName)
+		{
+			return i;
+		}
+	}
+
+	return INDEX_NONE;
+}
+
+#undef LOCTEXT_NAMESPACE
+
diff --git a/Engine/Plugins/Animation/IKRig/Source/IKRig/Private/Rig/Solvers/IKRig_PBIKSolver.cpp b/Engine/Plugins/Animation/IKRig/Source/IKRig/Private/Rig/Solvers/IKRig_PBIKSolver.cpp
deleted file mode 100644
