[IK Retargeter] Fixed PinBone Op "Copy Local Position Relative Scaled" mode to work correctly with source scale

diff --git a/Engine/Plugins/Animation/IKRig/Source/IKRig/Private/Retargeter/RetargetOps/PinBoneOp.cpp b/Engine/Plugins/Animation/IKRig/Source/IKRig/Private/Retargeter/RetargetOps/PinBoneOp.cpp
--- a/Engine/Plugins/Animation/IKRig/Source/IKRig/Private/Retargeter/RetargetOps/PinBoneOp.cpp
+++ b/Engine/Plugins/Animation/IKRig/Source/IKRig/Private/Retargeter/RetargetOps/PinBoneOp.cpp
@@ -38,34 +38,23 @@ const FTransform& FPinBoneData::GetParentTransform(int32 BoneIndex, const FRetar
 	return Pose[ParentIndex];
 }
 
-void FPinBoneData::CachePinData(ERetargetSourceOrTarget SkeletonToPinTo, const UIKRetargetProcessor* Processor)
+void FPinBoneData::CachePinData(ERetargetSourceOrTarget InSkeletonToCopyFrom, const UIKRetargetProcessor* Processor)
 {
-	SourceScaleFactorUsedForCachedData = Processor->GetGlobalSettings().SourceScaleFactor;
-	
-	auto ScaleSourceGlobalTransform = [this](FTransform& InOutTransformToScale)
-    {
-    	InOutTransformToScale.SetTranslation( InOutTransformToScale.GetTranslation() * SourceScaleFactorUsedForCachedData );
-    };
+	SourceScaleFactor = Processor->GetGlobalSettings().SourceScaleFactor;
 
     // get skeletons we are copying from/to
-    const FRetargetSkeleton& SkeletonToCopyFrom = Processor->GetSkeleton(SkeletonToPinTo);
+    const FRetargetSkeleton& SkeletonToCopyFrom = Processor->GetSkeleton(InSkeletonToCopyFrom);
     const FRetargetSkeleton& SkeletonToCopyTo = Processor->GetSkeleton(ERetargetSourceOrTarget::Target);
 
 	// get relevant transforms from the retarget poses
-    FTransform BoneToCopyFromRefPoseGlobal = SkeletonToCopyFrom.RetargetGlobalPose[BoneToCopyFrom.BoneIndex];
-    FTransform BoneToCopyToRefPoseGlobal = SkeletonToCopyTo.RetargetGlobalPose[BoneToCopyTo.BoneIndex];
 	FTransform BoneToCopyFromParentRefPoseGlobal = GetParentTransform(BoneToCopyFrom.BoneIndex, SkeletonToCopyFrom, SkeletonToCopyFrom.RetargetGlobalPose);
-	FTransform BoneToCopyToParentRefPoseGlobal = GetParentTransform(BoneToCopyFrom.BoneIndex, SkeletonToCopyFrom, SkeletonToCopyFrom.RetargetGlobalPose);
-	
-	// scale transforms according to the source scale factor
-    ScaleSourceGlobalTransform(BoneToCopyFromRefPoseGlobal);
-    ScaleSourceGlobalTransform(BoneToCopyToRefPoseGlobal);
-	ScaleSourceGlobalTransform(BoneToCopyFromParentRefPoseGlobal);
-	ScaleSourceGlobalTransform(BoneToCopyToParentRefPoseGlobal);
+	FTransform BoneToCopyToParentRefPoseGlobal = GetParentTransform(BoneToCopyTo.BoneIndex, SkeletonToCopyTo, SkeletonToCopyTo.RetargetGlobalPose);
+	FTransform BoneToCopyFromRefPoseGlobal = SkeletonToCopyFrom.RetargetGlobalPose[BoneToCopyFrom.BoneIndex];
+	FTransform BoneToCopyToRefPoseGlobal = SkeletonToCopyTo.RetargetGlobalPose[BoneToCopyTo.BoneIndex];
 
 	// cache offset from BoneToCopyFrom to BoneToCopyTo in ref pose
     OffsetFromBoneToCopyFromInRefPose = BoneToCopyFromRefPoseGlobal.GetRelativeTransform(BoneToCopyToRefPoseGlobal);
-	// cache local transform of BoneToCopyFrom in ref pose
+	// cache local transform of BoneToCopyFrom
     LocalRefPoseBoneToCopyFrom = BoneToCopyFromParentRefPoseGlobal.GetRelativeTransform(BoneToCopyFromRefPoseGlobal);
 	// cache local transform of BoneToCopyTo in ref pose
     LocalRefPoseBoneToCopyTo = BoneToCopyToParentRefPoseGlobal.GetRelativeTransform(BoneToCopyToRefPoseGlobal);
@@ -110,7 +99,7 @@ bool UPinBoneOp::Initialize(
 		}
 
 		// force reinitialization
-		PinData.SourceScaleFactorUsedForCachedData = -1.0f;
+		PinData.SourceScaleFactor = -1.0f;
 	}
 
 #if WITH_EDITOR
@@ -149,7 +138,7 @@ void UPinBoneOp::Run(
 		}
 
 		// recache the offsets if the source is scaled differently
-		if (!FMath::IsNearlyEqual(PinData.SourceScaleFactorUsedForCachedData, Processor->GetGlobalSettings().SourceScaleFactor))
+		if (!FMath::IsNearlyEqual(PinData.SourceScaleFactor, Processor->GetGlobalSettings().SourceScaleFactor))
 		{
 			PinData.CachePinData(PinTo, Processor);
 		}
@@ -326,22 +315,20 @@ FTransform UPinBoneOp::GetNewBoneTransform(
 		
 	case EPinBoneTranslationMode::CopyLocalPositionRelativeScaled:
 		{
-			const double LengthOfBoneToCopyToInRefPose = PinData.LocalRefPoseBoneToCopyTo.GetTranslation().Size();
-			const double LengthOfBoneToCopyFromInRefPose = PinData.LocalRefPoseBoneToCopyFrom.GetTranslation().Size();
-			double RestPoseLengthRatio = 1.0f;
-			if (LengthOfBoneToCopyFromInRefPose > KINDA_SMALL_NUMBER)
-			{
-				RestPoseLengthRatio = LengthOfBoneToCopyToInRefPose / LengthOfBoneToCopyFromInRefPose;
-			}
-		
 			const FVector BoneToCopyFromTranslationCurrent = GetBoneToCopyFromCurrentVector();
-			double CurrentLength;
+			double LengthOfBoneToCopyFromCurrent;
 			FVector BoneToCopyFromTranslationNorm;
-			BoneToCopyFromTranslationCurrent.ToDirectionAndLength(BoneToCopyFromTranslationNorm, CurrentLength);
-			const FVector BoneToCopyFromRelativeOffset = BoneToCopyFromTranslationNorm * (CurrentLength * RestPoseLengthRatio);
+			BoneToCopyFromTranslationCurrent.ToDirectionAndLength(BoneToCopyFromTranslationNorm, LengthOfBoneToCopyFromCurrent);
+			
+			const double LengthOfBoneToCopyFromInRefPose = PinData.LocalRefPoseBoneToCopyFrom.GetTranslation().Size();
+			const double BoneLengthScaleFactor = LengthOfBoneToCopyFromCurrent / LengthOfBoneToCopyFromInRefPose;
+			
+			const double LengthOfBoneToCopyToInRefPose = PinData.LocalRefPoseBoneToCopyTo.GetTranslation().Size();
+			const FVector BoneToCopyFromRelativeOffset = BoneToCopyFromTranslationNorm * (LengthOfBoneToCopyToInRefPose * BoneLengthScaleFactor);
 		
 			const FTransform& ParentOfBoneToCopyToCurrent = FPinBoneData::GetParentTransform(PinData.BoneToCopyTo.BoneIndex, SkeletonToCopyTo, OutTargetGlobalPose);
 			Result.SetTranslation(ParentOfBoneToCopyToCurrent.GetTranslation() + BoneToCopyFromRelativeOffset);
+			
 			break;
 		}
 		
diff --git a/Engine/Plugins/Animation/IKRig/Source/IKRig/Public/Retargeter/RetargetOps/PinBoneOp.h b/Engine/Plugins/Animation/IKRig/Source/IKRig/Public/Retargeter/RetargetOps/PinBoneOp.h
--- a/Engine/Plugins/Animation/IKRig/Source/IKRig/Public/Retargeter/RetargetOps/PinBoneOp.h
+++ b/Engine/Plugins/Animation/IKRig/Source/IKRig/Public/Retargeter/RetargetOps/PinBoneOp.h
@@ -27,7 +27,7 @@ struct FPinBoneData
 	
 	static const FTransform& GetParentTransform(int32 BoneIndex, const FRetargetSkeleton& Skeleton, const TArray<FTransform>& Pose);
 
-	void CachePinData(ERetargetSourceOrTarget SkeletonToPinTo, const UIKRetargetProcessor* Processor);
+	void CachePinData(ERetargetSourceOrTarget InSkeletonToCopyFrom, const UIKRetargetProcessor* Processor);
 
 	// The bone to copy FROM.
 	UPROPERTY(EditAnywhere, Category=Settings)
@@ -37,7 +37,7 @@ struct FPinBoneData
 	FBoneReference BoneToCopyTo;
 
 	// the scale factor used to calculate all the cached offset (lazy evaluated when scale factor is modified)
-	float SourceScaleFactorUsedForCachedData;
+	float SourceScaleFactor;
 	// the local transform of the BoneToCopyTo in the ref pose
 	FTransform LocalRefPoseBoneToCopyTo;
 	// the local transform of the BoneToCopyTo in the ref pose
