[IK Retargeter] Fixed bug in CopyPose feature on mismatched skeletons

diff --git a/Engine/Plugins/Animation/IKRig/Source/IKRig/Private/Retargeter/IKRetargetProcessor.cpp b/Engine/Plugins/Animation/IKRig/Source/IKRig/Private/Retargeter/IKRetargetProcessor.cpp
--- a/Engine/Plugins/Animation/IKRig/Source/IKRig/Private/Retargeter/IKRetargetProcessor.cpp
+++ b/Engine/Plugins/Animation/IKRig/Source/IKRig/Private/Retargeter/IKRetargetProcessor.cpp
@@ -1510,54 +1510,122 @@ bool FRetargetChainPairIK::Initialize(
 }
 
 void FBasePoseCopier::CopyPoseFromSource(
-		const FName InRootBoneToCopy,
-		const FRetargetSkeleton& InSourceSkeleton,
-		const FRetargetSkeleton& InTargetSkeleton,
-		TArray<FTransform>& InOutSourcePose,
-		TArray<FTransform>& OutTargetPose)
+	const FName InRootBoneToCopy,
+	const FRetargetSkeleton& InSourceSkeleton,
+	const FTargetSkeleton& InTargetSkeleton,
+	TArray<FTransform>& InOutSourcePose,
+	TArray<FTransform>& OutTargetGlobalPose,
+	TArray<FTransform>& OutTargetInputLocalPose,
+	FIKRigLogger& Log)
 {
 	// check if we need to rebuild the cached bone mapping
 	if (SourceToTargetBoneIndexMap.IsEmpty() || InRootBoneToCopy != RootBoneToCopy)
 	{
-		// update the mapping from source-to-target bones (by name)
-		SourceToTargetBoneIndexMap.Reset();
-		RootBoneToCopy = InRootBoneToCopy;
-		RootBoneToCopyIndex = RootBoneToCopy == NAME_None ? INDEX_NONE : InSourceSkeleton.FindBoneIndexByName(RootBoneToCopy);
-		for (int32 SourceBoneIndex=0; SourceBoneIndex<InSourceSkeleton.BoneNames.Num(); ++SourceBoneIndex)
+		Initialize(InRootBoneToCopy, InSourceSkeleton, InTargetSkeleton, Log);
+	}
+	
+	// copy source bones directly to the target (no retargeting)
+	for (const TPair<int32, int32>& Pair : SourceToTargetBoneIndexMap)
+	{
+		const int32 SourceIndex = Pair.Value;
+		const int32 TargetIndex = Pair.Key;
+
+		// copy the pose in local space (this matches CopyPoseFromMesh behavior)
+		const int32 TargetParentIndex = InTargetSkeleton.GetParentIndex(TargetIndex);
+		const FTransform TargetParentGlobal = TargetParentIndex != INDEX_NONE ? OutTargetGlobalPose[TargetParentIndex] : FTransform::Identity;
+		const int32 SourceParentIndex = InSourceSkeleton.GetParentIndex(SourceIndex);
+		const FTransform SourceParentGlobal = SourceParentIndex != INDEX_NONE ? InOutSourcePose[SourceParentIndex] : FTransform::Identity;
+		const FTransform& SourceGlobal = InOutSourcePose[SourceIndex];
+		const FTransform& SourceLocal = SourceGlobal.GetRelativeTransform(SourceParentGlobal);
+
+		// convert to global space and store
+		OutTargetGlobalPose[TargetIndex] = SourceLocal * TargetParentGlobal;
+	}
+
+	// update children of RootBoneToCopy that were not in the source
+	for (int32 TargetChildIndex : ChildrenToUpdate)
+	{
+		const int32 ParentIndex = InTargetSkeleton.GetParentIndex(TargetChildIndex);
+		const FTransform ParentGlobalTransform = ParentIndex != INDEX_NONE ? OutTargetGlobalPose[ParentIndex] : FTransform::Identity;
+		const FTransform ChildLocalTransform = InTargetSkeleton.InputLocalPose[TargetChildIndex];
+		OutTargetGlobalPose[TargetChildIndex] = ChildLocalTransform * ParentGlobalTransform;
+	}
+	
+	// update input local pose to reflect the newly copied pose
+	// (this local pose is used when updating intermediate bones in subsequent retargeting steps)
+	const int32 TargetRootBoneToCopyIndex = RootBoneToCopy == NAME_None ? InTargetSkeleton.BoneNames.Num() : InTargetSkeleton.FindBoneIndexByName(RootBoneToCopy);
+	const int32 ParentOfCopyRootIndex = InTargetSkeleton.GetParentIndex(TargetRootBoneToCopyIndex);
+	InTargetSkeleton.UpdateLocalTransformsBelowBone(ParentOfCopyRootIndex, OutTargetInputLocalPose, InTargetSkeleton.OutputGlobalPose);
+}
+
+void FBasePoseCopier::Initialize(
+	const FName InRootBoneToCopy,
+	const FRetargetSkeleton& InSourceSkeleton,
+	const FTargetSkeleton& InTargetSkeleton,
+	FIKRigLogger& Log)
+{
+	// store the root bone we initialized with
+	RootBoneToCopy = InRootBoneToCopy;
+
+	// find the roots to copy from
+	const int32 SourceRootBoneToCopyIndex = RootBoneToCopy == NAME_None ? 0 : InSourceSkeleton.FindBoneIndexByName(RootBoneToCopy);
+	const int32 TargetRootBoneToCopyIndex = RootBoneToCopy == NAME_None ? 0 : InTargetSkeleton.FindBoneIndexByName(RootBoneToCopy);
+
+	// warn if user specified a root bone but it wasn't found
+	if (RootBoneToCopy != NAME_None)
+	{
+		if (SourceRootBoneToCopyIndex == INDEX_NONE)
 		{
-			// filter out bones above the RootBoneToCopy
-			if (RootBoneToCopyIndex != INDEX_NONE &&
-				SourceBoneIndex != RootBoneToCopyIndex &&
-				!InSourceSkeleton.IsParentOf(RootBoneToCopyIndex, SourceBoneIndex))
-			{
-				continue;
-			}
-		
-			// filter out bones not in the target skeleton
-			const int32 TargetBoneIndex = InTargetSkeleton.FindBoneIndexByName(InTargetSkeleton.BoneNames[SourceBoneIndex]);
-			if (TargetBoneIndex == INDEX_NONE)
-			{
-				continue;
-			}
-		
-			// store map of source/target bones in common
-			SourceToTargetBoneIndexMap.Emplace(SourceBoneIndex, TargetBoneIndex);
+			Log.LogWarning( FText::Format(
+				LOCTEXT("MissingSourceRootToCopyFrom", "The root bone to copy from, {0} was not found in the source mesh {1}"),
+				FText::FromName(RootBoneToCopy), FText::FromString(InSourceSkeleton.SkeletalMesh->GetName())));
+		}
+
+		if (TargetRootBoneToCopyIndex == INDEX_NONE)
+		{
+			Log.LogWarning( FText::Format(
+				LOCTEXT("MissingTargetRootToCopyFrom", "The root bone to copy from, {0} was not found in the target mesh {1}"),
+				FText::FromName(RootBoneToCopy), FText::FromString(InSourceSkeleton.SkeletalMesh->GetName())));
 		}
 	}
 
-	// copy source bones directly to the target (no retargeting)
-	for (const TPair<int32, int32>& Pair : SourceToTargetBoneIndexMap)
+	// update the mapping from source-to-target bones (by name)
+	SourceToTargetBoneIndexMap.Reset();
+	for (int32 SourceBoneIndex=0; SourceBoneIndex<InSourceSkeleton.BoneNames.Num(); ++SourceBoneIndex)
 	{
-		OutTargetPose[Pair.Key] = InOutSourcePose[Pair.Value];
+		// filter out bones above the RootBoneToCopy
+		if (SourceRootBoneToCopyIndex != INDEX_NONE &&
+			SourceBoneIndex != SourceRootBoneToCopyIndex &&
+			!InSourceSkeleton.IsParentOf(SourceRootBoneToCopyIndex, SourceBoneIndex))
+		{
+			continue;
+		}
+	
+		// filter out bones not in the target skeleton
+		const int32 TargetBoneIndex = InTargetSkeleton.FindBoneIndexByName(InSourceSkeleton.BoneNames[SourceBoneIndex]);
+		if (TargetBoneIndex == INDEX_NONE)
+		{
+			continue;
+		}
+	
+		// store map of source/target bones in common
+		SourceToTargetBoneIndexMap.Emplace(TargetBoneIndex, SourceBoneIndex);
 	}
 
-	// reset source bones above the RootBoneToCopy to retarget pose
-	if (InRootBoneToCopy != NAME_None)
+	// cache list of children that need updated
+	ChildrenToUpdate.Reset();
+	const int32 LastBranchIndex = InTargetSkeleton.GetCachedEndOfBranchIndex(TargetRootBoneToCopyIndex);
+	const bool bIsLeafBone = LastBranchIndex == INDEX_NONE;
+	if (!bIsLeafBone)
 	{
-		InOutSourcePose = InSourceSkeleton.RetargetGlobalPose;
-		for (const TPair<int32, int32>& Pair : SourceToTargetBoneIndexMap)
+		for (int32 ChildBoneIndex = TargetRootBoneToCopyIndex + 1; ChildBoneIndex <= LastBranchIndex; ChildBoneIndex++)
 		{
-			InOutSourcePose[Pair.Value] = OutTargetPose[Pair.Key];
+			if (SourceToTargetBoneIndexMap.Contains(ChildBoneIndex))
+			{
+				continue;
+			}
+
+			ChildrenToUpdate.Add(ChildBoneIndex);
 		}
 	}
 }
@@ -2377,11 +2445,9 @@ void UIKRetargetProcessor::GenerateBasePoses(TArray<FTransform>& InSourceGlobalP
 			SourceSkeleton,
 			TargetSkeleton,
 			InSourceGlobalPose,
-			TargetSkeleton.OutputGlobalPose);
-		
-		// update input local pose (used when updating intermediate bones)
-		const int32 ParentOfCopyRootIndex = TargetSkeleton.GetParentIndex(BasePoseCopier.GetIndexOfRootBoneToCopy());
-		TargetSkeleton.UpdateLocalTransformsBelowBone(ParentOfCopyRootIndex, TargetSkeleton.InputLocalPose, TargetSkeleton.OutputGlobalPose );
+			TargetSkeleton.OutputGlobalPose,
+			TargetSkeleton.InputLocalPose,
+			Log);
 	}
 }
 
diff --git a/Engine/Plugins/Animation/IKRig/Source/IKRig/Public/Retargeter/IKRetargetProcessor.h b/Engine/Plugins/Animation/IKRig/Source/IKRig/Public/Retargeter/IKRetargetProcessor.h
--- a/Engine/Plugins/Animation/IKRig/Source/IKRig/Public/Retargeter/IKRetargetProcessor.h
+++ b/Engine/Plugins/Animation/IKRig/Source/IKRig/Public/Retargeter/IKRetargetProcessor.h
@@ -494,20 +494,29 @@ struct FBasePoseCopier
 	void CopyPoseFromSource(
 		const FName InRootBoneToCopy,
 		const FRetargetSkeleton& InSourceSkeleton,
-		const FRetargetSkeleton& InTargetSkeleton,
+		const FTargetSkeleton& InTargetSkeleton,
 		TArray<FTransform>& InOutSourcePose,
-		TArray<FTransform>& OutTargetPose);
-
-	int32 GetIndexOfRootBoneToCopy() const { return RootBoneToCopyIndex; };
+		TArray<FTransform>& OutTargetGlobalPose,
+		TArray<FTransform>& OutTargetInputLocalPose,
+		FIKRigLogger& Log);
 
 private:
+
+	void Initialize(
+		const FName InRootBoneToCopy,
+		const FRetargetSkeleton& InSourceSkeleton,
+		const FTargetSkeleton& InTargetSkeleton,
+		FIKRigLogger& Log);
 	
 	// mapping from source bones to target bone indices
 	// NOTE: this map will only contain bones with the same name in both skeletons that are below the RootBoneToCopy (if specified)
-	TArray< TPair<int32, int32> > SourceToTargetBoneIndexMap;
+	TMap<int32, int32> SourceToTargetBoneIndexMap;
+
+	// children of bones that are copied that need manually updated
+	TArray<int32> ChildrenToUpdate;
+
 	// the bone to start copying from (includes all children recursive)
 	FName RootBoneToCopy = NAME_None;
-	int32 RootBoneToCopyIndex = INDEX_NONE;
 };
 
 /** The runtime processor that converts an input pose from a source skeleton into an output pose on a target skeleton.
