[IKRig] Major refactor replacing all runtime UObjects with UStructs
- Ported all IKRig solver types to UStructs
- Updated engine tests to reflect new scripting API

diff --git a/Engine/Plugins/Animation/IKRig/Config/BaseIKRig.ini b/Engine/Plugins/Animation/IKRig/Config/BaseIKRig.ini
--- a/Engine/Plugins/Animation/IKRig/Config/BaseIKRig.ini
+++ b/Engine/Plugins/Animation/IKRig/Config/BaseIKRig.ini
@@ -2,4 +2,6 @@
 +ClassRedirects=(OldName="/Script/IKRig.IKRigPBIKSolver", NewName="/Script/IKRig.IKRigFBIKSolver")
 +ClassRedirects=(OldName="/Script/IKRig.IKRig_PBIKBoneSettings", NewName="/Script/IKRig.IKRig_FBIKBoneSettings")
 +ClassRedirects=(OldName="/Script/IKRig.AnimNode_IKRetargeter", NewName="/Script/IKRig.AnimNode_RetargetPoseFromMesh")
++FunctionRedirects=(OldName="/Script/IKRigEditor.IKRigController.SetRootBone",NewName="/Script/IKRigEditor.IKRigController.SetStartBone")
++FunctionRedirects=(OldName="/Script/IKRigEditor.IKRigController.GetRootBone",NewName="/Script/IKRigEditor.IKRigController.GetStartBone")
 
diff --git a/Engine/Plugins/Animation/IKRig/Source/IKRig/Private/AnimNodes/AnimNode_IKRig.cpp b/Engine/Plugins/Animation/IKRig/Source/IKRig/Private/AnimNodes/AnimNode_IKRig.cpp
--- a/Engine/Plugins/Animation/IKRig/Source/IKRig/Private/AnimNodes/AnimNode_IKRig.cpp
+++ b/Engine/Plugins/Animation/IKRig/Source/IKRig/Private/AnimNodes/AnimNode_IKRig.cpp
@@ -6,7 +6,7 @@
 #include "Animation/AnimTrace.h"
 #include "Rig/IKRigProcessor.h"
 #include "Rig/IKRigDefinition.h"
-#include "Rig/Solvers/IKRigSolver.h"
+#include "Rig/Solvers/IKRigSolverBase.h"
 
 #include "Components/SkeletalMeshComponent.h"
 #include "ActorComponents/IKRigInterface.h"
@@ -110,7 +110,7 @@ void FAnimNode_IKRig::AssignGoalTargets()
 
 	if (bDriveWithSourceAsset && !GIsTransacting)
 	{
-		IKRigProcessor->CopyAllInputsFromSourceAssetAtRuntime(RigDefinitionAsset);
+		IKRigProcessor->CopyAllSettingsFromAsset(RigDefinitionAsset);
 		return;
 	}
 	#endif
diff --git a/Engine/Plugins/Animation/IKRig/Source/IKRig/Private/Retargeter/IKRetargetProcessor.cpp b/Engine/Plugins/Animation/IKRig/Source/IKRig/Private/Retargeter/IKRetargetProcessor.cpp
--- a/Engine/Plugins/Animation/IKRig/Source/IKRig/Private/Retargeter/IKRetargetProcessor.cpp
+++ b/Engine/Plugins/Animation/IKRig/Source/IKRig/Private/Retargeter/IKRetargetProcessor.cpp
@@ -2042,7 +2042,7 @@ void UIKRetargetProcessor::Initialize(
 				LOCTEXT("SuccessfulInit", "Success! The IK Retargeter is ready to transfer animation from the source, {0} to the target, {1}"),
 				FText::FromString(SourceSkeletalMesh->GetName()), FText::FromString(TargetSkeletalMesh->GetName())));
 	}
-
+	
 	// copy the initial settings from the asset
 	CopyIKRigSettingsFromAsset();
 	
@@ -2538,7 +2538,7 @@ void UIKRetargetProcessor::RunIKRetarget(
 		// run IK solve
 		IKRigProcessor->Solve();
 		// copy results of solve
-		IKRigProcessor->CopyOutputGlobalPoseToArray(OutTargetGlobalPose);
+		IKRigProcessor->GetOutputPoseGlobal(OutTargetGlobalPose);
 	}
 }
 
@@ -2991,12 +2991,12 @@ FName UIKRetargetProcessor::GetMappedChainName(
 
 #endif
 
-void UIKRetargetProcessor::CopyIKRigSettingsFromAsset()
+void UIKRetargetProcessor::CopyIKRigSettingsFromAsset() const
 {
 	// copy IK Rig settings
 	if (const UIKRigDefinition* TargetIKRig = RetargeterAsset->GetIKRig(ERetargetSourceOrTarget::Target))
 	{
-		IKRigProcessor->CopyAllInputsFromSourceAssetAtRuntime(TargetIKRig);
+		IKRigProcessor->CopyAllSettingsFromAsset(TargetIKRig);
 	}
 }
 
diff --git a/Engine/Plugins/Animation/IKRig/Source/IKRig/Private/Rig/IKRigDefinition.cpp b/Engine/Plugins/Animation/IKRig/Source/IKRig/Private/Rig/IKRigDefinition.cpp
--- a/Engine/Plugins/Animation/IKRig/Source/IKRig/Private/Rig/IKRigDefinition.cpp
+++ b/Engine/Plugins/Animation/IKRig/Source/IKRig/Private/Rig/IKRigDefinition.cpp
@@ -4,6 +4,7 @@
 
 #include "IKRigObjectVersion.h"
 #include "Engine/SkeletalMesh.h"
+#include "Rig/Solvers/IKRigSolverBase.h"
 
 #include UE_INLINE_GENERATED_CPP_BY_NAME(IKRigDefinition)
 
@@ -302,10 +303,73 @@ void UIKRigDefinition::Serialize(FArchive& Ar)
 
 void UIKRigDefinition::PostLoad()
 {
+	
 	Super::PostLoad();
 	
 	// very early versions of the asset may not have been set as standalone
 	SetFlags(RF_Standalone);
+
+	// validate the input FInstancedStruct is a valid subclass of FIKRigSolverVase
+	auto DerivesFromBaseSolverType = [](const FInstancedStruct& InConvertedSolver)
+	{
+		if (!InConvertedSolver.IsValid())
+		{
+			return false;
+		}
+		
+		bool bIsDerivedFromBase = false;
+		UStruct* CurrentSuperStruct = InConvertedSolver.GetScriptStruct()->GetSuperStruct();
+		while (CurrentSuperStruct)
+		{
+			if (CurrentSuperStruct == FIKRigSolverBase::StaticStruct())
+			{
+				bIsDerivedFromBase = true;
+				break;
+			}
+			CurrentSuperStruct = CurrentSuperStruct->GetSuperStruct();
+		}
+		
+		return bIsDerivedFromBase;
+	};
+	
+	// convert old UObject-based solver stack to new struct-based types
+	PRAGMA_DISABLE_DEPRECATION_WARNINGS
+	for (TObjectPtr<UDEPRECATED_IKRigSolver> Solver : Solvers_DEPRECATED)
+	{
+		FInstancedStruct ConvertedSolver;
+		Solver->ConvertToInstancedStruct(ConvertedSolver);
+		if (DerivesFromBaseSolverType(ConvertedSolver))
+		{
+			SolverStack.Add(ConvertedSolver);
+		}
+		else
+		{
+			UE_LOG(LogTemp, Warning, TEXT("IK Rig: unable to load old UObject based solver type. Conversion failed for type:  %s"), *Solver->GetName());
+		}
+	}
+	PRAGMA_ENABLE_DEPRECATION_WARNINGS
+}
+
+const TArray<const FIKRigSolverBase*> UIKRigDefinition::GetSolverArray() const
+{
+	TArray<const FIKRigSolverBase*> Result;
+	for (const FInstancedStruct& SolverStruct : SolverStack)
+	{
+		if (!ensureMsgf(SolverStruct.IsValid(), TEXT("IK Rig asset contains unknown solver type. Possibly due to an unloaded plugin?")))
+		{
+			// this can happen if asset references deleted or unknown IK Solver type
+			// which should only happen during development (if at all)
+			continue;
+		}
+		const FIKRigSolverBase* Solver = SolverStruct.GetPtr<FIKRigSolverBase>();
+		Result.Add(Solver);
+	}
+	return MoveTemp(Result);
+}
+
+const TArray<FInstancedStruct> UIKRigDefinition::GetSolverStructs() const
+{
+	return SolverStack;
 }
 
 const FBoneChain* UIKRigDefinition::GetRetargetChainByName(FName ChainName) const
diff --git a/Engine/Plugins/Animation/IKRig/Source/IKRig/Private/Rig/IKRigProcessor.cpp b/Engine/Plugins/Animation/IKRig/Source/IKRig/Private/Rig/IKRigProcessor.cpp
--- a/Engine/Plugins/Animation/IKRig/Source/IKRig/Private/Rig/IKRigProcessor.cpp
+++ b/Engine/Plugins/Animation/IKRig/Source/IKRig/Private/Rig/IKRigProcessor.cpp
@@ -2,7 +2,7 @@
 
 #include "Rig/IKRigProcessor.h"
 #include "Rig/IKRigDefinition.h"
-#include "Rig/Solvers/IKRigSolver.h"
+#include "Rig/Solvers/IKRigSolverBase.h"
 
 #include "Engine/SkeletalMesh.h"
 
@@ -110,16 +110,13 @@ void UIKRigProcessor::Initialize(
 		}
 
 		// warn if goal is not connected to any solvers
-		bool bIsGoalConnectedToAnySolver = false;
-		const TArray<UIKRigSolver*>& AssetSolvers = InRigAsset->GetSolverArray();
-		for (const UIKRigSolver* IKRigSolver : AssetSolvers)
+		const TArray<const FIKRigSolverBase*> AssetSolvers = InRigAsset->GetSolverArray();
+		TSet<FName> AllRequiredGoals;
+		for (const FIKRigSolverBase* AssetSolver : AssetSolvers)
 		{
-			if (IKRigSolver->IsGoalConnected(EffectorGoal->GoalName))
-			{
-				bIsGoalConnectedToAnySolver = true;
-			}
+			AssetSolver->GetRequiredGoals(AllRequiredGoals);
 		}
-		if (!bIsGoalConnectedToAnySolver)
+		if (!AllRequiredGoals.Contains(EffectorGoal->GoalName))
 		{
 			Log.LogWarning(FText::Format(
 				LOCTEXT("DisconnectedGoal", "IK Rig, '{0}' has a Goal, '{1}' that is not connected to any solvers. It will have no effect."),
@@ -130,12 +127,11 @@ void UIKRigProcessor::Initialize(
 	}
 
 	// create copies of all the solvers in the IK rig
-	const TArray<UIKRigSolver*>& AssetSolvers = InRigAsset->GetSolverArray();
+	const TArray<FInstancedStruct>& AssetSolvers = InRigAsset->GetSolverStructs();
 	Solvers.Reset(AssetSolvers.Num());
-	int32 SolverIndex = 0;
-	for (const UIKRigSolver* IKRigSolver : AssetSolvers)
+	for (const FInstancedStruct& AssetSolverStruct : AssetSolvers)
 	{
-		if (!IKRigSolver)
+		if (!AssetSolverStruct.IsValid())
 		{
 			// this can happen if asset references deleted IK Solver type which should only happen during development (if at all)
 			Log.LogWarning(FText::Format(
@@ -144,20 +140,22 @@ void UIKRigProcessor::Initialize(
 			continue;
 		}
 
-		// create duplicate solver instance with unique name
-		const FName BaseName = FName(*(IKRigSolver->GetName() + "_SolverInstance_"));
-		const FName UniqueName = MakeUniqueObjectName(this, IKRigSolver->GetClass(), BaseName, EUniqueObjectNameOptions::GloballyUnique);
-		UIKRigSolver* Solver = DuplicateObject(IKRigSolver, this, UniqueName);
+		// copy the solver memory
+		const UScriptStruct* ScriptStruct = AssetSolverStruct.GetScriptStruct();
+		FIKRigSolverBase* NewSolver = static_cast<FIKRigSolverBase*>(FMemory::Malloc(ScriptStruct->GetStructureSize()));
+		ScriptStruct->InitializeStruct(NewSolver);
+		ScriptStruct->CopyScriptStruct(NewSolver, AssetSolverStruct.GetMemory());
+		// store it in the runtime stack
+		Solvers.Emplace(TUniquePtr<FIKRigSolverBase>(NewSolver));
 
 		// remove excluded goals from the solver
 		for (const FName& ExcludedGoal : ExcludedGoals)
 		{
-			Solver->RemoveGoal(ExcludedGoal);
+			NewSolver->OnGoalRemoved(ExcludedGoal);
 		}
 
 		// initialize it and store in the processor
-		Solver->Initialize(Skeleton);
-		Solvers.Add(Solver);
+		NewSolver->Initialize(Skeleton);
 	}
 
 	// validate retarget chains
@@ -188,16 +186,10 @@ bool UIKRigProcessor::IsIKRigCompatibleWithSkeleton(
 	// first we validate that all the required bones are in the input skeleton...
 	
 	TSet<FName> RequiredBones;
-	const TArray<UIKRigSolver*>& AssetSolvers = InRigAsset->GetSolverArray();
-	for (const UIKRigSolver* Solver : AssetSolvers)
+	const TArray<const FIKRigSolverBase*> AssetSolvers = InRigAsset->GetSolverArray();
+	for (const FIKRigSolverBase* AssetSolver : AssetSolvers)
 	{
-		const FName RootBone = Solver->GetRootBone();
-		if (RootBone != NAME_None)
-		{
-			RequiredBones.Add(RootBone);
-		}
-
-		Solver->GetBonesWithSettings(RequiredBones);
+		AssetSolver->GetRequiredBones(RequiredBones);
 	}
 
 	const TArray<UIKRigEffectorGoal*>& Goals = InRigAsset->GetGoalArray();
@@ -206,6 +198,9 @@ bool UIKRigProcessor::IsIKRigCompatibleWithSkeleton(
 		RequiredBones.Add(Goal->BoneName);
 	}
 
+	// strip None out (in case of lost bones in solver)
+	RequiredBones.Remove(NAME_None);
+
 	bool bAllRequiredBonesFound = true;
 	for (const FName& RequiredBone : RequiredBones)
 	{
@@ -316,24 +311,21 @@ void UIKRigProcessor::Solve(const FTransform& ComponentToWorld)
 	ResolveFinalGoalTransforms(ComponentToWorld);
 
 	// run all the solvers
-	for (UIKRigSolver* Solver : Solvers)
+	for (const TUniquePtr<FIKRigSolverBase>& Solver : Solvers)
 	{
-		#if WITH_EDITOR
 		if (Solver->IsEnabled())
 		{
 			Solver->Solve(Skeleton, GoalContainer);
 		}
-		#else
-		Solver->Solve(Skeleton, GoalContainer);
-		#endif
 	}
 
 	// make sure rotations are normalized coming out
 	Skeleton.NormalizeRotations(Skeleton.CurrentPoseGlobal);
 }
 
-void UIKRigProcessor::CopyOutputGlobalPoseToArray(TArray<FTransform>& OutputPoseGlobal) const
+void UIKRigProcessor::GetOutputPoseGlobal(TArray<FTransform>& OutputPoseGlobal) const
 {
+	// TODO, remove this copy
 	OutputPoseGlobal = Skeleton.CurrentPoseGlobal;
 }
 
@@ -352,7 +344,7 @@ void UIKRigProcessor::SetNeedsInitialized()
 	bTriedToInitialize = false;
 };
 
-void UIKRigProcessor::CopyAllInputsFromSourceAssetAtRuntime(const UIKRigDefinition* SourceAsset)
+void UIKRigProcessor::CopyAllSettingsFromAsset(const UIKRigDefinition* SourceAsset)
 {
 	check(SourceAsset)
 	
@@ -369,12 +361,13 @@ void UIKRigProcessor::CopyAllInputsFromSourceAssetAtRuntime(const UIKRigDefiniti
 	}
 
 	// copy solver settings
-	const TArray<UIKRigSolver*>& AssetSolvers = SourceAsset->GetSolverArray();
+	const TArray<const FIKRigSolverBase*> AssetSolvers = SourceAsset->GetSolverArray();
 	check(Solvers.Num() == AssetSolvers.Num()); // if number of solvers has been changed, processor should have been reinitialized
 	for (int32 SolverIndex=0; SolverIndex<Solvers.Num(); ++SolverIndex)
 	{
-		Solvers[SolverIndex]->SetEnabled(AssetSolvers[SolverIndex]->IsEnabled());
-		Solvers[SolverIndex]->UpdateSolverSettings(AssetSolvers[SolverIndex]);
+		const FIKRigSolverBase& AssetSolver = *AssetSolvers[SolverIndex];
+		FIKRigSolverBase* RunningSolver = Solvers[SolverIndex].Get();
+		RunningSolver->UpdateSettingsFromAsset(AssetSolver);
 	}
 }
 
diff --git a/Engine/Plugins/Animation/IKRig/Source/IKRig/Private/Rig/Solvers/IKRigBodyMoverSolver.cpp b/Engine/Plugins/Animation/IKRig/Source/IKRig/Private/Rig/Solvers/IKRigBodyMoverSolver.cpp
new file mode 100644
--- /dev/null
+++ b/Engine/Plugins/Animation/IKRig/Source/IKRig/Private/Rig/Solvers/IKRigBodyMoverSolver.cpp
@@ -0,0 +1,253 @@
+// Copyright Epic Games, Inc. All Rights Reserved.
+
+#include "Rig/Solvers/IKRigBodyMoverSolver.h"
+#include "Rig/IKRigDataTypes.h"
+#include "Rig/IKRigSkeleton.h"
+#include "Rig/Solvers/PointsToRotation.h"
+
+#include UE_INLINE_GENERATED_CPP_BY_NAME(IKRigBodyMoverSolver)
+
+#define LOCTEXT_NAMESPACE "BodyMoverSolver"
+
+void FIKRigBodyMoverSolver::Initialize(const FIKRigSkeleton& InIKRigSkeleton)
+{
+	BodyBoneIndex = InIKRigSkeleton.GetBoneIndexFromName(Settings.StartBone);
+}
+
+void FIKRigBodyMoverSolver::Solve(FIKRigSkeleton& InIKRigSkeleton, const FIKRigGoalContainer& InGoals)
+{
+	// no body bone specified
+	if (BodyBoneIndex == INDEX_NONE)
+	{
+		return;
+	}
+
+	// no goals connected
+	if (AllGoalSettings.IsEmpty())
+	{
+		return;
+	}
+	
+	// ensure body bone exists
+	check(InIKRigSkeleton.RefPoseGlobal.IsValidIndex(BodyBoneIndex));
+	
+	// calculate a "best fit" transform from deformed goal locations
+	TArray<FVector> InitialPoints;
+	TArray<FVector> CurrentPoints;
+	for (const FIKRigBodyMoverGoalSettings& GoalSetting : AllGoalSettings)
+	{
+		const FIKRigGoal* Goal = InGoals.FindGoalByName(GoalSetting.Goal);
+		if (!Goal)
+		{
+			return;
+		}
+
+		const int32 BoneIndex = InIKRigSkeleton.GetBoneIndexFromName(GoalSetting.BoneName);
+		const FVector InitialPosition = InIKRigSkeleton.CurrentPoseGlobal[BoneIndex].GetTranslation();
+		const FVector GoalPosition = Goal->FinalBlendedPosition;
+		const FVector FinalPosition = FMath::Lerp(InitialPosition, GoalPosition, GoalSetting.InfluenceMultiplier);
+		
+		InitialPoints.Add(InitialPosition);
+		CurrentPoints.Add(FinalPosition);
+	}
+	
+	FVector InitialCentroid;
+	FVector CurrentCentroid;
+	const FQuat RotationOffset = GetRotationFromDeformedPoints(
+		InitialPoints,
+		CurrentPoints,
+		InitialCentroid,
+		CurrentCentroid);
+
+	// alpha blend the position offset and add it to the current bone location
+	const FVector Offset = (CurrentCentroid - InitialCentroid);
+	const FVector Weight(
+		Offset.X > 0.f ? Settings.PositionPositiveX : Settings.PositionNegativeX,
+		Offset.Y > 0.f ? Settings.PositionPositiveY : Settings.PositionNegativeY,
+		Offset.Z > 0.f ? Settings.PositionPositiveZ : Settings.PositionNegativeZ);
+
+	// the bone transform to modify
+	FTransform& CurrentBodyTransform = InIKRigSkeleton.CurrentPoseGlobal[BodyBoneIndex];
+	CurrentBodyTransform.AddToTranslation(Offset * (Weight*Settings.PositionAlpha));
+
+	// do per-axis alpha blend
+	FVector Euler = RotationOffset.Euler() * FVector(Settings.RotateXAlpha, Settings.RotateYAlpha, Settings.RotateZAlpha);
+	FQuat FinalRotationOffset = FQuat::MakeFromEuler(Euler);
+	// alpha blend the entire rotation offset
+	FinalRotationOffset = FQuat::FastLerp(FQuat::Identity, FinalRotationOffset, Settings.RotationAlpha).GetNormalized();
+	// add rotation offset to original rotation
+	CurrentBodyTransform.SetRotation(FinalRotationOffset * CurrentBodyTransform.GetRotation());
+
+	// do FK update of children
+	InIKRigSkeleton.PropagateGlobalPoseBelowBone(BodyBoneIndex);
+}
+
+void FIKRigBodyMoverSolver::GetRequiredBones(TSet<FName>& OutRequiredBones) const
+{
+	OutRequiredBones.Add(Settings.StartBone);
+}
+
+void FIKRigBodyMoverSolver::GetRequiredGoals(TSet<FName>& OutRequiredGoals) const
+{
+	for (const FIKRigBodyMoverGoalSettings& GoalSetting : AllGoalSettings)
+	{
+		OutRequiredGoals.Add(GoalSetting.Goal);
+	}
+}
+
+FIKRigSolverSettingsBase* FIKRigBodyMoverSolver::GetSolverSettings()
+{
+	return &Settings;
+}
+
+const UScriptStruct* FIKRigBodyMoverSolver::GetSolverSettingsType() const
+{
+	return FIKRigBodyMoverSettings::StaticStruct();
+}
+
+void FIKRigBodyMoverSolver::AddGoal(const UIKRigEffectorGoal& InNewGoal)
+{
+	FIKRigBodyMoverGoalSettings GoalSettings;
+	GoalSettings.Goal = InNewGoal.GoalName;
+	GoalSettings.BoneName = InNewGoal.BoneName;
+	AllGoalSettings.Add(GoalSettings);
+}
+
+void FIKRigBodyMoverSolver::OnGoalRemoved(const FName& InGoalName)
+{
+	// check goal exists in this solver
+	const int32 GoalIndex = GetIndexOfGoal(InGoalName);
+	if (GoalIndex == INDEX_NONE)
+	{
+		return;
+	}
+
+	// remove it
+	AllGoalSettings.RemoveAt(GoalIndex);
+}
+
+void FIKRigBodyMoverSolver::OnGoalRenamed(const FName& InOldName, const FName& InNewName)
+{
+	// ensure goal even exists in this solver
+	const int32 GoalIndex = GetIndexOfGoal(InOldName);
+	if (GoalIndex == INDEX_NONE)
+	{
+		return;
+	}
+
+	// rename
+	AllGoalSettings[GoalIndex].Goal = InNewName;
+}
+
+void FIKRigBodyMoverSolver::OnGoalMovedToDifferentBone(const FName& InGoalName, const FName& InNewBoneName)
+{
+	// ensure goal even exists in this solver
+	const int32 GoalIndex = GetIndexOfGoal(InGoalName);
+	if (GoalIndex == INDEX_NONE)
+	{
+		return;
+	}
+
+	// rename
+	AllGoalSettings[GoalIndex].BoneName = InNewBoneName;
+}
+
+FIKRigGoalSettingsBase* FIKRigBodyMoverSolver::GetGoalSettings(const FName& InGoalName)
+{
+	const int32 GoalIndex = GetIndexOfGoal(InGoalName);
+	if (GoalIndex == INDEX_NONE)
+	{
+		return nullptr;
+	}
+
+	return &AllGoalSettings[GoalIndex];
+}
+
+const UScriptStruct* FIKRigBodyMoverSolver::GetGoalSettingsType() const
+{
+	return FIKRigBodyMoverGoalSettings::StaticStruct();
+}
+
+void FIKRigBodyMoverSolver::GetGoalsWithSettings(TSet<FName>& OutGoalsWithSettings) const
+{
+	for (const FIKRigBodyMoverGoalSettings& GoalSettings : AllGoalSettings)
+	{
+		OutGoalsWithSettings.Add(GoalSettings.Goal);
+	}
+}
+
+void FIKRigBodyMoverSolver::SetStartBone(const FName& InStartBoneName)
+{
+	Settings.StartBone = InStartBoneName;
+}
+
+FIKRigBodyMoverSettings UIKRigBodyMoverController::GetSolverSettings()
+{
+	return *static_cast<FIKRigBodyMoverSettings*>(SolverToControl->GetSolverSettings());
+}
+
+void UIKRigBodyMoverController::SetSolverSettings(FIKRigBodyMoverSettings InSettings)
+{
+	SolverToControl->SetSolverSettings(&InSettings);
+}
+
+FIKRigBodyMoverGoalSettings UIKRigBodyMoverController::GetGoalSettings(const FName InGoalName)
+{
+	if (FIKRigGoalSettingsBase* GoalSettings = SolverToControl->GetGoalSettings(InGoalName))
+	{
+		return *static_cast<FIKRigBodyMoverGoalSettings*>(GoalSettings);
+	}
+	return FIKRigBodyMoverGoalSettings();
+}
+
+void UIKRigBodyMoverController::SetGoalSettings(const FName InGoalName, FIKRigBodyMoverGoalSettings InSettings)
+{
+	SolverToControl->SetGoalSettings(InGoalName, &InSettings);
+}
+
+#if WITH_EDITOR
+
+UIKRigSolverControllerBase* FIKRigBodyMoverSolver::GetSolverController(UObject* Outer)
+{
+	return CreateControllerIfNeeded(Outer, UIKRigBodyMoverController::StaticClass());
+}
+
+FText FIKRigBodyMoverSolver::GetNiceName() const
+{
+	return FText(LOCTEXT("SolverName", "Body Mover"));
+}
+
+bool FIKRigBodyMoverSolver::GetWarningMessage(FText& OutWarningMessage) const
+{
+	if (Settings.StartBone == NAME_None)
+	{
+		OutWarningMessage = LOCTEXT("MissingRoot", "Missing start bone.");
+		return true;
+	}
+
+	if (AllGoalSettings.IsEmpty())
+	{
+		OutWarningMessage = LOCTEXT("MissingGoal", "Missing goals.");
+		return true;
+	}
+	
+	return false;
+}
+
+bool FIKRigBodyMoverSolver::IsBoneAffectedBySolver(const FName& InBoneName, const FIKRigSkeleton& InIKRigSkeleton) const
+{
+	return InIKRigSkeleton.IsBoneInDirectLineage(InBoneName, Settings.StartBone);
+}
+
+#endif
+
+int32 FIKRigBodyMoverSolver::GetIndexOfGoal(const FName& InGoalName) const
+{
+	return AllGoalSettings.IndexOfByPredicate([&InGoalName](const FIKRigBodyMoverGoalSettings& InElement)
+	{
+		return InElement.Goal == InGoalName;
+	});
+}
+
+#undef LOCTEXT_NAMESPACE
+
diff --git a/Engine/Plugins/Animation/IKRig/Source/IKRig/Private/Rig/Solvers/IKRigLimbSolver.cpp b/Engine/Plugins/Animation/IKRig/Source/IKRig/Private/Rig/Solvers/IKRigLimbSolver.cpp
new file mode 100644
--- /dev/null
+++ b/Engine/Plugins/Animation/IKRig/Source/IKRig/Private/Rig/Solvers/IKRigLimbSolver.cpp
@@ -0,0 +1,250 @@
+// Copyright Epic Games, Inc. All Rights Reserved.
+
+
+#include "Rig/Solvers/IKRigLimbSolver.h"
+
+#include "Rig/IKRigDataTypes.h"
+
+#include UE_INLINE_GENERATED_CPP_BY_NAME(IKRigLimbSolver)
+
+#define LOCTEXT_NAMESPACE "IKRig_LimbSolver"
+
+void FIKRigLimbSolver::Initialize(const FIKRigSkeleton& InIKRigSkeleton)
+{
+	Solver.Reset();
+	ChildrenToUpdate.Empty();
+	
+	if (Settings.GoalName == NAME_None || Settings.EndBone == NAME_None || Settings.StartBone == NAME_None)
+	{
+		return;
+	}
+
+	int32 BoneIndex = InIKRigSkeleton.GetBoneIndexFromName(Settings.EndBone);
+	const int32 RootIndex = InIKRigSkeleton.GetBoneIndexFromName(Settings.StartBone);
+	if (BoneIndex == INDEX_NONE || RootIndex == INDEX_NONE)
+	{
+		return;
+	}
+
+	// populate indices
+	TArray<int32> BoneIndices( {BoneIndex} );
+	BoneIndex = InIKRigSkeleton.GetParentIndex(BoneIndex);
+	while (BoneIndex != INDEX_NONE && BoneIndex >= RootIndex)
+	{
+		BoneIndices.Add(BoneIndex);
+		BoneIndex = InIKRigSkeleton.GetParentIndex(BoneIndex);
+	};
+
+	// if chain is not long enough
+	if (BoneIndices.Num() < 3)
+	{
+		return;
+	}
+
+	// sort the chain from root to end
+	Algo::Reverse(BoneIndices);
+
+	// initialize solver
+	for (int32 Index: BoneIndices)
+	{
+		const FVector Location = InIKRigSkeleton.CurrentPoseGlobal[Index].GetLocation();
+		Solver.AddLink(Location, Index);
+	}
+
+	const bool bInitialized = Solver.Initialize();
+	if (bInitialized)
+	{
+		// store children that needs propagation once solved
+		TArray<int32> Children;
+		for (int32 Index = 0; Index < BoneIndices.Num()-1; ++Index)
+		{
+			// store children if not already handled by the solver (part if the links)
+			InIKRigSkeleton.GetChildIndices(BoneIndices[Index], Children);
+			const int32 NextIndex = BoneIndices[Index+1];
+			for (const int32 ChildIndex: Children)
+			{
+				if (ChildIndex != NextIndex)
+				{
+					ChildrenToUpdate.Add(ChildIndex);
+					GatherChildren(ChildIndex, InIKRigSkeleton, ChildrenToUpdate);
+				}
+			}
+		}
+		// store end bone children
+		GatherChildren(BoneIndices.Last(), InIKRigSkeleton, ChildrenToUpdate);
+	}
+}
+
+void FIKRigLimbSolver::Solve(FIKRigSkeleton& InIKRigSkeleton, const FIKRigGoalContainer& InGoals)
+{
+	if (Solver.NumLinks() < 3)
+	{
+		return;
+	}
+
+	// fetch the goal transform
+	const FIKRigGoal* IKGoal = InGoals.FindGoalByName(Settings.GoalName);
+	if (!IKGoal)
+	{
+		return;
+	}
+	const FVector& GoalLocation = IKGoal->FinalBlendedPosition;
+	const FQuat& GoalRotation = IKGoal->FinalBlendedRotation;
+
+	// run the solve
+	const bool bModifiedLimb = Solver.Solve(
+		InIKRigSkeleton.CurrentPoseGlobal,
+		GoalLocation,
+		GoalRotation,
+		Settings);
+
+	// propagate if needed
+	if (bModifiedLimb)
+	{
+		// update chain bones local transform
+		for (int32 Index = 0; Index < Solver.NumLinks(); Index++)
+		{
+			InIKRigSkeleton.UpdateLocalTransformFromGlobal(Solver.GetBoneIndex(Index));
+		}
+
+		// propagate to children
+		for (const int32 ChildIndex: ChildrenToUpdate)
+		{
+			InIKRigSkeleton.UpdateGlobalTransformFromLocal(ChildIndex);
+		}
+	}
+}
+
+void FIKRigLimbSolver::GetRequiredBones(TSet<FName>& OutRequiredBones) const
+{
+	OutRequiredBones.Add(Settings.StartBone);
+	OutRequiredBones.Add(Settings.EndBone);
+}
+
+void FIKRigLimbSolver::GetRequiredGoals(TSet<FName>& OutRequiredGoals) const
+{
+	OutRequiredGoals.Add(Settings.GoalName);
+}
+
+FIKRigSolverSettingsBase* FIKRigLimbSolver::GetSolverSettings()
+{
+	return &Settings;
+}
+
+const UScriptStruct* FIKRigLimbSolver::GetSolverSettingsType() const
+{
+	return FIKRigLimbSolverSettings::StaticStruct();
+}
+
+void FIKRigLimbSolver::AddGoal(const UIKRigEffectorGoal& InNewGoal)
+{
+	Settings.GoalName = InNewGoal.GoalName;
+	Settings.EndBone = InNewGoal.BoneName;
+}
+
+void FIKRigLimbSolver::OnGoalRenamed(const FName& InOldName, const FName& InNewName)
+{
+	if (Settings.GoalName == InOldName)
+	{
+		Settings.GoalName = InNewName;
+	}
+}
+
+void FIKRigLimbSolver::OnGoalMovedToDifferentBone(const FName& InGoalName, const FName& InNewBoneName)
+{
+	if (Settings.GoalName == InGoalName)
+	{
+		Settings.EndBone = InNewBoneName;
+	}
+}
+
+void FIKRigLimbSolver::OnGoalRemoved(const FName& InGoalName)
+{
+	if (Settings.GoalName == InGoalName)
+	{
+		Settings.GoalName = NAME_None;
+		Settings.EndBone = NAME_None;
+	}
+}
+
+bool FIKRigLimbSolver::UsesStartBone() const
+{
+	return true;
+}
+
+void FIKRigLimbSolver::SetStartBone(const FName& InStartBoneName)
+{
+	Settings.StartBone = InStartBoneName;
+}
+
+FName FIKRigLimbSolver::GetStartBone() const
+{
+	return Settings.StartBone;
+}
+
+FIKRigLimbSolverSettings UIKRigLimbSolverController::GetSolverSettings()
+{
+	return *static_cast<FIKRigLimbSolverSettings*>(SolverToControl->GetSolverSettings());
+}
+
+void UIKRigLimbSolverController::SetSolverSettings(FIKRigLimbSolverSettings InSettings)
+{
+	SolverToControl->SetSolverSettings(&InSettings);
+}
+
+#if WITH_EDITOR
+
+UIKRigSolverControllerBase* FIKRigLimbSolver::GetSolverController(UObject* Outer)
+{
+	return CreateControllerIfNeeded(Outer, UIKRigLimbSolverController::StaticClass());
+}
+
+
+FText FIKRigLimbSolver::GetNiceName() const
+{
+	return FText(LOCTEXT("SolverName", "Limb IK"));
+}
+
+bool FIKRigLimbSolver::GetWarningMessage(FText& OutWarningMessage) const
+{
+	if (Settings.GoalName == NAME_None)
+	{
+		OutWarningMessage = LOCTEXT("MissingGoal", "Missing goal.");
+		return true;
+	}
+
+	if (Settings.StartBone == NAME_None)
+	{
+		OutWarningMessage = LOCTEXT("MissingRoot", "Missing root.");
+		return true;
+	}
+
+	if (Solver.NumLinks() < 3)
+	{
+		OutWarningMessage = LOCTEXT("Requires3BonesChain", "Requires at least 3 bones between root and goal.");
+		return true;
+	}
+	
+	return false;
+}
+
+bool FIKRigLimbSolver::IsBoneAffectedBySolver(const FName& BoneName, const FIKRigSkeleton& IKRigSkeleton) const
+{
+	return IKRigSkeleton.IsBoneInDirectLineage(BoneName, Settings.StartBone);
+}
+
+#endif
+
+void FIKRigLimbSolver::GatherChildren(const int32 BoneIndex, const FIKRigSkeleton& InSkeleton, TArray<int32>& OutChildren)
+{
+	TArray<int32> Children;
+	InSkeleton.GetChildIndices(BoneIndex, Children);
+	for (int32 ChildIndex: Children)
+	{
+		OutChildren.Add(ChildIndex);
+		GatherChildren(ChildIndex, InSkeleton, OutChildren);
+	}
+}
+
+#undef LOCTEXT_NAMESPACE
+
diff --git a/Engine/Plugins/Animation/IKRig/Source/IKRig/Private/Rig/Solvers/IKRigPoleSolver.cpp b/Engine/Plugins/Animation/IKRig/Source/IKRig/Private/Rig/Solvers/IKRigPoleSolver.cpp
new file mode 100644
--- /dev/null
+++ b/Engine/Plugins/Animation/IKRig/Source/IKRig/Private/Rig/Solvers/IKRigPoleSolver.cpp
@@ -0,0 +1,296 @@
+// Copyright Epic Games, Inc. All Rights Reserved.
+
+#include "Rig/Solvers/IKRigPoleSolver.h"
+#include "Rig/IKRigDataTypes.h"
+#include "Rig/IKRigSkeleton.h"
+
+#include UE_INLINE_GENERATED_CPP_BY_NAME(IKRigPoleSolver)
+
+#define LOCTEXT_NAMESPACE "PoleSolver"
+
+void FIKRigPoleSolver::Initialize(const FIKRigSkeleton& InIKRigSkeleton)
+{
+	Chain.Empty();
+	
+	int32 BoneIndex = InIKRigSkeleton.GetBoneIndexFromName(Settings.EndBone);
+	const int32 RootIndex = InIKRigSkeleton.GetBoneIndexFromName(Settings.StartBone);
+
+	if (BoneIndex == INDEX_NONE || RootIndex == INDEX_NONE)
+	{
+		return;
+	}
+
+	// populate chain
+	Chain.Add(BoneIndex);
+	BoneIndex = InIKRigSkeleton.GetParentIndex(BoneIndex);
+	while (BoneIndex != INDEX_NONE && BoneIndex >= RootIndex)
+	{
+		Chain.Add(BoneIndex);
+		BoneIndex = InIKRigSkeleton.GetParentIndex(BoneIndex);
+	};
+
+	// if chain is not long enough
+	if (Chain.Num() < 3)
+	{
+		Chain.Empty();
+		return;
+	}
+	
+	// sort the chain from root to end
+	Algo::Reverse(Chain);
+	
+	// store children that needs propagation once solved
+	TArray<int32> Children;
+	for (int32 Index = 0; Index < Chain.Num()-1; ++Index)
+	{
+		// store children if not already handled by the solver
+		InIKRigSkeleton.GetChildIndices(Chain[Index], Children);
+		const int32 NextIndex = Chain[Index+1];
+		for (const int32 ChildIndex: Children)
+		{
+			if (ChildIndex != NextIndex)
+			{
+				ChildrenToUpdate.Add(ChildIndex);
+				GatherChildren(ChildIndex, InIKRigSkeleton, ChildrenToUpdate);
+			}
+		}
+	}
+}
+
+void FIKRigPoleSolver::Solve(FIKRigSkeleton& InIKRigSkeleton, const FIKRigGoalContainer& InGoals)
+{
+	if (Chain.Num() < 3)
+	{
+		return;
+	}
+	
+	const FIKRigGoal* IKGoal = InGoals.FindGoalByName(Settings.AimAtGoal);
+	if (!IKGoal)
+	{
+		return;
+	}
+
+	const bool bHasAlpha = Settings.Alpha > KINDA_SMALL_NUMBER;
+	if (!bHasAlpha)
+	{
+		return;
+	}
+
+	const int32 RootIndex = Chain[0];
+	const int32 KneeIndex = Chain[1];
+	const int32 EndIndex = Chain.Last();
+
+	TArray<FTransform>& InOutTransforms = InIKRigSkeleton.CurrentPoseGlobal;
+
+	// initial configuration
+	const FVector RootLocation = InOutTransforms[RootIndex].GetLocation();
+	const FVector KneeLocation = InOutTransforms[KneeIndex].GetLocation();
+	const FVector EndLocation = InOutTransforms[EndIndex].GetLocation();
+	
+	const FVector RootToEnd = (EndLocation - RootLocation).GetSafeNormal();
+	const FVector RootToKnee = (KneeLocation - RootLocation).GetSafeNormal();
+	if (RootToEnd.IsZero() || RootToKnee.IsZero())
+	{
+		return;
+	}
+	const FVector InitPlane = FVector::CrossProduct(RootToEnd, RootToKnee).GetSafeNormal();
+
+	// target configuration
+	const FVector& GoalLocation = IKGoal->FinalBlendedPosition;
+	const FVector RootToPole = (GoalLocation - RootLocation).GetSafeNormal();
+	if (RootToPole.IsZero())
+	{
+		return;
+	}
+	const FVector TargetPlane = FVector::CrossProduct(RootToEnd, RootToPole).GetSafeNormal();
+
+	// compute delta rotation from InitialPlane to TargetPlane
+	if (InitPlane.IsZero() || InitPlane.Equals(TargetPlane))
+	{
+    	return;
+    }
+	
+	const FQuat DeltaRotation = FQuat::FindBetweenNormals(InitPlane, TargetPlane);
+	if (DeltaRotation.IsIdentity())
+	{
+		return;
+	}
+	
+	// update transforms
+	for (int32 Index = 0; Index < Chain.Num()-1; Index++)
+	{
+		int32 BoneIndex = Chain[Index];
+		FTransform& BoneTransform = InOutTransforms[BoneIndex];
+
+		// rotation
+		const FQuat BoneRotation = BoneTransform.GetRotation();
+		const FQuat TargetRotation = FMath::Lerp(BoneRotation, DeltaRotation * BoneRotation, Settings.Alpha);
+		BoneTransform.SetRotation(TargetRotation);
+
+		// translation
+		const FVector BoneTranslation = BoneTransform.GetLocation();
+		const FVector TargetTranslation = FMath::Lerp( BoneTranslation, RootLocation + DeltaRotation.RotateVector(BoneTranslation - RootLocation), Settings.Alpha);
+		BoneTransform.SetTranslation(TargetTranslation);
+	}
+
+	// propagate to children
+	for (const int32 ChildIndex: ChildrenToUpdate)
+	{
+		InIKRigSkeleton.UpdateGlobalTransformFromLocal(ChildIndex);
+	}
+}
+
+void FIKRigPoleSolver::GetRequiredBones(TSet<FName>& OutRequiredBones) const
+{
+	OutRequiredBones.Add(Settings.StartBone);
+	OutRequiredBones.Add(Settings.EndBone);
+}
+
+void FIKRigPoleSolver::GetRequiredGoals(TSet<FName>& OutRequiredGoals) const
+{
+	OutRequiredGoals.Add(Settings.AimAtGoal);
+}
+
+FIKRigSolverSettingsBase* FIKRigPoleSolver::GetSolverSettings()
+{
+	return &Settings;
+}
+
+const UScriptStruct* FIKRigPoleSolver::GetSolverSettingsType() const
+{
+	return FIKRigPoleSolverSettings::StaticStruct();
+}
+
+void FIKRigPoleSolver::AddGoal(const UIKRigEffectorGoal& NewGoal)
+{
+	Settings.AimAtGoal = NewGoal.GoalName;
+}
+
+void FIKRigPoleSolver::OnGoalRenamed(const FName& InOldName, const FName& InNewName)
+{
+	if (Settings.AimAtGoal == InOldName)
+	{
+		Settings.AimAtGoal = InNewName;
+	}
+}
+
+void FIKRigPoleSolver::OnGoalRemoved(const FName& InName)
+{
+	if (Settings.AimAtGoal == InName)
+	{
+		Settings.AimAtGoal = NAME_None;
+	}
+}
+
+void FIKRigPoleSolver::OnGoalMovedToDifferentBone(const FName& InGoalName, const FName& InNewBoneName)
+{
+	if (Settings.AimAtGoal == InGoalName)
+	{
+		Settings.EndBone = InNewBoneName;
+	}
+}
+
+void FIKRigPoleSolver::SetStartBone(const FName& InRootBoneName)
+{
+	Settings.StartBone = InRootBoneName;
+}
+
+FName FIKRigPoleSolver::GetStartBone() const
+{
+	return Settings.StartBone;
+}
+
+void FIKRigPoleSolver::SetEndBone(const FName& InEndBoneName)
+{
+	Settings.EndBone = InEndBoneName;
+}
+
+FName FIKRigPoleSolver::GetEndBone() const
+{
+	return Settings.EndBone;
+}
+
+FIKRigPoleSolverSettings UIKRigPoleSolverController::GetSolverSettings()
+{
+	return *static_cast<FIKRigPoleSolverSettings*>(SolverToControl->GetSolverSettings());
+}
+
+void UIKRigPoleSolverController::SetSolverSettings(FIKRigPoleSolverSettings InSettings)
+{
+	SolverToControl->SetSolverSettings(&InSettings);
+}
+
+#if WITH_EDITOR
+
+UIKRigSolverControllerBase* FIKRigPoleSolver::GetSolverController(UObject* Outer)
+{
+	return CreateControllerIfNeeded(Outer, UIKRigPoleSolverController::StaticClass());
+}
+
+FText FIKRigPoleSolver::GetNiceName() const
+{
+	return FText(LOCTEXT("SolverName", "Pole Solver"));
+}
+
+bool FIKRigPoleSolver::GetWarningMessage(FText& OutWarningMessage) const
+{
+	if (Settings.StartBone == NAME_None)
+	{
+		OutWarningMessage = LOCTEXT("MissingRoot", "Missing root bone.");
+		return true;
+	}
+	
+	if (Settings.EndBone == NAME_None)
+	{
+		OutWarningMessage = LOCTEXT("MissingEnd", "Missing end bone.");
+		return true;
+	}
+	
+	if (Settings.AimAtGoal == NAME_None)
+	{
+		OutWarningMessage = LOCTEXT("MissingGoal", "Missing aim goal.");
+		return true;
+	}
+	
+	if (Chain.Num() < 3)
+	{
+		OutWarningMessage = LOCTEXT("Requires3BonesChain", "Requires at least 3 bones between root and end bones.");
+		return true;
+	}
+
+	return false;
+}
+
+bool FIKRigPoleSolver::IsBoneAffectedBySolver(const FName& InBoneName, const FIKRigSkeleton& InIKRigSkeleton) const
+{
+	const bool bAffected = InIKRigSkeleton.IsBoneInDirectLineage(InBoneName, Settings.StartBone);
+	if (!bAffected)
+	{
+		return false;
+	}
+	
+	const int32 EndIndex = Chain.IsEmpty() ? INDEX_NONE : Chain.Last();
+	if (EndIndex == INDEX_NONE)
+	{
+		return false;
+	}
+
+	const int32 ChildIndex = InIKRigSkeleton.GetBoneIndexFromName(InBoneName);
+	return ChildIndex <= EndIndex;
+}
+
+#endif
+
+void FIKRigPoleSolver::GatherChildren(const int32 BoneIndex, const FIKRigSkeleton& InSkeleton, TArray<int32>& OutChildren)
+{
+	TArray<int32> Children;
+	InSkeleton.GetChildIndices(BoneIndex, Children);
+	for (int32 ChildIndex: Children)
+	{
+		OutChildren.Add(ChildIndex);
+		GatherChildren(ChildIndex, InSkeleton, OutChildren);
+	}
+}
+
+#undef LOCTEXT_NAMESPACE
+
diff --git a/Engine/Plugins/Animation/IKRig/Source/IKRig/Private/Rig/Solvers/IKRigSolver.cpp b/Engine/Plugins/Animation/IKRig/Source/IKRig/Private/Rig/Solvers/IKRigSolver.cpp
deleted file mode 100644
diff --git a/Engine/Plugins/Animation/IKRig/Source/IKRig/Private/Rig/Solvers/IKRigSolverBase.cpp b/Engine/Plugins/Animation/IKRig/Source/IKRig/Private/Rig/Solvers/IKRigSolverBase.cpp
new file mode 100644
--- /dev/null
+++ b/Engine/Plugins/Animation/IKRig/Source/IKRig/Private/Rig/Solvers/IKRigSolverBase.cpp
@@ -0,0 +1,137 @@
+// Copyright Epic Games, Inc. All Rights Reserved.
+
+#include "Rig/Solvers/IKRigSolverBase.h"
+
+#include UE_INLINE_GENERATED_CPP_BY_NAME(IKRigSolverBase)
+
+void FIKRigSolverBase::UpdateSettingsFromAsset(const FIKRigSolverBase& InAssetSolver)
+{
+	FIKRigSolverBase& AssetSolver = const_cast<FIKRigSolverBase&>(InAssetSolver);
+
+	// copy enabled flag
+	SetEnabled(AssetSolver.IsEnabled());
+	
+	// copy solver settings
+	const UScriptStruct* SolverSettingsType = GetSolverSettingsType();
+	const FIKRigSolverSettingsBase* CopyFrom = AssetSolver.GetSolverSettings();
+	FIKRigSolverSettingsBase* CopyTo = GetSolverSettings();
+	CopyAllEditableProperties(SolverSettingsType, CopyFrom, CopyTo);
+
+	// (optional) copy goal settings
+	if (InAssetSolver.UsesCustomGoalSettings())
+	{
+		const UScriptStruct* GoalSettingsType = GetGoalSettingsType();
+		TSet<FName> GoalsWithSettings;
+		InAssetSolver.GetGoalsWithSettings(GoalsWithSettings);
+		for (FName Goal : GoalsWithSettings)
+		{
+			const FIKRigGoalSettingsBase* CopyFromGoalSettings = AssetSolver.GetGoalSettings(Goal);
+			FIKRigGoalSettingsBase* CopyToGoalSettings = GetGoalSettings(Goal);
+			CopyAllEditableProperties(GoalSettingsType, CopyFromGoalSettings, CopyToGoalSettings);	
+		}
+	}
+	
+	// (optional) copy bones settings
+	if (InAssetSolver.UsesCustomBoneSettings())
+	{
+		const UScriptStruct* BoneSettingsType = GetBoneSettingsType();
+		TSet<FName> BonesWithSettings;
+		InAssetSolver.GetBonesWithSettings(BonesWithSettings);
+		for (FName Bone : BonesWithSettings)
+		{
+			const FIKRigBoneSettingsBase* CopyFromBoneSettings = AssetSolver.GetBoneSettings(Bone);
+			FIKRigBoneSettingsBase* CopyToBoneSettings = GetBoneSettings(Bone);
+			CopyAllEditableProperties(BoneSettingsType, CopyFromBoneSettings, CopyToBoneSettings);	
+		}
+	}
+}
+
+void FIKRigSolverBase::SetSolverSettings(const FIKRigSolverSettingsBase* InSettings)
+{
+	CopyAllEditableProperties(GetSolverSettingsType(), InSettings, GetSolverSettings());
+}
+
+void FIKRigSolverBase::SetGoalSettings(const FName& InGoalName, FIKRigGoalSettingsBase* InSettings)
+{
+	ensure(UsesCustomGoalSettings());
+	CopyAllEditableProperties(GetGoalSettingsType(), InSettings, GetGoalSettings(InGoalName));
+}
+
+void FIKRigSolverBase::SetBoneSettings(const FName& InBoneName, FIKRigBoneSettingsBase* InSettings)
+{
+	ensure(UsesCustomBoneSettings());
+	CopyAllEditableProperties(GetBoneSettingsType(), InSettings, GetBoneSettings(InBoneName));
+}
+
+void FIKRigSolverBase::CopyAllEditableProperties(
+	const UScriptStruct* SettingsType,
+	const FIKRigSettingsBase* CopyFrom,
+	FIKRigSettingsBase* CopyTo)
+{
+	// ensure source and destination were provided
+	if (!(CopyFrom && CopyTo))
+	{
+		return;
+	}
+	
+	// ensure the provided SettingsType is an IKRig settings stype
+	if (!ensure(SettingsType && SettingsType->IsChildOf(FIKRigSettingsBase::StaticStruct())))
+	{
+		return;
+	}
+
+	// ensure CopyFrom is instance of or derives from SettingsType
+	if (!ensure(SettingsType->IsChildOf(CopyFrom->StaticStruct())))
+	{
+		return;
+	}
+
+	// ensure CopyTo is instance of or derives from SettingsType
+	if (!ensure(SettingsType->IsChildOf(CopyTo->StaticStruct())))
+	{
+		return;
+	}
+
+	// iterate through all properties of the SettingsType
+	for (TFieldIterator<FProperty> It(SettingsType); It; ++It)
+	{
+		const FProperty* Property = *It;
+
+		// don't copy properties unless they are marked BlueprintReadWrite
+		if (Property->HasAnyPropertyFlags(CPF_BlueprintReadOnly | CPF_Transient))
+		{
+			continue;
+		}
+
+		// copy the value from CopyFrom to CopyTo
+		const uint8* SrcValuePtr = Property->ContainerPtrToValuePtr<uint8>(CopyFrom);
+		uint8* DestValuePtr = Property->ContainerPtrToValuePtr<uint8>(CopyTo);
+		if (SrcValuePtr && DestValuePtr)
+		{
+			Property->CopyCompleteValue(DestValuePtr, SrcValuePtr);
+		}
+	}
+}
+
+#if WITH_EDITOR
+
+UIKRigSolverControllerBase* FIKRigSolverBase::GetSolverController(UObject* Outer)
+{
+	return CreateControllerIfNeeded(Outer, UIKRigSolverControllerBase::StaticClass());
+}
+
+UIKRigSolverControllerBase* FIKRigSolverBase::CreateControllerIfNeeded(UObject* Outer, const UClass* ClassType)
+{
+	if (!Controller.IsValid())
+	{
+		if (ensure(ClassType && ClassType->IsChildOf(UIKRigSolverControllerBase::StaticClass())))
+		{
+			Controller = TStrongObjectPtr(NewObject<UIKRigSolverControllerBase>(Outer, ClassType));
+			Controller->SolverToControl = this;
+		}
+	}
+	return Controller.Get();
+}
+
+#endif
+
diff --git a/Engine/Plugins/Animation/IKRig/Source/IKRig/Private/Rig/Solvers/IKRig_BodyMover.cpp b/Engine/Plugins/Animation/IKRig/Source/IKRig/Private/Rig/Solvers/IKRig_BodyMover.cpp
deleted file mode 100644
diff --git a/Engine/Plugins/Animation/IKRig/Source/IKRig/Private/Rig/Solvers/IKRig_FBIKSolver.cpp b/Engine/Plugins/Animation/IKRig/Source/IKRig/Private/Rig/Solvers/IKRig_FBIKSolver.cpp
--- a/Engine/Plugins/Animation/IKRig/Source/IKRig/Private/Rig/Solvers/IKRig_FBIKSolver.cpp
+++ b/Engine/Plugins/Animation/IKRig/Source/IKRig/Private/Rig/Solvers/IKRig_FBIKSolver.cpp
@@ -7,15 +7,15 @@
 
 #include UE_INLINE_GENERATED_CPP_BY_NAME(IKRig_FBIKSolver)
 
-#define LOCTEXT_NAMESPACE "UIKRigFBIKSolver"
+#define LOCTEXT_NAMESPACE "FFBIKSolver"
 
-void UIKRigFBIKSolver::Initialize(const FIKRigSkeleton& InSkeleton)
+void FIKRigFullBodyIKSolver::Initialize(const FIKRigSkeleton& InSkeleton)
 {	
 	// check how many effectors are assigned to a bone
 	int NumEffectors = 0;
-	for (const UIKRig_FBIKEffector* Effector : Effectors)
+	for (const FIKRigFBIKGoalSettings& Effector : AllGoalSettings)
 	{
-		if (InSkeleton.GetBoneIndexFromName(Effector->BoneName) != INDEX_NONE)
+		if (InSkeleton.GetBoneIndexFromName(Effector.BoneName) != INDEX_NONE)
 		{
 			++NumEffectors; // bone is set and exists!
 		}
@@ -23,7 +23,7 @@ void UIKRigFBIKSolver::Initialize(const FIKRigSkeleton& InSkeleton)
 
 	// validate inputs are ready to be initialized
 	const bool bHasEffectors = NumEffectors > 0;
-	const bool bRootIsAssigned = RootBone != NAME_None;
+	const bool bRootIsAssigned = Settings.RootBone != NAME_None;
 	if (!(bHasEffectors && bRootIsAssigned))
 	{
 		return; // not setup yet
@@ -42,33 +42,33 @@ void UIKRigFBIKSolver::Initialize(const FIKRigSkeleton& InSkeleton)
 		const FTransform OrigTransform = InSkeleton.RefPoseGlobal[BoneIndex];
 		const FVector InOrigPosition = OrigTransform.GetLocation();
 		const FQuat InOrigRotation = OrigTransform.GetRotation();
-		const bool bIsRoot = Name == RootBone;
+		const bool bIsRoot = Name == Settings.RootBone;
 		Solver.AddBone(Name, ParentIndex, InOrigPosition, InOrigRotation, bIsRoot);
 	}
 
 	// create effectors
-	for (UIKRig_FBIKEffector* Effector : Effectors)
+	for (FIKRigFBIKGoalSettings& Effector : AllGoalSettings)
 	{
-		Effector->IndexInSolver = Solver.AddEffector(Effector->BoneName);
+		Effector.IndexInSolver = Solver.AddEffector(Effector.BoneName);
 	}
 		
 	// initialize solver
 	Solver.Initialize();
 }
 
-void UIKRigFBIKSolver::Solve(FIKRigSkeleton& IKRigSkeleton, const FIKRigGoalContainer& Goals)
+void FIKRigFullBodyIKSolver::Solve(FIKRigSkeleton& InIKRigSkeleton, const FIKRigGoalContainer& InGoals)
 {
 	if (!Solver.IsReadyToSimulate())
 	{
 		return;
 	}
 
-	if (Solver.GetNumBones() != IKRigSkeleton.BoneNames.Num())
+	if (Solver.GetNumBones() != InIKRigSkeleton.BoneNames.Num())
 	{
 		return;
 	}
 
-	TArray<FTransform>& InOutTransforms = IKRigSkeleton.CurrentPoseGlobal;
+	TArray<FTransform>& InOutTransforms = InIKRigSkeleton.CurrentPoseGlobal;
 	
 	// set bones to input pose
 	for(int32 BoneIndex = 0; BoneIndex < Solver.GetNumBones(); BoneIndex++)
@@ -77,59 +77,58 @@ void UIKRigFBIKSolver::Solve(FIKRigSkeleton& IKRigSkeleton, const FIKRigGoalCont
 	}
 
 	// update bone settings
-	for (const UIKRig_FBIKBoneSettings* BoneSetting : BoneSettings)
+	for (const FIKRigFBIKBoneSettings& BoneSetting : AllBoneSettings)
 	{
-		const int32 BoneIndex = Solver.GetBoneIndex(BoneSetting->Bone);
+		const int32 BoneIndex = Solver.GetBoneIndex(BoneSetting.Bone);
 		if (PBIK::FBoneSettings* InternalSettings = Solver.GetBoneSettings(BoneIndex))
 		{
-			BoneSetting->CopyToCoreStruct(*InternalSettings);
+			BoneSetting.CopyToCoreStruct(*InternalSettings);
 		}
 	}
 
 	// update effectors
-	for (const UIKRig_FBIKEffector* Effector : Effectors)
+	for (const FIKRigFBIKGoalSettings& GoalSettings : AllGoalSettings)
 	{
-		if (Effector->IndexInSolver < 0)
+		if (GoalSettings.IndexInSolver < 0)
 		{
 			continue;
 		}
 		
-		const FIKRigGoal* Goal = Goals.FindGoalByName(Effector->GoalName);
+		const FIKRigGoal* Goal = InGoals.FindGoalByName(GoalSettings.Goal);
 		if (!Goal)
 		{
 			return;
 		}
 
-		PBIK::FEffectorSettings Settings;
-		Settings.PositionAlpha = 1.0f; // this is constant because IKRig manages offset alphas itself
-		Settings.RotationAlpha = 1.0f; // this is constant because IKRig manages offset alphas itself
-		Settings.StrengthAlpha = Effector->StrengthAlpha;
-		Settings.ChainDepth = Effector->ChainDepth;
-		Settings.PullChainAlpha = Effector->PullChainAlpha;
-		Settings.PinRotation = Effector->PinRotation;
+		PBIK::FEffectorSettings EffectorSettings;
+		EffectorSettings.PositionAlpha = 1.0f; // this is constant because IKRig manages offset alphas itself
+		EffectorSettings.RotationAlpha = 1.0f; // this is constant because IKRig manages offset alphas itself
+		EffectorSettings.StrengthAlpha = GoalSettings.StrengthAlpha;
+		EffectorSettings.ChainDepth = GoalSettings.ChainDepth;
+		EffectorSettings.PullChainAlpha = GoalSettings.PullChainAlpha;
+		EffectorSettings.PinRotation = GoalSettings.PinRotation;
 		
 		Solver.SetEffectorGoal(
-			Effector->IndexInSolver,
+			GoalSettings.IndexInSolver,
 			Goal->FinalBlendedPosition,
 			Goal->FinalBlendedRotation,
-			Settings);
+			EffectorSettings);
 	}
 
 	// update settings
-	FPBIKSolverSettings Settings;
-	Settings.Iterations = Iterations;
-	Settings.SubIterations = SubIterations;
-	Settings.MassMultiplier = MassMultiplier;
-	Settings.bAllowStretch = bAllowStretch;
-	Settings.RootBehavior = RootBehavior;
-	Settings.PrePullRootSettings = PrePullRootSettings;
-	Settings.GlobalPullChainAlpha = PullChainAlpha;
-	Settings.MaxAngle = MaxAngle;
-	Settings.OverRelaxation = OverRelaxation;
-	Settings.bStartSolveFromInputPose_DEPRECATED = bStartSolveFromInputPose_DEPRECATED;
+	FPBIKSolverSettings SolverSettings;
+	SolverSettings.Iterations = Settings.Iterations;
+	SolverSettings.SubIterations = Settings.SubIterations;
+	SolverSettings.MassMultiplier = Settings.MassMultiplier;
+	SolverSettings.bAllowStretch = Settings.bAllowStretch;
+	SolverSettings.RootBehavior = Settings.RootBehavior;
+	SolverSettings.PrePullRootSettings = Settings.PrePullRootSettings;
+	SolverSettings.GlobalPullChainAlpha = Settings.GlobalPullChainAlpha;
+	SolverSettings.MaxAngle = Settings.MaxAngle;
+	SolverSettings.OverRelaxation = Settings.OverRelaxation;
 
 	// solve
-	Solver.Solve(Settings);
+	Solver.Solve(SolverSettings);
 
 	// copy transforms back
 	for(int32 BoneIndex = 0; BoneIndex < Solver.GetNumBones(); BoneIndex++)
@@ -138,223 +137,238 @@ void UIKRigFBIKSolver::Solve(FIKRigSkeleton& IKRigSkeleton, const FIKRigGoalCont
 	}
 }
 
-void UIKRigFBIKSolver::GetBonesWithSettings(TSet<FName>& OutBonesWithSettings) const
+void FIKRigFullBodyIKSolver::GetRequiredBones(TSet<FName>& OutRequiredBones) const
 {
-	for (UIKRig_FBIKBoneSettings* BoneSetting : BoneSettings)
+	OutRequiredBones.Add(Settings.RootBone);
+	// TODO in the future make this bone agnostic, stop storing goal bone names and find dynamically at Init() time
+	for (const FIKRigFBIKGoalSettings& GoalSettings : AllGoalSettings)
 	{
-		if (BoneSetting)
-		{
-			OutBonesWithSettings.Add(BoneSetting->Bone);
-		}
+		OutRequiredBones.Add(GoalSettings.BoneName);
 	}
 }
 
-void UIKRigFBIKSolver::UpdateSolverSettings(UIKRigSolver* InSettings)
+void FIKRigFullBodyIKSolver::GetRequiredGoals(TSet<FName>& OutRequiredGoals) const
 {
-	if(UIKRigFBIKSolver* Settings = Cast<UIKRigFBIKSolver>(InSettings))
+	for (const FIKRigFBIKGoalSettings& GoalSettings : AllGoalSettings)
 	{
-		Iterations = Settings->Iterations;
-		SubIterations = Settings->SubIterations;
-		MassMultiplier = Settings->MassMultiplier;
-		bAllowStretch = Settings->bAllowStretch;
-		RootBehavior = Settings->RootBehavior;
-		PrePullRootSettings = Settings->PrePullRootSettings;
-		PullChainAlpha = Settings->PullChainAlpha;
-		MaxAngle = Settings->MaxAngle;
-		OverRelaxation = Settings->OverRelaxation;
-		bStartSolveFromInputPose_DEPRECATED = Settings->bStartSolveFromInputPose_DEPRECATED;
-
-		// copy effector settings
-		for (const UIKRig_FBIKEffector* InEffector : Settings->Effectors)
-		{
-			for (UIKRig_FBIKEffector* Effector : Effectors)
-			{
-				if (Effector && Effector->GoalName == InEffector->GoalName)
-				{
-					Effector->CopySettings(InEffector);
-					break;
-				}
-			}
-		}
-
-		// copy bone settings
-		for (const UIKRig_FBIKBoneSettings* InBoneSetting : Settings->BoneSettings)
-		{
-			for (UIKRig_FBIKBoneSettings* BoneSetting : BoneSettings)
-			{
-				if (BoneSetting && BoneSetting->Bone == InBoneSetting->Bone)
-				{
-					BoneSetting->CopySettings(InBoneSetting);
-					break;	
-				}
-			}
-		}
+		OutRequiredGoals.Add(GoalSettings.Goal);
 	}
 }
 
-void UIKRigFBIKSolver::RemoveGoal(const FName& GoalName)
+FIKRigSolverSettingsBase* FIKRigFullBodyIKSolver::GetSolverSettings()
 {
-	// ensure goal even exists in this solver
-	const int32 GoalIndex = GetIndexOfGoal(GoalName);
-	if (GoalIndex == INDEX_NONE)
-	{
-		return;
-	}
-
-	// remove it
-	Effectors.RemoveAt(GoalIndex);
+	return &Settings;
 }
 
-bool UIKRigFBIKSolver::IsGoalConnected(const FName& GoalName) const
+const UScriptStruct* FIKRigFullBodyIKSolver::GetSolverSettingsType() const
 {
-	return GetIndexOfGoal(GoalName) != INDEX_NONE;
+	return FIKRigFBIKSettings::StaticStruct();
 }
 
-#if WITH_EDITOR
+// GOAL SETTINGS
 
-FText UIKRigFBIKSolver::GetNiceName() const
+void FIKRigFullBodyIKSolver::AddGoal(const UIKRigEffectorGoal& InNewGoal)
 {
-	return FText(LOCTEXT("SolverName", "Full Body IK"));
+	AllGoalSettings.Emplace(InNewGoal.GoalName, InNewGoal.BoneName);
 }
 
-bool UIKRigFBIKSolver::GetWarningMessage(FText& OutWarningMessage) const
+void FIKRigFullBodyIKSolver::OnGoalRenamed(const FName& InOldName, const FName& InNewName)
 {
-	if (RootBone == NAME_None)
-	{
-		OutWarningMessage = LOCTEXT("MissingRoot", "Missing root bone.");
-		return true;
-	}
-
-	if (Effectors.IsEmpty())
+	// ensure goal even exists in this solver
+	const int32 GoalIndex = GetIndexOfGoal(InOldName);
+	if (GoalIndex == INDEX_NONE)
 	{
-		OutWarningMessage = LOCTEXT("MissingGoal", "Missing goals.");
-		return true;
+		return;
 	}
-	
-	return false;
-}
 
-void UIKRigFBIKSolver::AddGoal(const UIKRigEffectorGoal& NewGoal)
-{
-	UIKRig_FBIKEffector* NewEffector = NewObject<UIKRig_FBIKEffector>(this, UIKRig_FBIKEffector::StaticClass());
-	NewEffector->GoalName = NewGoal.GoalName;
-	NewEffector->BoneName = NewGoal.BoneName;
-	Effectors.Add(NewEffector);
+	// rename
+	AllGoalSettings[GoalIndex].Goal = InNewName;
 }
 
-void UIKRigFBIKSolver::RenameGoal(const FName& OldName, const FName& NewName)
+void FIKRigFullBodyIKSolver::OnGoalMovedToDifferentBone(const FName& InGoalName, const FName& InBoneName)
 {
 	// ensure goal even exists in this solver
-	const int32 GoalIndex = GetIndexOfGoal(OldName);
+	const int32 GoalIndex = GetIndexOfGoal(InGoalName);
 	if (GoalIndex == INDEX_NONE)
 	{
 		return;
 	}
 
-	// rename
-	Effectors[GoalIndex]->Modify();
-	Effectors[GoalIndex]->GoalName = NewName;
+	// goal moved to different bone
+	AllGoalSettings[GoalIndex].BoneName = InBoneName;
 }
 
-void UIKRigFBIKSolver::SetGoalBone(const FName& GoalName, const FName& NewBoneName)
+void FIKRigFullBodyIKSolver::OnGoalRemoved(const FName& InGoalName)
 {
 	// ensure goal even exists in this solver
-	const int32 GoalIndex = GetIndexOfGoal(GoalName);
+	const int32 GoalIndex = GetIndexOfGoal(InGoalName);
 	if (GoalIndex == INDEX_NONE)
 	{
 		return;
 	}
 
-	// rename
-	Effectors[GoalIndex]->Modify();
-	Effectors[GoalIndex]->BoneName = NewBoneName;
+	// remove it
+	AllGoalSettings.RemoveAt(GoalIndex);
 }
 
-void UIKRigFBIKSolver::SetRootBone(const FName& RootBoneName)
+bool FIKRigFullBodyIKSolver::UsesCustomGoalSettings() const
 {
-	RootBone = RootBoneName;
+	return true;
 }
 
-UObject* UIKRigFBIKSolver::GetGoalSettings(const FName& GoalName) const
+FIKRigGoalSettingsBase* FIKRigFullBodyIKSolver::GetGoalSettings(const FName& InGoalName)
 {
-	const int32 GoalIndex = GetIndexOfGoal(GoalName);
+	const int32 GoalIndex = GetIndexOfGoal(InGoalName);
 	if (GoalIndex == INDEX_NONE)
 	{
 		return nullptr;
 	}
 
-	return Effectors[GoalIndex];
+	return &AllGoalSettings[GoalIndex];
 }
 
-void UIKRigFBIKSolver::AddBoneSetting(const FName& BoneName)
+const UScriptStruct* FIKRigFullBodyIKSolver::GetGoalSettingsType() const
 {
-	if (GetBoneSetting(BoneName))
+	return FIKRigFBIKGoalSettings::StaticStruct();
+}
+
+void FIKRigFullBodyIKSolver::GetGoalsWithSettings(TSet<FName>& OutGoalsWithSettings) const
+{
+	for (const FIKRigFBIKGoalSettings& GoalSettings : AllGoalSettings)
 	{
-		return; // already have settings on this bone
+		OutGoalsWithSettings.Add(GoalSettings.Goal);
 	}
+}
 
-	UIKRig_FBIKBoneSettings* NewBoneSettings = NewObject<UIKRig_FBIKBoneSettings>(this, UIKRig_FBIKBoneSettings::StaticClass());
-	NewBoneSettings->Bone = BoneName;
-	BoneSettings.Add(NewBoneSettings);
+// START BONE
+
+bool FIKRigFullBodyIKSolver::UsesStartBone() const
+{
+	return true;
 }
 
-void UIKRigFBIKSolver::RemoveBoneSetting(const FName& BoneName)
+FName FIKRigFullBodyIKSolver::GetStartBone() const
 {
-	UIKRig_FBIKBoneSettings* BoneSettingToRemove = nullptr; 
-	for (UIKRig_FBIKBoneSettings* BoneSetting : BoneSettings)
+	return Settings.RootBone;
+}
+
+void FIKRigFullBodyIKSolver::SetStartBone(const FName& InStartBoneName)
+{
+	Settings.RootBone = InStartBoneName;
+}
+
+// BONE SETTINGS
+
+bool FIKRigFullBodyIKSolver::UsesCustomBoneSettings() const
+{
+	return true;
+}
+
+void FIKRigFullBodyIKSolver::AddSettingsToBone(const FName& InBoneName)
+{
+	if (AllBoneSettings.ContainsByPredicate([&](const FIKRigFBIKBoneSettings& Element){return Element.Bone == InBoneName;}))
 	{
-		if (BoneSetting->Bone == BoneName)
-		{
-			BoneSettingToRemove = BoneSetting;
-			break; // can only be one with this name
-		}
+		// bone already has settings
+		return;
 	}
+	
+	AllBoneSettings.Emplace(InBoneName);
+}
 
-	if (BoneSettingToRemove)
+void FIKRigFullBodyIKSolver::RemoveSettingsOnBone(const FName& InBoneName)
+{
+	AllBoneSettings.RemoveAll([&](const FIKRigFBIKBoneSettings& Element)
 	{
-		BoneSettings.Remove(BoneSettingToRemove);
-	}
+		return Element.Bone == InBoneName;
+	});
 }
 
-UObject* UIKRigFBIKSolver::GetBoneSetting(const FName& BoneName) const
+FIKRigBoneSettingsBase* FIKRigFullBodyIKSolver::GetBoneSettings(const FName& InBoneName)
 {
-	for (UIKRig_FBIKBoneSettings* BoneSetting : BoneSettings)
+	for (FIKRigFBIKBoneSettings& BoneSetting : AllBoneSettings)
 	{
-		if (BoneSetting && BoneSetting->Bone == BoneName)
+		if (BoneSetting.Bone == InBoneName)
 		{
-			return BoneSetting;
+			return &BoneSetting;
 		}
 	}
 	
 	return nullptr;
 }
 
-void UIKRigFBIKSolver::DrawBoneSettings(
-	const FName& BoneName,
-	const FIKRigSkeleton& IKRigSkeleton,
-	FPrimitiveDrawInterface* PDI) const
+const UScriptStruct* FIKRigFullBodyIKSolver::GetBoneSettingsType() const
 {
+	return FIKRigFBIKBoneSettings::StaticStruct();
+}
+
+bool FIKRigFullBodyIKSolver::HasSettingsOnBone(const FName& InBoneName) const
+{
+	for (const FIKRigFBIKBoneSettings& BoneSetting : AllBoneSettings)
+	{
+		if (BoneSetting.Bone == InBoneName)
+		{
+			return true;
+		}
+	}
+
+	return false;
+}
+
+void FIKRigFullBodyIKSolver::GetBonesWithSettings(TSet<FName>& OutBonesWithSettings) const
+{
+	for (const FIKRigFBIKBoneSettings& BoneSetting : AllBoneSettings)
+	{
+		OutBonesWithSettings.Add(BoneSetting.Bone);
+	}
+}
+
+#if WITH_EDITOR
+
+UIKRigSolverControllerBase* FIKRigFullBodyIKSolver::GetSolverController(UObject* Outer)
+{
+	return CreateControllerIfNeeded(Outer, UIKRigFBIKController::StaticClass());
+}
+
+FText FIKRigFullBodyIKSolver::GetNiceName() const
+{
+	return FText(LOCTEXT("SolverName", "Full Body IK"));
+}
+
+bool FIKRigFullBodyIKSolver::GetWarningMessage(FText& OutWarningMessage) const
+{
+	if (Settings.RootBone == NAME_None)
+	{
+		OutWarningMessage = LOCTEXT("MissingStart", "Missing start bone.");
+		return true;
+	}
+
+	if (AllGoalSettings.IsEmpty())
+	{
+		OutWarningMessage = LOCTEXT("MissingGoal", "Missing goals.");
+		return true;
+	}
 	
+	return false;
 }
 
-bool UIKRigFBIKSolver::IsBoneAffectedBySolver(const FName& BoneName, const FIKRigSkeleton& IKRigSkeleton) const
+bool FIKRigFullBodyIKSolver::IsBoneAffectedBySolver(const FName& InBoneName, const FIKRigSkeleton& InIKRigSkeleton) const
 {
 	// nothing is affected by solver without a root bone assigned or at least 1 effector
-	if (RootBone == NAME_None || Effectors.IsEmpty())
+	if (Settings.RootBone == NAME_None || AllGoalSettings.IsEmpty())
 	{
 		return false;
 	}
 
 	// has to be BELOW root
-	if (!IKRigSkeleton.IsBoneInDirectLineage(BoneName, RootBone))
+	if (!InIKRigSkeleton.IsBoneInDirectLineage(InBoneName, Settings.RootBone))
 	{
 		return false;
 	}
 
 	// has to be ABOVE an effector
-	for (UIKRig_FBIKEffector* Effector : Effectors)
+	for (const FIKRigFBIKGoalSettings& Effector : AllGoalSettings)
 	{
-		if (IKRigSkeleton.IsBoneInDirectLineage(Effector->BoneName, BoneName))
+		if (InIKRigSkeleton.IsBoneInDirectLineage(Effector.BoneName, InBoneName))
 		{
 			return true;
 		}
@@ -363,39 +377,59 @@ bool UIKRigFBIKSolver::IsBoneAffectedBySolver(const FName& BoneName, const FIKRi
 	return false;
 }
 
-void UIKRigFBIKSolver::PostLoad()
+#endif
+
+int32 FIKRigFullBodyIKSolver::GetIndexOfGoal(const FName& InGoalName) const
 {
-	Super::PostLoad();
-	
-	// patch for loading old effectors, we blow them away
-	bool bHasNullEffector = false;
-	for (const UIKRig_FBIKEffector* Effector : Effectors)
+	return AllGoalSettings.IndexOfByPredicate([InGoalName](const FIKRigFBIKGoalSettings& Element)
 	{
-		if (!Effector)
-		{
-			bHasNullEffector = true;
-		}
-	}
+		return Element.Goal == InGoalName;
+	});
+}
+
+//
+// BEGIN CONTROLLER
+//
+
+FIKRigFBIKSettings UIKRigFBIKController::GetSolverSettings()
+{
+	return *static_cast<FIKRigFBIKSettings*>(SolverToControl->GetSolverSettings());
+}
 
-	if (bHasNullEffector)
+void UIKRigFBIKController::SetSolverSettings(FIKRigFBIKSettings InSettings)
+{
+	SolverToControl->SetSolverSettings(&InSettings);
+}
+
+FIKRigFBIKGoalSettings UIKRigFBIKController::GetGoalSettings(const FName InGoalName)
+{
+	if (FIKRigGoalSettingsBase* GoalSettings = SolverToControl->GetGoalSettings(InGoalName))
 	{
-		Effectors.Empty();
+		return *static_cast<FIKRigFBIKGoalSettings*>(GoalSettings);
 	}
+	return FIKRigFBIKGoalSettings();
 }
 
-#endif
+void UIKRigFBIKController::SetGoalSettings(const FName InGoalName, FIKRigFBIKGoalSettings InSettings)
+{
+	SolverToControl->SetGoalSettings(InGoalName, &InSettings);
+}
 
-int32 UIKRigFBIKSolver::GetIndexOfGoal(const FName& OldName) const
+FIKRigFBIKBoneSettings UIKRigFBIKController::GetBoneSettings(const FName InBoneName)
 {
-	for (int32 i=0; i<Effectors.Num(); ++i)
+	if (FIKRigFBIKBoneSettings* BoneSettings = static_cast<FIKRigFBIKBoneSettings*>(SolverToControl->GetBoneSettings(InBoneName)))
 	{
-		if (Effectors[i]->GoalName == OldName)
-		{
-			return i;
-		}
+		return *BoneSettings;
 	}
+	
+	FIKRigFBIKBoneSettings NewBoneSettings;
+	NewBoneSettings.Bone = InBoneName;
+	return NewBoneSettings;
+}
 
-	return INDEX_NONE;
+void UIKRigFBIKController::SetBoneSettings(const FName InBoneName, FIKRigFBIKBoneSettings InSettings)
+{
+	SolverToControl->SetBoneSettings(InBoneName, &InSettings);
 }
 
 #undef LOCTEXT_NAMESPACE
diff --git a/Engine/Plugins/Animation/IKRig/Source/IKRig/Private/Rig/Solvers/IKRig_LimbSolver.cpp b/Engine/Plugins/Animation/IKRig/Source/IKRig/Private/Rig/Solvers/IKRig_LimbSolver.cpp
deleted file mode 100644
diff --git a/Engine/Plugins/Animation/IKRig/Source/IKRig/Private/Rig/Solvers/IKRig_PoleSolver.cpp b/Engine/Plugins/Animation/IKRig/Source/IKRig/Private/Rig/Solvers/IKRig_PoleSolver.cpp
deleted file mode 100644
diff --git a/Engine/Plugins/Animation/IKRig/Source/IKRig/Private/Rig/Solvers/IKRig_SetTransform.cpp b/Engine/Plugins/Animation/IKRig/Source/IKRig/Private/Rig/Solvers/IKRig_SetTransform.cpp
deleted file mode 100644
diff --git a/Engine/Plugins/Animation/IKRig/Source/IKRig/Private/Rig/Solvers/SetTransformSolver.cpp b/Engine/Plugins/Animation/IKRig/Source/IKRig/Private/Rig/Solvers/SetTransformSolver.cpp
new file mode 100644
--- /dev/null
+++ b/Engine/Plugins/Animation/IKRig/Source/IKRig/Private/Rig/Solvers/SetTransformSolver.cpp
@@ -0,0 +1,155 @@
+// Copyright Epic Games, Inc. All Rights Reserved.
+
+#include "Rig/Solvers/SetTransformSolver.h"
+#include "Rig/IKRigDataTypes.h"
+#include "Rig/IKRigSkeleton.h"
+
+#include UE_INLINE_GENERATED_CPP_BY_NAME(SetTransformSolver)
+
+#define LOCTEXT_NAMESPACE "SetTransformSolver"
+
+void FIKRigSetTransform::Initialize(const FIKRigSkeleton& IKRigSkeleton)
+{
+	BoneIndex = IKRigSkeleton.GetBoneIndexFromName(Settings.BoneToAffect);
+}
+
+void FIKRigSetTransform::Solve(FIKRigSkeleton& InIKRigSkeleton, const FIKRigGoalContainer& InGoals)
+{
+	// no bone specified
+	if (BoneIndex == INDEX_NONE)
+	{
+		return;
+	}
+
+	// no goal specified
+	const FIKRigGoal* InGoal = InGoals.FindGoalByName(Settings.Goal);
+	if (!InGoal)
+	{
+		return;
+	}
+
+	// check that settings are such that there is anything to do at all
+	const bool bPositionEnabled = Settings.PositionAlpha > UE_KINDA_SMALL_NUMBER;
+	const bool bRotationEnabled = Settings.RotationAlpha > UE_KINDA_SMALL_NUMBER;
+	const bool bAnythingEnabled = bPositionEnabled || bRotationEnabled;
+	const bool bHasAlpha = Settings.Alpha > UE_KINDA_SMALL_NUMBER;
+	if (!(bAnythingEnabled && bHasAlpha))
+	{
+		return;
+	}
+
+	FTransform& CurrentTransform = InIKRigSkeleton.CurrentPoseGlobal[BoneIndex];	
+
+	if (bPositionEnabled)
+	{
+		const FVector TargetPosition = FMath::Lerp(CurrentTransform.GetTranslation(), InGoal->FinalBlendedPosition, Settings.PositionAlpha * Settings.Alpha);
+		CurrentTransform.SetLocation(TargetPosition);
+	}
+	
+	if (bRotationEnabled)
+	{
+		const FQuat TargetRotation = FMath::Lerp(CurrentTransform.GetRotation(), InGoal->FinalBlendedRotation, Settings.RotationAlpha * Settings.Alpha);
+		CurrentTransform.SetRotation(TargetRotation);
+	}
+
+	if (Settings.bPropagateToChildren)
+	{
+		InIKRigSkeleton.PropagateGlobalPoseBelowBone(BoneIndex);	
+	}
+}
+
+void FIKRigSetTransform::GetRequiredBones(TSet<FName>& OutRequiredBones) const
+{
+	OutRequiredBones.Add(Settings.BoneToAffect);
+}
+
+void FIKRigSetTransform::GetRequiredGoals(TSet<FName>& OutRequiredGoals) const
+{
+	if (Settings.Goal != NAME_Name)
+	{
+		OutRequiredGoals.Add(Settings.Goal);
+	}
+}
+
+FIKRigSolverSettingsBase* FIKRigSetTransform::GetSolverSettings()
+{
+	return &Settings;
+}
+
+const UScriptStruct* FIKRigSetTransform::GetSolverSettingsType() const
+{
+	return FIKRigSetTransformSettings::StaticStruct();
+}
+
+void FIKRigSetTransform::OnGoalRemoved(const FName& InGoalName)
+{
+	if (Settings.Goal == InGoalName)
+	{
+		Settings.Goal = NAME_None;
+		Settings.BoneToAffect = NAME_None;
+	}	
+}
+
+void FIKRigSetTransform::AddGoal(const UIKRigEffectorGoal& InNewGoal)
+{
+	Settings.Goal = InNewGoal.GoalName;
+	Settings.BoneToAffect = InNewGoal.BoneName;
+}
+
+void FIKRigSetTransform::OnGoalRenamed(const FName& InOldName, const FName& InNewName)
+{
+	if (Settings.Goal == InOldName)
+	{
+		Settings.Goal = InNewName;
+	}
+}
+
+void FIKRigSetTransform::OnGoalMovedToDifferentBone(const FName& InGoalName, const FName& InNewBoneName)
+{
+	if (Settings.Goal == InGoalName)
+	{
+		Settings.BoneToAffect = InNewBoneName;
+	}
+}
+
+FIKRigSetTransformSettings UIKRigSetTransformController::GetSolverSettings()
+{
+	return *static_cast<FIKRigSetTransformSettings*>(SolverToControl->GetSolverSettings());
+}
+
+void UIKRigSetTransformController::SetSolverSettings(FIKRigSetTransformSettings InSettings)
+{
+	SolverToControl->SetSolverSettings(&InSettings);
+}
+
+#if WITH_EDITOR
+
+UIKRigSolverControllerBase* FIKRigSetTransform::GetSolverController(UObject* Outer)
+{
+	return CreateControllerIfNeeded(Outer, UIKRigSetTransformController::StaticClass());
+}
+
+FText FIKRigSetTransform::GetNiceName() const
+{
+	return FText(LOCTEXT("SolverName", "Set Transform"));
+}
+
+bool FIKRigSetTransform::GetWarningMessage(FText& OutWarningMessage) const
+{
+	if (Settings.BoneToAffect == NAME_None)
+	{
+		OutWarningMessage = LOCTEXT("MissingGoal", "Missing goal.");
+		return true;
+	}
+	return false;
+}
+
+bool FIKRigSetTransform::IsBoneAffectedBySolver(const FName& InBoneName, const FIKRigSkeleton& InIKRigSkeleton) const
+{
+	return InIKRigSkeleton.IsBoneInDirectLineage(InBoneName, Settings.BoneToAffect);
+}
+
+#endif
+
+#undef LOCTEXT_NAMESPACE
+
diff --git a/Engine/Plugins/Animation/IKRig/Source/IKRig/Public/Retargeter/IKRetargetProcessor.h b/Engine/Plugins/Animation/IKRig/Source/IKRig/Public/Retargeter/IKRetargetProcessor.h
--- a/Engine/Plugins/Animation/IKRig/Source/IKRig/Public/Retargeter/IKRetargetProcessor.h
+++ b/Engine/Plugins/Animation/IKRig/Source/IKRig/Public/Retargeter/IKRetargetProcessor.h
@@ -562,9 +562,6 @@ public:
 		const TMap<FName, float>& SpeedValuesFromCurves,
 		const float DeltaTime,
 		const FRetargetProfile& Profile);
-	
-	/** Apply the settings stored in the IKRig asset. */
-	void CopyIKRigSettingsFromAsset();
 
 	/** Does a partial reinitialization (at runtime) whenever the retarget pose is swapped to a different or if the
 	 * pose has been modified. Does nothing if the pose has not changed. */
@@ -612,7 +609,11 @@ public:
 	/** Set that this processor needs to be reinitialized. */
 	void SetNeedsInitialized();
 	
+	/** Apply the settings stored in the IKRig asset. */
+	void CopyIKRigSettingsFromAsset() const;
+	
 #if WITH_EDITOR
+	
 private:
 	DECLARE_MULTICAST_DELEGATE(FOnRetargeterInitialized);
 	FOnRetargeterInitialized RetargeterInitialized;
diff --git a/Engine/Plugins/Animation/IKRig/Source/IKRig/Public/Rig/IKRigDefinition.h b/Engine/Plugins/Animation/IKRig/Source/IKRig/Public/Rig/IKRigDefinition.h
--- a/Engine/Plugins/Animation/IKRig/Source/IKRig/Public/Rig/IKRigDefinition.h
+++ b/Engine/Plugins/Animation/IKRig/Source/IKRig/Public/Rig/IKRigDefinition.h
@@ -3,10 +3,11 @@
 #pragma once
 
 #include "BoneContainer.h"
-#include "IKRigLogger.h"
 #include "UObject/Object.h"
+#include "StructUtils/InstancedStruct.h"
 #include "Interfaces/Interface_PreviewMeshProvider.h"
 #include "IKRigSkeleton.h"
+#include "Rig/Solvers/IKRigSolverBase.h"
 #if WITH_EDITOR
 #include "SAdvancedTransformInputBox.h"
 #endif
@@ -14,10 +15,10 @@
 #include "IKRigDefinition.generated.h"
 
 class IPropertyHandle;
-class UIKRigSolver;
+class UDEPRECATED_IKRigSolver;
 
 #if WITH_EDITORONLY_DATA
-UENUM(BlueprintType)
+UENUM()
 enum class EIKRigGoalPreviewMode : uint8
 {
 	Additive		UMETA(DisplayName = "Additive"),
@@ -160,7 +161,7 @@ struct IKRIG_API FBoneChain
 	FName IKGoalName;
 };
 
-USTRUCT(Blueprintable)
+USTRUCT(BlueprintType)
 struct IKRIG_API FRetargetDefinition
 {
 	GENERATED_BODY()
@@ -181,7 +182,7 @@ struct IKRIG_API FRetargetDefinition
 	FBoneChain* GetEditableBoneChainByName(FName ChainName);
 };
 
-UCLASS(Blueprintable)
+UCLASS(BlueprintType)
 class IKRIG_API UIKRigDefinition : public UObject, public IInterface_PreviewMeshProvider
 {
 	GENERATED_BODY()
@@ -201,7 +202,10 @@ public:
 	const TArray<UIKRigEffectorGoal*>& GetGoalArray() const { return Goals; };
 
 	/** runtime, read-only access to array of Solvers, all modifications must go through UIKRigController */
-	const TArray<UIKRigSolver*>& GetSolverArray() const { return Solvers; };
+	const TArray<const FIKRigSolverBase*> GetSolverArray() const;
+
+	/** runtime, read-only access to array of Solvers structs, all modifications must go through UIKRigController */
+	const TArray<FInstancedStruct> GetSolverStructs() const;
 
 	/** runtime, read-only access to Retarget Definition, all modifications must go through UIKRigController */
 	const TArray<FBoneChain>& GetRetargetChains() const { return RetargetDefinition.BoneChains; };
@@ -259,14 +263,19 @@ private:
 	/** goals, used as effectors by solvers that support them */
 	UPROPERTY()
 	TArray<TObjectPtr<UIKRigEffectorGoal>> Goals;
-	
+
 	/** polymorphic stack of solvers, executed in serial fashion where output of prior solve is input to the next */
-	UPROPERTY(instanced)
-	TArray<TObjectPtr<UIKRigSolver>> Solvers;
+	UPROPERTY(meta = (ExcludeBaseStruct, BaseStruct = "/Script/IKRig.IKRigSolverBase"))
+	TArray<FInstancedStruct> SolverStack;
 
 	/** bone chains for IK Retargeter */
 	UPROPERTY()
 	FRetargetDefinition RetargetDefinition;
 
 	friend class UIKRigController;
+
+	PRAGMA_DISABLE_DEPRECATION_WARNINGS 
+	UPROPERTY(meta = (DeprecatedProperty, DeprecationMessage = "The pre 5.6 stack of UObject based solvers. Use SolverStack instead."))
+	TArray<TObjectPtr<UDEPRECATED_IKRigSolver>> Solvers_DEPRECATED;
+	PRAGMA_ENABLE_DEPRECATION_WARNINGS 
 };
diff --git a/Engine/Plugins/Animation/IKRig/Source/IKRig/Public/Rig/IKRigProcessor.h b/Engine/Plugins/Animation/IKRig/Source/IKRig/Public/Rig/IKRigProcessor.h
--- a/Engine/Plugins/Animation/IKRig/Source/IKRig/Public/Rig/IKRigProcessor.h
+++ b/Engine/Plugins/Animation/IKRig/Source/IKRig/Public/Rig/IKRigProcessor.h
@@ -10,7 +10,7 @@
 #include "IKRigProcessor.generated.h"
 
 class UIKRigDefinition;
-class UIKRigSolver;
+struct FIKRigSolverBase;
 
 USTRUCT()
 struct FGoalBone
@@ -36,7 +36,7 @@ public:
 	* 2. Initialize() with an IKRigDefinition asset
 	* 3. each tick, call SetIKGoal() and SetInputPoseGlobal()
 	* 4. Call Solve()
-	* 5. Copy output transforms with CopyOutputGlobalPoseToArray()
+	* 5. Copy output transforms with GetOutputPoseGlobal()
 	* 
 	*/
 
@@ -76,7 +76,7 @@ public:
 	void Solve(const FTransform& WorldToComponent = FTransform::Identity);
 
 	/** Get the results after calling Solve() */
-	void CopyOutputGlobalPoseToArray(TArray<FTransform>& OutputPoseGlobal) const;
+	void GetOutputPoseGlobal(TArray<FTransform>& OutputPoseGlobal) const;
 
 	/** Reset all internal data structures. Will require re-initialization before solving again. */
 	void Reset();
@@ -108,11 +108,12 @@ public:
 	/** logging system */
 	FIKRigLogger Log;
 
-	/** Used to propagate setting values from the source asset at runtime (settings that do not require re-initialization) */
-	void CopyAllInputsFromSourceAssetAtRuntime(const UIKRigDefinition* SourceAsset);
-
 	/** Get read-only access to currently running solvers */
-	const TArray<TObjectPtr<UIKRigSolver>>& GetSolvers(){ return Solvers; };
+	const TArray<TUniquePtr<FIKRigSolverBase>>& GetSolvers(){ return Solvers; };
+	
+	/** Used to propagate setting values from the source asset at runtime (settings that do not require re-initialization)
+	 * This allows the user to edit their IK Rig asset and see the effects in PIE. */
+	void CopyAllSettingsFromAsset(const UIKRigDefinition* SourceAsset);
 	
 private:
 
@@ -120,8 +121,7 @@ private:
 	void ResolveFinalGoalTransforms(const FTransform& WorldToComponent);
 
 	/** the stack of solvers to run in order */
-	UPROPERTY(Transient)
-	TArray<TObjectPtr<UIKRigSolver>> Solvers;
+	TArray<TUniquePtr<FIKRigSolverBase>> Solvers;
 
 	/** the named transforms that solvers use as end effectors */
 	FIKRigGoalContainer GoalContainer;
diff --git a/Engine/Plugins/Animation/IKRig/Source/IKRig/Public/Rig/IKRigSkeleton.h b/Engine/Plugins/Animation/IKRig/Source/IKRig/Public/Rig/IKRigSkeleton.h
--- a/Engine/Plugins/Animation/IKRig/Source/IKRig/Public/Rig/IKRigSkeleton.h
+++ b/Engine/Plugins/Animation/IKRig/Source/IKRig/Public/Rig/IKRigSkeleton.h
@@ -9,7 +9,6 @@
 struct FBoneChain;
 class USkeletalMesh;
 class UIKRigDefinition;
-class UIKRigSolver;
 
 /** Data used just to initialize an IKRigSkeleton from outside systems
  *
diff --git a/Engine/Plugins/Animation/IKRig/Source/IKRig/Public/Rig/Solvers/IKRigBodyMoverSolver.h b/Engine/Plugins/Animation/IKRig/Source/IKRig/Public/Rig/Solvers/IKRigBodyMoverSolver.h
new file mode 100644
--- /dev/null
+++ b/Engine/Plugins/Animation/IKRig/Source/IKRig/Public/Rig/Solvers/IKRigBodyMoverSolver.h
@@ -0,0 +1,254 @@
+﻿// Copyright Epic Games, Inc. All Rights Reserved.
+
+#pragma once
+
+#include "Rig/IKRigDefinition.h"
+#include "IKRigSolverBase.h"
+
+#include "IKRigBodyMoverSolver.generated.h"
+
+USTRUCT(BlueprintType)
+struct IKRIG_API FIKRigBodyMoverGoalSettings : public FIKRigGoalSettingsBase
+{
+	GENERATED_BODY()
+
+	// The bone associated with this goal.
+	UPROPERTY(BlueprintReadOnly, VisibleAnywhere, Category = "Goal Settings")
+	FName BoneName;
+
+	// Scale the influence this goal has on the body. Range is 0-10. Default is 1.0.
+	UPROPERTY(BlueprintReadWrite, EditAnywhere, Category = "Goal Settings", meta = (ClampMin = "0", ClampMax = "10", UIMin = "0.0", UIMax = "10.0"))
+	float InfluenceMultiplier = 1.0f;
+};
+
+USTRUCT(BlueprintType)
+struct IKRIG_API FIKRigBodyMoverSettings : public FIKRigSolverSettingsBase
+{
+	GENERATED_BODY()
+
+	// The target bone to move with the effectors.
+	UPROPERTY(BlueprintReadOnly, VisibleAnywhere, Category = "Body Mover Settings")
+	FName StartBone;
+
+	// Blend the translational effect of this solver on/off. Range is 0-1. Default is 1.0.
+	UPROPERTY(BlueprintReadWrite, EditAnywhere, Category = "Body Mover Settings", meta = (UIMin = "0.0", UIMax = "1.0"))
+	float PositionAlpha = 1.0f;
+
+	// Multiply the POSITIVE X translation. Range is 0-1. Default is 1.0.
+	UPROPERTY(BlueprintReadWrite, EditAnywhere, Category = "Body Mover Settings", meta = (UIMin = "0.0", UIMax = "1.0"))
+	float PositionPositiveX = 1.0f;
+
+	// Multiply the NEGATIVE X translation. Range is 0-1. Default is 1.0.
+	UPROPERTY(BlueprintReadWrite, EditAnywhere, Category = "Body Mover Settings", meta = (UIMin = "0.0", UIMax = "1.0"))
+	float PositionNegativeX = 1.0f;
+
+	// Multiply the POSITIVE Y translation. Range is 0-1. Default is 1.0.
+	UPROPERTY(BlueprintReadWrite, EditAnywhere, Category = "Body Mover Settings", meta = (UIMin = "0.0", UIMax = "1.0"))
+	float PositionPositiveY = 1.0f;
+
+	// Multiply the NEGATIVE Y translation. Range is 0-1. Default is 1.0.
+	UPROPERTY(BlueprintReadWrite, EditAnywhere, Category = "Body Mover Settings", meta = (UIMin = "0.0", UIMax = "1.0"))
+	float PositionNegativeY = 1.0f;
+
+	// Multiply the POSITIVE Z translation. Range is 0-1. Default is 1.0.
+	UPROPERTY(BlueprintReadWrite, EditAnywhere, Category = "Body Mover Settings", meta = (UIMin = "0.0", UIMax = "1.0"))
+	float PositionPositiveZ = 1.0f;
+
+	//* Multiply the NEGATIVE Z translation. Range is 0-1. Default is 1.0.
+	UPROPERTY(BlueprintReadWrite, EditAnywhere, Category = "Body Mover Settings", meta = (UIMin = "0.0", UIMax = "1.0"))
+	float PositionNegativeZ = 1.0f;
+
+	// Blend the total rotational effect on/off. Range is 0-1. Default is 1.0.
+	UPROPERTY(BlueprintReadWrite, EditAnywhere, Category = "Body Mover Settings", meta = (UIMin = "0.0", UIMax = "1.0"))
+	float RotationAlpha = 1.0f;
+
+	// Blend the X-axis rotational effect on/off. Range is 0-1. Default is 1.0.
+	UPROPERTY(BlueprintReadWrite, EditAnywhere, Category = "Body Mover Settings", meta = (UIMin = "0.0", UIMax = "1.0"))
+	float RotateXAlpha = 1.0f;
+
+	// Blend the Y-axis rotational effect on/off. Range is 0-1. Default is 1.0.
+	UPROPERTY(BlueprintReadWrite, EditAnywhere, Category = "Body Mover Settings", meta = (UIMin = "0.0", UIMax = "1.0"))
+	float RotateYAlpha = 1.0f;
+
+	// Blend the Z-axis rotational effect on/off. Range is 0-1. Default is 1.0.
+	UPROPERTY(BlueprintReadWrite, EditAnywhere, Category = "Body Mover Settings", meta = (UIMin = "0.0", UIMax = "1.0"))
+	float RotateZAlpha = 1.0f;
+};
+
+USTRUCT(BlueprintType)
+struct IKRIG_API FIKRigBodyMoverSolver : public FIKRigSolverBase
+{
+	GENERATED_BODY()
+	
+	UPROPERTY()
+	FIKRigBodyMoverSettings Settings;
+	
+	UPROPERTY()
+	TArray<FIKRigBodyMoverGoalSettings> AllGoalSettings;
+
+	// FIKRigSolverBase interface
+	virtual void Initialize(const FIKRigSkeleton& InIKRigSkeleton) override;
+	virtual void Solve(FIKRigSkeleton& InIKRigSkeleton, const FIKRigGoalContainer& InGoals) override;
+	virtual void GetRequiredBones(TSet<FName>& OutRequiredBones) const override;
+	virtual void GetRequiredGoals(TSet<FName>& OutRequiredGoals) const override;
+
+	// solver settings
+	virtual FIKRigSolverSettingsBase* GetSolverSettings() override;
+	virtual const UScriptStruct* GetSolverSettingsType() const override;
+	
+	// goals
+	virtual void AddGoal(const UIKRigEffectorGoal& InNewGoal) override;
+	virtual void OnGoalRemoved(const FName& InGoalName) override;
+	virtual void OnGoalRenamed(const FName& InOldName, const FName& InNewName) override;
+	virtual void OnGoalMovedToDifferentBone(const FName& InGoalName, const FName& InNewBoneName) override;
+
+	// goal settings
+	virtual bool UsesCustomGoalSettings() const override { return true; };
+	virtual FIKRigGoalSettingsBase* GetGoalSettings(const FName& InGoalName) override;
+	virtual const UScriptStruct* GetGoalSettingsType() const override;
+	virtual void GetGoalsWithSettings(TSet<FName>& OutGoalsWithSettings) const override;
+	
+	// start bone can be set on this solver
+	virtual bool UsesStartBone() const override { return true; };
+	virtual FName GetStartBone() const override { return Settings.StartBone; };
+	virtual void SetStartBone(const FName& InRootBoneName) override;
+
+#if WITH_EDITOR
+	// return custom controller for scripting this solver
+	virtual UIKRigSolverControllerBase* GetSolverController(UObject* Outer) override;
+	// UI stuff
+	virtual FText GetNiceName() const override;
+	virtual bool GetWarningMessage(FText& OutWarningMessage) const override;
+	virtual bool IsBoneAffectedBySolver(const FName& InBoneName, const FIKRigSkeleton& InIKRigSkeleton) const override;
+#endif
+
+private:
+
+	int32 GetIndexOfGoal(const FName& InGoalName) const;
+
+	int32 BodyBoneIndex;
+};
+
+/* The blueprint/python API for modifying a Body Mover solver's settings in an IK Rig.
+ * Can adjust Solver and Goal settings. */
+UCLASS(BlueprintType)
+class IKRIG_API UIKRigBodyMoverController : public UIKRigSolverControllerBase
+{
+	GENERATED_BODY()
+	
+public:
+
+	/* Get the current solver settings as a struct.
+	 * @return FIKRigBodyMoverSettings struct with the current settings used by the solver. */
+	UFUNCTION(BlueprintCallable, Category = IKRig)
+	FIKRigBodyMoverSettings GetSolverSettings();
+
+	/* Set the solver settings. Input is a custom struct type for this solver.
+	 * @param InSettings a FIKRigBodyMoverSettings struct containing all the settings to apply to this solver */
+	UFUNCTION(BlueprintCallable, Category = IKRig)
+	void SetSolverSettings(FIKRigBodyMoverSettings InSettings);
+
+	/* Get the settings for the specified goal.
+	 * @param InGoalName the name of the goal to get settings for
+	 * @return FIKRigBodyMoverGoalSettings struct (empty if the specified goal does not belong to this solver) */
+	UFUNCTION(BlueprintCallable, Category = IKRig)
+	FIKRigBodyMoverGoalSettings GetGoalSettings(const FName InGoalName);
+
+	/* Set the settings for the specified goal.
+	 * @param InGoalName: the name of the goal to assign the settings to.
+	 * @param InSettings: a custom struct type with all the settings for a goal */
+	UFUNCTION(BlueprintCallable, Category = IKRig)
+	void SetGoalSettings(const FName InGoalName, FIKRigBodyMoverGoalSettings InSettings);
+};
+
+PRAGMA_DISABLE_DEPRECATION_WARNINGS
+
+UCLASS(Deprecated)
+class
+UE_DEPRECATED(5.6, "This type has been replaced with FBodyMoverGoalSettings.")
+IKRIG_API UDEPRECATED_IKRig_BodyMoverEffector : public UObject
+{
+	GENERATED_BODY()
+
+public:
+	
+	UPROPERTY()
+	FName GoalName;
+	UPROPERTY()
+	FName BoneName;
+	UPROPERTY()
+	float InfluenceMultiplier = 1.0f;
+};
+
+UCLASS(Deprecated)
+class
+UE_DEPRECATED(5.6, "This type has been replaced with FBodyMoverSolver.")
+IKRIG_API UDEPRECATED_IKRig_BodyMover : public UDEPRECATED_IKRigSolver
+{
+	GENERATED_BODY()
+	
+public:
+
+	// upgrade patch to new UStruct-based solver
+	virtual void ConvertToInstancedStruct(FInstancedStruct& OutInstancedStruct) override
+	{
+		OutInstancedStruct.InitializeAs(FIKRigBodyMoverSolver::StaticStruct());
+		FIKRigBodyMoverSolver& NewSolver = OutInstancedStruct.GetMutable<FIKRigBodyMoverSolver>();
+
+		// load solver settings
+		NewSolver.SetEnabled(IsEnabled());
+		NewSolver.Settings.StartBone = RootBone;
+		NewSolver.Settings.PositionAlpha = PositionAlpha;
+		NewSolver.Settings.PositionPositiveX = PositionPositiveX;
+		NewSolver.Settings.PositionNegativeX = PositionNegativeX;
+		NewSolver.Settings.PositionPositiveY = PositionPositiveY;
+		NewSolver.Settings.PositionNegativeY = PositionNegativeY;
+		NewSolver.Settings.PositionPositiveZ = PositionPositiveZ;
+		NewSolver.Settings.PositionNegativeZ = PositionNegativeZ;
+		NewSolver.Settings.RotationAlpha = RotationAlpha;
+		NewSolver.Settings.RotateXAlpha = RotateXAlpha;
+		NewSolver.Settings.RotateYAlpha = RotateYAlpha;
+		NewSolver.Settings.RotateZAlpha = RotateZAlpha;
+		// load goal settings
+		for (const TObjectPtr<UDEPRECATED_IKRig_BodyMoverEffector>& Effector : Effectors_DEPRECATED)
+		{
+			FIKRigBodyMoverGoalSettings NewSettings;
+			NewSettings.BoneName = Effector->BoneName;
+			NewSettings.Goal = Effector->GoalName;
+			NewSettings.InfluenceMultiplier = Effector->InfluenceMultiplier;
+			NewSolver.AllGoalSettings.Add(NewSettings);
+		}
+	};
+	
+	UPROPERTY()
+	FName RootBone;
+	UPROPERTY()
+	float PositionAlpha = 1.0f;
+	UPROPERTY()
+	float PositionPositiveX = 1.0f;
+	UPROPERTY()
+	float PositionNegativeX = 1.0f;
+	UPROPERTY()
+	float PositionPositiveY = 1.0f;
+	UPROPERTY()
+	float PositionNegativeY = 1.0f;
+	UPROPERTY()
+	float PositionPositiveZ = 1.0f;
+	UPROPERTY()
+	float PositionNegativeZ = 1.0f;
+	UPROPERTY()
+	float RotationAlpha = 1.0f;
+	UPROPERTY()
+	float RotateXAlpha = 1.0f;
+	UPROPERTY()
+	float RotateYAlpha = 1.0f;
+	UPROPERTY()
+	float RotateZAlpha = 1.0f;
+
+	UE_DEPRECATED(5.6, "IK Rig refactored to not use UObjects")
+	UPROPERTY(meta = (DeprecatedProperty))
+	TArray<TObjectPtr<UDEPRECATED_IKRig_BodyMoverEffector>> Effectors_DEPRECATED;
+	
+};
+
+PRAGMA_ENABLE_DEPRECATION_WARNINGS
diff --git a/Engine/Plugins/Animation/IKRig/Source/IKRig/Public/Rig/Solvers/IKRigLimbSolver.h b/Engine/Plugins/Animation/IKRig/Source/IKRig/Public/Rig/Solvers/IKRigLimbSolver.h
new file mode 100644
--- /dev/null
+++ b/Engine/Plugins/Animation/IKRig/Source/IKRig/Public/Rig/Solvers/IKRigLimbSolver.h
@@ -0,0 +1,177 @@
+﻿// Copyright Epic Games, Inc. All Rights Reserved.
+
+#pragma once
+
+#include "Rig/IKRigDefinition.h"
+#include "IKRigSolverBase.h"
+
+#include "Rig/Solvers/LimbSolver.h"
+
+#include "IKRigLimbSolver.generated.h"
+
+
+USTRUCT(BlueprintType)
+struct IKRIG_API FIKRigLimbSolverSettings : public FLimbSolverSettings
+{
+	GENERATED_BODY()
+	
+	/** The first bone in the IK chain, for example the "hip" in a leg, or the "shoulder" in an arm */
+	UPROPERTY(BlueprintReadOnly, VisibleAnywhere, DisplayName = "Root Bone", Category = "Limb IK Settings")
+	FName StartBone = NAME_None;
+
+	/** The name of the IK goal to drive the end bone towards */
+	UPROPERTY(BlueprintReadOnly, VisibleAnywhere, Category = "Limb IK Effector")
+	FName GoalName = NAME_None;
+
+	/** The name of the last bone in the IK chain. This is the bone you want to reach the goal. */
+	UPROPERTY(BlueprintReadOnly, VisibleAnywhere, Category = "Limb IK Effector")
+	FName EndBone = NAME_None;
+};
+
+USTRUCT(BlueprintType)
+struct IKRIG_API FIKRigLimbSolver : public FIKRigSolverBase
+{
+	GENERATED_BODY()
+	
+	UPROPERTY()
+	FIKRigLimbSolverSettings Settings;
+	
+	// runtime
+	virtual void Initialize(const FIKRigSkeleton& InIKRigSkeleton) override;
+	virtual void Solve(FIKRigSkeleton& InIKRigSkeleton, const FIKRigGoalContainer& InGoals) override;
+	virtual void GetRequiredBones(TSet<FName>& OutRequiredBones) const override;
+	virtual void GetRequiredGoals(TSet<FName>& OutRequiredGoals) const override;
+	
+	// solver settings
+	virtual FIKRigSolverSettingsBase* GetSolverSettings() override;
+	virtual const UScriptStruct* GetSolverSettingsType() const override;
+
+	// goals
+	virtual void AddGoal(const UIKRigEffectorGoal& InNewGoal) override;
+	virtual void OnGoalRenamed(const FName& InOldName, const FName& InNewName) override;
+	virtual void OnGoalMovedToDifferentBone(const FName& InGoalName, const FName& InNewBoneName) override;
+	virtual void OnGoalRemoved(const FName& InGoalName) override;
+	
+	// start bone
+	virtual bool UsesStartBone() const override;
+	virtual void SetStartBone(const FName& InStartBoneName) override;
+	virtual FName GetStartBone() const override;
+
+#if WITH_EDITOR
+	// return custom controller for scripting this solver
+	virtual UIKRigSolverControllerBase* GetSolverController(UObject* Outer) override;
+	
+	// general UI
+	virtual FText GetNiceName() const override;
+	virtual bool GetWarningMessage(FText& OutWarningMessage) const override;
+	virtual bool IsBoneAffectedBySolver(const FName& InBoneName, const FIKRigSkeleton& InIKRigSkeleton) const override;
+#endif
+	
+private:
+
+	static void GatherChildren(const int32 BoneIndex, const FIKRigSkeleton& InSkeleton, TArray<int32>& OutChildren);
+	
+	FLimbSolver Solver;
+	TArray<int32> ChildrenToUpdate;
+	
+};
+
+/* The blueprint/python API for modifying an Full-Body IK solver's settings in an IK Rig.
+ * Can adjust Solver, Goal and Bone settings. */
+UCLASS(BlueprintType)
+class IKRIG_API UIKRigLimbSolverController : public UIKRigSolverControllerBase
+{
+	GENERATED_BODY()
+	
+public:
+
+	/* Get the current solver settings as a struct.
+	 * @return FIKRigLimbSolverSettings struct with the current settings used by the solver. */
+	UFUNCTION(BlueprintCallable, Category = IKRig)
+	FIKRigLimbSolverSettings GetSolverSettings();
+
+	/* Set the solver settings. Input is a custom struct type for this solver.
+	 * @param InSettings a FIKRigLimbSolverSettings struct containing all the settings to apply to this solver */
+	UFUNCTION(BlueprintCallable, Category = IKRig)
+	void SetSolverSettings(FIKRigLimbSolverSettings InSettings);
+};
+
+PRAGMA_DISABLE_DEPRECATION_WARNINGS
+
+UCLASS(Deprecated)
+class
+UE_DEPRECATED(5.6, "This type has been replaced with FIKRigLimbSolver.")
+IKRIG_API UDEPRECATED_IKRig_LimbEffector : public UObject
+{
+	GENERATED_BODY()
+
+public:
+	UDEPRECATED_IKRig_LimbEffector() { SetFlags(RF_Transactional); }
+
+	UPROPERTY()
+	FName GoalName = NAME_None;
+	UPROPERTY()
+	FName BoneName = NAME_None;
+};
+
+UCLASS(Deprecated)
+class UE_DEPRECATED(5.6, "This type has been replaced with FIKRigLimbSolver.")
+IKRIG_API UDEPRECATED_IKRig_LimbSolver : public UDEPRECATED_IKRigSolver
+{
+	GENERATED_BODY()
+
+public:
+
+	UDEPRECATED_IKRig_LimbSolver()
+	{
+		Effector_DEPRECATED = CreateDefaultSubobject<UDEPRECATED_IKRig_LimbEffector>(TEXT("Effector"));
+	}
+
+	// upgrade patch to new UStruct-based solver
+	virtual void ConvertToInstancedStruct(FInstancedStruct& OutInstancedStruct) override
+	{
+		OutInstancedStruct.InitializeAs(FIKRigLimbSolver::StaticStruct());
+		FIKRigLimbSolver& NewSolver = OutInstancedStruct.GetMutable<FIKRigLimbSolver>();
+		NewSolver.SetEnabled(IsEnabled());
+		NewSolver.Settings.StartBone = RootName;
+		NewSolver.Settings.GoalName = Effector_DEPRECATED->GoalName;
+		NewSolver.Settings.EndBone = Effector_DEPRECATED->BoneName;
+		NewSolver.Settings.ReachPrecision = ReachPrecision;
+		NewSolver.Settings.HingeRotationAxis = HingeRotationAxis;
+		NewSolver.Settings.MaxIterations = MaxIterations;
+		NewSolver.Settings.bEnableLimit = bEnableLimit;
+		NewSolver.Settings.MinRotationAngle = MinRotationAngle;
+		NewSolver.Settings.bAveragePull = bAveragePull;
+		NewSolver.Settings.PullDistribution = PullDistribution;
+		NewSolver.Settings.ReachStepAlpha = ReachStepAlpha;
+		NewSolver.Settings.bEnableTwistCorrection = bEnableTwistCorrection;
+		NewSolver.Settings.EndBoneForwardAxis = EndBoneForwardAxis;
+	}
+	
+	UPROPERTY()
+	FName RootName = NAME_None;
+	UPROPERTY()
+	float ReachPrecision = 0.01f;
+	UPROPERTY()
+	TEnumAsByte<EAxis::Type> HingeRotationAxis = EAxis::None;
+	UPROPERTY()
+	int32 MaxIterations = 12;
+	UPROPERTY()
+	bool bEnableLimit = false;
+	UPROPERTY()
+	float MinRotationAngle = 15.f;
+	UPROPERTY()
+	bool bAveragePull = true;
+	UPROPERTY()
+	float PullDistribution = 0.5f;
+	UPROPERTY()
+	float ReachStepAlpha = 0.7f;
+	UPROPERTY()
+	bool bEnableTwistCorrection = false;
+	UPROPERTY()
+	TEnumAsByte<EAxis::Type> EndBoneForwardAxis = EAxis::Y;
+	UPROPERTY(meta = (DeprecatedProperty))
+	TObjectPtr<UDEPRECATED_IKRig_LimbEffector> Effector_DEPRECATED;
+};
+
+PRAGMA_ENABLE_DEPRECATION_WARNINGS
diff --git a/Engine/Plugins/Animation/IKRig/Source/IKRig/Public/Rig/Solvers/IKRigPoleSolver.h b/Engine/Plugins/Animation/IKRig/Source/IKRig/Public/Rig/Solvers/IKRigPoleSolver.h
new file mode 100644
--- /dev/null
+++ b/Engine/Plugins/Animation/IKRig/Source/IKRig/Public/Rig/Solvers/IKRigPoleSolver.h
@@ -0,0 +1,156 @@
+﻿// Copyright Epic Games, Inc. All Rights Reserved.
+
+#pragma once
+
+#include "Rig/IKRigDefinition.h"
+#include "Rig/Solvers/IKRigSolverBase.h"
+
+#include "IKRigPoleSolver.generated.h"
+
+USTRUCT(BlueprintType)
+struct IKRIG_API FIKRigPoleSolverSettings : public FIKRigSolverSettingsBase
+{
+	GENERATED_BODY()
+
+	UPROPERTY(BlueprintReadOnly, VisibleAnywhere, Category = "Pole Solver Settings")
+	FName StartBone = NAME_None;
+
+	UPROPERTY(BlueprintReadOnly, VisibleAnywhere, Category = "Pole Solver Settings")
+	FName EndBone = NAME_None;
+
+	UPROPERTY(BlueprintReadOnly, VisibleAnywhere, Category = "Pole Solver Settings")
+	FName AimAtGoal = NAME_None;
+
+	// Blend the effect on/off. Range is 0-1. Default is 1.0. 
+	UPROPERTY(BlueprintReadWrite, EditAnywhere, Category = "Pole Solver Settings", meta = (ClampMin = "0.0", ClampMax = "1.0", UIMin = "0.0", UIMax = "1.0"))
+	float Alpha = 1.0f;
+};
+
+USTRUCT(BlueprintType)
+struct IKRIG_API FIKRigPoleSolver : public FIKRigSolverBase
+{
+	GENERATED_BODY()
+	
+	UPROPERTY()
+	FIKRigPoleSolverSettings Settings;
+
+	// runtime FIKRigSolverBase interface
+	virtual void Initialize(const FIKRigSkeleton& InIKRigSkeleton) override;
+	virtual void Solve(FIKRigSkeleton& InIKRigSkeleton, const FIKRigGoalContainer& InGoals) override;
+	virtual void GetRequiredBones(TSet<FName>& OutRequiredBones) const override;
+	virtual void GetRequiredGoals(TSet<FName>& OutRequiredGoals) const override;
+
+	// solver settings
+	virtual FIKRigSolverSettingsBase* GetSolverSettings() override;
+	virtual const UScriptStruct* GetSolverSettingsType() const override;
+	
+	// goals
+	virtual void AddGoal(const UIKRigEffectorGoal& NewGoal) override;
+	virtual void OnGoalRenamed(const FName& InOldName, const FName& InNewName) override;
+	virtual void OnGoalRemoved(const FName& InName) override;
+	virtual void OnGoalMovedToDifferentBone(const FName& InGoalName, const FName& InNewBoneName) override;
+	
+	// start bone can be set on this solver
+	virtual bool UsesStartBone() const override { return true; };
+	virtual void SetStartBone(const FName& InRootBoneName) override;
+	virtual FName GetStartBone() const override;
+	
+	// end bone can be set on this solver
+	virtual bool UsesEndBone() const override { return true; };
+	virtual void SetEndBone(const FName& InEndBoneName) override;
+	virtual FName GetEndBone() const override;
+
+#if WITH_EDITOR
+	// return custom controller for scripting this solver
+	virtual UIKRigSolverControllerBase* GetSolverController(UObject* Outer) override;
+	
+	// UI
+	virtual FText GetNiceName() const override;
+	virtual bool GetWarningMessage(FText& OutWarningMessage) const override;
+	virtual bool IsBoneAffectedBySolver(const FName& InBoneName, const FIKRigSkeleton& InIKRigSkeleton) const override;
+#endif
+
+private:
+	
+	TArray<int32> Chain;
+	TArray<int32> ChildrenToUpdate;
+
+	static void GatherChildren(const int32 BoneIndex, const FIKRigSkeleton& InSkeleton, TArray<int32>& OutChildren);
+};
+
+/* The blueprint/python API for modifying a Pole solver's settings in an IK Rig. */
+UCLASS(BlueprintType)
+class IKRIG_API UIKRigPoleSolverController : public UIKRigSolverControllerBase
+{
+	GENERATED_BODY()
+	
+public:
+
+	/* Get the current solver settings as a struct.
+	 * @return FIKRigPoleSolverSettings struct with the current settings used by the solver. */
+	UFUNCTION(BlueprintCallable, Category = IKRig)
+	FIKRigPoleSolverSettings GetSolverSettings();
+
+	/* Set the solver settings. Input is a custom struct type for this solver.
+	 * @param InSettings a FIKRigPoleSolverSettings struct containing all the settings to apply to this solver */
+	UFUNCTION(BlueprintCallable, Category = IKRig)
+	void SetSolverSettings(FIKRigPoleSolverSettings InSettings);
+};
+
+PRAGMA_DISABLE_DEPRECATION_WARNINGS
+
+UCLASS(Deprecated)
+class
+UE_DEPRECATED(5.6, "This type has been replaced with FPoleSolver.")
+IKRIG_API UDEPRECATED_IKRig_PoleSolverEffector : public UObject
+{
+	GENERATED_BODY()
+
+public:
+	
+	UPROPERTY()
+	FName GoalName = NAME_None;
+	UPROPERTY()
+	FName BoneName = NAME_None;
+	UPROPERTY()
+	float Alpha = 1.0f;
+};
+
+UCLASS(Deprecated)
+class
+UE_DEPRECATED(5.6, "This type has been replaced with FPoleSolver.")
+IKRIG_API UDEPRECATED_IKRig_PoleSolver : public UDEPRECATED_IKRigSolver
+{
+	GENERATED_BODY()
+
+public:
+
+	UDEPRECATED_IKRig_PoleSolver()
+	{
+		Effector_DEPRECATED = CreateDefaultSubobject<UDEPRECATED_IKRig_PoleSolverEffector>(TEXT("Effector"));
+	};
+	
+	UPROPERTY()
+	FName RootName = NAME_None;
+	UPROPERTY()
+	FName EndName = NAME_None;
+	UE_DEPRECATED(5.6, "IK Rig refactored to not use UObjects")
+	UPROPERTY(meta = (DeprecatedProperty))
+	TObjectPtr<UDEPRECATED_IKRig_PoleSolverEffector> Effector_DEPRECATED;
+
+	// upgrade patch to new UStruct-based solver
+	virtual void ConvertToInstancedStruct(FInstancedStruct& OutInstancedStruct) override
+	{
+		OutInstancedStruct.InitializeAs(FIKRigPoleSolver::StaticStruct());
+		FIKRigPoleSolver& NewSolver = OutInstancedStruct.GetMutable<FIKRigPoleSolver>();
+
+		// load solver and effector settings
+		NewSolver.SetEnabled(IsEnabled());
+		NewSolver.Settings.StartBone = RootName;
+		NewSolver.Settings.EndBone = EndName;
+		NewSolver.Settings.AimAtGoal = Effector_DEPRECATED->GoalName;
+		NewSolver.Settings.Alpha = Effector_DEPRECATED->Alpha;
+	};
+};
+
+PRAGMA_ENABLE_DEPRECATION_WARNINGS
\ No newline at end of file
diff --git a/Engine/Plugins/Animation/IKRig/Source/IKRig/Public/Rig/Solvers/IKRigSolver.h b/Engine/Plugins/Animation/IKRig/Source/IKRig/Public/Rig/Solvers/IKRigSolver.h
deleted file mode 100644
diff --git a/Engine/Plugins/Animation/IKRig/Source/IKRig/Public/Rig/Solvers/IKRigSolverBase.h b/Engine/Plugins/Animation/IKRig/Source/IKRig/Public/Rig/Solvers/IKRigSolverBase.h
new file mode 100644
--- /dev/null
+++ b/Engine/Plugins/Animation/IKRig/Source/IKRig/Public/Rig/Solvers/IKRigSolverBase.h
@@ -0,0 +1,298 @@
+// Copyright Epic Games, Inc. All Rights Reserved.
+
+#pragma once
+
+#include "UObject/Object.h"
+#include "UObject/StrongObjectPtr.h"
+
+#include "IKRigSolverBase.generated.h"
+
+class UIKRigEffectorGoal;
+class FPrimitiveDrawInterface;
+struct FIKRigGoalContainer;
+struct FIKRigSkeleton;
+struct FInstancedStruct;
+
+// This is the base class for all types of settings used by IKRig
+USTRUCT(BlueprintType)
+struct IKRIG_API FIKRigSettingsBase
+{
+	GENERATED_BODY()
+	
+	virtual ~FIKRigSettingsBase() {};
+};
+
+// This is the base class for defining editable per-goal settings for your custom IKRig solver.
+// For example, your solver could have a "Strength" value for goals, which can go here.
+// NOTE: the derived type must be returned by the solver's GetGoalSettingsType() and GetGoalSettings()
+USTRUCT(BlueprintType)
+struct IKRIG_API FIKRigGoalSettingsBase : public FIKRigSettingsBase
+{
+	GENERATED_BODY()
+
+	/** The IKRig Goal that these settings are applied to. */
+	UPROPERTY(BlueprintReadOnly, VisibleAnywhere, Category = "Goal Settings")
+	FName Goal;
+	
+	virtual ~FIKRigGoalSettingsBase() {};
+};
+
+// This is the base class for defining editable per-bone settings for your custom IKRig solver.
+// For example, your solver may have rotation limits per-bone which can be stored here.
+// NOTE: the derived type must be returned by the solver's GetBoneSettingsType() and GetBoneSettings()
+USTRUCT(BlueprintType)
+struct IKRIG_API FIKRigBoneSettingsBase : public FIKRigSettingsBase
+{
+	GENERATED_BODY()
+
+	/** The bone these settings are applied to. */
+	UPROPERTY(BlueprintReadOnly, VisibleAnywhere, Category = "Bone Settings")
+	FName Bone;
+	
+	virtual ~FIKRigBoneSettingsBase() {};
+};
+
+// This is the base class for defining editable settings for your custom IKRig solver.
+// All user-configurable properties for your solver should be stored in a subclass of this
+// NOTE: the derived type must be returned by the solver's GetSolverSettingsType() and GetSolverSettings()
+USTRUCT(BlueprintType)
+struct IKRIG_API FIKRigSolverSettingsBase : public FIKRigSettingsBase
+{
+	GENERATED_BODY()
+	
+	virtual ~FIKRigSolverSettingsBase() {};
+};
+
+// this is the base class for creating your own solver type that integrates into the IK Rig framework/editor.
+USTRUCT(BlueprintType)
+struct IKRIG_API FIKRigSolverBase
+{
+	GENERATED_BODY()
+
+	virtual ~FIKRigSolverBase() {};
+
+	// get if this solver is enabled
+	bool IsEnabled() const { return bIsEnabled; };
+	// turn solver on/off (will be skipped during execution if disabled)
+	void SetEnabled(const bool bEnabled){ bIsEnabled = bEnabled; };
+	
+	// RUNTIME
+	// (required) override to setup internal data based on ref pose
+	virtual void Initialize(const FIKRigSkeleton& InIKRigSkeleton) { checkNoEntry(); };
+	// (required) override Solve() to evaluate new output pose (InOutGlobalTransform)
+	virtual void Solve(FIKRigSkeleton& InIKRigSkeleton, const FIKRigGoalContainer& InGoals) { checkNoEntry(); };
+	// (required) override to add any bones which this solver requires (helps validate skeleton compatibility)
+	virtual void GetRequiredBones(TSet<FName>& OutRequiredBones) const  { checkNoEntry(); };
+	// (required) override to add any goals that have been added to this solver
+	virtual void GetRequiredGoals(TSet<FName>& OutRequiredGoals) const  { checkNoEntry(); };
+	// END RUNTIME
+
+	// (optional) override to add custom logic when updating settings from the source asset at runtime
+	// NOTE: you can safely cast this to your own solver type and copy any relevant settings at runtime
+	// This is necessary because at runtime, the IKRigProcessor creates a copy of your solver struct
+	// and the copy must be notified of changes made to the editor instance so that users can see their edits update in PIE.
+	// NOTE: the default implementation will copy solver, goal and bone settings if your solver has them
+	virtual void UpdateSettingsFromAsset(const FIKRigSolverBase& InAssetSolver);
+	
+	// SOLVER SETTINGS
+	// (required) override this to support custom settings for your solver type
+	virtual FIKRigSolverSettingsBase* GetSolverSettings()   { checkNoEntry(); return nullptr; };
+	// (required) override this and return the type your solver uses to store its settings
+	virtual const UScriptStruct* GetSolverSettingsType() const { return nullptr; };
+	// (optional) override to add custom logic if needed when outside systems set solver settings.
+	virtual void SetSolverSettings(const FIKRigSolverSettingsBase* InSettings);
+	// END SOLVER SETTINGS
+
+	// GOALS
+	// (required) override to support ADDING a new goal to custom solver
+	virtual void AddGoal(const UIKRigEffectorGoal& InNewGoal)  { checkNoEntry(); };
+	// (required) override to support RENAMING an existing goal
+	virtual void OnGoalRenamed(const FName& InOldName, const FName& InNewName)  { checkNoEntry(); };
+	// (required) override to support CHANGING BONE for an existing goal
+	virtual void OnGoalMovedToDifferentBone(const FName& InGoalName, const FName& InNewBoneName)  { checkNoEntry(); };
+	// (required) override to support REMOVING a goal from custom solver
+	virtual void OnGoalRemoved(const FName& InGoalName)  { checkNoEntry(); };
+	// END Goals
+
+	// GOAL SETTINGS
+	// (optional) override and return true to tell systems if this solver supports custom per-goal settings
+	virtual bool UsesCustomGoalSettings() const { return false; };
+	// (required if UsesCustomGoalSettings())
+	// override to support supplying goals settings specific to this solver to outside systems for editing/UI
+	// returns a pointer to the settings stored for the given bone (null if bone does not have settings)
+	virtual FIKRigGoalSettingsBase* GetGoalSettings(const FName& InGoalName) { ensure(!UsesCustomGoalSettings()); return nullptr; };
+	// (optional) override if you need custom logic when setting goal settings from outside systems.
+	virtual void SetGoalSettings(const FName& InGoalName, FIKRigGoalSettingsBase* InSettings);
+	// (required if UsesCustomGoalSettings()) returns the type used for the goal settings in this solver
+	virtual const UScriptStruct* GetGoalSettingsType() const { ensure(!UsesCustomGoalSettings()); return nullptr; };
+	// (required if UsesCustomGoalSettings()) override to support telling outside systems which goals this solver has settings for.
+	virtual void GetGoalsWithSettings(TSet<FName>& OutGoalsWithSettings) const { ensure(!UsesCustomGoalSettings()); };
+	// END GOAL SETTINGS
+
+	// START BONE
+	// (optional) determines if your solver requires a start bone to be specified (most do)
+	// override and return true if your solver requires the user to specify a start bone
+	virtual bool UsesStartBone() const { return false; };
+	// (required if UsesStartBone()) if solver requires a start bone, then override this to return it.
+	virtual FName GetStartBone() const { ensure(!UsesStartBone()); return NAME_None; };
+	// (required if UsesStartBone()) override to support SETTING ROOT BONE for the solver
+	virtual void SetStartBone(const FName& InRootBoneName){ ensure(!UsesStartBone()); };
+	// END START BONE
+
+	// END BONE
+	// (optional) implement if your solver requires an end bone
+	// override and return true if this solver requires an end bone to be specified
+	virtual bool UsesEndBone() const { return false; };
+	// (required if UsesEndBone()) override to support getting end bone of the solver
+	virtual FName GetEndBone() const { ensure(!UsesEndBone()); return NAME_None; };
+	// (required if UsesEndBone()) override to support setting end bone for the solver
+	virtual void SetEndBone(const FName& InEndBoneName){ ensure(!UsesEndBone()); };
+	// END ROOT BONE
+
+	// BONE SETTINGS
+	// (optional) implement if your solver supports per-bone settings
+	// override and return true to tell systems if this solver supports per-bone settings
+	virtual bool UsesCustomBoneSettings() const { return false; };
+	// (required if UsesCustomBoneSettings()) override to support ADDING settings to a particular bone
+	virtual void AddSettingsToBone(const FName& InBoneName){ ensure(!UsesCustomBoneSettings()); };
+	// (required if UsesCustomBoneSettings()) override to support REMOVING settings on a particular bone
+	virtual void RemoveSettingsOnBone(const FName& InBoneName){ ensure(!UsesCustomBoneSettings()); };
+	// (required if UsesCustomBoneSettings()) override to support supplying per-bone settings to outside systems for editing/UI
+	// NOTE: This must be overriden on solvers that use bone setting and return nullptr if no settings for the given bone
+	virtual FIKRigBoneSettingsBase* GetBoneSettings(const FName& InBoneName) { ensure(!UsesCustomBoneSettings()); return nullptr; };
+	// (optional) override if you need custom logic when per-bone settings are applied from outside systems
+	virtual void SetBoneSettings(const FName& InBoneName, FIKRigBoneSettingsBase* InSettings);
+	// (required if UsesCustomBoneSettings()) returns the type used for the bone settings in this solver
+	virtual const UScriptStruct* GetBoneSettingsType() const {  ensure(!UsesCustomBoneSettings()); return nullptr; };
+	// (required if UsesCustomBoneSettings()) override to inform outside systems if the given bone has custom settings on it
+	virtual bool HasSettingsOnBone(const FName& InBoneName) const { ensure(!UsesCustomBoneSettings()); return false; };
+	// (required if UsesCustomBoneSettings()) override to support telling outside systems which bones this solver has settings for.
+	// NOTE: Only ADD to the incoming set, do not remove from it.
+	virtual void GetBonesWithSettings(TSet<FName>& OutBonesWithSettings) const { ensure(!UsesCustomBoneSettings()); };
+	// END BONE SETTINGS
+
+#if WITH_EDITOR
+
+	// SCRIPTING API
+	// (optional, but recommended) provide a scripting object to edit your custom solver type in the editor via BP/Python
+	virtual UIKRigSolverControllerBase* GetSolverController(UObject* Outer);
+	// END SCRIPTING API
+
+	// GENERAL UI
+	// (required) override to give your solver a nice name to display in the UI
+	virtual FText GetNiceName() const { checkNoEntry() return FText::GetEmpty(); };
+	// (optional, but recommended) override to provide warning to user during setup of any missing components. return false if no warnings.
+	virtual bool GetWarningMessage(FText& OutWarningMessage) const { return false; };
+	// (optional, but recommended) return true if the supplied Bone is affected by this solver - this provides UI feedback for user
+	virtual bool IsBoneAffectedBySolver(const FName& InBoneName, const FIKRigSkeleton& InIKRigSkeleton) const { return false; };
+	// END GENERAL UI
+
+	// TODO: This will be supported in a future version of Unreal
+	// (optional) override to do custom debug drawing in the editor viewport
+	virtual void DrawBoneSettings(const FName& InBoneName, const FIKRigSkeleton& InIKRigSkeleton, FPrimitiveDrawInterface* PDI) const {};
+
+protected:
+
+	// helper function for solvers to instantiate their own controller
+	UIKRigSolverControllerBase* CreateControllerIfNeeded(UObject* Outer, const UClass* ClassType);
+
+	// an optional custom editor controller used to edit this solver by script/blueprint
+	TStrongObjectPtr<UIKRigSolverControllerBase> Controller = nullptr;
+	
+#endif
+
+private:
+	
+	// generic function to set all types of solver settings
+	static void CopyAllEditableProperties(const UScriptStruct* SettingsType, const FIKRigSettingsBase* CopyFrom, FIKRigSettingsBase* CopyTo);
+
+	// all solvers can be enabled/disabled in the UI or API
+	UPROPERTY()
+	bool bIsEnabled = true;
+};
+
+UCLASS(BlueprintType)
+class IKRIG_API UIKRigSolverControllerBase : public UObject
+{
+	GENERATED_BODY()
+	
+public:
+
+	/* Turn the solver on/off in the IK Rig solver stack.
+	 * @param bIsEnabled if true, the solver will execute, false otherwise. */
+	UFUNCTION(BlueprintCallable, Category = IKRig)
+	void SetEnabled(bool bIsEnabled) { SolverToControl->SetEnabled(bIsEnabled); };
+
+	/* Get if the solver is on or off in the IK Rig solver stack.
+	 * @param bIsEnabled returns true if the solver is enabled, false otherwise. */
+	UFUNCTION(BlueprintCallable, Category = IKRig)
+	bool GetEnabled() { return SolverToControl->IsEnabled(); };
+
+	// the solver this controller controls
+	FIKRigSolverBase* SolverToControl = nullptr;
+};
+
+//
+// BEGIN DEPRECATED CODE
+//
+// NOTE: As of Unreal 5.6, the old UObject based IK Rig system has been deprecated.
+// These are left here to load old UObject data into the new format. But at some point in the future, support for these IK Rigs will be removed!
+//
+
+PRAGMA_DISABLE_DEPRECATION_WARNINGS
+
+// NOTE on UIKRigSolver deprecation / upgrade path:
+// In 5.6 the UObject's in IK Rig have been replaced with runtime-friendly UStruct based data structures
+// As part of this change, solvers based on UIKRigSolver will no longer function and must be upgraded to the new base FIKRigSolverBase.
+// Old assets using UIKRigSolver-based solvers can be loaded and patched using the ConvertToInstancedStruct() function (see below)
+// All solvers that ship with the IKRig plugin have been ported to the new system and should work exactly as before.
+UCLASS(Deprecated)
+class UE_DEPRECATED(5.6, "This type has been replaced by FIKRigSolverBase. UIKRigSolver-based solvers no longer work, please refactor into the new FIKRigSolverBase struct.")
+IKRIG_API UDEPRECATED_IKRigSolver : public UObject
+{
+	GENERATED_BODY()
+	
+public:
+
+	// this is the deprecation upgrade path for solvers that inherit from UIKRigSolver to convert them to FIKRigSolverBase
+	// override this and supply a struct into OutInstancedStruct that derives from FIKRigSolverBase that implements your custom solver.
+	virtual void ConvertToInstancedStruct(FInstancedStruct& OutInstancedStruct){};
+
+	// DEPRECATED UObject-based solver interface
+	virtual void Initialize(const FIKRigSkeleton& IKRigSkeleton) {};
+	virtual void Solve(FIKRigSkeleton& IKRigSkeleton, const FIKRigGoalContainer& Goals) {};
+	virtual FName GetRootBone() const { return NAME_None; };
+	virtual void GetBonesWithSettings(TSet<FName>& OutBonesWithSettings) const {};
+	bool IsEnabled() const { return bIsEnabled; };
+	void SetEnabled(const bool bEnabled){ bIsEnabled = bEnabled; };
+	virtual void UpdateSolverSettings(UDEPRECATED_IKRigSolver* InSettings){};
+	virtual void RemoveGoal(const FName& GoalName){};
+	virtual bool IsGoalConnected(const FName& GoalName) const {return false;};
+#if WITH_EDITORONLY_DATA
+	virtual FText GetNiceName() const { return FText::GetEmpty(); };
+	virtual bool GetWarningMessage(FText& OutWarningMessage) const { return false; };
+	virtual void AddGoal(const UIKRigEffectorGoal& NewGoal){};
+	virtual void RenameGoal(const FName& OldName, const FName& NewName){};
+	virtual void SetGoalBone(const FName& GoalName, const FName& NewBoneName){};
+	virtual UObject* GetGoalSettings(const FName& GoalName) const {return nullptr;};
+	virtual void SetRootBone(const FName& RootBoneName){};
+	virtual bool RequiresRootBone() const { return false; };
+	virtual void SetEndBone(const FName& EndBoneName){};
+	virtual FName GetEndBone() const { return NAME_None; };
+	virtual bool RequiresEndBone() const { return false; };
+	virtual void AddBoneSetting(const FName& BoneName){};
+	virtual void RemoveBoneSetting(const FName& BoneName){};
+	virtual UObject* GetBoneSetting(const FName& BoneName) const { ensure(!UsesBoneSettings()); return nullptr; };
+	virtual bool UsesBoneSettings() const { return false;};
+	virtual void DrawBoneSettings(const FName& BoneName, const FIKRigSkeleton& IKRigSkeleton, FPrimitiveDrawInterface* PDI) const {};
+	virtual bool IsBoneAffectedBySolver(const FName& BoneName, const FIKRigSkeleton& IKRigSkeleton) const { return false; };
+	// END DEPRECATED UObject-based solver interface
+#endif
+
+private:
+	
+	UPROPERTY()
+	bool bIsEnabled = true;
+};
+
+PRAGMA_ENABLE_DEPRECATION_WARNINGS
diff --git a/Engine/Plugins/Animation/IKRig/Source/IKRig/Public/Rig/Solvers/IKRig_BodyMover.h b/Engine/Plugins/Animation/IKRig/Source/IKRig/Public/Rig/Solvers/IKRig_BodyMover.h
deleted file mode 100644
diff --git a/Engine/Plugins/Animation/IKRig/Source/IKRig/Public/Rig/Solvers/IKRig_FBIKSolver.h b/Engine/Plugins/Animation/IKRig/Source/IKRig/Public/Rig/Solvers/IKRig_FBIKSolver.h
--- a/Engine/Plugins/Animation/IKRig/Source/IKRig/Public/Rig/Solvers/IKRig_FBIKSolver.h
+++ b/Engine/Plugins/Animation/IKRig/Source/IKRig/Public/Rig/Solvers/IKRig_FBIKSolver.h
@@ -3,23 +3,19 @@
 #pragma once
 
 #include "Rig/IKRigDefinition.h"
-#include "IKRigSolver.h"
+#include "IKRigSolverBase.h"
 #include "Core/PBIKSolver.h"
 #include "PBIK_Shared.h"
 
 #include "IKRig_FBIKSolver.generated.h"
 
-UCLASS(BlueprintType)
-class IKRIG_API UIKRig_FBIKEffector : public UObject
+USTRUCT(BlueprintType)
+struct IKRIG_API FIKRigFBIKGoalSettings : public FIKRigGoalSettingsBase
 {
 	GENERATED_BODY()
 
-public:
-	UIKRig_FBIKEffector() { SetFlags(RF_Transactional); }
-	
-	/** The Goal that is driving this effector's transform. */
-	UPROPERTY(BlueprintReadOnly, VisibleAnywhere, Category = "Full Body IK Effector")
-	FName GoalName;
+	FIKRigFBIKGoalSettings() = default;
+	FIKRigFBIKGoalSettings(FName InGoalName, FName InBoneName) : BoneName(InBoneName) { Goal = InGoalName;}
 	
 	/** The bone that this effector will pull on. */
 	UPROPERTY(BlueprintReadOnly, VisibleAnywhere, Category = "Full Body IK Effector")
@@ -55,36 +51,15 @@ public:
 
 	UPROPERTY(Transient)
 	int32 IndexInSolver = -1;
-
-	void CopySettings(const UIKRig_FBIKEffector* Other)
-	{
-		ChainDepth = Other->ChainDepth;
-		StrengthAlpha = Other->StrengthAlpha;
-		PullChainAlpha = Other->PullChainAlpha;
-		PinRotation = Other->PinRotation;
-	}
 };
 
-UCLASS(BlueprintType)
-class IKRIG_API UIKRig_FBIKBoneSettings : public UObject
+USTRUCT(BlueprintType)
+struct IKRIG_API FIKRigFBIKBoneSettings : public FIKRigBoneSettingsBase
 {
 	GENERATED_BODY()
 
-public:
-	
-	UIKRig_FBIKBoneSettings()
-		: Bone(NAME_None), 
-		X(EPBIKLimitType::Free),
-		Y(EPBIKLimitType::Free),
-		Z(EPBIKLimitType::Free),
-		PreferredAngles(FVector::ZeroVector)
-	{
-		SetFlags(RF_Transactional);
-	}
-
-	/** The bone these settings are applied to. */
-	UPROPERTY(BlueprintReadOnly, VisibleAnywhere, Category = Bone, meta = (Constant, CustomWidget = "BoneName"))
-	FName Bone;
+	FIKRigFBIKBoneSettings() = default;
+	FIKRigFBIKBoneSettings(FName InBoneName) { Bone = InBoneName; }
 	
 	/** Range is 0 to 1 (Default is 0). At higher values, the bone will resist rotating (forcing other bones to compensate). */
 	UPROPERTY(BlueprintReadWrite, EditAnywhere, Category = Stiffness, meta = (ClampMin = "0", ClampMax = "1", UIMin = "0.0", UIMax = "1.0"))
@@ -99,7 +74,7 @@ public:
 	 *Limited: rotation is clamped between the min/max angles relative to the Skeletal Mesh reference pose.
 	 *Locked: no rotation is allowed in this axis (will remain at reference pose angle). */
 	UPROPERTY(BlueprintReadWrite, EditAnywhere, Category = Limits)
-	EPBIKLimitType X;
+	EPBIKLimitType X = EPBIKLimitType::Free;
 	/**Range is -180 to 0 (Default is 0). Degrees of rotation in the negative X direction to allow when joint is in "Limited" mode. */
 	UPROPERTY(BlueprintReadWrite, EditAnywhere, Category = Limits, meta = (ClampMin = "-180", ClampMax = "0", UIMin = "-180.0", UIMax = "0.0"))
 	float MinX = 0.0f;
@@ -112,7 +87,7 @@ public:
 	*Limited: rotation is clamped between the min/max angles relative to the Skeletal Mesh reference pose.
 	*Locked: no rotation is allowed in this axis (will remain at input pose angle). */
 	UPROPERTY(BlueprintReadWrite, EditAnywhere, Category = Limits)
-	EPBIKLimitType Y;
+	EPBIKLimitType Y = EPBIKLimitType::Free;
 	/**Range is -180 to 0 (Default is 0). Degrees of rotation in the negative Y direction to allow when joint is in "Limited" mode. */
 	UPROPERTY(BlueprintReadWrite, EditAnywhere, Category = Limits, meta = (ClampMin = "-180", ClampMax = "0", UIMin = "-180.0", UIMax = "0.0"))
 	float MinY = 0.0f;
@@ -125,7 +100,7 @@ public:
 	*Limited: rotation is clamped between the min/max angles relative to the Skeletal Mesh reference pose.
 	*Locked: no rotation is allowed in this axis (will remain at input pose angle). */
 	UPROPERTY(BlueprintReadWrite, EditAnywhere, Category = Limits)
-	EPBIKLimitType Z;
+	EPBIKLimitType Z = EPBIKLimitType::Free;
 	/**Range is -180 to 0 (Default is 0). Degrees of rotation in the negative Z direction to allow when joint is in "Limited" mode. */
 	UPROPERTY(BlueprintReadWrite, EditAnywhere, Category = Limits, meta = (ClampMin = "-180", ClampMax = "0", UIMin = "-180.0", UIMax = "0.0"))
 	float MinZ = 0.0f;
@@ -142,7 +117,7 @@ public:
 	 * This happens by moving the effector at the tip of the chain towards the root of the chain.
 	 * This can be used to coerce knees and elbows to bend in the anatomically "correct" direction without resorting to limits (which may require more iterations to converge).*/
 	UPROPERTY(BlueprintReadWrite, EditAnywhere, Category = PreferredAngles)
-	FVector PreferredAngles;
+	FVector PreferredAngles = FVector::ZeroVector;
 
 	void CopyToCoreStruct(PBIK::FBoneSettings& Settings) const
 	{
@@ -162,40 +137,21 @@ public:
 		Settings.PreferredAngles.Yaw = PreferredAngles.Z;
 		Settings.PreferredAngles.Roll = PreferredAngles.X;
 	}
-
-	void CopySettings(const UIKRig_FBIKBoneSettings* Other)
-	{	
-		RotationStiffness = Other->RotationStiffness;
-		PositionStiffness = Other->PositionStiffness;
-		X = Other->X;
-		MinX = Other->MinX;
-		MaxX = Other->MaxX;
-		Y = Other->Y;
-		MinY = Other->MinY;
-		MaxY = Other->MaxY;
-		Z = Other->Z;
-		MinZ = Other->MinZ;
-		MaxZ = Other->MaxZ;
-		bUsePreferredAngles = Other->bUsePreferredAngles;
-		PreferredAngles = Other->PreferredAngles;
-	}
 };
 
-UCLASS(BlueprintType, EditInlineNew, config = Engine, hidecategories = UObject)
-class IKRIG_API UIKRigFBIKSolver : public UIKRigSolver
+USTRUCT(BlueprintType)
+struct IKRIG_API FIKRigFBIKSettings : public FIKRigSolverSettingsBase
 {
 	GENERATED_BODY()
 
-public:
-
-	/** All bones above this bone in the hierarchy will be completely ignored by the solver. Typically this is set to
+	/** All bones above this bone in the hierarchy will be completely ignored by the solver. Typically, this is set to
 	 * the top-most skinned bone in the Skeletal Mesh (ie Pelvis, Hips etc), NOT the actual root of the entire skeleton.
 	 *
 	 * If you want to use the solver on a single chain of bones, and NOT translate the chain, ensure that "PinRoot" is
 	 * checked on to disable the root from translating to reach the effector goals.*/
-	UPROPERTY(BlueprintReadWrite, VisibleAnywhere, Category = "Full Body IK Settings")
+	UPROPERTY(BlueprintReadOnly, VisibleAnywhere, Category = "Full Body IK Settings")
 	FName RootBone;
-
+	
 	/** High iteration counts can help solve complex joint configurations with competing constraints, but will increase runtime cost. Default is 20. */
 	UPROPERTY(BlueprintReadWrite, EditAnywhere, Category = SolverSettings, meta = (ClampMin = "0", ClampMax = "1000", UIMin = "0.0", UIMax = "200.0"))
 	int32 Iterations = 20;
@@ -225,7 +181,7 @@ public:
 
 	/** A global multiplier for all Pull Chain Alpha values on all effectors. Range is 0.0 to 1.0. Default is 1.0. */
 	UPROPERTY(BlueprintReadWrite, EditAnywhere, Category = AdvancedSettings, meta = (ClampMin = "0", ClampMax = "1", UIMin = "0.0", UIMax = "1.0"))
-	float PullChainAlpha = 1.0f;
+	float GlobalPullChainAlpha = 1.0f;
 
 	/** Maximum angle that a joint can be rotated per constraint iteration. Lower this value if the solve is diverging. Range is 0.0 to 180.0. Default is 30. */
 	UPROPERTY(BlueprintReadWrite, EditAnywhere, Category = AdvancedSettings, meta = (ClampMin = "0", ClampMax = "45", UIMin = "0.0", UIMax = "180.0"))
@@ -234,73 +190,300 @@ public:
 	/** Pushes constraints beyond their normal amount to speed up convergence. Increasing this may speed up convergence, but at the cost of stability. Range is 1.0 - 2.0. Default is 1.3. */
 	UPROPERTY(BlueprintReadWrite, EditAnywhere, Category = AdvancedSettings, meta = (ClampMin = "1",  ClampMax = "2", UIMin = "1.0", UIMax = "2.0"))
 	float OverRelaxation = 1.3f;
-	
-	/** DEPRECATED: When true, the solver is reset each tick to start from the current input pose. Default is true.
-	 * If false, incoming animated poses are ignored and the solver starts from the results of the previous solve.
-	 * In very limited circumstances, it can be beneficial to use the pose from the previous frame.*/
-	UPROPERTY()
-	bool bStartSolveFromInputPose_DEPRECATED = true;
+};
 
-	UFUNCTION(BlueprintCallable, Category = "IKRig")
-	TArray<UIKRig_FBIKEffector*> GetEffectors()
-	{
-		return Effectors;
-	}
 
-	UFUNCTION(BlueprintCallable, Category = "IKRig")
-	TArray<UIKRig_FBIKBoneSettings*> GetBoneSettings()
-	{
-		return BoneSettings;
-	}
+USTRUCT(BlueprintType)
+struct IKRIG_API FIKRigFullBodyIKSolver : public FIKRigSolverBase
+{
+	GENERATED_BODY()
 
 	UPROPERTY()
-	TArray<TObjectPtr<UIKRig_FBIKEffector>> Effectors;
-
+	FIKRigFBIKSettings Settings;
+	
 	UPROPERTY()
-	TArray<TObjectPtr<UIKRig_FBIKBoneSettings>> BoneSettings;
+	TArray<FIKRigFBIKGoalSettings> AllGoalSettings;
 
-	/** UIKRigSolver interface */
-	// runtime
-	virtual void Initialize(const FIKRigSkeleton& IKRigSkeleton) override;
-	virtual void Solve(FIKRigSkeleton& IKRigSkeleton, const FIKRigGoalContainer& Goals) override;
-
-	virtual FName GetRootBone() const override { return RootBone; };
+	UPROPERTY()
+	TArray<FIKRigFBIKBoneSettings> AllBoneSettings;
+	
+	// FIKRigSolverBase runtime interface
+	virtual void Initialize(const FIKRigSkeleton& InIKRigSkeleton) override;
+	virtual void Solve(FIKRigSkeleton& InIKRigSkeleton, const FIKRigGoalContainer& InGoals) override;
+	virtual void GetRequiredBones(TSet<FName>& OutRequiredBones) const override;
+	virtual void GetRequiredGoals(TSet<FName>& OutRequiredGoals) const override;
+	// END FIKRigSolverBase runtime interface
+
+	// solver settings
+	virtual FIKRigSolverSettingsBase* GetSolverSettings() override;
+	virtual const UScriptStruct* GetSolverSettingsType() const override;
+	// END solver settings
+	
+	// goals
+	virtual void AddGoal(const UIKRigEffectorGoal& InNewGoal) override;
+	virtual void OnGoalRenamed(const FName& InOldName, const FName& InNewName) override;
+	virtual void OnGoalMovedToDifferentBone(const FName& InGoalName, const FName& InNewBoneName) override;
+	virtual void OnGoalRemoved(const FName& InGoalName) override;
+	// END goals
+
+	// goal settings
+	virtual bool UsesCustomGoalSettings() const override;
+	virtual FIKRigGoalSettingsBase* GetGoalSettings(const FName& InGoalName) override;
+	virtual const UScriptStruct* GetGoalSettingsType() const override;
+	virtual void GetGoalsWithSettings(TSet<FName>& OutGoalsWithSettings) const override;
+	// END goal settings
+
+	// start bone
+	virtual bool UsesStartBone() const override;
+	virtual void SetStartBone(const FName& InStartBoneName) override;
+	virtual FName GetStartBone() const override;
+	// END start bone
+	
+	// bone settings
+	virtual bool UsesCustomBoneSettings() const override;
+	virtual void AddSettingsToBone(const FName& InBoneName) override;
+	virtual void RemoveSettingsOnBone(const FName& InBoneName) override;
+	virtual FIKRigBoneSettingsBase* GetBoneSettings(const FName& InBoneName) override;
+	virtual const UScriptStruct* GetBoneSettingsType() const override;
+	virtual bool HasSettingsOnBone(const FName& InBoneName) const override;
 	virtual void GetBonesWithSettings(TSet<FName>& OutBonesWithSettings) const override;
-
-	virtual void UpdateSolverSettings(UIKRigSolver* InSettings) override;
-	virtual void RemoveGoal(const FName& GoalName) override;
-	virtual bool IsGoalConnected(const FName& GoalName) const override;
+	// END bone settings
 	
 #if WITH_EDITOR
+	// return custom controller for scripting this solver
+	virtual UIKRigSolverControllerBase* GetSolverController(UObject* Outer) override;
+	// ui
 	virtual FText GetNiceName() const override;
 	virtual bool GetWarningMessage(FText& OutWarningMessage) const override;
-	// goals
-	virtual void AddGoal(const UIKRigEffectorGoal& NewGoal) override;
-	virtual void RenameGoal(const FName& OldName, const FName& NewName) override;
-	virtual void SetGoalBone(const FName& GoalName, const FName& NewBoneName) override;
-	virtual UObject* GetGoalSettings(const FName& GoalName) const override;
-	// bone settings
-	virtual void AddBoneSetting(const FName& BoneName) override;
-	virtual void RemoveBoneSetting(const FName& BoneName) override;
-	virtual UObject* GetBoneSetting(const FName& BoneName) const override;
-	virtual bool UsesBoneSettings() const override{ return true;};
-	virtual void DrawBoneSettings(const FName& BoneName, const FIKRigSkeleton& IKRigSkeleton, FPrimitiveDrawInterface* PDI) const override;
-	// root bone can be set on this solver
-	virtual bool RequiresRootBone() const override { return true; };
-	virtual void SetRootBone(const FName& RootBoneName) override;
-	virtual bool IsBoneAffectedBySolver(const FName& BoneName, const FIKRigSkeleton& IKRigSkeleton) const override;
-	/** END UIKRigSolver interface */
-
-	//** UObject */
-	virtual void PostLoad() override;
-	//** END UObject */
+	virtual bool IsBoneAffectedBySolver(const FName& InBoneName, const FIKRigSkeleton& InIKRigSkeleton) const override;
 #endif
 
 private:
 
 	FPBIKSolver Solver;
 
-	int32 GetIndexOfGoal(const FName& GoalName) const;
+	int32 GetIndexOfGoal(const FName& InGoalName) const;
+};
+
+
+/* The blueprint/python API for modifying an Full-Body IK solver's settings in an IK Rig.
+ * Can adjust Solver, Goal and Bone settings. */
+UCLASS(BlueprintType)
+class IKRIG_API UIKRigFBIKController : public UIKRigSolverControllerBase
+{
+	GENERATED_BODY()
+	
+public:
+
+	/* Get the current solver settings as a struct.
+	 * @return FIKRigFBIKSettings struct with the current settings used by the solver. */
+	UFUNCTION(BlueprintCallable, Category = IKRig)
+	FIKRigFBIKSettings GetSolverSettings();
+
+	/* Set the solver settings. Input is a custom struct type for this solver.
+	 * @param InSettings a FIKRigFBIKSettings struct containing all the settings to apply to this solver */
+	UFUNCTION(BlueprintCallable, Category = IKRig)
+	void SetSolverSettings(FIKRigFBIKSettings InSettings);
+
+	/* Get the settings for the specified goal.
+	 * @param InGoalName the name of the goal to get settings for
+	 * @return FIKRigFBIKGoalSettings struct (empty if the specified goal does not belong to this solver) */
+	UFUNCTION(BlueprintCallable, Category = IKRig)
+	FIKRigFBIKGoalSettings GetGoalSettings(const FName InGoalName);
+
+	/* Set the settings for the specified goal.
+	 * @param InGoalName: the name of the goal to assign the settings to.
+	 * @param InSettings: a custom struct type with all the settings for an FBIK goal */
+	UFUNCTION(BlueprintCallable, Category = IKRig)
+	void SetGoalSettings(const FName InGoalName, FIKRigFBIKGoalSettings InSettings);
+
+	/* Get the settings associated with a particular bone.
+	 * Note that you must AddBoneSettings() using the IK Rig controller before a bone will have settings on it.
+	 * @param InBoneName the name of the bone to get settings for
+	 * @return FIKRigFBIKBoneSettings struct holding all the settings for the specified bone (or empty if the bone did not have settings) */
+	UFUNCTION(BlueprintCallable, Category = IKRig)
+	FIKRigFBIKBoneSettings GetBoneSettings(const FName InBoneName);
+
+	/* Apply settings to a given bone
+	 * @param InBoneName the name of the bone to apply the settings to
+	 * @param InSettings a FIKRigFBIKBoneSettings struct containing the settings */
+	UFUNCTION(BlueprintCallable, Category = IKRig)
+	void SetBoneSettings(const FName InBoneName, FIKRigFBIKBoneSettings InSettings);
+};
+
+//
+// BEGIN DEPRECATED CODE 5.6
+//
+
+PRAGMA_DISABLE_DEPRECATION_WARNINGS
+
+UCLASS(Deprecated)
+class
+UE_DEPRECATED(5.6, "This type has been replaced with FFBIKGoalSettings.")
+IKRIG_API UDEPRECATED_IKRig_FBIKEffector : public UObject
+{
+	GENERATED_BODY()
+	
+public:
+	
+	UPROPERTY()
+	FName GoalName;
+	UPROPERTY()
+	FName BoneName;
+	UPROPERTY()
+	int32 ChainDepth = 0;
+	UPROPERTY()
+	float StrengthAlpha = 1.0f;
+	UPROPERTY()
+	float PullChainAlpha = 1.0f;
+	UPROPERTY()
+	float PinRotation = 1.0f;
+};
+
+UCLASS(Deprecated)
+class
+UE_DEPRECATED(5.6, "This type has been replaced with FFBIKBoneSettings.")
+IKRIG_API UDEPRECATED_IKRig_FBIKBoneSettings : public UObject
+{
+	GENERATED_BODY()
+
+public:
+	
+	UPROPERTY()
+	FName Bone;
+	UPROPERTY()
+	float RotationStiffness = 0.0f;
+	UPROPERTY()
+	float PositionStiffness = 0.0f;
+	UPROPERTY()
+	EPBIKLimitType X;
+	UPROPERTY()
+	float MinX = 0.0f;
+	UPROPERTY()
+	float MaxX = 0.0f;
+	UPROPERTY()
+	EPBIKLimitType Y;
+	UPROPERTY()
+	float MinY = 0.0f;
+	UPROPERTY()
+	float MaxY = 0.0f;
+	UPROPERTY()
+	EPBIKLimitType Z;
+	UPROPERTY()
+	float MinZ = 0.0f;
+	UPROPERTY()
+	float MaxZ = 0.0f;
+	UPROPERTY()
+	bool bUsePreferredAngles = false;
+	UPROPERTY()
+	FVector PreferredAngles;
 };
 
+UCLASS(Deprecated)
+class
+UE_DEPRECATED(5.6, "This type has been replaced with FFBIKSolver.")
+IKRIG_API UDEPRECATED_IKRigFBIKSolver : public UDEPRECATED_IKRigSolver
+{
+	GENERATED_BODY()
+
+public:
+
+	// upgrade patch to new UStruct-based solver
+	virtual void ConvertToInstancedStruct(FInstancedStruct& OutInstancedStruct) override
+	{
+		OutInstancedStruct.InitializeAs(FIKRigFullBodyIKSolver::StaticStruct());
+		
+		// load solver settings
+		FIKRigFullBodyIKSolver& NewSolver = OutInstancedStruct.GetMutable<FIKRigFullBodyIKSolver>();
+		NewSolver.SetEnabled(IsEnabled());
+		NewSolver.Settings.RootBone = RootBone;
+		NewSolver.Settings.Iterations = Iterations;
+		NewSolver.Settings.SubIterations = SubIterations;
+		NewSolver.Settings.MassMultiplier = MassMultiplier;
+		NewSolver.Settings.bAllowStretch = bAllowStretch;
+		NewSolver.Settings.RootBehavior = RootBehavior;
+		NewSolver.Settings.PrePullRootSettings = PrePullRootSettings;
+		NewSolver.Settings.GlobalPullChainAlpha = PullChainAlpha;
+		NewSolver.Settings.MaxAngle = MaxAngle;
+		NewSolver.Settings.OverRelaxation = OverRelaxation;
+
+		// load effectors into solver
+		for (UDEPRECATED_IKRig_FBIKEffector* OldEffector : Effectors_DEPRECATED)
+		{
+			FIKRigFBIKGoalSettings NewEffector;
+			NewEffector.Goal = OldEffector->GoalName;
+			NewEffector.BoneName = OldEffector->BoneName;
+			NewEffector.ChainDepth = OldEffector->ChainDepth;
+			NewEffector.StrengthAlpha = OldEffector->StrengthAlpha;
+			NewEffector.PullChainAlpha = OldEffector->PullChainAlpha;
+			NewEffector.PinRotation = OldEffector->PinRotation;
+			NewSolver.AllGoalSettings.Add(NewEffector);
+		}
+		
+		// load bone settings into solver
+		for (UDEPRECATED_IKRig_FBIKBoneSettings* OldBoneSetting : BoneSettings_DEPRECATED)
+		{
+			FIKRigFBIKBoneSettings NewBoneSetting;
+			NewBoneSetting.Bone = OldBoneSetting->Bone;
+			NewBoneSetting.RotationStiffness = OldBoneSetting->RotationStiffness;
+			NewBoneSetting.PositionStiffness = OldBoneSetting->PositionStiffness;
+			NewBoneSetting.X = OldBoneSetting->X;
+			NewBoneSetting.MinX = OldBoneSetting->MinX;
+			NewBoneSetting.MaxX = OldBoneSetting->MaxX;
+			NewBoneSetting.Y = OldBoneSetting->Y;
+			NewBoneSetting.MinY = OldBoneSetting->MinY;
+			NewBoneSetting.MaxY = OldBoneSetting->MaxY;
+			NewBoneSetting.Z = OldBoneSetting->Z;
+			NewBoneSetting.MinZ = OldBoneSetting->MinZ;
+			NewBoneSetting.MaxZ = OldBoneSetting->MaxZ;
+			NewBoneSetting.bUsePreferredAngles = OldBoneSetting->bUsePreferredAngles;
+			NewBoneSetting.PreferredAngles = OldBoneSetting->PreferredAngles;
+			NewSolver.AllBoneSettings.Add(NewBoneSetting);
+		}
+	};
+	
+	UPROPERTY()
+	FName RootBone;
+	UPROPERTY()
+	int32 Iterations = 20;
+	UPROPERTY()
+	int32 SubIterations = 0;
+	UPROPERTY()
+	float MassMultiplier = 1.0f;
+	UPROPERTY()
+	bool bAllowStretch = false;
+	UPROPERTY()
+	EPBIKRootBehavior RootBehavior = EPBIKRootBehavior::PrePull;
+	UPROPERTY()
+	FRootPrePullSettings PrePullRootSettings;
+	UPROPERTY()
+	float PullChainAlpha = 1.0f;
+	UPROPERTY()
+	float MaxAngle = 30.f;
+	UPROPERTY()
+	float OverRelaxation = 1.3f;
+
+	UE_DEPRECATED(5.6, "Deprecated UObject based FBIK solver.")
+	UFUNCTION(meta = (DeprecatedFunction))
+	TArray<UDEPRECATED_IKRig_FBIKEffector*> GetEffectors()
+	{
+		return Effectors_DEPRECATED;
+	}
+
+	UE_DEPRECATED(5.6, "Deprecated UObject based FBIK solver.")
+	UFUNCTION(meta = (DeprecatedFunction))
+	TArray<UDEPRECATED_IKRig_FBIKBoneSettings*> GetBoneSettings()
+	{
+		return BoneSettings_DEPRECATED;
+	}
+
+	UE_DEPRECATED(5.6, "IK Rig refactored to not use UObjects")
+	UPROPERTY(meta = (DeprecatedProperty))
+	TArray<TObjectPtr<UDEPRECATED_IKRig_FBIKEffector>> Effectors_DEPRECATED;
+
+	UE_DEPRECATED(5.6, "IK Rig refactored to not use UObjects")
+	UPROPERTY(meta = (DeprecatedProperty))
+	TArray<TObjectPtr<UDEPRECATED_IKRig_FBIKBoneSettings>> BoneSettings_DEPRECATED;
+};
 
+PRAGMA_ENABLE_DEPRECATION_WARNINGS
diff --git a/Engine/Plugins/Animation/IKRig/Source/IKRig/Public/Rig/Solvers/IKRig_LimbSolver.h b/Engine/Plugins/Animation/IKRig/Source/IKRig/Public/Rig/Solvers/IKRig_LimbSolver.h
deleted file mode 100644
diff --git a/Engine/Plugins/Animation/IKRig/Source/IKRig/Public/Rig/Solvers/IKRig_PoleSolver.h b/Engine/Plugins/Animation/IKRig/Source/IKRig/Public/Rig/Solvers/IKRig_PoleSolver.h
deleted file mode 100644
diff --git a/Engine/Plugins/Animation/IKRig/Source/IKRig/Public/Rig/Solvers/IKRig_SetTransform.h b/Engine/Plugins/Animation/IKRig/Source/IKRig/Public/Rig/Solvers/IKRig_SetTransform.h
deleted file mode 100644
diff --git a/Engine/Plugins/Animation/IKRig/Source/IKRig/Public/Rig/Solvers/LimbSolver.h b/Engine/Plugins/Animation/IKRig/Source/IKRig/Public/Rig/Solvers/LimbSolver.h
--- a/Engine/Plugins/Animation/IKRig/Source/IKRig/Public/Rig/Solvers/LimbSolver.h
+++ b/Engine/Plugins/Animation/IKRig/Source/IKRig/Public/Rig/Solvers/LimbSolver.h
@@ -3,13 +3,17 @@
 #pragma once
 
 #include "CoreMinimal.h"
+#include "IKRigSolverBase.h"
 #include "LimbSolver.generated.h"
 
-USTRUCT()
-struct IKRIG_API FLimbSolverSettings
+USTRUCT(BlueprintType)
+struct IKRIG_API FLimbSolverSettings : public FIKRigSolverSettingsBase
 {
 	GENERATED_BODY()
 
+	// make this polymorphic
+	virtual ~FLimbSolverSettings() {};
+
 	/** Precision (distance to the target) */
 	UPROPERTY(EditAnywhere, Category = "Limb Solver Settings", meta = (ClampMin = "0.0"))
 	float ReachPrecision = 0.01f;
diff --git a/Engine/Plugins/Animation/IKRig/Source/IKRig/Public/Rig/Solvers/SetTransformSolver.h b/Engine/Plugins/Animation/IKRig/Source/IKRig/Public/Rig/Solvers/SetTransformSolver.h
new file mode 100644
--- /dev/null
+++ b/Engine/Plugins/Animation/IKRig/Source/IKRig/Public/Rig/Solvers/SetTransformSolver.h
@@ -0,0 +1,151 @@
+// Copyright Epic Games, Inc. All Rights Reserved.
+
+#pragma once
+
+#include "Rig/IKRigDefinition.h"
+#include "IKRigSolverBase.h"
+
+#include "SetTransformSolver.generated.h"
+
+USTRUCT(BlueprintType)
+struct IKRIG_API FIKRigSetTransformSettings : public FIKRigSolverSettingsBase
+{
+	GENERATED_BODY()
+	
+	// The goal used to affect the bone transform
+	UPROPERTY(BlueprintReadOnly, VisibleAnywhere, Category = "Set Transform Settings")
+	FName Goal;
+
+	// The bone to affect
+	UPROPERTY(BlueprintReadOnly, VisibleAnywhere, Category = "Set Transform Settings")
+	FName BoneToAffect;
+
+	// Blend the translation on/off. Range is 0-1. Default is 1.0.
+	UPROPERTY(BlueprintReadWrite, EditAnywhere, Category = "Set Transform Effector", meta = (ClampMin = "0.0", ClampMax = "1.0", UIMin = "0.0", UIMax = "1.0"))
+	float PositionAlpha = 1.0;
+
+	// Blend the rotation on/off. Range is 0-1. Default is 1.0.
+	UPROPERTY(BlueprintReadWrite, EditAnywhere, Category = "Set Transform Effector", meta = (ClampMin = "0.0", ClampMax = "1.0", UIMin = "0.0", UIMax = "1.0"))
+	float RotationAlpha = 1.0f;
+
+	// Blend the total effect on/off. Range is 0-1. Default is 1.0.
+	UPROPERTY(BlueprintReadWrite, EditAnywhere, Category = "Set Transform Effector", meta = (ClampMin = "0.0", ClampMax = "1.0", UIMin = "0.0", UIMax = "1.0"))
+	float Alpha = 1.0f;
+
+	// If true, the transform modification will propagate to the hierarchy below the target bone.
+	UPROPERTY(EditAnywhere, Category = "Set Transform Settings")
+	bool bPropagateToChildren = true;
+};
+
+USTRUCT(BlueprintType)
+struct IKRIG_API FIKRigSetTransform : public FIKRigSolverBase
+{
+	GENERATED_BODY()
+	
+	UPROPERTY()
+	FIKRigSetTransformSettings Settings;
+
+	// runtime
+	virtual void Initialize(const FIKRigSkeleton& IKRigSkeleton) override;
+	virtual void Solve(FIKRigSkeleton& InIKRigSkeleton, const FIKRigGoalContainer& InGoals) override;
+	virtual void GetRequiredBones(TSet<FName>& OutRequiredBones) const override;
+	virtual void GetRequiredGoals(TSet<FName>& OutRequiredGoals) const override;
+
+	// settings
+	virtual FIKRigSolverSettingsBase* GetSolverSettings() override;
+	virtual const UScriptStruct* GetSolverSettingsType() const override;
+	
+	// goals
+	virtual void AddGoal(const UIKRigEffectorGoal& InNewGoal) override;
+	virtual void OnGoalRemoved(const FName& InGoalName) override;
+	virtual void OnGoalRenamed(const FName& InOldName, const FName& InNewName) override;
+	virtual void OnGoalMovedToDifferentBone(const FName& InGoalName, const FName& InNewBoneName) override;
+	
+#if WITH_EDITOR
+	// return custom controller for scripting this solver
+	virtual UIKRigSolverControllerBase* GetSolverController(UObject* Outer) override;
+	// ui
+	virtual FText GetNiceName() const override;
+	virtual bool GetWarningMessage(FText& OutWarningMessage) const override;
+	virtual bool IsBoneAffectedBySolver(const FName& InBoneName, const FIKRigSkeleton& InIKRigSkeleton) const override;
+#endif
+
+private:
+
+	int32 BoneIndex = INDEX_NONE;
+};
+
+/* The blueprint/python API for modifying Set Transform settings in an IK Rig. */
+UCLASS(BlueprintType)
+class IKRIG_API UIKRigSetTransformController : public UIKRigSolverControllerBase
+{
+	GENERATED_BODY()
+	
+public:
+
+	/* Get the current solver settings as a struct.
+	 * @return FIKRigSetTransformSettings struct with the current settings used by the solver. */
+	UFUNCTION(BlueprintCallable, Category = IKRig)
+	FIKRigSetTransformSettings GetSolverSettings();
+
+	/* Set the solver settings. Input is a custom struct type for this solver.
+	 * @param InSettings a FIKRigSetTransformSettings struct containing all the settings to apply to this solver */
+	UFUNCTION(BlueprintCallable, Category = IKRig)
+	void SetSolverSettings(FIKRigSetTransformSettings InSettings);
+};
+
+PRAGMA_DISABLE_DEPRECATION_WARNINGS
+
+UCLASS(Deprecated)
+class
+UE_DEPRECATED(5.6, "This type has been replaced with FSetTransformGoalSettings.")
+IKRIG_API UDEPRECATED_IKRig_SetTransformEffector : public UObject
+{
+	GENERATED_BODY()
+
+public:
+	
+	UPROPERTY()
+	bool bEnablePosition = true;
+	UPROPERTY()
+	bool bEnableRotation = true;
+	UPROPERTY()
+	float Alpha = 1.0f;
+};
+
+UCLASS(Deprecated)
+class
+UE_DEPRECATED(5.6, "This type has been replaced with FSetTransform.")
+IKRIG_API UDEPRECATED_IKRig_SetTransform : public UDEPRECATED_IKRigSolver
+{
+	GENERATED_BODY()
+
+public:
+	
+	UDEPRECATED_IKRig_SetTransform()
+	{
+		Effector_DEPRECATED = CreateDefaultSubobject<UDEPRECATED_IKRig_SetTransformEffector>(TEXT("Effector"));
+	}
+
+	// upgrade patch to new UStruct-based solver
+	virtual void ConvertToInstancedStruct(FInstancedStruct& OutInstancedStruct) override
+	{
+		OutInstancedStruct.InitializeAs(FIKRigSetTransform::StaticStruct());
+		FIKRigSetTransform& NewSolver = OutInstancedStruct.GetMutable<FIKRigSetTransform>();
+		NewSolver.SetEnabled(IsEnabled());
+		NewSolver.Settings.Goal = Goal;
+		NewSolver.Settings.BoneToAffect = RootBone;
+		NewSolver.Settings.Alpha = Effector_DEPRECATED->Alpha;
+		NewSolver.Settings.PositionAlpha = Effector_DEPRECATED->bEnablePosition ? 1.0f : 0.0f;
+		NewSolver.Settings.RotationAlpha = Effector_DEPRECATED->bEnableRotation ? 1.0f : 0.0f;
+	};
+
+	UPROPERTY()
+	FName Goal;
+	UPROPERTY()
+	FName RootBone;
+	UPROPERTY(meta = (DeprecatedProperty))
+	TObjectPtr<UDEPRECATED_IKRig_SetTransformEffector> Effector_DEPRECATED;
+};
+
+PRAGMA_ENABLE_DEPRECATION_WARNINGS
diff --git a/Engine/Plugins/Animation/IKRig/Source/IKRigDeveloper/Private/AnimGraphNode_IKRig.cpp b/Engine/Plugins/Animation/IKRig/Source/IKRigDeveloper/Private/AnimGraphNode_IKRig.cpp
--- a/Engine/Plugins/Animation/IKRig/Source/IKRigDeveloper/Private/AnimGraphNode_IKRig.cpp
+++ b/Engine/Plugins/Animation/IKRig/Source/IKRigDeveloper/Private/AnimGraphNode_IKRig.cpp
@@ -4,7 +4,7 @@
 #include "Animation/AnimInstance.h"
 #include "Components/SkeletalMeshComponent.h"
 #include "Rig/IKRigDefinition.h"
-#include "Rig/Solvers/IKRigSolver.h"
+#include "Rig/Solvers/IKRigSolverBase.h"
 #include "Kismet2/CompilerResultsLog.h"
 #include "AnimationGraphSchema.h"
 #include "ScopedTransaction.h"
diff --git a/Engine/Plugins/Animation/IKRig/Source/IKRigEditor/Private/IKRigEditor.cpp b/Engine/Plugins/Animation/IKRig/Source/IKRigEditor/Private/IKRigEditor.cpp
--- a/Engine/Plugins/Animation/IKRig/Source/IKRigEditor/Private/IKRigEditor.cpp
+++ b/Engine/Plugins/Animation/IKRig/Source/IKRigEditor/Private/IKRigEditor.cpp
@@ -57,6 +57,9 @@ void FIKRigEditor::StartupModule()
 
 	// register detail customizations
 	FPropertyEditorModule& PropertyEditorModule = FModuleManager::LoadModuleChecked<FPropertyEditorModule>("PropertyEditor");
+	// custom instanced struct details
+	PropertyEditorModule.RegisterCustomClassLayout(UStructViewer::StaticClass()->GetFName(), FOnGetDetailCustomizationInstance::CreateStatic(&FStructViewerCustomization::MakeInstance));
+	ClassesToUnregisterOnShutdown.Add(UStructViewer::StaticClass()->GetFName());
 	// custom IK rig bone details
 	PropertyEditorModule.RegisterCustomClassLayout(UIKRigBoneDetails::StaticClass()->GetFName(), FOnGetDetailCustomizationInstance::CreateStatic(&FIKRigGenericDetailCustomization::MakeInstance));
 	ClassesToUnregisterOnShutdown.Add(UIKRigBoneDetails::StaticClass()->GetFName());
diff --git a/Engine/Plugins/Animation/IKRig/Source/IKRigEditor/Private/RigEditor/IKRigAutoFBIK.cpp b/Engine/Plugins/Animation/IKRig/Source/IKRigEditor/Private/RigEditor/IKRigAutoFBIK.cpp
--- a/Engine/Plugins/Animation/IKRig/Source/IKRigEditor/Private/RigEditor/IKRigAutoFBIK.cpp
+++ b/Engine/Plugins/Animation/IKRig/Source/IKRigEditor/Private/RigEditor/IKRigAutoFBIK.cpp
@@ -60,14 +60,14 @@ void FAutoFBIKCreator::CreateFBIKSetup(const UIKRigController& IKRigController,
 	}
 
 	// create IK solver and attach all the goals to it
-	const int32 SolverIndex = IKRigController.AddSolver(UIKRigFBIKSolver::StaticClass());
+	const int32 SolverIndex = IKRigController.AddSolver(FIKRigFullBodyIKSolver::StaticStruct());
 	for (const FName& GoalName : GoalNames)
 	{
 		IKRigController.ConnectGoalToSolver(GoalName, SolverIndex);	
 	}
 
 	// set the root of the solver
-	const bool bSetRoot = IKRigController.SetRootBone(CharacterizeResults.AutoRetargetDefinition.RetargetDefinition.RootBone, SolverIndex);
+	const bool bSetRoot = IKRigController.SetStartBone(CharacterizeResults.AutoRetargetDefinition.RetargetDefinition.RootBone, SolverIndex);
 	if (!bSetRoot)
 	{
 		Results.Outcome = EAutoFBIKResult::MissingRootBone;
@@ -75,11 +75,11 @@ void FAutoFBIKCreator::CreateFBIKSetup(const UIKRigController& IKRigController,
 	}
 
 	// update solver settings for retargeting
-	UIKRigFBIKSolver* Solver = CastChecked<UIKRigFBIKSolver>(IKRigController.GetSolverAtIndex(SolverIndex));
+	FIKRigFullBodyIKSolver* Solver = static_cast<FIKRigFullBodyIKSolver*>(IKRigController.GetSolverAtIndex(SolverIndex));
 	// set the root behavior to "free", allows pelvis motion only when needed to reach goals
-	Solver->RootBehavior = EPBIKRootBehavior::Free;
+	Solver->Settings.RootBehavior = EPBIKRootBehavior::Free;
 	// removing pull chain alpha on all goals "calms" the motion down, especially when retargeting arms
-	Solver->PullChainAlpha = 0.0f;
+	Solver->Settings.GlobalPullChainAlpha = 0.0f;
 
 	// assign bone settings from template
 	const FAutoCharacterizer& Characterizer = IKRigController.GetAutoCharacterizer();
@@ -108,27 +108,27 @@ void FAutoFBIKCreator::CreateFBIKSetup(const UIKRigController& IKRigController,
 		}
 
 		// get the bone settings for this bone for the FBIK solver (or add one if there is none)
-		auto GetOrAddBoneSettings = [SolverIndex, &IKRigController, BoneFullName]() -> UIKRig_FBIKBoneSettings*
+		auto GetOrAddBoneSettings = [SolverIndex, &IKRigController, BoneFullName]() -> FIKRigFBIKBoneSettings*
 		{
-			if (!IKRigController.GetBoneSettings(BoneFullName, SolverIndex))
+			if (!IKRigController.GetSettingsForBone(BoneFullName, SolverIndex))
 			{
 				IKRigController.AddBoneSetting(BoneFullName, SolverIndex);
 			}
 			
-			return CastChecked<UIKRig_FBIKBoneSettings>(IKRigController.GetBoneSettings(BoneFullName, SolverIndex));
+			return reinterpret_cast<FIKRigFBIKBoneSettings*>(IKRigController.GetSettingsForBone(BoneFullName, SolverIndex));
 		};
 
 		// apply rotational stiffness
 		if (BoneSettings.RotationStiffness > 0.f)
 		{
-			UIKRig_FBIKBoneSettings* Settings = GetOrAddBoneSettings();
+			FIKRigFBIKBoneSettings* Settings = GetOrAddBoneSettings();
 			Settings->RotationStiffness = FMath::Clamp(BoneSettings.RotationStiffness, 0.f, 1.0f);
 		}
 		
 		// apply preferred angles
 		if (BoneSettings.PreferredAxis != EPreferredAxis::None)
 		{
-			UIKRig_FBIKBoneSettings* Settings = GetOrAddBoneSettings();
+			FIKRigFBIKBoneSettings* Settings = GetOrAddBoneSettings();
 			Settings->bUsePreferredAngles = true;
 			Settings->PreferredAngles = BoneSettings.GetPreferredAxisAsAngles();
 		}
@@ -136,7 +136,7 @@ void FAutoFBIKCreator::CreateFBIKSetup(const UIKRigController& IKRigController,
 		// apply locked axes
 		if (BoneSettings.bIsHinge)
 		{
-			UIKRig_FBIKBoneSettings* Settings = GetOrAddBoneSettings();
+			FIKRigFBIKBoneSettings* Settings = GetOrAddBoneSettings();
 			BoneSettings.LockNonPreferredAxes(Settings->X, Settings->Y, Settings->Z);	
 		}
 	}
diff --git a/Engine/Plugins/Animation/IKRig/Source/IKRigEditor/Private/RigEditor/IKRigController.cpp b/Engine/Plugins/Animation/IKRig/Source/IKRigEditor/Private/RigEditor/IKRigController.cpp
--- a/Engine/Plugins/Animation/IKRig/Source/IKRigEditor/Private/RigEditor/IKRigController.cpp
+++ b/Engine/Plugins/Animation/IKRig/Source/IKRigEditor/Private/RigEditor/IKRigController.cpp
@@ -5,12 +5,13 @@
 #include "IKRigEditor.h"
 #include "Rig/IKRigDefinition.h"
 #include "Rig/IKRigProcessor.h"
-#include "Rig/Solvers/IKRigSolver.h"
+#include "Rig/Solvers/IKRigSolverBase.h"
 
 #include "Engine/SkeletalMesh.h"
 #include "ScopedTransaction.h"
 #include "RigEditor/IKRigAutoCharacterizer.h"
 #include "RigEditor/IKRigAutoFBIK.h"
+#include "RigEditor/IKRigDetailCustomizations.h"
 
 #include UE_INLINE_GENERATED_CPP_BY_NAME(IKRigController)
 
@@ -20,6 +21,8 @@ UIKRigController::UIKRigController()
 {
 	AutoCharacterizer = MakeUnique<FAutoCharacterizer>();
 	AutoFBIKCreator = MakeUnique<FAutoFBIKCreator>();
+
+	StructViewer = CreateDefaultSubobject<UStructViewer>(TEXT("SettingsViewer"));
 }
 
 UIKRigController* UIKRigController::GetController(const UIKRigDefinition* InIKRigDefinition)
@@ -51,7 +54,7 @@ TObjectPtr<UIKRigDefinition>& UIKRigController::GetAssetPtr()
 
 bool UIKRigController::AddBoneSetting(const FName BoneName, int32 SolverIndex) const
 {
-	check(Asset->Solvers.IsValidIndex(SolverIndex))
+	check(Asset->SolverStack.IsValidIndex(SolverIndex))
 
 	if (!CanAddBoneSetting(BoneName, SolverIndex))
 	{
@@ -61,15 +64,15 @@ bool UIKRigController::AddBoneSetting(const FName BoneName, int32 SolverIndex) c
 	
 	FScopedTransaction Transaction(LOCTEXT("AddBoneSetting_Label", "Add Bone Setting"));
 	FScopedReinitializeIKRig Reinitialize(this);
-	UIKRigSolver* Solver = Asset->Solvers[SolverIndex];
-	Solver->Modify();
-	Solver->AddBoneSetting(BoneName);
+	Asset->Modify();
+	FIKRigSolverBase* Solver = GetSolverAtIndex(SolverIndex);
+	Solver->AddSettingsToBone(BoneName);
 	return true;
 }
 
 bool UIKRigController::CanAddBoneSetting(const FName BoneName, int32 SolverIndex) const
 {
-	const UIKRigSolver* Solver = GetSolverAtIndex(SolverIndex);
+	const FIKRigSolverBase* Solver = GetSolverAtIndex(SolverIndex);
 	if (!Solver)
 	{
 		UE_LOG(LogIKRigEditor, Warning, TEXT("Solver does not exist at index: %d."), SolverIndex);
@@ -82,13 +85,13 @@ bool UIKRigController::CanAddBoneSetting(const FName BoneName, int32 SolverIndex
 		return false; // bone doesn't exist
 	}
 
-	if (!Asset->Solvers[SolverIndex]->UsesBoneSettings())
+	if (!Solver->UsesCustomBoneSettings())
 	{
 		UE_LOG(LogIKRigEditor, Warning, TEXT("Solver does not support bone settings."));
 		return false; // solver doesn't support per-bone settings
 	}
 
-	if (Asset->Solvers[SolverIndex]->GetBoneSetting(BoneName))
+	if (Solver->HasSettingsOnBone(BoneName))
 	{
 		UE_LOG(LogIKRigEditor, Warning, TEXT("Solver already has settings on bone, %s."), *BoneName.ToString());
 		return false; // solve already has settings on this bone
@@ -99,7 +102,7 @@ bool UIKRigController::CanAddBoneSetting(const FName BoneName, int32 SolverIndex
 
 bool UIKRigController::RemoveBoneSetting(const FName BoneName, int32 SolverIndex) const
 {
-	UIKRigSolver* Solver = GetSolverAtIndex(SolverIndex);
+	FIKRigSolverBase* Solver = GetSolverAtIndex(SolverIndex);
 	if (!Solver)
 	{
 		UE_LOG(LogIKRigEditor, Warning, TEXT("Solver does not exist at index: %d."), SolverIndex);
@@ -114,22 +117,23 @@ bool UIKRigController::RemoveBoneSetting(const FName BoneName, int32 SolverIndex
 
 	FScopedTransaction Transaction(LOCTEXT("RemoveBoneSetting_Label", "Remove Bone Setting"));
 	FScopedReinitializeIKRig Reinitialize(this);
-	Solver->Modify();
-	Solver->RemoveBoneSetting(BoneName);
+	Asset->Modify();
+	
+	Solver->RemoveSettingsOnBone(BoneName);
 
 	return true;
 }
 
 bool UIKRigController::CanRemoveBoneSetting(const FName BoneName, int32 SolverIndex) const
 {
-	const UIKRigSolver* Solver = GetSolverAtIndex(SolverIndex);
+	const FIKRigSolverBase* Solver = GetSolverAtIndex(SolverIndex);
 	if (!Solver)
 	{
 		UE_LOG(LogIKRigEditor, Warning, TEXT("Solver does not exist at index: %d."), SolverIndex);
 		return false; // solver doesn't exist
 	}
 
-	if (!Solver->UsesBoneSettings())
+	if (!Solver->UsesCustomBoneSettings())
 	{
 		UE_LOG(LogIKRigEditor, Warning, TEXT("Solver does not support bone settings."));
 		return false; // solver doesn't use bone settings
@@ -141,7 +145,7 @@ bool UIKRigController::CanRemoveBoneSetting(const FName BoneName, int32 SolverIn
 		return false; // bone doesn't exist
 	}
 
-	if (!Solver->GetBoneSetting(BoneName))
+	if (!Solver->HasSettingsOnBone(BoneName))
 	{
 		UE_LOG(LogIKRigEditor, Warning, TEXT("Solver does not have settings on bone, %s."), *BoneName.ToString());
 		return false; // solver doesn't have any settings for this bone
@@ -152,14 +156,18 @@ bool UIKRigController::CanRemoveBoneSetting(const FName BoneName, int32 SolverIn
 
 UObject* UIKRigController::GetBoneSettings(const FName BoneName, int32 SolverIndex) const
 {
-	const UIKRigSolver* Solver = GetSolverAtIndex(SolverIndex);
+	return nullptr;
+}
+
+FIKRigBoneSettingsBase* UIKRigController::GetSettingsForBone(const FName BoneName, int32 SolverIndex) const
+{
+	FIKRigSolverBase* Solver = GetSolverAtIndex(SolverIndex);
 	if (!Solver)
 	{
-		UE_LOG(LogIKRigEditor, Warning, TEXT("Solver does not exist at index: %d."), SolverIndex);
-		return nullptr; // solver doesn't exist
+		return nullptr;
 	}
-	
-	return Solver->GetBoneSetting(BoneName);
+
+	return Solver->GetBoneSettings(BoneName);
 }
 
 bool UIKRigController::DoesBoneHaveSettings(const FName BoneName) const
@@ -169,9 +177,10 @@ bool UIKRigController::DoesBoneHaveSettings(const FName BoneName) const
 		return false; // bone doesn't exist (do not spam here)
 	}
 	
-	for (UIKRigSolver* Solver : Asset->Solvers)
+	for (int32 SolverIndex=0; SolverIndex < Asset->SolverStack.Num(); ++SolverIndex)
 	{
-		if (UObject* BoneSetting = Solver->GetBoneSetting(BoneName))
+		const FIKRigSolverBase* Solver = GetSolverAtIndex(SolverIndex);
+		if (Solver->HasSettingsOnBone(BoneName))
 		{
 			return true;
 		}
@@ -596,6 +605,12 @@ bool UIKRigController::ValidateChain(
 	}
 }
 
+UStructViewer* UIKRigController::GetStructViewer()
+{
+	StructViewer->Reset();
+	return StructViewer;
+}
+
 FName UIKRigController::GetRetargetChainFromBone(const FName BoneName, const FIKRigSkeleton* OptionalSkeleton) const
 {
 	const FIKRigSkeleton& Skeleton = OptionalSkeleton ? *OptionalSkeleton : GetIKRigSkeleton();
@@ -686,6 +701,18 @@ const FIKRigSkeleton& UIKRigController::GetIKRigSkeleton() const
 	return Asset->Skeleton;
 }
 
+int32 UIKRigController::AddSolver(const FString InIKRigSolverType) const
+{
+	UScriptStruct* SolverType = FindObject<UScriptStruct>(nullptr, *InIKRigSolverType);
+	if (!SolverType)
+	{
+		UE_LOG(LogIKRigEditor, Warning, TEXT("Specified solver type was not found, %s."), *InIKRigSolverType);
+		return INDEX_NONE;
+	}
+
+	return AddSolver(SolverType);
+}
+
 bool UIKRigController::SetBoneExcluded(const FName BoneName, const bool bExclude) const
 {
 	// does bone exist?
@@ -738,29 +765,34 @@ FTransform UIKRigController::GetRefPoseTransformOfBone(const FName BoneName) con
 // SOLVERS
 //
 
-int32 UIKRigController::AddSolver(TSubclassOf<UIKRigSolver> InIKRigSolverClass) const
+int32 UIKRigController::AddSolver(UScriptStruct* InIKRigSolverType) const
 {
 	check(Asset)
 
-	if (!InIKRigSolverClass)
+	if (!InIKRigSolverType)
+	{
+		UE_LOG(LogIKRigEditor, Warning, TEXT("Could not add solver to IK Rig. No solver type specified."));
+		return INDEX_NONE;
+	}
+
+	if (!InIKRigSolverType->IsChildOf(FIKRigSolverBase::StaticStruct()))
 	{
-		UE_LOG(LogIKRigEditor, Warning, TEXT("Could not add solver to IK Rig. Invalid solver class specified."));
+		UE_LOG(LogIKRigEditor, Warning, TEXT("Could not add solver to IK Rig. Invalid solver type specified. Must be child of FIKRigSolverBase."));
 		return INDEX_NONE;
 	}
 
 	FScopedTransaction Transaction(LOCTEXT("AddSolver_Label", "Add Solver"));
 	FScopedReinitializeIKRig Reinitialize(this);
 	Asset->Modify();
-	UIKRigSolver* NewSolver = NewObject<UIKRigSolver>(Asset, InIKRigSolverClass, NAME_None, RF_Transactional);
-	check(NewSolver);
-	return Asset->Solvers.Add(NewSolver);
+	const int32 NewSolverIndex = Asset->SolverStack.Emplace(InIKRigSolverType);
+	return NewSolverIndex;
 }
 
 bool UIKRigController::RemoveSolver(const int32 SolverIndex) const
 {
 	check(Asset)
 	
-	if (!Asset->Solvers.IsValidIndex(SolverIndex))
+	if (!Asset->SolverStack.IsValidIndex(SolverIndex))
 	{
 		UE_LOG(LogIKRigEditor, Warning, TEXT("Solver not removed. Invalid index, %d."), SolverIndex);
 		return false;
@@ -769,20 +801,20 @@ bool UIKRigController::RemoveSolver(const int32 SolverIndex) const
 	FScopedTransaction Transaction(LOCTEXT("RemoveSolver_Label", "Remove Solver"));
 	FScopedReinitializeIKRig Reinitialize(this);
 	Asset->Modify();
-	Asset->Solvers.RemoveAt(SolverIndex);
+	Asset->SolverStack.RemoveAt(SolverIndex);
 	return true;
 }
 
 bool UIKRigController::MoveSolverInStack(int32 SolverToMoveIndex, int32 TargetSolverIndex) const
 {
-	if (!Asset->Solvers.IsValidIndex(SolverToMoveIndex))
+	if (!Asset->SolverStack.IsValidIndex(SolverToMoveIndex))
 	{
 		UE_LOG(LogIKRigEditor, Warning, TEXT("Solver not moved. Invalid source index, %d."), SolverToMoveIndex);
 		return false;
 	}
 
 	// allow a target 1 greater than the last element (for dragging below last element of list)
-	if (TargetSolverIndex > Asset->Solvers.Num() || TargetSolverIndex < 0)
+	if (TargetSolverIndex > Asset->SolverStack.Num() || TargetSolverIndex < 0)
 	{
 		UE_LOG(LogIKRigEditor, Warning, TEXT("Solver not moved. Invalid target index, %d."), TargetSolverIndex);
 		return false;
@@ -797,17 +829,16 @@ bool UIKRigController::MoveSolverInStack(int32 SolverToMoveIndex, int32 TargetSo
 	FScopedTransaction Transaction(LOCTEXT("ReorderSolver_Label", "Reorder Solvers"));
 	FScopedReinitializeIKRig Reinitialize(this);
 	Asset->Modify();
-	
-	UIKRigSolver* SolverToMove = Asset->Solvers[SolverToMoveIndex];
-	Asset->Solvers.RemoveAt(SolverToMoveIndex);
-	Asset->Solvers.Insert(SolverToMove, TargetSolverIndex);
+
+	Asset->SolverStack.Swap(SolverToMoveIndex, TargetSolverIndex);
 
 	return true;
 }
 
 bool UIKRigController::SetSolverEnabled(int32 SolverIndex, bool bIsEnabled) const
 {
-	if (!Asset->Solvers.IsValidIndex(SolverIndex))
+	FIKRigSolverBase* Solver = GetSolverAtIndex(SolverIndex);
+	if (!Solver)
 	{
 		UE_LOG(LogIKRigEditor, Warning, TEXT("Solver not enabled or disabled. Invalid index, %d."), SolverIndex);
 		return false;
@@ -815,59 +846,61 @@ bool UIKRigController::SetSolverEnabled(int32 SolverIndex, bool bIsEnabled) cons
 	
 	FScopedTransaction Transaction(LOCTEXT("SetSolverEndabled_Label", "Enable/Disable Solver"));
 	FScopedReinitializeIKRig Reinitialize(this);
-	UIKRigSolver* Solver = Asset->Solvers[SolverIndex];
-	Solver->Modify();
+	Asset->Modify();
 	Solver->SetEnabled(bIsEnabled);
 	return true;
 }
 
 bool UIKRigController::GetSolverEnabled(int32 SolverIndex) const
 {
-	if (!Asset->Solvers.IsValidIndex(SolverIndex))
+	const FIKRigSolverBase* Solver = GetSolverAtIndex(SolverIndex);
+	if (!Solver)
 	{
 		UE_LOG(LogIKRigEditor, Warning, TEXT("Invalid solver index, %d."), SolverIndex);
 		return false;
 	}
 
-	return Asset->Solvers[SolverIndex]->IsEnabled();
+	return Solver->IsEnabled();
 }
 
-bool UIKRigController::SetRootBone(const FName RootBoneName, int32 SolverIndex) const
+bool UIKRigController::SetStartBone(const FName RootBoneName, int32 SolverIndex) const
 {
-	if (!Asset->Solvers.IsValidIndex(SolverIndex))
+	FIKRigSolverBase* Solver = GetSolverAtIndex(SolverIndex);
+	if (!Solver)
 	{
-		UE_LOG(LogIKRigEditor, Warning, TEXT("Root bone not set. Invalid solver index, %d."), SolverIndex);
+		UE_LOG(LogIKRigEditor, Warning, TEXT("Start bone not set. Invalid solver index, %d."), SolverIndex);
 		return false; // solver doesn't exist
 	}
 
 	if (Asset->Skeleton.GetBoneIndexFromName(RootBoneName) == INDEX_NONE)
 	{
-		UE_LOG(LogIKRigEditor, Warning, TEXT("Root bone not set. Invalid bone specified, %s."), *RootBoneName.ToString());
+		UE_LOG(LogIKRigEditor, Warning, TEXT("Start bone not set. Invalid bone specified, %s."), *RootBoneName.ToString());
 		return false; // bone doesn't exist
 	}
 
-	FScopedTransaction Transaction(LOCTEXT("SetRootBone_Label", "Set Root Bone"));
+	FScopedTransaction Transaction(LOCTEXT("SetStartBone_Label", "Set Start Bone"));
 	FScopedReinitializeIKRig Reinitialize(this);
-	UIKRigSolver* Solver = Asset->Solvers[SolverIndex];
-	Solver->Modify();
-	Solver->SetRootBone(RootBoneName);
+	Asset->Modify();
+	Solver->SetStartBone(RootBoneName);
 	return true;
 }
 
-FName UIKRigController::GetRootBone(int32 SolverIndex) const
+FName UIKRigController::GetStartBone(int32 SolverIndex) const
 {
-	if (!Asset->Solvers.IsValidIndex(SolverIndex))
+	const FIKRigSolverBase* Solver = GetSolverAtIndex(SolverIndex);
+	if (!Solver)
 	{
 		UE_LOG(LogIKRigEditor, Warning, TEXT("Could not query root bone. Invalid solver index, %d."), SolverIndex);
 		return NAME_None; // solver doesn't exist
 	}
 
-	return Asset->Solvers[SolverIndex]->GetRootBone();
+	return Solver->GetStartBone();
 }
 
 bool UIKRigController::SetEndBone(const FName EndBoneName, int32 SolverIndex) const
 {
-	if (!Asset->Solvers.IsValidIndex(SolverIndex))
+	FIKRigSolverBase* Solver = GetSolverAtIndex(SolverIndex);
+	if (!Solver)
 	{
 		UE_LOG(LogIKRigEditor, Warning, TEXT("End bone not set. Invalid solver index, %d."), SolverIndex);
 		return false; // solver doesn't exist
@@ -879,7 +912,7 @@ bool UIKRigController::SetEndBone(const FName EndBoneName, int32 SolverIndex) co
 		return false; // bone doesn't exist
 	}
 
-	if (!Asset->Solvers[SolverIndex]->RequiresEndBone())
+	if (!Solver->UsesEndBone())
 	{
 		UE_LOG(LogIKRigEditor, Warning, TEXT("End bone not set. Specified solver does not support end bones."));
 		return false; //
@@ -887,63 +920,117 @@ bool UIKRigController::SetEndBone(const FName EndBoneName, int32 SolverIndex) co
 
 	FScopedTransaction Transaction(LOCTEXT("SetEndBone_Label", "Set End Bone"));
 	FScopedReinitializeIKRig Reinitialize(this);
-	UIKRigSolver* Solver = Asset->Solvers[SolverIndex];
-	Solver->Modify();
+	Asset->Modify();
 	Solver->SetEndBone(EndBoneName);
 	return true;
 }
 
 FName UIKRigController::GetEndBone(int32 SolverIndex) const
 {
-	if (!Asset->Solvers.IsValidIndex(SolverIndex))
+	const FIKRigSolverBase* Solver = GetSolverAtIndex(SolverIndex);
+	if (!Solver)
 	{
 		UE_LOG(LogIKRigEditor, Warning, TEXT("End bone not queried. Invalid solver index, %d."), SolverIndex);
 		return NAME_None; // solver doesn't exist
 	}
 
-	return Asset->Solvers[SolverIndex]->GetEndBone();
+	return Solver->GetEndBone();
+}
+
+UIKRigSolverControllerBase* UIKRigController::GetSolverController(int32 SolverIndex)
+{
+	FIKRigSolverBase* SolverAtIndex = GetSolverAtIndex(SolverIndex);
+	if (!SolverAtIndex)
+	{
+		return nullptr;
+	}
+
+	return SolverAtIndex->GetSolverController(this);
+}
+
+int32 UIKRigController::GetIndexOfSolver(UIKRigSolverControllerBase* InController)
+{
+	if (!InController)
+	{
+		return INDEX_NONE;
+	}
+	
+	for (int32 SolverIndex = 0; SolverIndex<Asset->SolverStack.Num(); SolverIndex++)
+	{
+		if (GetSolverAtIndex(SolverIndex) == InController->SolverToControl)
+		{
+			return SolverIndex;
+		}
+	}
+	
+	return INDEX_NONE;
 }
 
-const TArray<UIKRigSolver*>& UIKRigController::GetSolverArray() const
+const TArray<FIKRigSolverBase*> UIKRigController::GetSolverArray() const
 {
-	return Asset->Solvers;
+	TArray<FIKRigSolverBase*> SolverPointers;
+	for (int32 SolverIndex = 0; SolverIndex<Asset->SolverStack.Num(); SolverIndex++)
+	{
+		SolverPointers.Add(GetSolverAtIndex(SolverIndex));
+	}
+	return MoveTemp(SolverPointers);
 }
 
 FString UIKRigController::GetSolverUniqueName(int32 SolverIndex) const
 {
 	check(Asset)
-	if (!Asset->Solvers.IsValidIndex(SolverIndex))
+	const FIKRigSolverBase* Solver = GetSolverAtIndex(SolverIndex);
+	if (!Solver)
 	{
 		checkNoEntry()
 		return "";
 	}
 
-	return FString::FromInt(SolverIndex+1) + " - " + Asset->Solvers[SolverIndex]->GetNiceName().ToString();
+	return FString::FromInt(SolverIndex+1) + " - " + Solver->GetNiceName().ToString();
 		
 }
 
 int32 UIKRigController::GetNumSolvers() const
 {
 	check(Asset)
-	return Asset->Solvers.Num();
+	return Asset->SolverStack.Num();
 }
 
-UIKRigSolver* UIKRigController::GetSolverAtIndex(int32 Index) const
+FIKRigSolverBase* UIKRigController::GetSolverAtIndex(int32 Index) const
 {
 	check(Asset)
 
-	if (Asset->Solvers.IsValidIndex(Index))
+	if (Asset->SolverStack.IsValidIndex(Index))
 	{
-		return Asset->Solvers[Index];
+		return Asset->SolverStack[Index].GetMutablePtr<FIKRigSolverBase>();
 	}
 	
 	return nullptr;
 }
 
-int32 UIKRigController::GetIndexOfSolver(UIKRigSolver* Solver) const
+FInstancedStruct* UIKRigController::GetSolverStructAtIndex(int32 Index) const
 {
 	check(Asset)
-	return Asset->Solvers.Find(Solver);
+
+	if (Asset->SolverStack.IsValidIndex(Index))
+	{
+		return &Asset->SolverStack[Index];
+	}
+	
+	return nullptr;
+}
+
+int32 UIKRigController::GetIndexOfSolver(FIKRigSolverBase* Solver) const
+{
+	for (int32 SolverIndex = 0; SolverIndex<Asset->SolverStack.Num(); SolverIndex++)
+	{
+		if (GetSolverAtIndex(SolverIndex) == Solver)
+		{
+			return SolverIndex;
+		}
+	}
+	
+	return INDEX_NONE;
 }
 
 // -------------------------------------------------------
@@ -969,11 +1056,11 @@ FName UIKRigController::AddNewGoal(const FName GoalName, const FName BoneName) c
 	
 	FScopedTransaction Transaction(LOCTEXT("AddNewGoal_Label", "Add New Goal"));
 	FScopedReinitializeIKRig Reinitialize(this, true /*bGoalsChanged*/);
+	Asset->Modify();
 
 	UIKRigEffectorGoal* NewGoal = NewObject<UIKRigEffectorGoal>(Asset, UIKRigEffectorGoal::StaticClass(), NAME_None, RF_Transactional);
 	NewGoal->BoneName = BoneName;
 	NewGoal->GoalName = GoalName;
-	Asset->Modify();
 	Asset->Goals.Add(NewGoal);
 
 	// set initial transform
@@ -993,17 +1080,17 @@ bool UIKRigController::RemoveGoal(const FName GoalName) const
 
 	FScopedTransaction Transaction(LOCTEXT("RemoveGoal_Label", "Remove Goal"));
 	FScopedReinitializeIKRig Reinitialize(this, true /*bGoalsChanged*/);
-
+	Asset->Modify();
+	
 	// remove from all the solvers
 	const FName GoalToRemove = Asset->Goals[GoalIndex]->GoalName;
-	for (UIKRigSolver* Solver : Asset->Solvers)
+	for (int32 SolverIndex=0; SolverIndex<Asset->SolverStack.Num(); SolverIndex++)
 	{
-		Solver->Modify();
-		Solver->RemoveGoal(GoalToRemove);
+		FIKRigSolverBase* Solver = GetSolverAtIndex(SolverIndex);
+		Solver->OnGoalRemoved(GoalToRemove);
 	}
 
 	// remove from core system
-	Asset->Modify();
 	Asset->Goals.RemoveAt(GoalIndex);
 
 	// clean any retarget chains that might reference the missing goal
@@ -1039,9 +1126,9 @@ FName UIKRigController::RenameGoal(const FName OldName, const FName PotentialNew
 	
 	FScopedTransaction Transaction(LOCTEXT("RenameGoal_Label", "Rename Goal"));
 	FScopedReinitializeIKRig Reinitialize(this, true /*bGoalsChanged*/);
+	Asset->Modify();
 
 	// rename in core
-	Asset->Modify();
 	Asset->Goals[GoalIndex]->Modify();
 	Asset->Goals[GoalIndex]->GoalName = NewName;
 
@@ -1055,10 +1142,10 @@ FName UIKRigController::RenameGoal(const FName OldName, const FName PotentialNew
 	}
 
 	// rename in solvers
-	for (UIKRigSolver* Solver : Asset->Solvers)
+	for (int32 SolverIndex=0; SolverIndex<Asset->SolverStack.Num(); SolverIndex++)
 	{
-		Solver->Modify();
-		Solver->RenameGoal(OldName, NewName);
+		FIKRigSolverBase* Solver = GetSolverAtIndex(SolverIndex);
+		Solver->OnGoalRenamed(OldName, NewName);
 	}
 
 	return NewName;
@@ -1121,6 +1208,7 @@ bool UIKRigController::SetGoalBone(const FName GoalName, const FName NewBoneName
 	
 	FScopedTransaction Transaction(LOCTEXT("SetGoalBone_Label", "Set Goal Bone"));
 	FScopedReinitializeIKRig Reinitialize(this);
+	Asset->Modify();
 
 	// update goal's bone name and it's initial transform
 	TObjectPtr<UIKRigEffectorGoal> Goal = Asset->Goals[GoalIndex]; 
@@ -1129,10 +1217,10 @@ bool UIKRigController::SetGoalBone(const FName GoalName, const FName NewBoneName
 	Goal->InitialTransform = GetRefPoseTransformOfBone(NewBoneName);
 	
 	// update in solvers
-	for (UIKRigSolver* Solver : Asset->Solvers)
+	for (int32 SolverIndex=0; SolverIndex<Asset->SolverStack.Num(); SolverIndex++)
 	{
-		Solver->Modify();
-		Solver->SetGoalBone(GoalName, NewBoneName);
+		FIKRigSolverBase* Solver = GetSolverAtIndex(SolverIndex);
+		Solver->OnGoalMovedToDifferentBone(GoalName, NewBoneName);
 	}
 
 	return true;
@@ -1178,7 +1266,8 @@ bool UIKRigController::ConnectGoalToSolver(const FName GoalName, int32 SolverInd
 	}
 	
 	// can't add goal to a solver with an invalid index
-	if (!Asset->Solvers.IsValidIndex(SolverIndex))
+	FIKRigSolverBase* Solver = GetSolverAtIndex(SolverIndex);
+	if (!Solver)
 	{
 		UE_LOG(LogIKRigEditor, Warning, TEXT("Trying to connect Goal, %s to a unknown solver with index, %d."), *GoalName.ToString(), SolverIndex);
 		return false;
@@ -1186,9 +1275,7 @@ bool UIKRigController::ConnectGoalToSolver(const FName GoalName, int32 SolverInd
 
 	FScopedTransaction Transaction(LOCTEXT("ConnectGoalSolver_Label", "Connect Goal to Solver"));
 	FScopedReinitializeIKRig Reinitialize(this);
-	UIKRigSolver* Solver = Asset->Solvers[SolverIndex];
 	const UIKRigEffectorGoal& Goal = *GetGoal(GoalName);
-	Solver->Modify();
 	Solver->AddGoal(Goal);
 	return true;
 }
@@ -1203,7 +1290,8 @@ bool UIKRigController::DisconnectGoalFromSolver(const FName GoalToRemove, int32
 	}
 	
 	// can't remove goal from a solver with an invalid index
-	if (!Asset->Solvers.IsValidIndex(SolverIndex))
+	FIKRigSolverBase* Solver = GetSolverAtIndex(SolverIndex);
+	if (!Solver)
 	{
 		UE_LOG(LogIKRigEditor, Warning, TEXT("Trying to disconnect Goal, %s from an unknown solver with index, %d."), *GoalToRemove.ToString(), SolverIndex);
 		return false;
@@ -1211,33 +1299,35 @@ bool UIKRigController::DisconnectGoalFromSolver(const FName GoalToRemove, int32
 
     FScopedTransaction Transaction(LOCTEXT("DisconnectGoalSolver_Label", "Disconnect Goal from Solver"));
 	FScopedReinitializeIKRig Reinitialize(this);
-	UIKRigSolver* Solver = Asset->Solvers[SolverIndex];
-	Solver->Modify();
-	Solver->RemoveGoal(GoalToRemove);
+	Asset->Modify();
+	
+	Solver->OnGoalRemoved(GoalToRemove);
 	return true;
 }
 
 bool UIKRigController::IsGoalConnectedToSolver(const FName GoalName, int32 SolverIndex) const
 {
-	if (!Asset->Solvers.IsValidIndex(SolverIndex))
+	const FIKRigSolverBase* Solver = GetSolverAtIndex(SolverIndex);
+	if (!Solver)
 	{
 		return false;
 	}
 
-	return Asset->Solvers[SolverIndex]->IsGoalConnected(GoalName);
+	TSet<FName> RequiredGoals;
+	Solver->GetRequiredGoals(RequiredGoals);
+
+	return RequiredGoals.Contains(GoalName);
 }
 
 bool UIKRigController::IsGoalConnectedToAnySolver(const FName GoalName) const
 {
-	for (const TObjectPtr<UIKRigSolver>& Solver : Asset->Solvers)
+	TSet<FName> RequiredGoals;
+	for (int32 SolverIndex=0; SolverIndex<Asset->SolverStack.Num(); SolverIndex++)
 	{
-		if (Solver->IsGoalConnected(GoalName))
-		{
-			return true;
-		}
+		GetSolverAtIndex(SolverIndex)->GetRequiredGoals(RequiredGoals);
 	}
 
-	return false;
+	return RequiredGoals.Contains(GoalName);
 }
 
 const TArray<UIKRigEffectorGoal*>& UIKRigController::GetAllGoals() const
@@ -1258,18 +1348,18 @@ UIKRigEffectorGoal* UIKRigController::GetGoal(const FName GoalName) const
 
 UObject* UIKRigController::GetGoalSettingsForSolver(const FName GoalName, int32 SolverIndex) const
 {
-	const int32 GoalIndex = GetGoalIndex(GoalName);
-	if (GoalIndex == INDEX_NONE)
-	{
-		return nullptr; // no goal with that index
-	}
+	return nullptr;
+}
 
-	if (!Asset->Solvers.IsValidIndex(SolverIndex))
+bool UIKRigController::DoesSolverHaveCustomGoalSettings(const FName GoalName, int32 SolverIndex) const
+{
+	FIKRigSolverBase* Solver = GetSolverAtIndex(SolverIndex);
+	if (!Solver)
 	{
-		return nullptr; // no solver with that index
+		return false;
 	}
-
-	return Asset->Solvers[SolverIndex]->GetGoalSettings(GoalName);
+	
+	return Solver->GetGoalSettings(GoalName) != nullptr;
 }
 
 FTransform UIKRigController::GetGoalCurrentTransform(const FName GoalName) const
@@ -1368,8 +1458,8 @@ void UIKRigController::BroadcastNeedsReinitialized() const
 {
 	// initialize all solvers
 	const FIKRigSkeleton& IKRigSkeleton = GetIKRigSkeleton();
-	const TArray<UIKRigSolver*>& Solvers = GetSolverArray(); 
-	for (UIKRigSolver* Solver: Solvers)
+	const TArray<FIKRigSolverBase*> Solvers = GetSolverArray(); 
+	for (FIKRigSolverBase* Solver: Solvers)
 	{
 		Solver->Initialize(IKRigSkeleton);
 	}
diff --git a/Engine/Plugins/Animation/IKRig/Source/IKRigEditor/Private/RigEditor/IKRigDefinitionFactory.cpp b/Engine/Plugins/Animation/IKRig/Source/IKRigEditor/Private/RigEditor/IKRigDefinitionFactory.cpp
--- a/Engine/Plugins/Animation/IKRig/Source/IKRigEditor/Private/RigEditor/IKRigDefinitionFactory.cpp
+++ b/Engine/Plugins/Animation/IKRig/Source/IKRigEditor/Private/RigEditor/IKRigDefinitionFactory.cpp
@@ -1,6 +1,8 @@
 // Copyright Epic Games, Inc. All Rights Reserved.
 
 #include "RigEditor/IKRigDefinitionFactory.h"
+
+#include "AssetToolsModule.h"
 #include "Rig/IKRigDefinition.h"
 #include "AssetTypeCategories.h"
 #include "ContentBrowserModule.h"
@@ -54,5 +56,33 @@ FString UIKRigDefinitionFactory::GetDefaultNewAssetName() const
 {
 	return FString(TEXT("IK_NewIKRig"));
 }
+
+UIKRigDefinition* UIKRigDefinitionFactory::CreateNewIKRigAsset(const FString& InPackagePath, const FString& InAssetName)
+{
+	// add a trailing forward slash in case user forgot
+	FString DesiredPackagePath = InPackagePath;
+	if (!DesiredPackagePath.EndsWith(TEXT("/")))
+	{
+		DesiredPackagePath.Append(TEXT("/"));
+	}
+	DesiredPackagePath = DesiredPackagePath / InAssetName;
+	
+	// create unique package and asset names
+	FString UniqueAssetName = InAssetName;
+	FString UniquePackageName;
+	const FAssetToolsModule& AssetToolsModule = FModuleManager::LoadModuleChecked<FAssetToolsModule>("AssetTools");
+	AssetToolsModule.Get().CreateUniqueAssetName(DesiredPackagePath, TEXT(""), UniquePackageName, UniqueAssetName);
+	if (UniquePackageName.EndsWith(UniqueAssetName))
+	{
+		UniquePackageName = UniquePackageName.LeftChop(UniqueAssetName.Len() + 1);
+	}
+	
+	// create the new IK Rig asset
+	UIKRigDefinitionFactory* Factory = NewObject<UIKRigDefinitionFactory>();
+	UObject* NewAsset = AssetToolsModule.Get().CreateAsset(*UniqueAssetName, *UniquePackageName, nullptr, Factory);
+	UIKRigDefinition* IKRig = Cast<UIKRigDefinition>(NewAsset);
+	return IKRig;
+}
+
 #undef LOCTEXT_NAMESPACE
 
diff --git a/Engine/Plugins/Animation/IKRig/Source/IKRigEditor/Private/RigEditor/IKRigDetailCustomizations.cpp b/Engine/Plugins/Animation/IKRig/Source/IKRigEditor/Private/RigEditor/IKRigDetailCustomizations.cpp
--- a/Engine/Plugins/Animation/IKRig/Source/IKRigEditor/Private/RigEditor/IKRigDetailCustomizations.cpp
+++ b/Engine/Plugins/Animation/IKRig/Source/IKRigEditor/Private/RigEditor/IKRigDetailCustomizations.cpp
@@ -2,6 +2,7 @@
 
 #include "RigEditor/IKRigDetailCustomizations.h"
 
+#include "Editor.h"
 #include "DetailLayoutBuilder.h"
 #include "DetailCategoryBuilder.h"
 #include "DetailWidgetRow.h"
@@ -18,6 +19,48 @@ namespace IKRigDetailCustomizationsConstants
 	static const float ItemWidth = 125.0f;
 }
 
+TSharedRef<IDetailCustomization> FStructViewerCustomization::MakeInstance()
+{
+	return MakeShareable(new FStructViewerCustomization);
+}
+
+void FStructViewerCustomization::CustomizeDetails(IDetailLayoutBuilder& DetailBuilder)
+{
+	TArray<TWeakObjectPtr<UObject>> ObjectsBeingCustomized;
+	DetailBuilder.GetObjectsBeingCustomized(ObjectsBeingCustomized);
+	UStructViewer* StructViewer = Cast<UStructViewer>(ObjectsBeingCustomized[0]);
+	if (!ensure(StructViewer))
+	{
+		return;
+	}
+	if (!StructViewer->IsValid())
+	{
+		return;
+	}
+
+	// show the struct in the details panel
+	TSharedPtr<FStructOnScope> StructOnScope = StructViewer->GetStructOnScope();
+	FName StructTitle = StructViewer->GetStructName();
+	IDetailCategoryBuilder& CategoryBuilder = DetailBuilder.EditCategory(StructTitle);
+	TArray<IDetailPropertyRow*> OutRows;
+	TArray<TSharedPtr<IPropertyHandle>> Handles = CategoryBuilder.AddAllExternalStructureProperties(StructOnScope.ToSharedRef(), EPropertyLocation::Default, &OutRows);
+
+	// setup callbacks to begin/end an Undo transaction when property is edited
+	UObject* ObjectToTransact = StructViewer->GetStructOwner();
+	for (IDetailPropertyRow* Row : OutRows)
+	{
+		Row->GetPropertyHandle()->SetOnPropertyValuePreChange(FSimpleDelegate::CreateLambda([ObjectToTransact]()
+		{
+			GEditor->BeginTransaction(LOCTEXT("StructViewer", "Edited property."));
+			ObjectToTransact->Modify();
+		}));
+		Row->GetPropertyHandle()->SetOnPropertyValueChanged(FSimpleDelegate::CreateLambda([ObjectToTransact]()
+		{
+			GEditor->EndTransaction();
+		}));
+	}
+}
+
 const TArray<FText>& FIKRigGenericDetailCustomization::GetButtonLabels()
 {
 	static const TArray<FText> ButtonLabels
diff --git a/Engine/Plugins/Animation/IKRig/Source/IKRigEditor/Private/RigEditor/IKRigEditMode.cpp b/Engine/Plugins/Animation/IKRig/Source/IKRigEditor/Private/RigEditor/IKRigEditMode.cpp
--- a/Engine/Plugins/Animation/IKRig/Source/IKRigEditor/Private/RigEditor/IKRigEditMode.cpp
+++ b/Engine/Plugins/Animation/IKRig/Source/IKRigEditor/Private/RigEditor/IKRigEditMode.cpp
@@ -247,7 +247,7 @@ void FIKRigEditMode::GetBoneColors(
 			// record which bones in the skeleton are affected by this solver
 			const FIKRigSkeleton& Skeleton = Processor->GetSkeleton();
 			const UIKRigController* AssetController = Controller->AssetController;
-			if(const UIKRigSolver* SelectedSolver = AssetController->GetSolverAtIndex(SelectedSolvers[0].Get()->IndexInStack))
+			if(const FIKRigSolverBase* SelectedSolver = AssetController->GetSolverAtIndex(SelectedSolvers[0].Get()->IndexInStack))
 			{
 				for (int32 BoneIndex=0; BoneIndex < Skeleton.BoneNames.Num(); ++BoneIndex)
 				{
diff --git a/Engine/Plugins/Animation/IKRig/Source/IKRigEditor/Private/RigEditor/IKRigEditorController.cpp b/Engine/Plugins/Animation/IKRig/Source/IKRigEditor/Private/RigEditor/IKRigEditorController.cpp
--- a/Engine/Plugins/Animation/IKRig/Source/IKRigEditor/Private/RigEditor/IKRigEditorController.cpp
+++ b/Engine/Plugins/Animation/IKRig/Source/IKRigEditor/Private/RigEditor/IKRigEditorController.cpp
@@ -21,6 +21,7 @@
 #include "ScopedTransaction.h"
 #include "Framework/Notifications/NotificationManager.h"
 #include "RigEditor/IKRigAutoFBIK.h"
+#include "RigEditor/IKRigDetailCustomizations.h"
 #include "Widgets/Notifications/SNotificationList.h"
 
 #include UE_INLINE_GENERATED_CPP_BY_NAME(IKRigEditorController)
@@ -849,28 +850,33 @@ bool FIKRigEditorController::PromptToAddDefaultSolver() const
 	}
 
 	TArray<TSharedPtr<FIKRigSolverTypeAndName>> SolverTypes;
-	for (TObjectIterator<UClass> ClassIt; ClassIt; ++ClassIt)
+	for (TObjectIterator<UStruct> StructIt; StructIt; ++StructIt)
 	{
-		UClass* Class = *ClassIt;
-		if (!Class->IsNative())
+		UScriptStruct* Struct = Cast<UScriptStruct>(*StructIt);
+		if (!Struct)
+		{
+			continue;
+		}
+		
+		if (!Struct->IsNative())
 		{
 			continue;
 		}
 
-		if (!ClassIt->IsChildOf(UIKRigSolver::StaticClass()))
+		if (!Struct->IsChildOf(FIKRigSolverBase::StaticStruct()))
 		{
 			continue;
 		}
 
-		if (Class == UIKRigSolver::StaticClass())
+		if (Struct == FIKRigSolverBase::StaticStruct())
 		{
 			continue; // skip base class
 		}
 
-		const UIKRigSolver* SolverCDO = CastChecked<UIKRigSolver>(Class->ClassDefaultObject);
+		//const UIKRigSolver* SolverCDO = CastChecked<UIKRigSolver>(Class->ClassDefaultObject);
 		TSharedPtr<FIKRigSolverTypeAndName> SolverType = MakeShared<FIKRigSolverTypeAndName>();
-		SolverType->NiceName = SolverCDO->GetNiceName();
-		SolverType->SolverType = TSubclassOf<UIKRigSolver>(Class);
+		SolverType->NiceName = FText::FromString(Struct->GetName());
+		SolverType->SolverType = Struct;
 		SolverTypes.Add(SolverType);
 	}
 
@@ -878,7 +884,7 @@ bool FIKRigEditorController::PromptToAddDefaultSolver() const
 	TSharedPtr<FIKRigSolverTypeAndName> SelectedSolver = SolverTypes[0];
 	for (TSharedPtr<FIKRigSolverTypeAndName>& SolverType :SolverTypes)
 	{
-		if (SolverType->SolverType == UIKRigFBIKSolver::StaticClass())
+		if (SolverType->SolverType == FIKRigFullBodyIKSolver::StaticStruct())
 		{
 			SelectedSolver = SolverType;
 			break;
@@ -951,10 +957,13 @@ void FIKRigEditorController::ShowDetailsForBone(const FName BoneName) const
 
 void FIKRigEditorController::ShowDetailsForBoneSettings(const FName& BoneName, int32 SolverIndex) const
 {
-	if (UObject* BoneSettings = AssetController->GetBoneSettings(BoneName, SolverIndex))
-	{
-		DetailsView->SetObject(BoneSettings);
-	}
+	FIKRigSolverBase* Solver = AssetController->GetSolverAtIndex(SolverIndex);
+	UStructViewer* StructViewer = AssetController->GetStructViewer();
+	StructViewer->SetStructToView(
+		Solver->GetBoneSettingsType(),
+		reinterpret_cast<uint8*>(Solver->GetBoneSettings(BoneName)),
+		AssetController->GetAsset());
+	ShowStructViewerInDetailsView(StructViewer);
 }
 
 void FIKRigEditorController::ShowDetailsForGoal(const FName& GoalName) const
@@ -964,19 +973,24 @@ void FIKRigEditorController::ShowDetailsForGoal(const FName& GoalName) const
 
 void FIKRigEditorController::ShowDetailsForGoalSettings(const FName GoalName, const int32 SolverIndex) const
 {
-	// get solver that owns this effector
-	if (const UIKRigSolver* SolverWithEffector = AssetController->GetSolverAtIndex(SolverIndex))
-	{
-		if (UObject* EffectorSettings = SolverWithEffector->GetGoalSettings(GoalName))
-		{
-			DetailsView->SetObject(EffectorSettings);
-		}
-	}
+	FIKRigSolverBase* Solver = AssetController->GetSolverAtIndex(SolverIndex);
+	UStructViewer* StructViewer = AssetController->GetStructViewer();
+	StructViewer->SetStructToView(
+		Solver->GetGoalSettingsType(),
+		reinterpret_cast<uint8*>(Solver->GetGoalSettings(GoalName)),
+		AssetController->GetAsset());
+	ShowStructViewerInDetailsView(StructViewer);
 }
 
 void FIKRigEditorController::ShowDetailsForSolver(const int32 SolverIndex) const
 {
-	DetailsView->SetObject(AssetController->GetSolverAtIndex(SolverIndex));
+	FIKRigSolverBase* Solver = AssetController->GetSolverAtIndex(SolverIndex);
+	UStructViewer* StructViewer = AssetController->GetStructViewer();
+	StructViewer->SetStructToView(
+		Solver->GetSolverSettingsType(),
+		reinterpret_cast<uint8*>(Solver->GetSolverSettings()),
+		AssetController->GetAsset());
+	ShowStructViewerInDetailsView(StructViewer);
 }
 
 void FIKRigEditorController::ShowAssetDetails() const
@@ -992,31 +1006,8 @@ void FIKRigEditorController::ShowDetailsForElements(const TArray<TSharedPtr<FIKR
 		return;
 	}
 
+	// edit the last selected element
 	const TSharedPtr<FIKRigTreeElement>& LastItem = InItems.Last();
-
-	// check is the items are all of the same type
-	const bool bContainsSeveralTypes = InItems.ContainsByPredicate( [LastItem](const TSharedPtr<FIKRigTreeElement>& Item)
-	{
-		return Item->ElementType != LastItem->ElementType;
-	});
-
-	// if all elements are similar then treat them once
-	if (!bContainsSeveralTypes)
-	{
-		TArray<TWeakObjectPtr<>> Objects;
-		for (const TSharedPtr<FIKRigTreeElement>& Item: InItems)
-		{
-			TWeakObjectPtr<> Object = Item->GetObject();
-			if (Object.IsValid())
-			{
-				Objects.Add(Object);
-			}
-		}
-		DetailsView->SetObjects(Objects);
-		return;
-	}
-
-	// fallback to the last selected element
 	switch (LastItem->ElementType)
 	{
 	case IKRigTreeElementType::BONE:
@@ -1041,6 +1032,12 @@ void FIKRigEditorController::ShowDetailsForElements(const TArray<TSharedPtr<FIKR
 	}
 }
 
+void FIKRigEditorController::ShowStructViewerInDetailsView(UStructViewer* StructViewer) const
+{
+	constexpr bool bForceRefresh = true; // must forcibly refresh because struct viewer is a recycled UObject
+	DetailsView->SetObject(StructViewer, bForceRefresh);
+}
+
 void FIKRigEditorController::OnFinishedChangingDetails(const FPropertyChangedEvent& PropertyChangedEvent) const
 {
 	const bool bPreviewMeshChanged = PropertyChangedEvent.GetPropertyName() == UIKRigDefinition::GetPreviewMeshPropertyName();
diff --git a/Engine/Plugins/Animation/IKRig/Source/IKRigEditor/Private/RigEditor/IKRigSkeletonCommands.cpp b/Engine/Plugins/Animation/IKRig/Source/IKRigEditor/Private/RigEditor/IKRigSkeletonCommands.cpp
--- a/Engine/Plugins/Animation/IKRig/Source/IKRigEditor/Private/RigEditor/IKRigSkeletonCommands.cpp
+++ b/Engine/Plugins/Animation/IKRig/Source/IKRigEditor/Private/RigEditor/IKRigSkeletonCommands.cpp
@@ -10,7 +10,7 @@ void FIKRigSkeletonCommands::RegisterCommands()
 	UI_COMMAND(DeleteElement, "Delete Goal or Effector", "Delete the selected IK Goal, Effector or Bone Settings.", EUserInterfaceActionType::Button, FInputChord(EKeys::Delete), FInputChord(EKeys::BackSpace));
 	UI_COMMAND(ConnectGoalToSolver, "Connect Goal to Selected Solver", "Connect this goal to the selected solver.", EUserInterfaceActionType::Button, FInputChord());
 	UI_COMMAND(DisconnectGoalFromSolver, "Disconnect Goal from Selected Solver", "Disconnect this goal from the selected solver.", EUserInterfaceActionType::Button, FInputChord());
-	UI_COMMAND(SetRootBoneOnSolver, "Set Root Bone on Selected Solver", "Set the Root Bone setting on the selected solver.", EUserInterfaceActionType::Button, FInputChord());
+	UI_COMMAND(SetStartBoneOnSolver, "Set Start Bone on Selected Solver", "Set the Start Bone on the selected solver.", EUserInterfaceActionType::Button, FInputChord());
 	UI_COMMAND(SetEndBoneOnSolver, "Set End Bone on Selected Solver", "Set End Bone on selected solver.", EUserInterfaceActionType::Button, FInputChord());
 	UI_COMMAND(AddBoneSettings, "Add Settings to Selected Bone", "Apply settings to the selected bone for the selected solver (defined per-solver for limits, stiffness etc).", EUserInterfaceActionType::Button, FInputChord());
 	UI_COMMAND(RemoveBoneSettings, "Remove Settings on Selected Bone", "Remove all settings on the selected bone for the selected solver.", EUserInterfaceActionType::Button, FInputChord());
diff --git a/Engine/Plugins/Animation/IKRig/Source/IKRigEditor/Private/RigEditor/SIKRigHierarchy.cpp b/Engine/Plugins/Animation/IKRig/Source/IKRigEditor/Private/RigEditor/SIKRigHierarchy.cpp
--- a/Engine/Plugins/Animation/IKRig/Source/IKRigEditor/Private/RigEditor/SIKRigHierarchy.cpp
+++ b/Engine/Plugins/Animation/IKRig/Source/IKRigEditor/Private/RigEditor/SIKRigHierarchy.cpp
@@ -3,7 +3,7 @@
 #include "RigEditor/SIKRigHierarchy.h"
 
 #include "Rig/IKRigProcessor.h"
-#include "Rig/Solvers/IKRigSolver.h"
+#include "Rig/Solvers/IKRigSolverBase.h"
 #include "Engine/SkeletalMesh.h"
 #include "IPersonaToolkit.h"
 #include "SKismetInspector.h"
@@ -39,19 +39,20 @@ void FIKRigTreeElement::RequestRename()
 	OnRenameRequested.ExecuteIfBound();
 }
 
-TWeakObjectPtr< UObject > FIKRigTreeElement::GetObject() const
+void FIKRigTreeElement::GetStruct(const UScriptStruct*& OutType, uint8*& OutMemory) const
 {
 	const TSharedPtr<FIKRigEditorController>& Controller = EditorController.Pin();
 	if (!Controller.IsValid())
 	{
-		return nullptr; 
+		return; 
 	}
 
 	if (Controller->AssetController == nullptr)
 	{
-		return nullptr;
+		return;
 	}
 
+	/* TODO
 	switch(ElementType)
 	{
 	case IKRigTreeElementType::BONE:
@@ -65,13 +66,13 @@ TWeakObjectPtr< UObject > FIKRigTreeElement::GetObject() const
 	case IKRigTreeElementType::GOAL:
 		return Controller->AssetController->GetGoal(GoalName);
 	case IKRigTreeElementType::SOLVERGOAL:
-		if (const UIKRigSolver* SolverWithEffector = Controller->AssetController->GetSolverAtIndex(EffectorIndex))
+		if (const FIKRigSolverBase* SolverWithEffector = Controller->AssetController->GetSolverAtIndex(EffectorIndex))
 		{
-			return SolverWithEffector->GetGoalSettings(EffectorGoalName);
+			return nullptr; // TODO SolverWithEffector->GetGoalSettings(EffectorGoalName);
 		}
 	default:
 		return nullptr;
-	}
+	}*/
 }
 
 FName FIKRigTreeElement::GetChainName() const
@@ -659,13 +660,13 @@ bool SIKRigHierarchy::IsElementConnectedToSolver(TWeakPtr<FIKRigTreeElement> InT
 {
 	const UIKRigController* AssetController = EditorController.Pin()->AssetController;
 	const TSharedPtr<FIKRigTreeElement> TreeElement = InTreeElement.Pin();
+	const FIKRigSolverBase* Solver = AssetController->GetSolverAtIndex(SolverIndex);
 	
-	if (!AssetController->GetSolverArray().IsValidIndex(SolverIndex))
+	if (!Solver)
 	{
 		return false; // not a valid solver index
 	}
-
-	const UIKRigSolver* Solver = AssetController->GetSolverAtIndex(SolverIndex);
+	
 	if (TreeElement->ElementType == IKRigTreeElementType::BONE)
 	{
 		// is this bone affected by this solver?
@@ -675,7 +676,7 @@ bool SIKRigHierarchy::IsElementConnectedToSolver(TWeakPtr<FIKRigTreeElement> InT
 	if (TreeElement->ElementType == IKRigTreeElementType::BONE_SETTINGS)
 	{
 		// is this bone setting belonging to the solver?
-		return (Solver->GetBoneSetting(TreeElement->BoneSettingBoneName) != nullptr);
+		return Solver->HasSettingsOnBone(TreeElement->BoneSettingBoneName);
 	}
 
 	if (TreeElement->ElementType == IKRigTreeElementType::GOAL)
@@ -779,7 +780,7 @@ void SIKRigHierarchy::BindCommands()
         FExecuteAction::CreateSP(this, &SIKRigHierarchy::HandleDisconnectGoalFromSolver),
         FCanExecuteAction::CreateSP(this, &SIKRigHierarchy::CanDisconnectGoalFromSolvers));
 
-	CommandList->MapAction(Commands.SetRootBoneOnSolver,
+	CommandList->MapAction(Commands.SetStartBoneOnSolver,
         FExecuteAction::CreateSP(this, &SIKRigHierarchy::HandleSetRootBoneOnSolvers),
         FCanExecuteAction::CreateSP(this, &SIKRigHierarchy::CanSetRootBoneOnSolvers));
 
@@ -847,7 +848,7 @@ void SIKRigHierarchy::FillContextMenu(FMenuBuilder& MenuBuilder)
 	MenuBuilder.EndSection();
 
 	MenuBuilder.BeginSection("BoneSettings", LOCTEXT("BoneSettingsOperations", "Bone Settings"));
-	MenuBuilder.AddMenuEntry(Actions.SetRootBoneOnSolver);
+	MenuBuilder.AddMenuEntry(Actions.SetStartBoneOnSolver);
 	MenuBuilder.AddMenuEntry(Actions.SetEndBoneOnSolver);
 	MenuBuilder.AddMenuEntry(Actions.AddBoneSettings);
 	MenuBuilder.AddMenuEntry(Actions.RemoveBoneSettings);
@@ -1059,7 +1060,7 @@ void SIKRigHierarchy::HandleSetRootBoneOnSolvers()
 	int32 SolverToShow = 0;
 	for (const TSharedPtr<FSolverStackElement>& Solver : SelectedSolvers)
 	{
-		AssetController->SetRootBone(RootBoneName, Solver->IndexInStack);
+		AssetController->SetStartBone(RootBoneName, Solver->IndexInStack);
 		SolverToShow = Solver->IndexInStack;
 	}
 
@@ -1092,7 +1093,7 @@ bool SIKRigHierarchy::CanSetRootBoneOnSolvers()
 	Controller->GetSelectedSolvers(SelectedSolvers);
 	for (const TSharedPtr<FSolverStackElement>& Solver : SelectedSolvers)
 	{
-		if (AssetController->GetSolverAtIndex(Solver->IndexInStack)->RequiresRootBone())
+		if (AssetController->GetSolverAtIndex(Solver->IndexInStack)->UsesStartBone())
 		{
 			return true;
 		}
@@ -1154,7 +1155,7 @@ bool SIKRigHierarchy::CanSetEndBoneOnSolvers() const
 	Controller->GetSelectedSolvers(SelectedSolvers);
 	for (const TSharedPtr<FSolverStackElement>& Solver : SelectedSolvers)
 	{
-		if (AssetController->GetSolverAtIndex(Solver->IndexInStack)->RequiresEndBone())
+		if (AssetController->GetSolverAtIndex(Solver->IndexInStack)->UsesEndBone())
 		{
 			return true;
 		}
@@ -1719,7 +1720,7 @@ void SIKRigHierarchy::RefreshTreeView(bool IsInitialSetup /*=false*/)
 	const TArray<UIKRigEffectorGoal*>& Goals = AssetController->GetAllGoals();
 	
 	// get all solvers
-	const TArray<UIKRigSolver*>& Solvers = AssetController->GetSolverArray();
+	const TArray<FIKRigSolverBase*> Solvers = AssetController->GetSolverArray();
 	// record bone element indices
 	TMap<FName, int32> BoneTreeElementIndices;
 
@@ -1770,7 +1771,7 @@ void SIKRigHierarchy::RefreshTreeView(bool IsInitialSetup /*=false*/)
 		// create all "Bone Setting" tree elements for this bone
 		for (int32 SolverIndex=0; SolverIndex<Solvers.Num(); ++SolverIndex)
 		{
-			if (!Solvers[SolverIndex]->GetBoneSetting(BoneName))
+			if (!Solvers[SolverIndex]->HasSettingsOnBone(BoneName))
 			{
 				continue;
 			}
@@ -1832,8 +1833,7 @@ void SIKRigHierarchy::RefreshTreeView(bool IsInitialSetup /*=false*/)
 			// add all solver settings connected to this goal
 			for (int32 SolverIndex=0; SolverIndex<Solvers.Num(); ++SolverIndex)
 			{
-				UObject* GoalSettings = AssetController->GetGoalSettingsForSolver(Goal->GoalName, SolverIndex);
-				if (!GoalSettings)
+				if (!AssetController->DoesSolverHaveCustomGoalSettings(Goal->GoalName, SolverIndex))
 				{
 					continue;
 				}
diff --git a/Engine/Plugins/Animation/IKRig/Source/IKRigEditor/Private/RigEditor/SIKRigRetargetChainList.cpp b/Engine/Plugins/Animation/IKRig/Source/IKRigEditor/Private/RigEditor/SIKRigRetargetChainList.cpp
--- a/Engine/Plugins/Animation/IKRig/Source/IKRigEditor/Private/RigEditor/SIKRigRetargetChainList.cpp
+++ b/Engine/Plugins/Animation/IKRig/Source/IKRigEditor/Private/RigEditor/SIKRigRetargetChainList.cpp
@@ -768,7 +768,7 @@ void SIKRigRetargetChainList::MirrorSelectedChains() const
 		if (GoalOnOldChainName != NAME_None && GoalOnNewChain)
 		{
 			// connect to the same solvers
-			const TArray<UIKRigSolver*>& AllSolvers =  AssetController->GetSolverArray();
+			const TArray<FIKRigSolverBase*> AllSolvers =  AssetController->GetSolverArray();
 			for (int32 SolverIndex=0; SolverIndex<AllSolvers.Num(); ++SolverIndex)
 			{
 				if (AssetController->IsGoalConnectedToSolver(GoalOnOldChainName, SolverIndex))
diff --git a/Engine/Plugins/Animation/IKRig/Source/IKRigEditor/Private/RigEditor/SIKRigSolverStack.cpp b/Engine/Plugins/Animation/IKRig/Source/IKRigEditor/Private/RigEditor/SIKRigSolverStack.cpp
--- a/Engine/Plugins/Animation/IKRig/Source/IKRigEditor/Private/RigEditor/SIKRigSolverStack.cpp
+++ b/Engine/Plugins/Animation/IKRig/Source/IKRigEditor/Private/RigEditor/SIKRigSolverStack.cpp
@@ -5,7 +5,7 @@
 #include "RigEditor/IKRigEditorStyle.h"
 #include "RigEditor/IKRigToolkit.h"
 #include "RigEditor/IKRigEditorController.h"
-#include "Rig/Solvers/IKRigSolver.h"
+#include "Rig/Solvers/IKRigSolverBase.h"
 
 #include "Framework/Commands/UICommandList.h"
 #include "Framework/MultiBox/MultiBoxBuilder.h"
@@ -160,7 +160,7 @@ void SIKRigSolverStackItem::Construct(
 bool SIKRigSolverStackItem::GetWarningMessage(FText& Message) const
 {
 	constexpr bool bFromAsset = false;
-	if (const UIKRigSolver* Solver = GetSolver(bFromAsset))
+	if (const FIKRigSolverBase* Solver = GetSolver(bFromAsset))
 	{
 		return Solver->GetWarningMessage(Message);
 	}
@@ -171,7 +171,7 @@ bool SIKRigSolverStackItem::GetWarningMessage(FText& Message) const
 bool SIKRigSolverStackItem::IsSolverEnabled() const
 {
 	constexpr bool bFromAsset = true;
-	if (const UIKRigSolver* Solver = GetSolver(bFromAsset))
+	if (const FIKRigSolverBase* Solver = GetSolver(bFromAsset))
 	{
 		return Solver->IsEnabled();
 	}
@@ -179,7 +179,7 @@ bool SIKRigSolverStackItem::IsSolverEnabled() const
 	return false;
 }
 
-UIKRigSolver* SIKRigSolverStackItem::GetSolver(const bool bFromAsset) const
+const FIKRigSolverBase* SIKRigSolverStackItem::GetSolver(const bool bFromAsset) const
 {
 	if (!(StackElement.IsValid() && SolverStack.IsValid()))
 	{
@@ -206,7 +206,7 @@ UIKRigSolver* SIKRigSolverStackItem::GetSolver(const bool bFromAsset) const
 		{
 			if (Processor->GetSolvers().IsValidIndex(SolverIndex))
 			{
-				return Processor->GetSolvers()[SolverIndex];
+				return Processor->GetSolvers()[SolverIndex].Get();
 			}
 		}
 	}
@@ -312,16 +312,23 @@ void SIKRigSolverStack::BuildAddNewMenu(FMenuBuilder& MenuBuilder)
 	MenuBuilder.BeginSection("AddNewSolver", LOCTEXT("AddOperations", "Add New Solver"));
 	
 	// add menu option to create each solver type
-	TArray<UClass*> SolverClasses;
-	for(TObjectIterator<UClass> ClassIt; ClassIt; ++ClassIt)
+	TArray<UStruct*> SolverTypes;
+	const UStruct* BaseStruct = FIKRigSolverBase::StaticStruct();
+	for(TObjectIterator<UStruct> StructIt; StructIt; ++StructIt)
 	{
-		UClass* Class = *ClassIt;
+		UStruct* Struct = *StructIt;
 
-		if(Class->IsChildOf(UIKRigSolver::StaticClass()) && !Class->HasAnyClassFlags(CLASS_Abstract))
+		if(!Struct->IsChildOf(BaseStruct) || Struct == BaseStruct)
 		{
-			UIKRigSolver* SolverCDO = Cast<UIKRigSolver>(Class->GetDefaultObject());
-			FUIAction Action = FUIAction( FExecuteAction::CreateSP(this, &SIKRigSolverStack::AddNewSolver, Class));
-			MenuBuilder.AddMenuEntry(FText::FromString(SolverCDO->GetNiceName().ToString()), FText::GetEmpty(), FSlateIcon(), Action);
+			continue;
+		}
+
+		if (UScriptStruct* ScriptStruct = Cast<UScriptStruct>(Struct))
+		{
+			//UIKRigSolver* SolverCDO = Cast<UIKRigSolver>(Class->GetDefaultObject());
+			FString TypeName = ScriptStruct->GetName();
+			FUIAction Action = FUIAction( FExecuteAction::CreateSP(this, &SIKRigSolverStack::AddNewSolver, ScriptStruct));
+			MenuBuilder.AddMenuEntry(FText::FromString(TypeName), FText::GetEmpty(), FSlateIcon(), Action);
 		}
 	}
 	
@@ -347,7 +354,7 @@ bool SIKRigSolverStack::IsAddSolverEnabled() const
 	return false;
 }
 
-void SIKRigSolverStack::AddNewSolver(UClass* Class)
+void SIKRigSolverStack::AddNewSolver(UScriptStruct* SolverType)
 {
 	const TSharedPtr<FIKRigEditorController> Controller = EditorController.Pin();
 	if (!Controller.IsValid())
@@ -362,7 +369,7 @@ void SIKRigSolverStack::AddNewSolver(UClass* Class)
 	}
 	
 	// add the solver
-	const int32 NewSolverIndex = AssetController->AddSolver(Class);
+	const int32 NewSolverIndex = AssetController->AddSolver(SolverType);
 	// update stack view
 	RefreshStackView();
 	Controller->RefreshTreeView(); // updates solver indices in tree items
@@ -415,7 +422,7 @@ void SIKRigSolverStack::RefreshStackView()
 	const int32 NumSolvers = AssetController->GetNumSolvers();
 	for (int32 i=0; i<NumSolvers; ++i)
 	{
-		const UIKRigSolver* Solver = AssetController->GetSolverAtIndex(i);
+		const FIKRigSolverBase* Solver = AssetController->GetSolverAtIndex(i);
 		const FText DisplayName = Solver ? FText::FromString(AssetController->GetSolverUniqueName(i)) : UnknownSolverTxt;
 		TSharedPtr<FSolverStackElement> SolverItem = FSolverStackElement::Make(DisplayName, i);
 		ListViewItems.Add(SolverItem);
diff --git a/Engine/Plugins/Animation/IKRig/Source/IKRigEditor/Public/RetargetEditor/IKRetargeterController.h b/Engine/Plugins/Animation/IKRig/Source/IKRigEditor/Public/RetargetEditor/IKRetargeterController.h
--- a/Engine/Plugins/Animation/IKRig/Source/IKRigEditor/Public/RetargetEditor/IKRetargeterController.h
+++ b/Engine/Plugins/Animation/IKRig/Source/IKRigEditor/Public/RetargetEditor/IKRetargeterController.h
@@ -31,7 +31,7 @@ enum class EAutoMapChainType : uint8
 
 // A stateless singleton (1-per-asset) class used to make modifications to a UIKRetargeter asset.
 // Use UIKRetargeter.GetController() to get the controller for the asset you want to modify.  
-UCLASS(config = Engine, hidecategories = UObject)
+UCLASS(BlueprintType, hidecategories = UObject)
 class IKRIGEDITOR_API UIKRetargeterController : public UObject
 {
 	GENERATED_BODY()
@@ -58,15 +58,15 @@ private:
 public:
 	
 	/** Use this to get the controller for the given retargeter asset */
-	UFUNCTION(BlueprintCallable, Category=IKRetargeter, meta = (BlueprintThreadSafe))
+	UFUNCTION(BlueprintCallable, Category=IKRetargeter, meta = (BlueprintThreadSafe, DisplayName = "Get IK Retarget Controller"))
 	static UIKRetargeterController* GetController(const UIKRetargeter* InRetargeterAsset);
 	
 	// Set the IK Rig to use as the source or target (to copy animation FROM/TO) 
-	UFUNCTION(BlueprintCallable, BlueprintPure=false, Category=IKRetargeter, meta = (BlueprintThreadSafe))
+	UFUNCTION(BlueprintCallable, BlueprintPure=false, Category=IKRetargeter, meta = (BlueprintThreadSafe, DisplayName = "Set IK Rig"))
 	void SetIKRig(const ERetargetSourceOrTarget SourceOrTarget, UIKRigDefinition* IKRig) const;
 	
 	// Get either source or target IK Rig 
-	UFUNCTION(BlueprintCallable, Category=IKRetargeter, meta = (BlueprintThreadSafe))
+	UFUNCTION(BlueprintCallable, Category=IKRetargeter, meta = (BlueprintThreadSafe, DisplayName = "Get IK Rig"))
 	const UIKRigDefinition* GetIKRig(const ERetargetSourceOrTarget SourceOrTarget) const;
 
 	// Set the preview skeletal mesh for either source or target
@@ -220,11 +220,11 @@ private:
 	// convenience to get chain settings UObject by name
 	URetargetChainSettings* GetChainSettings(const FName& TargetChainName) const;
 
+public:
+
 	//
 	// RETARGET POSE PUBLIC/SCRIPTING API
 	//
-
-public:
 	
 	// Add new retarget pose. Returns the name of the new retarget pose.
 	UFUNCTION(BlueprintCallable, BlueprintPure=false, Category=IKRetargeter, meta =(BlueprintThreadSafe))
diff --git a/Engine/Plugins/Animation/IKRig/Source/IKRigEditor/Public/RigEditor/IKRigAnimInstance.h b/Engine/Plugins/Animation/IKRig/Source/IKRigEditor/Public/RigEditor/IKRigAnimInstance.h
--- a/Engine/Plugins/Animation/IKRig/Source/IKRigEditor/Public/RigEditor/IKRigAnimInstance.h
+++ b/Engine/Plugins/Animation/IKRig/Source/IKRigEditor/Public/RigEditor/IKRigAnimInstance.h
@@ -10,7 +10,7 @@
 
 class UIKRigDefinition;
 
-UCLASS(transient, NotBlueprintable)
+UCLASS(transient, NotBlueprintType, NotBlueprintable)
 class UIKRigAnimInstance : public UAnimPreviewInstance
 {
 	GENERATED_UCLASS_BODY()
diff --git a/Engine/Plugins/Animation/IKRig/Source/IKRigEditor/Public/RigEditor/IKRigController.h b/Engine/Plugins/Animation/IKRig/Source/IKRigEditor/Public/RigEditor/IKRigController.h
--- a/Engine/Plugins/Animation/IKRig/Source/IKRigEditor/Public/RigEditor/IKRigController.h
+++ b/Engine/Plugins/Animation/IKRig/Source/IKRigEditor/Public/RigEditor/IKRigController.h
@@ -3,10 +3,12 @@
 #pragma once
 
 #include "CoreMinimal.h"
+#include "StructUtils/InstancedStruct.h"
 #include "UObject/Object.h"
 
 #include "IKRigController.generated.h"
 
+class UStructViewer;
 struct FAutoCharacterizeResults;
 struct FAutoCharacterizer;
 struct FAutoFBIKCreator;
@@ -15,13 +17,17 @@ struct FRetargetDefinition;
 struct FIKRigInputSkeleton;
 struct FReferenceSkeleton;
 struct FIKRigSkeleton;
-class UIKRigSolver;
+struct FIKRigSolverBase;
+struct FIKRigBoneSettingsBase;
+struct FIKRigGoalSettingsBase;
+struct FInstancedStruct;
 class UIKRigDefinition;
 class UIKRigEffectorGoal;
 class USkeletalMesh;
 class USkeleton;
 struct FBoneChain;
 
+
 // A singleton (per-asset) class used to make modifications to a UIKRigDefinition asset
 //
 // All modifications to an IKRigDefinition must go through this controller.
@@ -30,7 +36,7 @@ struct FBoneChain;
 // of a running IK Rig processor instance.
 // The API here is split into public/scripting sections which are accessible from Blueprint/Python and sections that
 // are only relevant to editors in C++.
-UCLASS(config = Engine, hidecategories = UObject)
+UCLASS(BlueprintType, hidecategories = UObject)
 class IKRIGEDITOR_API UIKRigController : public UObject
 {
 	GENERATED_BODY()
@@ -44,16 +50,16 @@ public:
 	//
 	
 	// Use this to get the controller for the given IKRig
-	UFUNCTION(BlueprintCallable, Category=IKRig)
+	UFUNCTION(BlueprintCallable, Category=IKRig, meta = (DisplayName = "Get IK Rig Controller"))
 	static UIKRigController* GetController(const UIKRigDefinition* InIKRigDefinition);
 	
 	// Sets the preview mesh to use. Loads the hierarchy into the asset's IKRigSkeleton.
 	// Returns true if the mesh was able to be set. False if it was incompatible for any reason. 
-	UFUNCTION(BlueprintCallable, BlueprintPure=false, Category=IKRig)
+	UFUNCTION(BlueprintCallable, BlueprintPure=false, Category=IKRig, meta = (DisplayName = "Set IK Rig Skeletal Mesh"))
 	bool SetSkeletalMesh(USkeletalMesh* SkeletalMesh) const;
 
 	// Get the skeletal mesh this asset is initialized with 
-	UFUNCTION(BlueprintCallable, Category=IKRig)
+	UFUNCTION(BlueprintCallable, Category=IKRig, meta = (DisplayName = "Get IK Rig Skeletal Mesh"))
 	USkeletalMesh* GetSkeletalMesh() const;
 
 	// Returns true if the provided skeletal mesh could be used with this IK Rig.
@@ -85,21 +91,14 @@ public:
 	//
 
 	// Add a new solver of the given type to the bottom of the stack. Returns the stack index.
+	// The type here is the full package path of the UStruct type. For example: /Script/IKRig.FullBodyIKSolver
 	UFUNCTION(BlueprintCallable, BlueprintPure=false, Category=IKRig)
-	int32 AddSolver(TSubclassOf<UIKRigSolver> InSolverClass) const;
+	int32 AddSolver(const FString InIKRigSolverType) const;
 
 	// Remove the solver at the given stack index. 
 	UFUNCTION(BlueprintCallable, BlueprintPure=false, Category=IKRig)
 	bool RemoveSolver(const int32 SolverIndex) const;
 
-	// Get access to the given solver. 
-	UFUNCTION(BlueprintCallable, Category=IKRig)
-	UIKRigSolver* GetSolverAtIndex(int32 Index) const;
-
-	// Get access to the given solver. 
-	UFUNCTION(BlueprintCallable, Category=IKRig)
-	int32 GetIndexOfSolver(UIKRigSolver* Solver) const;
-
 	// Get the number of solvers in the stack. 
 	UFUNCTION(BlueprintCallable, Category=IKRig)
 	int32 GetNumSolvers() const;
@@ -116,13 +115,13 @@ public:
 	UFUNCTION(BlueprintCallable, Category=IKRig)
 	bool GetSolverEnabled(int32 SolverIndex) const;
 
-	// Set the root bone on a given solver. (not all solvers support root bones, checks CanSetRootBone() first) 
+	// Set the root bone on a given solver. (not all solvers support root bones, checks CanSetStartBone() first) 
 	UFUNCTION(BlueprintCallable, BlueprintPure=false, Category=IKRig)
-	bool SetRootBone(const FName RootBoneName, int32 SolverIndex) const;
+	bool SetStartBone(const FName RootBoneName, int32 SolverIndex) const;
 
-	// Get the name of the root bone on a given solver. (not all solvers support root bones, checks CanSetRootBone() first) 
+	// Get the name of the root bone on a given solver. (not all solvers support start bones, checks CanSetStartBone() first) 
 	UFUNCTION(BlueprintCallable, Category=IKRig)
-	FName GetRootBone(int32 SolverIndex) const;
+	FName GetStartBone(int32 SolverIndex) const;
 
 	// Set the end bone on a given solver. (not all solvers require extra end bones, checks CanSetEndBone() first) 
 	UFUNCTION(BlueprintCallable, BlueprintPure=false, Category=IKRig)
@@ -132,12 +131,33 @@ public:
 	UFUNCTION(BlueprintCallable, Category=IKRig)
 	FName GetEndBone(int32 SolverIndex) const;
 
+	// Get a controller to customize the settings for a given solver in the stack 
+	UFUNCTION(BlueprintCallable, Category=IKRig)
+	UIKRigSolverControllerBase* GetSolverController(int32 SolverIndex);
+
+	// Get the index for a given solver in the stack
+	// The input here is a solver controller (as returned by GetSolverController())
+	UFUNCTION(BlueprintCallable, Category=IKRig)
+	int32 GetIndexOfSolver(UIKRigSolverControllerBase* InController);
+
 	//
 	// SOLVERS C++ ONLY API
 	//
+
+	// Add a solver to the IK Rig asset's solver stack
+	int32 AddSolver(UScriptStruct* InIKRigSolverType) const;
+
+	// Get access to the given solver. 
+	FIKRigSolverBase* GetSolverAtIndex(int32 Index) const;
+
+	// Get access to the instanced struct for a given solver. 
+	FInstancedStruct* GetSolverStructAtIndex(int32 Index) const;
+
+	// Get access to the given solver. 
+	int32 GetIndexOfSolver(FIKRigSolverBase* Solver) const;
 	
 	// Get read-only access to the array of solvers. 
-	const TArray<UIKRigSolver*>& GetSolverArray() const;
+	const TArray<FIKRigSolverBase*> GetSolverArray() const;
 
 	// Get unique label for a given solver. Returns dash separated index and name like so, "0 - SolverName". 
 	FString GetSolverUniqueName(int32 SolverIndex) const;
@@ -192,12 +212,17 @@ public:
 
 	// Get access to the list of Goals. 
 	UFUNCTION(BlueprintCallable, Category=IKRig)
-  const TArray<UIKRigEffectorGoal*>& GetAllGoals() const;
+	const TArray<UIKRigEffectorGoal*>& GetAllGoals() const;
 
 	// Get the UObject for the settings associated with the given Goal in the given Solver.
 	// Solvers can define their own per-Goal settings depending on their needs. These are termed "Effectors".
+	UE_DEPRECATED(5.6, "In 5.6+ each solver provides it's own controller. Use GetSolverController(SolverIndex).GetGoalSettings(GoalName) instead.")
 	UFUNCTION(BlueprintCallable, Category=IKRig)
 	UObject* GetGoalSettingsForSolver(const FName GoalName, int32 SolverIndex) const;
+
+	// returns true if the goal has custom settings for the given solver
+	UFUNCTION(BlueprintCallable, Category=IKRig)
+	bool DoesSolverHaveCustomGoalSettings(const FName GoalName, int32 SolverIndex) const;
 	
 	//
 	// GOALS C++ ONLY API
@@ -247,18 +272,22 @@ public:
 	UFUNCTION(BlueprintCallable, BlueprintPure=false, Category=IKRig)
 	bool RemoveBoneSetting(const FName BoneName, int32 SolverIndex) const;
 
-	// Get the generic (Solver-specific) Bone settings UObject for this Bone in the given Solver.
-	UFUNCTION(BlueprintCallable, BlueprintPure=false, Category=IKRig)
-	UObject* GetBoneSettings(const FName BoneName, int32 SolverIndex) const;
-
 	// Get the global-space retarget pose transform of the given Bone.
 	UFUNCTION(BlueprintCallable, Category=IKRig)
 	FTransform GetRefPoseTransformOfBone(const FName BoneName) const;
 
+	// Get the generic (Solver-specific) Bone settings UObject for this Bone in the given Solver.
+	UE_DEPRECATED(5.6, "In 5.6+ each solver provides it's own controller. Use GetSolverController(SolverIndex).GetBoneSettings(BoneName) instead.")
+	UFUNCTION(BlueprintCallable, BlueprintPure=false, Category=IKRig)
+	UObject* GetBoneSettings(const FName BoneName, int32 SolverIndex) const;
+
 	//
 	// BONES C++ ONLY API
 	//
 
+	// Get the generic (Solver-specific) Bone settings UStruct for this Bone in the given Solver.
+	FIKRigBoneSettingsBase* GetSettingsForBone(const FName BoneName, int32 SolverIndex) const;
+
 	// Returns true if this Bone can have settings in the given Solver. 
 	bool CanAddBoneSetting(const FName BoneName, int32 SolverIndex) const;
 
@@ -369,6 +398,9 @@ public:
 	// Optionally provide a runtime skeleton from an IKRigProcessor to get indices for a running instance (otherwise uses stored hierarchy in asset)
 	bool ValidateChain( const FName ChainName, const FIKRigSkeleton* OptionalSkeleton, TSet<int32>& OutChainIndices) const;
 
+	// 
+	UStructViewer* GetStructViewer();
+
 private:
 
 	// Called whenever the rig is modified in such a way that would require re-initialization by dependent systems.
@@ -410,6 +442,10 @@ private:
 	UPROPERTY(transient)
 	TObjectPtr<UIKRigDefinition> Asset = nullptr;
 
+	// a dummy UObject to wrap UStructs for the details panel
+	UPROPERTY(transient)
+	TObjectPtr<UStructViewer> StructViewer = nullptr;
+
 	// Static global map of assets to their associated controller
 	static TMap<UIKRigDefinition*, UIKRigController*> Controllers;
 
diff --git a/Engine/Plugins/Animation/IKRig/Source/IKRigEditor/Public/RigEditor/IKRigDefinitionFactory.h b/Engine/Plugins/Animation/IKRig/Source/IKRigEditor/Public/RigEditor/IKRigDefinitionFactory.h
--- a/Engine/Plugins/Animation/IKRig/Source/IKRigEditor/Public/RigEditor/IKRigDefinitionFactory.h
+++ b/Engine/Plugins/Animation/IKRig/Source/IKRigEditor/Public/RigEditor/IKRigDefinitionFactory.h
@@ -32,4 +32,12 @@ public:
 	virtual FString GetDefaultNewAssetName() const override;
 	virtual bool ShouldShowInNewMenu() const override;
 	// END UFactory interface
+
+	/**
+	 * Create a new IK Rig rig asset at a specified location in your project's content folder.
+	 * @param InDesiredPackagePath The package path to where the new IK Rig will be placed. (ie "/Game/MyIKRigs/")
+	 * @param InAssetName The name of the new asset (ie "IK_MyNewRig")
+	 */
+	UFUNCTION(BlueprintCallable, Category = "IK Rig", DisplayName="Create New IK Rig Asset")
+	static UIKRigDefinition* CreateNewIKRigAsset(const FString& InPackagePath, const FString& InAssetName);
 };
diff --git a/Engine/Plugins/Animation/IKRig/Source/IKRigEditor/Public/RigEditor/IKRigDetailCustomizations.h b/Engine/Plugins/Animation/IKRig/Source/IKRigEditor/Public/RigEditor/IKRigDetailCustomizations.h
--- a/Engine/Plugins/Animation/IKRig/Source/IKRigEditor/Public/RigEditor/IKRigDetailCustomizations.h
+++ b/Engine/Plugins/Animation/IKRig/Source/IKRigEditor/Public/RigEditor/IKRigDetailCustomizations.h
@@ -7,6 +7,73 @@
 #include "ScopedTransaction.h"
 #include "RigEditor/IKRigEditorController.h"
 
+#include "IKRigDetailCustomizations.generated.h"
+
+// a thin wrapper around UStruct data to display in a details panel
+UCLASS(Blueprintable)
+class UStructViewer : public UObject
+{
+	GENERATED_BODY()
+
+public:
+
+	// assign a struct to view in the detail panel
+	// InStructOwner is the owning UObject to use for Undo transactions when editing the struct properties
+	void SetStructToView(const UScriptStruct* InStructType, uint8* InStructMemory, UObject* InStructOwner)
+	{
+		StructType = InStructType;
+		StructMemory = InStructMemory;
+		StructOwner = InStructOwner;
+	}
+
+	bool IsValid() const
+	{
+		return StructType && StructMemory;
+	}
+
+	void Reset()
+	{
+		StructType = nullptr;
+		StructMemory = nullptr;
+		StructOwner = nullptr;
+	}
+
+	TSharedPtr<FStructOnScope> GetStructOnScope()
+	{
+		return MakeShared<FStructOnScope>(StructType, StructMemory);
+	}
+
+	FName GetStructName() const
+	{
+		return FName(*StructType->GetName());
+	}
+
+	UObject* GetStructOwner() const
+	{
+		return StructOwner.Get();
+	};
+
+private:
+	
+	// we need the type reflection data and the memory to build a StructOnScope to display in details panel
+	const UScriptStruct* StructType = nullptr;
+	uint8* StructMemory = nullptr;
+
+	// object that owns the struct for Undo/Redo transactions
+	TWeakObjectPtr<UObject> StructOwner = nullptr;
+};
+
+class FStructViewerCustomization : public IDetailCustomization
+{
+public:
+	static TSharedRef<IDetailCustomization> MakeInstance();
+
+	// IDetailCustomization interface
+	virtual void CustomizeDetails(IDetailLayoutBuilder& DetailBuilder) override;
+	// END IDetailCustomization interface
+};
+
+
 class FIKRigGenericDetailCustomization : public IDetailCustomization
 {
 public:
diff --git a/Engine/Plugins/Animation/IKRig/Source/IKRigEditor/Public/RigEditor/IKRigEditorController.h b/Engine/Plugins/Animation/IKRig/Source/IKRigEditor/Public/RigEditor/IKRigEditorController.h
--- a/Engine/Plugins/Animation/IKRig/Source/IKRigEditor/Public/RigEditor/IKRigEditorController.h
+++ b/Engine/Plugins/Animation/IKRig/Source/IKRigEditor/Public/RigEditor/IKRigEditorController.h
@@ -2,7 +2,7 @@
 
 #pragma once
 
-#include "Rig/Solvers/IKRigSolver.h"
+#include "Rig/Solvers/IKRigSolverBase.h"
 #include "SIKRigRetargetChainList.h"
 #include "Templates/SubclassOf.h"
 #include "Animation/AnimationAsset.h"
@@ -21,6 +21,8 @@ class FSolverStackElement;
 class FIKRigTreeElement;
 class UAnimInstance;
 class UDebugSkelMeshComponent;
+class UStructViewer;
+class UIKRigProcessor;
 
 enum EIKRigSelectionType : int8
 {
@@ -181,6 +183,8 @@ public:
 	void ShowAssetDetails() const;
 	// show selected items in details view 
 	void ShowDetailsForElements(const TArray<TSharedPtr<FIKRigTreeElement>>& InItems) const;
+	// show the currently configured StructViewer object in the details tab
+	void ShowStructViewerInDetailsView(UStructViewer* StructViewer) const;
 	// callback when detail is edited 
 	void OnFinishedChangingDetails(const FPropertyChangedEvent& PropertyChangedEvent) const;
 	// returns true if the supplied UObject is being shown in the details panel
@@ -283,5 +287,5 @@ private:
 struct FIKRigSolverTypeAndName
 {
 	FText NiceName;
-	TSubclassOf<UIKRigSolver> SolverType;
+	UScriptStruct* SolverType;
 };
\ No newline at end of file
diff --git a/Engine/Plugins/Animation/IKRig/Source/IKRigEditor/Public/RigEditor/IKRigSkeletonCommands.h b/Engine/Plugins/Animation/IKRig/Source/IKRigEditor/Public/RigEditor/IKRigSkeletonCommands.h
--- a/Engine/Plugins/Animation/IKRig/Source/IKRigEditor/Public/RigEditor/IKRigSkeletonCommands.h
+++ b/Engine/Plugins/Animation/IKRig/Source/IKRigEditor/Public/RigEditor/IKRigSkeletonCommands.h
@@ -28,7 +28,7 @@ public:
 	TSharedPtr< FUICommandInfo > DisconnectGoalFromSolver;
 
 	/** set root bone on selected solvers */
-	TSharedPtr< FUICommandInfo > SetRootBoneOnSolver;
+	TSharedPtr< FUICommandInfo > SetStartBoneOnSolver;
 
 	/** set end bone on selected solvers (Pole Solvers only)*/
 	TSharedPtr< FUICommandInfo > SetEndBoneOnSolver;
diff --git a/Engine/Plugins/Animation/IKRig/Source/IKRigEditor/Public/RigEditor/SIKRigHierarchy.h b/Engine/Plugins/Animation/IKRig/Source/IKRigEditor/Public/RigEditor/SIKRigHierarchy.h
--- a/Engine/Plugins/Animation/IKRig/Source/IKRigEditor/Public/RigEditor/SIKRigHierarchy.h
+++ b/Engine/Plugins/Animation/IKRig/Source/IKRigEditor/Public/RigEditor/SIKRigHierarchy.h
@@ -60,7 +60,7 @@ public:
 	FOnRenameRequested OnRenameRequested;
 
 	/** get the underlying object */
-	TWeakObjectPtr< UObject > GetObject() const;
+	void GetStruct(const UScriptStruct*& OutType, uint8*& OutMemory) const;
 
 	/** get the name of the retarget chain this element belongs to (or NAME_None if not in a chain) */
 	FName GetChainName() const;
diff --git a/Engine/Plugins/Animation/IKRig/Source/IKRigEditor/Public/RigEditor/SIKRigSolverStack.h b/Engine/Plugins/Animation/IKRig/Source/IKRigEditor/Public/RigEditor/SIKRigSolverStack.h
--- a/Engine/Plugins/Animation/IKRig/Source/IKRigEditor/Public/RigEditor/SIKRigSolverStack.h
+++ b/Engine/Plugins/Animation/IKRig/Source/IKRigEditor/Public/RigEditor/SIKRigSolverStack.h
@@ -11,7 +11,7 @@
 #include "Widgets/Views/SListView.h"
 #include "Framework/Commands/UICommandList.h"
 
-class UIKRigSolver;
+struct FIKRigSolverBase;
 class FIKRigEditorController;
 class SIKRigSolverStack;
 class FIKRigEditorToolkit;
@@ -60,7 +60,7 @@ public:
 	bool IsSolverEnabled() const;
 
 private:
-	UIKRigSolver* GetSolver(const bool bFromAsset) const;
+	const FIKRigSolverBase* GetSolver(const bool bFromAsset) const;
 	TWeakPtr<FSolverStackElement> StackElement;
 	TWeakPtr<SIKRigSolverStack> SolverStack;
 };
@@ -105,7 +105,7 @@ private:
 	/** END solver stack menu */
 	
 	/** menu command callback for adding a new solver */
-	void AddNewSolver(UClass* Class);
+	void AddNewSolver(UScriptStruct* SolverType);
 	/** delete solver from stack */
 	void DeleteSolver(TSharedPtr<FSolverStackElement> SolverToDelete);
 	/** when a solver is selected on in the stack view */
