[IKRig] Renaming files to adhere to new standards.

diff --git a/Engine/Plugins/Animation/IKRig/Source/IKRig/Private/Rig/Solvers/IKRigFullBodyIK.cpp b/Engine/Plugins/Animation/IKRig/Source/IKRig/Private/Rig/Solvers/IKRigFullBodyIK.cpp
new file mode 100644
--- /dev/null
+++ b/Engine/Plugins/Animation/IKRig/Source/IKRig/Private/Rig/Solvers/IKRigFullBodyIK.cpp
@@ -0,0 +1,436 @@
+// Copyright Epic Games, Inc. All Rights Reserved.
+
+#include "Rig/Solvers/IKRigFullBodyIK.h"
+
+#include "Rig/IKRigDataTypes.h"
+#include "Rig/IKRigSkeleton.h"
+
+#include UE_INLINE_GENERATED_CPP_BY_NAME(IKRigFullBodyIK)
+
+#define LOCTEXT_NAMESPACE "FFBIKSolver"
+
+void FIKRigFullBodyIKSolver::Initialize(const FIKRigSkeleton& InSkeleton)
+{	
+	// check how many effectors are assigned to a bone
+	int NumEffectors = 0;
+	for (const FIKRigFBIKGoalSettings& Effector : AllGoalSettings)
+	{
+		if (InSkeleton.GetBoneIndexFromName(Effector.BoneName) != INDEX_NONE)
+		{
+			++NumEffectors; // bone is set and exists!
+		}
+	}
+
+	// validate inputs are ready to be initialized
+	const bool bHasEffectors = NumEffectors > 0;
+	const bool bRootIsAssigned = Settings.RootBone != NAME_None;
+	if (!(bHasEffectors && bRootIsAssigned))
+	{
+		return; // not setup yet
+	}
+
+	// reset all internal data
+	Solver.Reset();
+
+	// create bones
+	for (int BoneIndex = 0; BoneIndex < InSkeleton.BoneNames.Num(); ++BoneIndex)
+	{
+		const FName& Name = InSkeleton.BoneNames[BoneIndex];
+
+		// get the parent bone solver index
+		const int32 ParentIndex = InSkeleton.GetParentIndexThatIsNotExcluded(BoneIndex);
+		const FTransform OrigTransform = InSkeleton.RefPoseGlobal[BoneIndex];
+		const FVector InOrigPosition = OrigTransform.GetLocation();
+		const FQuat InOrigRotation = OrigTransform.GetRotation();
+		const bool bIsRoot = Name == Settings.RootBone;
+		Solver.AddBone(Name, ParentIndex, InOrigPosition, InOrigRotation, bIsRoot);
+	}
+
+	// create effectors
+	for (FIKRigFBIKGoalSettings& Effector : AllGoalSettings)
+	{
+		Effector.IndexInSolver = Solver.AddEffector(Effector.BoneName);
+	}
+		
+	// initialize solver
+	Solver.Initialize();
+}
+
+void FIKRigFullBodyIKSolver::Solve(FIKRigSkeleton& InIKRigSkeleton, const FIKRigGoalContainer& InGoals)
+{
+	if (!Solver.IsReadyToSimulate())
+	{
+		return;
+	}
+
+	if (Solver.GetNumBones() != InIKRigSkeleton.BoneNames.Num())
+	{
+		return;
+	}
+
+	TArray<FTransform>& InOutTransforms = InIKRigSkeleton.CurrentPoseGlobal;
+	
+	// set bones to input pose
+	for(int32 BoneIndex = 0; BoneIndex < Solver.GetNumBones(); BoneIndex++)
+	{
+		Solver.SetBoneTransform(BoneIndex, InOutTransforms[BoneIndex]);
+	}
+
+	// update bone settings
+	for (const FIKRigFBIKBoneSettings& BoneSetting : AllBoneSettings)
+	{
+		const int32 BoneIndex = Solver.GetBoneIndex(BoneSetting.Bone);
+		if (PBIK::FBoneSettings* InternalSettings = Solver.GetBoneSettings(BoneIndex))
+		{
+			BoneSetting.CopyToCoreStruct(*InternalSettings);
+		}
+	}
+
+	// update effectors
+	for (const FIKRigFBIKGoalSettings& GoalSettings : AllGoalSettings)
+	{
+		if (GoalSettings.IndexInSolver < 0)
+		{
+			continue;
+		}
+		
+		const FIKRigGoal* Goal = InGoals.FindGoalByName(GoalSettings.Goal);
+		if (!Goal)
+		{
+			return;
+		}
+
+		PBIK::FEffectorSettings EffectorSettings;
+		EffectorSettings.PositionAlpha = 1.0f; // this is constant because IKRig manages offset alphas itself
+		EffectorSettings.RotationAlpha = 1.0f; // this is constant because IKRig manages offset alphas itself
+		EffectorSettings.StrengthAlpha = GoalSettings.StrengthAlpha;
+		EffectorSettings.ChainDepth = GoalSettings.ChainDepth;
+		EffectorSettings.PullChainAlpha = GoalSettings.PullChainAlpha;
+		EffectorSettings.PinRotation = GoalSettings.PinRotation;
+		
+		Solver.SetEffectorGoal(
+			GoalSettings.IndexInSolver,
+			Goal->FinalBlendedPosition,
+			Goal->FinalBlendedRotation,
+			EffectorSettings);
+	}
+
+	// update settings
+	FPBIKSolverSettings SolverSettings;
+	SolverSettings.Iterations = Settings.Iterations;
+	SolverSettings.SubIterations = Settings.SubIterations;
+	SolverSettings.MassMultiplier = Settings.MassMultiplier;
+	SolverSettings.bAllowStretch = Settings.bAllowStretch;
+	SolverSettings.RootBehavior = Settings.RootBehavior;
+	SolverSettings.PrePullRootSettings = Settings.PrePullRootSettings;
+	SolverSettings.GlobalPullChainAlpha = Settings.GlobalPullChainAlpha;
+	SolverSettings.MaxAngle = Settings.MaxAngle;
+	SolverSettings.OverRelaxation = Settings.OverRelaxation;
+
+	// solve
+	Solver.Solve(SolverSettings);
+
+	// copy transforms back
+	for(int32 BoneIndex = 0; BoneIndex < Solver.GetNumBones(); BoneIndex++)
+	{
+		Solver.GetBoneGlobalTransform(BoneIndex, InOutTransforms[BoneIndex]);
+	}
+}
+
+void FIKRigFullBodyIKSolver::GetRequiredBones(TSet<FName>& OutRequiredBones) const
+{
+	OutRequiredBones.Add(Settings.RootBone);
+	// TODO in the future make this bone agnostic, stop storing goal bone names and find dynamically at Init() time
+	for (const FIKRigFBIKGoalSettings& GoalSettings : AllGoalSettings)
+	{
+		OutRequiredBones.Add(GoalSettings.BoneName);
+	}
+}
+
+void FIKRigFullBodyIKSolver::GetRequiredGoals(TSet<FName>& OutRequiredGoals) const
+{
+	for (const FIKRigFBIKGoalSettings& GoalSettings : AllGoalSettings)
+	{
+		OutRequiredGoals.Add(GoalSettings.Goal);
+	}
+}
+
+FIKRigSolverSettingsBase* FIKRigFullBodyIKSolver::GetSolverSettings()
+{
+	return &Settings;
+}
+
+const UScriptStruct* FIKRigFullBodyIKSolver::GetSolverSettingsType() const
+{
+	return FIKRigFBIKSettings::StaticStruct();
+}
+
+// GOAL SETTINGS
+
+void FIKRigFullBodyIKSolver::AddGoal(const UIKRigEffectorGoal& InNewGoal)
+{
+	AllGoalSettings.Emplace(InNewGoal.GoalName, InNewGoal.BoneName);
+}
+
+void FIKRigFullBodyIKSolver::OnGoalRenamed(const FName& InOldName, const FName& InNewName)
+{
+	// ensure goal even exists in this solver
+	const int32 GoalIndex = GetIndexOfGoal(InOldName);
+	if (GoalIndex == INDEX_NONE)
+	{
+		return;
+	}
+
+	// rename
+	AllGoalSettings[GoalIndex].Goal = InNewName;
+}
+
+void FIKRigFullBodyIKSolver::OnGoalMovedToDifferentBone(const FName& InGoalName, const FName& InBoneName)
+{
+	// ensure goal even exists in this solver
+	const int32 GoalIndex = GetIndexOfGoal(InGoalName);
+	if (GoalIndex == INDEX_NONE)
+	{
+		return;
+	}
+
+	// goal moved to different bone
+	AllGoalSettings[GoalIndex].BoneName = InBoneName;
+}
+
+void FIKRigFullBodyIKSolver::OnGoalRemoved(const FName& InGoalName)
+{
+	// ensure goal even exists in this solver
+	const int32 GoalIndex = GetIndexOfGoal(InGoalName);
+	if (GoalIndex == INDEX_NONE)
+	{
+		return;
+	}
+
+	// remove it
+	AllGoalSettings.RemoveAt(GoalIndex);
+}
+
+bool FIKRigFullBodyIKSolver::UsesCustomGoalSettings() const
+{
+	return true;
+}
+
+FIKRigGoalSettingsBase* FIKRigFullBodyIKSolver::GetGoalSettings(const FName& InGoalName)
+{
+	const int32 GoalIndex = GetIndexOfGoal(InGoalName);
+	if (GoalIndex == INDEX_NONE)
+	{
+		return nullptr;
+	}
+
+	return &AllGoalSettings[GoalIndex];
+}
+
+const UScriptStruct* FIKRigFullBodyIKSolver::GetGoalSettingsType() const
+{
+	return FIKRigFBIKGoalSettings::StaticStruct();
+}
+
+void FIKRigFullBodyIKSolver::GetGoalsWithSettings(TSet<FName>& OutGoalsWithSettings) const
+{
+	for (const FIKRigFBIKGoalSettings& GoalSettings : AllGoalSettings)
+	{
+		OutGoalsWithSettings.Add(GoalSettings.Goal);
+	}
+}
+
+// START BONE
+
+bool FIKRigFullBodyIKSolver::UsesStartBone() const
+{
+	return true;
+}
+
+FName FIKRigFullBodyIKSolver::GetStartBone() const
+{
+	return Settings.RootBone;
+}
+
+void FIKRigFullBodyIKSolver::SetStartBone(const FName& InStartBoneName)
+{
+	Settings.RootBone = InStartBoneName;
+}
+
+// BONE SETTINGS
+
+bool FIKRigFullBodyIKSolver::UsesCustomBoneSettings() const
+{
+	return true;
+}
+
+void FIKRigFullBodyIKSolver::AddSettingsToBone(const FName& InBoneName)
+{
+	if (AllBoneSettings.ContainsByPredicate([&](const FIKRigFBIKBoneSettings& Element){return Element.Bone == InBoneName;}))
+	{
+		// bone already has settings
+		return;
+	}
+	
+	AllBoneSettings.Emplace(InBoneName);
+}
+
+void FIKRigFullBodyIKSolver::RemoveSettingsOnBone(const FName& InBoneName)
+{
+	AllBoneSettings.RemoveAll([&](const FIKRigFBIKBoneSettings& Element)
+	{
+		return Element.Bone == InBoneName;
+	});
+}
+
+FIKRigBoneSettingsBase* FIKRigFullBodyIKSolver::GetBoneSettings(const FName& InBoneName)
+{
+	for (FIKRigFBIKBoneSettings& BoneSetting : AllBoneSettings)
+	{
+		if (BoneSetting.Bone == InBoneName)
+		{
+			return &BoneSetting;
+		}
+	}
+	
+	return nullptr;
+}
+
+const UScriptStruct* FIKRigFullBodyIKSolver::GetBoneSettingsType() const
+{
+	return FIKRigFBIKBoneSettings::StaticStruct();
+}
+
+bool FIKRigFullBodyIKSolver::HasSettingsOnBone(const FName& InBoneName) const
+{
+	for (const FIKRigFBIKBoneSettings& BoneSetting : AllBoneSettings)
+	{
+		if (BoneSetting.Bone == InBoneName)
+		{
+			return true;
+		}
+	}
+
+	return false;
+}
+
+void FIKRigFullBodyIKSolver::GetBonesWithSettings(TSet<FName>& OutBonesWithSettings) const
+{
+	for (const FIKRigFBIKBoneSettings& BoneSetting : AllBoneSettings)
+	{
+		OutBonesWithSettings.Add(BoneSetting.Bone);
+	}
+}
+
+#if WITH_EDITOR
+
+UIKRigSolverControllerBase* FIKRigFullBodyIKSolver::GetSolverController(UObject* Outer)
+{
+	return CreateControllerIfNeeded(Outer, UIKRigFBIKController::StaticClass());
+}
+
+FText FIKRigFullBodyIKSolver::GetNiceName() const
+{
+	return FText(LOCTEXT("SolverName", "Full Body IK"));
+}
+
+bool FIKRigFullBodyIKSolver::GetWarningMessage(FText& OutWarningMessage) const
+{
+	if (Settings.RootBone == NAME_None)
+	{
+		OutWarningMessage = LOCTEXT("MissingStart", "Missing start bone.");
+		return true;
+	}
+
+	if (AllGoalSettings.IsEmpty())
+	{
+		OutWarningMessage = LOCTEXT("MissingGoal", "Missing goals.");
+		return true;
+	}
+	
+	return false;
+}
+
+bool FIKRigFullBodyIKSolver::IsBoneAffectedBySolver(const FName& InBoneName, const FIKRigSkeleton& InIKRigSkeleton) const
+{
+	// nothing is affected by solver without a root bone assigned or at least 1 effector
+	if (Settings.RootBone == NAME_None || AllGoalSettings.IsEmpty())
+	{
+		return false;
+	}
+
+	// has to be BELOW root
+	if (!InIKRigSkeleton.IsBoneInDirectLineage(InBoneName, Settings.RootBone))
+	{
+		return false;
+	}
+
+	// has to be ABOVE an effector
+	for (const FIKRigFBIKGoalSettings& Effector : AllGoalSettings)
+	{
+		if (InIKRigSkeleton.IsBoneInDirectLineage(Effector.BoneName, InBoneName))
+		{
+			return true;
+		}
+	}
+	
+	return false;
+}
+
+#endif
+
+int32 FIKRigFullBodyIKSolver::GetIndexOfGoal(const FName& InGoalName) const
+{
+	return AllGoalSettings.IndexOfByPredicate([InGoalName](const FIKRigFBIKGoalSettings& Element)
+	{
+		return Element.Goal == InGoalName;
+	});
+}
+
+//
+// BEGIN CONTROLLER
+//
+
+FIKRigFBIKSettings UIKRigFBIKController::GetSolverSettings()
+{
+	return *static_cast<FIKRigFBIKSettings*>(SolverToControl->GetSolverSettings());
+}
+
+void UIKRigFBIKController::SetSolverSettings(FIKRigFBIKSettings InSettings)
+{
+	SolverToControl->SetSolverSettings(&InSettings);
+}
+
+FIKRigFBIKGoalSettings UIKRigFBIKController::GetGoalSettings(const FName InGoalName)
+{
+	if (FIKRigGoalSettingsBase* GoalSettings = SolverToControl->GetGoalSettings(InGoalName))
+	{
+		return *static_cast<FIKRigFBIKGoalSettings*>(GoalSettings);
+	}
+	return FIKRigFBIKGoalSettings();
+}
+
+void UIKRigFBIKController::SetGoalSettings(const FName InGoalName, FIKRigFBIKGoalSettings InSettings)
+{
+	SolverToControl->SetGoalSettings(InGoalName, &InSettings);
+}
+
+FIKRigFBIKBoneSettings UIKRigFBIKController::GetBoneSettings(const FName InBoneName)
+{
+	if (FIKRigFBIKBoneSettings* BoneSettings = static_cast<FIKRigFBIKBoneSettings*>(SolverToControl->GetBoneSettings(InBoneName)))
+	{
+		return *BoneSettings;
+	}
+	
+	FIKRigFBIKBoneSettings NewBoneSettings;
+	NewBoneSettings.Bone = InBoneName;
+	return NewBoneSettings;
+}
+
+void UIKRigFBIKController::SetBoneSettings(const FName InBoneName, FIKRigFBIKBoneSettings InSettings)
+{
+	SolverToControl->SetBoneSettings(InBoneName, &InSettings);
+}
+
+#undef LOCTEXT_NAMESPACE
+
diff --git a/Engine/Plugins/Animation/IKRig/Source/IKRig/Private/Rig/Solvers/IKRigSetTransform.cpp b/Engine/Plugins/Animation/IKRig/Source/IKRig/Private/Rig/Solvers/IKRigSetTransform.cpp
new file mode 100644
--- /dev/null
+++ b/Engine/Plugins/Animation/IKRig/Source/IKRig/Private/Rig/Solvers/IKRigSetTransform.cpp
@@ -0,0 +1,155 @@
+// Copyright Epic Games, Inc. All Rights Reserved.
+
+#include "Rig/Solvers/IKRigSetTransform.h"
+#include "Rig/IKRigDataTypes.h"
+#include "Rig/IKRigSkeleton.h"
+
+#include UE_INLINE_GENERATED_CPP_BY_NAME(IKRigSetTransform)
+
+#define LOCTEXT_NAMESPACE "SetTransformSolver"
+
+void FIKRigSetTransform::Initialize(const FIKRigSkeleton& IKRigSkeleton)
+{
+	BoneIndex = IKRigSkeleton.GetBoneIndexFromName(Settings.BoneToAffect);
+}
+
+void FIKRigSetTransform::Solve(FIKRigSkeleton& InIKRigSkeleton, const FIKRigGoalContainer& InGoals)
+{
+	// no bone specified
+	if (BoneIndex == INDEX_NONE)
+	{
+		return;
+	}
+
+	// no goal specified
+	const FIKRigGoal* InGoal = InGoals.FindGoalByName(Settings.Goal);
+	if (!InGoal)
+	{
+		return;
+	}
+
+	// check that settings are such that there is anything to do at all
+	const bool bPositionEnabled = Settings.PositionAlpha > UE_KINDA_SMALL_NUMBER;
+	const bool bRotationEnabled = Settings.RotationAlpha > UE_KINDA_SMALL_NUMBER;
+	const bool bAnythingEnabled = bPositionEnabled || bRotationEnabled;
+	const bool bHasAlpha = Settings.Alpha > UE_KINDA_SMALL_NUMBER;
+	if (!(bAnythingEnabled && bHasAlpha))
+	{
+		return;
+	}
+
+	FTransform& CurrentTransform = InIKRigSkeleton.CurrentPoseGlobal[BoneIndex];	
+
+	if (bPositionEnabled)
+	{
+		const FVector TargetPosition = FMath::Lerp(CurrentTransform.GetTranslation(), InGoal->FinalBlendedPosition, Settings.PositionAlpha * Settings.Alpha);
+		CurrentTransform.SetLocation(TargetPosition);
+	}
+	
+	if (bRotationEnabled)
+	{
+		const FQuat TargetRotation = FMath::Lerp(CurrentTransform.GetRotation(), InGoal->FinalBlendedRotation, Settings.RotationAlpha * Settings.Alpha);
+		CurrentTransform.SetRotation(TargetRotation);
+	}
+
+	if (Settings.bPropagateToChildren)
+	{
+		InIKRigSkeleton.PropagateGlobalPoseBelowBone(BoneIndex);	
+	}
+}
+
+void FIKRigSetTransform::GetRequiredBones(TSet<FName>& OutRequiredBones) const
+{
+	OutRequiredBones.Add(Settings.BoneToAffect);
+}
+
+void FIKRigSetTransform::GetRequiredGoals(TSet<FName>& OutRequiredGoals) const
+{
+	if (Settings.Goal != NAME_Name)
+	{
+		OutRequiredGoals.Add(Settings.Goal);
+	}
+}
+
+FIKRigSolverSettingsBase* FIKRigSetTransform::GetSolverSettings()
+{
+	return &Settings;
+}
+
+const UScriptStruct* FIKRigSetTransform::GetSolverSettingsType() const
+{
+	return FIKRigSetTransformSettings::StaticStruct();
+}
+
+void FIKRigSetTransform::OnGoalRemoved(const FName& InGoalName)
+{
+	if (Settings.Goal == InGoalName)
+	{
+		Settings.Goal = NAME_None;
+		Settings.BoneToAffect = NAME_None;
+	}	
+}
+
+void FIKRigSetTransform::AddGoal(const UIKRigEffectorGoal& InNewGoal)
+{
+	Settings.Goal = InNewGoal.GoalName;
+	Settings.BoneToAffect = InNewGoal.BoneName;
+}
+
+void FIKRigSetTransform::OnGoalRenamed(const FName& InOldName, const FName& InNewName)
+{
+	if (Settings.Goal == InOldName)
+	{
+		Settings.Goal = InNewName;
+	}
+}
+
+void FIKRigSetTransform::OnGoalMovedToDifferentBone(const FName& InGoalName, const FName& InNewBoneName)
+{
+	if (Settings.Goal == InGoalName)
+	{
+		Settings.BoneToAffect = InNewBoneName;
+	}
+}
+
+FIKRigSetTransformSettings UIKRigSetTransformController::GetSolverSettings()
+{
+	return *static_cast<FIKRigSetTransformSettings*>(SolverToControl->GetSolverSettings());
+}
+
+void UIKRigSetTransformController::SetSolverSettings(FIKRigSetTransformSettings InSettings)
+{
+	SolverToControl->SetSolverSettings(&InSettings);
+}
+
+#if WITH_EDITOR
+
+UIKRigSolverControllerBase* FIKRigSetTransform::GetSolverController(UObject* Outer)
+{
+	return CreateControllerIfNeeded(Outer, UIKRigSetTransformController::StaticClass());
+}
+
+FText FIKRigSetTransform::GetNiceName() const
+{
+	return FText(LOCTEXT("SolverName", "Set Transform"));
+}
+
+bool FIKRigSetTransform::GetWarningMessage(FText& OutWarningMessage) const
+{
+	if (Settings.BoneToAffect == NAME_None)
+	{
+		OutWarningMessage = LOCTEXT("MissingGoal", "Missing goal.");
+		return true;
+	}
+	return false;
+}
+
+bool FIKRigSetTransform::IsBoneAffectedBySolver(const FName& InBoneName, const FIKRigSkeleton& InIKRigSkeleton) const
+{
+	return InIKRigSkeleton.IsBoneInDirectLineage(InBoneName, Settings.BoneToAffect);
+}
+
+#endif
+
+#undef LOCTEXT_NAMESPACE
+
diff --git a/Engine/Plugins/Animation/IKRig/Source/IKRig/Private/Rig/Solvers/IKRig_FBIKSolver.cpp b/Engine/Plugins/Animation/IKRig/Source/IKRig/Private/Rig/Solvers/IKRig_FBIKSolver.cpp
deleted file mode 100644
diff --git a/Engine/Plugins/Animation/IKRig/Source/IKRig/Private/Rig/Solvers/SetTransformSolver.cpp b/Engine/Plugins/Animation/IKRig/Source/IKRig/Private/Rig/Solvers/SetTransformSolver.cpp
deleted file mode 100644
diff --git a/Engine/Plugins/Animation/IKRig/Source/IKRig/Public/Rig/Solvers/IKRigFullBodyIK.h b/Engine/Plugins/Animation/IKRig/Source/IKRig/Public/Rig/Solvers/IKRigFullBodyIK.h
new file mode 100644
--- /dev/null
+++ b/Engine/Plugins/Animation/IKRig/Source/IKRig/Public/Rig/Solvers/IKRigFullBodyIK.h
@@ -0,0 +1,489 @@
+// Copyright Epic Games, Inc. All Rights Reserved.
+
+#pragma once
+
+#include "Rig/IKRigDefinition.h"
+#include "IKRigSolverBase.h"
+#include "Core/PBIKSolver.h"
+#include "PBIK_Shared.h"
+
+#include "IKRigFullBodyIK.generated.h"
+
+USTRUCT(BlueprintType)
+struct IKRIG_API FIKRigFBIKGoalSettings : public FIKRigGoalSettingsBase
+{
+	GENERATED_BODY()
+
+	FIKRigFBIKGoalSettings() = default;
+	FIKRigFBIKGoalSettings(FName InGoalName, FName InBoneName) : BoneName(InBoneName) { Goal = InGoalName;}
+	
+	/** The bone that this effector will pull on. */
+	UPROPERTY(BlueprintReadOnly, VisibleAnywhere, Category = "Full Body IK Effector")
+	FName BoneName;
+
+	/** Range 0-inf (default is 0). Explicitly set the number of bones up the hierarchy to consider part of this effector's 'chain'.
+	* The "chain" of bones is used to apply Preferred Angles, Pull Chain Alpha and Chain "Sub Solves".
+	* If left at 0, the solver will attempt to determine the root of the chain by searching up the hierarchy until it finds a branch or another effector, whichever it finds first.
+	*/
+	UPROPERTY(BlueprintReadWrite, EditAnywhere, Category = "Full Body IK Effector", meta = (ClampMin = "0", UIMin = "0"))
+	int32 ChainDepth = 0;
+
+	/** Range 0-1 (default is 1.0). The strength of the effector when pulling the bone towards it's target location.
+	* At 0.0, the effector does not pull at all, but the bones between the effector and the root will still slightly resist motion from other effectors.
+	* This can thus act as a "stabilizer" for parts of the body that you do not want to behave in a pure FK fashion.
+	*/
+	UPROPERTY(BlueprintReadWrite, EditAnywhere, Category = "Full Body IK Effector", meta = (ClampMin = "0", ClampMax = "1", UIMin = "0.0", UIMax = "1.0"))
+	float StrengthAlpha = 1.0f;
+
+	/** Range 0-1 (default is 1.0). When enabled (greater than 0.0), the solver internally partitions the skeleton into 'chains' which extend
+	 * from the effector up the hierarchy by "Chain Depth". If Chain Depth is 0, the chain root is set to the nearest fork in the skeleton.
+	* These chains are pre-rotated and translated, as a whole, towards the effector targets.
+	* This can improve the results for sparse bone chains, and significantly improve convergence on dense bone chains.
+	* But it may cause undesirable results in highly constrained bone chains (like robot arms).
+	*/
+	UPROPERTY(BlueprintReadWrite, EditAnywhere, Category = "Full Body IK Effector", meta = (ClampMin = "0", ClampMax = "1", UIMin = "0.0", UIMax = "1.0"))
+	float PullChainAlpha = 1.0f;
+
+	/** Range 0-1 (default is 1.0).
+	*Blends the effector bone rotation between the rotation of the effector transform (1.0) and the rotation of the input bone (0.0).*/
+	UPROPERTY(BlueprintReadWrite, EditAnywhere, Category = "Full Body IK Effector", meta = (ClampMin = "0", ClampMax = "1", UIMin = "0.0", UIMax = "1.0"))
+	float PinRotation = 1.0f;
+
+	UPROPERTY(Transient)
+	int32 IndexInSolver = -1;
+};
+
+USTRUCT(BlueprintType)
+struct IKRIG_API FIKRigFBIKBoneSettings : public FIKRigBoneSettingsBase
+{
+	GENERATED_BODY()
+
+	FIKRigFBIKBoneSettings() = default;
+	FIKRigFBIKBoneSettings(FName InBoneName) { Bone = InBoneName; }
+	
+	/** Range is 0 to 1 (Default is 0). At higher values, the bone will resist rotating (forcing other bones to compensate). */
+	UPROPERTY(BlueprintReadWrite, EditAnywhere, Category = Stiffness, meta = (ClampMin = "0", ClampMax = "1", UIMin = "0.0", UIMax = "1.0"))
+	float RotationStiffness = 0.0f;
+
+	/** Range is 0 to 1 (Default is 0). At higher values, the bone will resist translational motion (forcing other bones to compensate). */
+	UPROPERTY(BlueprintReadWrite, EditAnywhere, Category = Stiffness, meta = (ClampMin = "0", ClampMax = "1", UIMin = "0.0", UIMax = "1.0"))
+	float PositionStiffness = 0.0f;
+
+	/** Limit the rotation angle of the bone on the X axis.
+	 *Free: can rotate freely in this axis.
+	 *Limited: rotation is clamped between the min/max angles relative to the Skeletal Mesh reference pose.
+	 *Locked: no rotation is allowed in this axis (will remain at reference pose angle). */
+	UPROPERTY(BlueprintReadWrite, EditAnywhere, Category = Limits)
+	EPBIKLimitType X = EPBIKLimitType::Free;
+	/**Range is -180 to 0 (Default is 0). Degrees of rotation in the negative X direction to allow when joint is in "Limited" mode. */
+	UPROPERTY(BlueprintReadWrite, EditAnywhere, Category = Limits, meta = (ClampMin = "-180", ClampMax = "0", UIMin = "-180.0", UIMax = "0.0"))
+	float MinX = 0.0f;
+	/**Range is 0 to 180 (Default is 0). Degrees of rotation in the positive X direction to allow when joint is in "Limited" mode. */
+	UPROPERTY(BlueprintReadWrite, EditAnywhere, Category = Limits, meta = (ClampMin = "0", ClampMax = "180", UIMin = "0.0", UIMax = "180.0"))
+	float MaxX = 0.0f;
+
+	/** Limit the rotation angle of the bone on the Y axis.
+	*Free: can rotate freely in this axis.
+	*Limited: rotation is clamped between the min/max angles relative to the Skeletal Mesh reference pose.
+	*Locked: no rotation is allowed in this axis (will remain at input pose angle). */
+	UPROPERTY(BlueprintReadWrite, EditAnywhere, Category = Limits)
+	EPBIKLimitType Y = EPBIKLimitType::Free;
+	/**Range is -180 to 0 (Default is 0). Degrees of rotation in the negative Y direction to allow when joint is in "Limited" mode. */
+	UPROPERTY(BlueprintReadWrite, EditAnywhere, Category = Limits, meta = (ClampMin = "-180", ClampMax = "0", UIMin = "-180.0", UIMax = "0.0"))
+	float MinY = 0.0f;
+	/**Range is 0 to 180 (Default is 0). Degrees of rotation in the positive Y direction to allow when joint is in "Limited" mode. */
+	UPROPERTY(BlueprintReadWrite, EditAnywhere, Category = Limits, meta = (ClampMin = "0", ClampMax = "180", UIMin = "0.0", UIMax = "180.0"))
+	float MaxY = 0.0f;
+
+	/** Limit the rotation angle of the bone on the Z axis.
+	*Free: can rotate freely in this axis.
+	*Limited: rotation is clamped between the min/max angles relative to the Skeletal Mesh reference pose.
+	*Locked: no rotation is allowed in this axis (will remain at input pose angle). */
+	UPROPERTY(BlueprintReadWrite, EditAnywhere, Category = Limits)
+	EPBIKLimitType Z = EPBIKLimitType::Free;
+	/**Range is -180 to 0 (Default is 0). Degrees of rotation in the negative Z direction to allow when joint is in "Limited" mode. */
+	UPROPERTY(BlueprintReadWrite, EditAnywhere, Category = Limits, meta = (ClampMin = "-180", ClampMax = "0", UIMin = "-180.0", UIMax = "0.0"))
+	float MinZ = 0.0f;
+	/**Range is 0 to 180 (Default is 0). Degrees of rotation in the positive Z direction to allow when joint is in "Limited" mode. */
+	UPROPERTY(BlueprintReadWrite, EditAnywhere, Category = Limits, meta = (ClampMin = "0", ClampMax = "180", UIMin = "0.0", UIMax = "180.0"))
+	float MaxZ = 0.0f;
+
+	/**When true, this bone will "prefer" to rotate in the direction specified by the Preferred Angles when the chain it belongs to is compressed.
+	 * Preferred Angles should be the first method used to fix bones that bend in the wrong direction (rather than limits).
+	 * The resulting angles can be visualized on their own by temporarily setting the Solver iterations to 0 and moving the effectors.*/
+	UPROPERTY(BlueprintReadWrite, EditAnywhere, Category = PreferredAngles)
+	bool bUsePreferredAngles = false;
+	/**The local Euler angles (in degrees) used to rotate this bone when the chain it belongs to is squashed.
+	 * This happens by moving the effector at the tip of the chain towards the root of the chain.
+	 * This can be used to coerce knees and elbows to bend in the anatomically "correct" direction without resorting to limits (which may require more iterations to converge).*/
+	UPROPERTY(BlueprintReadWrite, EditAnywhere, Category = PreferredAngles)
+	FVector PreferredAngles = FVector::ZeroVector;
+
+	void CopyToCoreStruct(PBIK::FBoneSettings& Settings) const
+	{
+		Settings.RotationStiffness = RotationStiffness;
+		Settings.PositionStiffness = PositionStiffness;
+		Settings.X = static_cast<PBIK::ELimitType>(X);
+		Settings.MinX = MinX;
+		Settings.MaxX = MaxX;
+		Settings.Y = static_cast<PBIK::ELimitType>(Y);
+		Settings.MinY = MinY;
+		Settings.MaxY = MaxY;
+		Settings.Z = static_cast<PBIK::ELimitType>(Z);
+		Settings.MinZ = MinZ;
+		Settings.MaxZ = MaxZ;
+		Settings.bUsePreferredAngles = bUsePreferredAngles;
+		Settings.PreferredAngles.Pitch = PreferredAngles.Y;
+		Settings.PreferredAngles.Yaw = PreferredAngles.Z;
+		Settings.PreferredAngles.Roll = PreferredAngles.X;
+	}
+};
+
+USTRUCT(BlueprintType)
+struct IKRIG_API FIKRigFBIKSettings : public FIKRigSolverSettingsBase
+{
+	GENERATED_BODY()
+
+	/** All bones above this bone in the hierarchy will be completely ignored by the solver. Typically, this is set to
+	 * the top-most skinned bone in the Skeletal Mesh (ie Pelvis, Hips etc), NOT the actual root of the entire skeleton.
+	 *
+	 * If you want to use the solver on a single chain of bones, and NOT translate the chain, ensure that "PinRoot" is
+	 * checked on to disable the root from translating to reach the effector goals.*/
+	UPROPERTY(BlueprintReadOnly, VisibleAnywhere, Category = "Full Body IK Settings")
+	FName RootBone;
+	
+	/** High iteration counts can help solve complex joint configurations with competing constraints, but will increase runtime cost. Default is 20. */
+	UPROPERTY(BlueprintReadWrite, EditAnywhere, Category = SolverSettings, meta = (ClampMin = "0", ClampMax = "1000", UIMin = "0.0", UIMax = "200.0"))
+	int32 Iterations = 20;
+
+	/** Iterations used for sub-chains defined by the Chain Depth of the effectors. These are solved BEFORE the main iteration pass. Default is 0. */
+	UPROPERTY(BlueprintReadWrite, EditAnywhere, Category = SolverSettings, meta = (ClampMin = "0", ClampMax = "1000", UIMin = "0.0", UIMax = "200.0"))
+	int32 SubIterations = 0;
+
+	/** A global mass multiplier; higher values will make the joints more stiff, but require more iterations. Typical range is 0.0 to 10.0. */
+	UPROPERTY(BlueprintReadWrite, EditAnywhere, Category = SolverSettings, meta = (ClampMin = "0", UIMin = "0.0", UIMax = "10.0"))
+	float MassMultiplier = 1.0f;
+
+	/** If true, joints will translate to reach the effectors; causing bones to lengthen if necessary. Good for cartoon effects. Default is false. */
+	UPROPERTY(BlueprintReadWrite, EditAnywhere, Category = SolverSettings)
+	bool bAllowStretch = false;
+
+	/** (Default is PrePull) Set the behavior of the solver root.
+	*Pre Pull: translates and rotates the root (and all children) by the average motion of the stretched effectors to help achieve faster convergence when reaching far.
+	*Pin to Input: locks the translation and rotation of the root bone to the input pose. Overrides any bone settings applied to the root. Good for partial-body solves.
+	*Free: treats the root bone like any other and allows it to move freely or according to any bone settings applied to it. */
+	UPROPERTY(BlueprintReadWrite, EditAnywhere, Category = RootBehavior)
+	EPBIKRootBehavior RootBehavior = EPBIKRootBehavior::PrePull;
+
+	/** Settings only applicable when Root Behavior is set to "PrePull". Use these values to adjust the gross movement and orientation of the entire skeleton. */
+	UPROPERTY(BlueprintReadWrite, EditAnywhere, Category = RootBehavior)
+	FRootPrePullSettings PrePullRootSettings;
+
+	/** A global multiplier for all Pull Chain Alpha values on all effectors. Range is 0.0 to 1.0. Default is 1.0. */
+	UPROPERTY(BlueprintReadWrite, EditAnywhere, Category = AdvancedSettings, meta = (ClampMin = "0", ClampMax = "1", UIMin = "0.0", UIMax = "1.0"))
+	float GlobalPullChainAlpha = 1.0f;
+
+	/** Maximum angle that a joint can be rotated per constraint iteration. Lower this value if the solve is diverging. Range is 0.0 to 180.0. Default is 30. */
+	UPROPERTY(BlueprintReadWrite, EditAnywhere, Category = AdvancedSettings, meta = (ClampMin = "0", ClampMax = "45", UIMin = "0.0", UIMax = "180.0"))
+	float MaxAngle = 30.f;
+
+	/** Pushes constraints beyond their normal amount to speed up convergence. Increasing this may speed up convergence, but at the cost of stability. Range is 1.0 - 2.0. Default is 1.3. */
+	UPROPERTY(BlueprintReadWrite, EditAnywhere, Category = AdvancedSettings, meta = (ClampMin = "1",  ClampMax = "2", UIMin = "1.0", UIMax = "2.0"))
+	float OverRelaxation = 1.3f;
+};
+
+
+USTRUCT(BlueprintType)
+struct IKRIG_API FIKRigFullBodyIKSolver : public FIKRigSolverBase
+{
+	GENERATED_BODY()
+
+	UPROPERTY()
+	FIKRigFBIKSettings Settings;
+	
+	UPROPERTY()
+	TArray<FIKRigFBIKGoalSettings> AllGoalSettings;
+
+	UPROPERTY()
+	TArray<FIKRigFBIKBoneSettings> AllBoneSettings;
+	
+	// FIKRigSolverBase runtime interface
+	virtual void Initialize(const FIKRigSkeleton& InIKRigSkeleton) override;
+	virtual void Solve(FIKRigSkeleton& InIKRigSkeleton, const FIKRigGoalContainer& InGoals) override;
+	virtual void GetRequiredBones(TSet<FName>& OutRequiredBones) const override;
+	virtual void GetRequiredGoals(TSet<FName>& OutRequiredGoals) const override;
+	// END FIKRigSolverBase runtime interface
+
+	// solver settings
+	virtual FIKRigSolverSettingsBase* GetSolverSettings() override;
+	virtual const UScriptStruct* GetSolverSettingsType() const override;
+	// END solver settings
+	
+	// goals
+	virtual void AddGoal(const UIKRigEffectorGoal& InNewGoal) override;
+	virtual void OnGoalRenamed(const FName& InOldName, const FName& InNewName) override;
+	virtual void OnGoalMovedToDifferentBone(const FName& InGoalName, const FName& InNewBoneName) override;
+	virtual void OnGoalRemoved(const FName& InGoalName) override;
+	// END goals
+
+	// goal settings
+	virtual bool UsesCustomGoalSettings() const override;
+	virtual FIKRigGoalSettingsBase* GetGoalSettings(const FName& InGoalName) override;
+	virtual const UScriptStruct* GetGoalSettingsType() const override;
+	virtual void GetGoalsWithSettings(TSet<FName>& OutGoalsWithSettings) const override;
+	// END goal settings
+
+	// start bone
+	virtual bool UsesStartBone() const override;
+	virtual void SetStartBone(const FName& InStartBoneName) override;
+	virtual FName GetStartBone() const override;
+	// END start bone
+	
+	// bone settings
+	virtual bool UsesCustomBoneSettings() const override;
+	virtual void AddSettingsToBone(const FName& InBoneName) override;
+	virtual void RemoveSettingsOnBone(const FName& InBoneName) override;
+	virtual FIKRigBoneSettingsBase* GetBoneSettings(const FName& InBoneName) override;
+	virtual const UScriptStruct* GetBoneSettingsType() const override;
+	virtual bool HasSettingsOnBone(const FName& InBoneName) const override;
+	virtual void GetBonesWithSettings(TSet<FName>& OutBonesWithSettings) const override;
+	// END bone settings
+	
+#if WITH_EDITOR
+	// return custom controller for scripting this solver
+	virtual UIKRigSolverControllerBase* GetSolverController(UObject* Outer) override;
+	// ui
+	virtual FText GetNiceName() const override;
+	virtual bool GetWarningMessage(FText& OutWarningMessage) const override;
+	virtual bool IsBoneAffectedBySolver(const FName& InBoneName, const FIKRigSkeleton& InIKRigSkeleton) const override;
+#endif
+
+private:
+
+	FPBIKSolver Solver;
+
+	int32 GetIndexOfGoal(const FName& InGoalName) const;
+};
+
+
+/* The blueprint/python API for modifying an Full-Body IK solver's settings in an IK Rig.
+ * Can adjust Solver, Goal and Bone settings. */
+UCLASS(BlueprintType)
+class IKRIG_API UIKRigFBIKController : public UIKRigSolverControllerBase
+{
+	GENERATED_BODY()
+	
+public:
+
+	/* Get the current solver settings as a struct.
+	 * @return FIKRigFBIKSettings struct with the current settings used by the solver. */
+	UFUNCTION(BlueprintCallable, Category = IKRig)
+	FIKRigFBIKSettings GetSolverSettings();
+
+	/* Set the solver settings. Input is a custom struct type for this solver.
+	 * @param InSettings a FIKRigFBIKSettings struct containing all the settings to apply to this solver */
+	UFUNCTION(BlueprintCallable, Category = IKRig)
+	void SetSolverSettings(FIKRigFBIKSettings InSettings);
+
+	/* Get the settings for the specified goal.
+	 * @param InGoalName the name of the goal to get settings for
+	 * @return FIKRigFBIKGoalSettings struct (empty if the specified goal does not belong to this solver) */
+	UFUNCTION(BlueprintCallable, Category = IKRig)
+	FIKRigFBIKGoalSettings GetGoalSettings(const FName InGoalName);
+
+	/* Set the settings for the specified goal.
+	 * @param InGoalName: the name of the goal to assign the settings to.
+	 * @param InSettings: a custom struct type with all the settings for an FBIK goal */
+	UFUNCTION(BlueprintCallable, Category = IKRig)
+	void SetGoalSettings(const FName InGoalName, FIKRigFBIKGoalSettings InSettings);
+
+	/* Get the settings associated with a particular bone.
+	 * Note that you must AddBoneSettings() using the IK Rig controller before a bone will have settings on it.
+	 * @param InBoneName the name of the bone to get settings for
+	 * @return FIKRigFBIKBoneSettings struct holding all the settings for the specified bone (or empty if the bone did not have settings) */
+	UFUNCTION(BlueprintCallable, Category = IKRig)
+	FIKRigFBIKBoneSettings GetBoneSettings(const FName InBoneName);
+
+	/* Apply settings to a given bone
+	 * @param InBoneName the name of the bone to apply the settings to
+	 * @param InSettings a FIKRigFBIKBoneSettings struct containing the settings */
+	UFUNCTION(BlueprintCallable, Category = IKRig)
+	void SetBoneSettings(const FName InBoneName, FIKRigFBIKBoneSettings InSettings);
+};
+
+//
+// BEGIN DEPRECATED CODE 5.6
+//
+
+PRAGMA_DISABLE_DEPRECATION_WARNINGS
+
+UCLASS(Deprecated)
+class
+UE_DEPRECATED(5.6, "This type has been replaced with FFBIKGoalSettings.")
+IKRIG_API UDEPRECATED_IKRig_FBIKEffector : public UObject
+{
+	GENERATED_BODY()
+	
+public:
+	
+	UPROPERTY()
+	FName GoalName;
+	UPROPERTY()
+	FName BoneName;
+	UPROPERTY()
+	int32 ChainDepth = 0;
+	UPROPERTY()
+	float StrengthAlpha = 1.0f;
+	UPROPERTY()
+	float PullChainAlpha = 1.0f;
+	UPROPERTY()
+	float PinRotation = 1.0f;
+};
+
+UCLASS(Deprecated)
+class
+UE_DEPRECATED(5.6, "This type has been replaced with FFBIKBoneSettings.")
+IKRIG_API UDEPRECATED_IKRig_FBIKBoneSettings : public UObject
+{
+	GENERATED_BODY()
+
+public:
+	
+	UPROPERTY()
+	FName Bone;
+	UPROPERTY()
+	float RotationStiffness = 0.0f;
+	UPROPERTY()
+	float PositionStiffness = 0.0f;
+	UPROPERTY()
+	EPBIKLimitType X;
+	UPROPERTY()
+	float MinX = 0.0f;
+	UPROPERTY()
+	float MaxX = 0.0f;
+	UPROPERTY()
+	EPBIKLimitType Y;
+	UPROPERTY()
+	float MinY = 0.0f;
+	UPROPERTY()
+	float MaxY = 0.0f;
+	UPROPERTY()
+	EPBIKLimitType Z;
+	UPROPERTY()
+	float MinZ = 0.0f;
+	UPROPERTY()
+	float MaxZ = 0.0f;
+	UPROPERTY()
+	bool bUsePreferredAngles = false;
+	UPROPERTY()
+	FVector PreferredAngles;
+};
+
+UCLASS(Deprecated)
+class
+UE_DEPRECATED(5.6, "This type has been replaced with FFBIKSolver.")
+IKRIG_API UDEPRECATED_IKRigFBIKSolver : public UDEPRECATED_IKRigSolver
+{
+	GENERATED_BODY()
+
+public:
+
+	// upgrade patch to new UStruct-based solver
+	virtual void ConvertToInstancedStruct(FInstancedStruct& OutInstancedStruct) override
+	{
+		OutInstancedStruct.InitializeAs(FIKRigFullBodyIKSolver::StaticStruct());
+		
+		// load solver settings
+		FIKRigFullBodyIKSolver& NewSolver = OutInstancedStruct.GetMutable<FIKRigFullBodyIKSolver>();
+		NewSolver.SetEnabled(IsEnabled());
+		NewSolver.Settings.RootBone = RootBone;
+		NewSolver.Settings.Iterations = Iterations;
+		NewSolver.Settings.SubIterations = SubIterations;
+		NewSolver.Settings.MassMultiplier = MassMultiplier;
+		NewSolver.Settings.bAllowStretch = bAllowStretch;
+		NewSolver.Settings.RootBehavior = RootBehavior;
+		NewSolver.Settings.PrePullRootSettings = PrePullRootSettings;
+		NewSolver.Settings.GlobalPullChainAlpha = PullChainAlpha;
+		NewSolver.Settings.MaxAngle = MaxAngle;
+		NewSolver.Settings.OverRelaxation = OverRelaxation;
+
+		// load effectors into solver
+		for (UDEPRECATED_IKRig_FBIKEffector* OldEffector : Effectors_DEPRECATED)
+		{
+			FIKRigFBIKGoalSettings NewEffector;
+			NewEffector.Goal = OldEffector->GoalName;
+			NewEffector.BoneName = OldEffector->BoneName;
+			NewEffector.ChainDepth = OldEffector->ChainDepth;
+			NewEffector.StrengthAlpha = OldEffector->StrengthAlpha;
+			NewEffector.PullChainAlpha = OldEffector->PullChainAlpha;
+			NewEffector.PinRotation = OldEffector->PinRotation;
+			NewSolver.AllGoalSettings.Add(NewEffector);
+		}
+		
+		// load bone settings into solver
+		for (UDEPRECATED_IKRig_FBIKBoneSettings* OldBoneSetting : BoneSettings_DEPRECATED)
+		{
+			FIKRigFBIKBoneSettings NewBoneSetting;
+			NewBoneSetting.Bone = OldBoneSetting->Bone;
+			NewBoneSetting.RotationStiffness = OldBoneSetting->RotationStiffness;
+			NewBoneSetting.PositionStiffness = OldBoneSetting->PositionStiffness;
+			NewBoneSetting.X = OldBoneSetting->X;
+			NewBoneSetting.MinX = OldBoneSetting->MinX;
+			NewBoneSetting.MaxX = OldBoneSetting->MaxX;
+			NewBoneSetting.Y = OldBoneSetting->Y;
+			NewBoneSetting.MinY = OldBoneSetting->MinY;
+			NewBoneSetting.MaxY = OldBoneSetting->MaxY;
+			NewBoneSetting.Z = OldBoneSetting->Z;
+			NewBoneSetting.MinZ = OldBoneSetting->MinZ;
+			NewBoneSetting.MaxZ = OldBoneSetting->MaxZ;
+			NewBoneSetting.bUsePreferredAngles = OldBoneSetting->bUsePreferredAngles;
+			NewBoneSetting.PreferredAngles = OldBoneSetting->PreferredAngles;
+			NewSolver.AllBoneSettings.Add(NewBoneSetting);
+		}
+	};
+	
+	UPROPERTY()
+	FName RootBone;
+	UPROPERTY()
+	int32 Iterations = 20;
+	UPROPERTY()
+	int32 SubIterations = 0;
+	UPROPERTY()
+	float MassMultiplier = 1.0f;
+	UPROPERTY()
+	bool bAllowStretch = false;
+	UPROPERTY()
+	EPBIKRootBehavior RootBehavior = EPBIKRootBehavior::PrePull;
+	UPROPERTY()
+	FRootPrePullSettings PrePullRootSettings;
+	UPROPERTY()
+	float PullChainAlpha = 1.0f;
+	UPROPERTY()
+	float MaxAngle = 30.f;
+	UPROPERTY()
+	float OverRelaxation = 1.3f;
+
+	UE_DEPRECATED(5.6, "Deprecated UObject based FBIK solver.")
+	UFUNCTION(meta = (DeprecatedFunction))
+	TArray<UDEPRECATED_IKRig_FBIKEffector*> GetEffectors()
+	{
+		return Effectors_DEPRECATED;
+	}
+
+	UE_DEPRECATED(5.6, "Deprecated UObject based FBIK solver.")
+	UFUNCTION(meta = (DeprecatedFunction))
+	TArray<UDEPRECATED_IKRig_FBIKBoneSettings*> GetBoneSettings()
+	{
+		return BoneSettings_DEPRECATED;
+	}
+
+	UE_DEPRECATED(5.6, "IK Rig refactored to not use UObjects")
+	UPROPERTY(meta = (DeprecatedProperty))
+	TArray<TObjectPtr<UDEPRECATED_IKRig_FBIKEffector>> Effectors_DEPRECATED;
+
+	UE_DEPRECATED(5.6, "IK Rig refactored to not use UObjects")
+	UPROPERTY(meta = (DeprecatedProperty))
+	TArray<TObjectPtr<UDEPRECATED_IKRig_FBIKBoneSettings>> BoneSettings_DEPRECATED;
+};
+
+PRAGMA_ENABLE_DEPRECATION_WARNINGS
diff --git a/Engine/Plugins/Animation/IKRig/Source/IKRig/Public/Rig/Solvers/IKRigSetTransform.h b/Engine/Plugins/Animation/IKRig/Source/IKRig/Public/Rig/Solvers/IKRigSetTransform.h
new file mode 100644
--- /dev/null
+++ b/Engine/Plugins/Animation/IKRig/Source/IKRig/Public/Rig/Solvers/IKRigSetTransform.h
@@ -0,0 +1,151 @@
+// Copyright Epic Games, Inc. All Rights Reserved.
+
+#pragma once
+
+#include "Rig/IKRigDefinition.h"
+#include "IKRigSolverBase.h"
+
+#include "IKRigSetTransform.generated.h"
+
+USTRUCT(BlueprintType)
+struct IKRIG_API FIKRigSetTransformSettings : public FIKRigSolverSettingsBase
+{
+	GENERATED_BODY()
+	
+	// The goal used to affect the bone transform
+	UPROPERTY(BlueprintReadOnly, VisibleAnywhere, Category = "Set Transform Settings")
+	FName Goal;
+
+	// The bone to affect
+	UPROPERTY(BlueprintReadOnly, VisibleAnywhere, Category = "Set Transform Settings")
+	FName BoneToAffect;
+
+	// Blend the translation on/off. Range is 0-1. Default is 1.0.
+	UPROPERTY(BlueprintReadWrite, EditAnywhere, Category = "Set Transform Effector", meta = (ClampMin = "0.0", ClampMax = "1.0", UIMin = "0.0", UIMax = "1.0"))
+	float PositionAlpha = 1.0;
+
+	// Blend the rotation on/off. Range is 0-1. Default is 1.0.
+	UPROPERTY(BlueprintReadWrite, EditAnywhere, Category = "Set Transform Effector", meta = (ClampMin = "0.0", ClampMax = "1.0", UIMin = "0.0", UIMax = "1.0"))
+	float RotationAlpha = 1.0f;
+
+	// Blend the total effect on/off. Range is 0-1. Default is 1.0.
+	UPROPERTY(BlueprintReadWrite, EditAnywhere, Category = "Set Transform Effector", meta = (ClampMin = "0.0", ClampMax = "1.0", UIMin = "0.0", UIMax = "1.0"))
+	float Alpha = 1.0f;
+
+	// If true, the transform modification will propagate to the hierarchy below the target bone.
+	UPROPERTY(EditAnywhere, Category = "Set Transform Settings")
+	bool bPropagateToChildren = true;
+};
+
+USTRUCT(BlueprintType)
+struct IKRIG_API FIKRigSetTransform : public FIKRigSolverBase
+{
+	GENERATED_BODY()
+	
+	UPROPERTY()
+	FIKRigSetTransformSettings Settings;
+
+	// runtime
+	virtual void Initialize(const FIKRigSkeleton& IKRigSkeleton) override;
+	virtual void Solve(FIKRigSkeleton& InIKRigSkeleton, const FIKRigGoalContainer& InGoals) override;
+	virtual void GetRequiredBones(TSet<FName>& OutRequiredBones) const override;
+	virtual void GetRequiredGoals(TSet<FName>& OutRequiredGoals) const override;
+
+	// settings
+	virtual FIKRigSolverSettingsBase* GetSolverSettings() override;
+	virtual const UScriptStruct* GetSolverSettingsType() const override;
+	
+	// goals
+	virtual void AddGoal(const UIKRigEffectorGoal& InNewGoal) override;
+	virtual void OnGoalRemoved(const FName& InGoalName) override;
+	virtual void OnGoalRenamed(const FName& InOldName, const FName& InNewName) override;
+	virtual void OnGoalMovedToDifferentBone(const FName& InGoalName, const FName& InNewBoneName) override;
+	
+#if WITH_EDITOR
+	// return custom controller for scripting this solver
+	virtual UIKRigSolverControllerBase* GetSolverController(UObject* Outer) override;
+	// ui
+	virtual FText GetNiceName() const override;
+	virtual bool GetWarningMessage(FText& OutWarningMessage) const override;
+	virtual bool IsBoneAffectedBySolver(const FName& InBoneName, const FIKRigSkeleton& InIKRigSkeleton) const override;
+#endif
+
+private:
+
+	int32 BoneIndex = INDEX_NONE;
+};
+
+/* The blueprint/python API for modifying Set Transform settings in an IK Rig. */
+UCLASS(BlueprintType)
+class IKRIG_API UIKRigSetTransformController : public UIKRigSolverControllerBase
+{
+	GENERATED_BODY()
+	
+public:
+
+	/* Get the current solver settings as a struct.
+	 * @return FIKRigSetTransformSettings struct with the current settings used by the solver. */
+	UFUNCTION(BlueprintCallable, Category = IKRig)
+	FIKRigSetTransformSettings GetSolverSettings();
+
+	/* Set the solver settings. Input is a custom struct type for this solver.
+	 * @param InSettings a FIKRigSetTransformSettings struct containing all the settings to apply to this solver */
+	UFUNCTION(BlueprintCallable, Category = IKRig)
+	void SetSolverSettings(FIKRigSetTransformSettings InSettings);
+};
+
+PRAGMA_DISABLE_DEPRECATION_WARNINGS
+
+UCLASS(Deprecated)
+class
+UE_DEPRECATED(5.6, "This type has been replaced with FSetTransformGoalSettings.")
+IKRIG_API UDEPRECATED_IKRig_SetTransformEffector : public UObject
+{
+	GENERATED_BODY()
+
+public:
+	
+	UPROPERTY()
+	bool bEnablePosition = true;
+	UPROPERTY()
+	bool bEnableRotation = true;
+	UPROPERTY()
+	float Alpha = 1.0f;
+};
+
+UCLASS(Deprecated)
+class
+UE_DEPRECATED(5.6, "This type has been replaced with FSetTransform.")
+IKRIG_API UDEPRECATED_IKRig_SetTransform : public UDEPRECATED_IKRigSolver
+{
+	GENERATED_BODY()
+
+public:
+	
+	UDEPRECATED_IKRig_SetTransform()
+	{
+		Effector_DEPRECATED = CreateDefaultSubobject<UDEPRECATED_IKRig_SetTransformEffector>(TEXT("Effector"));
+	}
+
+	// upgrade patch to new UStruct-based solver
+	virtual void ConvertToInstancedStruct(FInstancedStruct& OutInstancedStruct) override
+	{
+		OutInstancedStruct.InitializeAs(FIKRigSetTransform::StaticStruct());
+		FIKRigSetTransform& NewSolver = OutInstancedStruct.GetMutable<FIKRigSetTransform>();
+		NewSolver.SetEnabled(IsEnabled());
+		NewSolver.Settings.Goal = Goal;
+		NewSolver.Settings.BoneToAffect = RootBone;
+		NewSolver.Settings.Alpha = Effector_DEPRECATED->Alpha;
+		NewSolver.Settings.PositionAlpha = Effector_DEPRECATED->bEnablePosition ? 1.0f : 0.0f;
+		NewSolver.Settings.RotationAlpha = Effector_DEPRECATED->bEnableRotation ? 1.0f : 0.0f;
+	};
+
+	UPROPERTY()
+	FName Goal;
+	UPROPERTY()
+	FName RootBone;
+	UPROPERTY(meta = (DeprecatedProperty))
+	TObjectPtr<UDEPRECATED_IKRig_SetTransformEffector> Effector_DEPRECATED;
+};
+
+PRAGMA_ENABLE_DEPRECATION_WARNINGS
diff --git a/Engine/Plugins/Animation/IKRig/Source/IKRig/Public/Rig/Solvers/IKRig_FBIKSolver.h b/Engine/Plugins/Animation/IKRig/Source/IKRig/Public/Rig/Solvers/IKRig_FBIKSolver.h
deleted file mode 100644
diff --git a/Engine/Plugins/Animation/IKRig/Source/IKRig/Public/Rig/Solvers/SetTransformSolver.h b/Engine/Plugins/Animation/IKRig/Source/IKRig/Public/Rig/Solvers/SetTransformSolver.h
deleted file mode 100644
diff --git a/Engine/Plugins/Animation/IKRig/Source/IKRigEditor/Private/RigEditor/IKRigAutoFBIK.cpp b/Engine/Plugins/Animation/IKRig/Source/IKRigEditor/Private/RigEditor/IKRigAutoFBIK.cpp
--- a/Engine/Plugins/Animation/IKRig/Source/IKRigEditor/Private/RigEditor/IKRigAutoFBIK.cpp
+++ b/Engine/Plugins/Animation/IKRig/Source/IKRigEditor/Private/RigEditor/IKRigAutoFBIK.cpp
@@ -3,7 +3,7 @@
 #include "RigEditor/IKRigAutoFBIK.h"
 
 #include "IKRigEditor.h"
-#include "Rig/Solvers/IKRig_FBIKSolver.h"
+#include "Rig/Solvers/IKRigFullBodyIK.h"
 #include "RigEditor/IKRigAutoCharacterizer.h"
 #include "RigEditor/IKRigController.h"
 #include "Templates/SubclassOf.h"
diff --git a/Engine/Plugins/Animation/IKRig/Source/IKRigEditor/Private/RigEditor/IKRigEditorController.cpp b/Engine/Plugins/Animation/IKRig/Source/IKRigEditor/Private/RigEditor/IKRigEditorController.cpp
--- a/Engine/Plugins/Animation/IKRig/Source/IKRigEditor/Private/RigEditor/IKRigEditorController.cpp
+++ b/Engine/Plugins/Animation/IKRig/Source/IKRigEditor/Private/RigEditor/IKRigEditorController.cpp
@@ -10,7 +10,7 @@
 #include "RigEditor/IKRigToolkit.h"
 #include "RigEditor/SIKRigAssetBrowser.h"
 #include "RigEditor/SIKRigOutputLog.h"
-#include "Rig/Solvers/IKRig_FBIKSolver.h"
+#include "Rig/Solvers/IKRigFullBodyIK.h"
 #include "Widgets/Input/SComboBox.h"
 #include "ContentBrowserModule.h"
 #include "IContentBrowserSingleton.h"
