IKRig Solver: updated warning message order for consistency

#rnx

diff --git a/Engine/Plugins/Animation/IKRig/Source/IKRig/Private/Rig/Solvers/IKRigLimbSolver.cpp b/Engine/Plugins/Animation/IKRig/Source/IKRig/Private/Rig/Solvers/IKRigLimbSolver.cpp
--- a/Engine/Plugins/Animation/IKRig/Source/IKRig/Private/Rig/Solvers/IKRigLimbSolver.cpp
+++ b/Engine/Plugins/Animation/IKRig/Source/IKRig/Private/Rig/Solvers/IKRigLimbSolver.cpp
@@ -207,15 +207,15 @@ FText FIKRigLimbSolver::GetNiceName() const
 
 bool FIKRigLimbSolver::GetWarningMessage(FText& OutWarningMessage) const
 {
-	if (Settings.GoalName == NAME_None)
+	if (Settings.StartBone == NAME_None)
 	{
-		OutWarningMessage = LOCTEXT("MissingGoal", "Missing goal.");
+		OutWarningMessage = LOCTEXT("MissingRoot", "Missing root.");
 		return true;
 	}
 
-	if (Settings.StartBone == NAME_None)
+	if (Settings.GoalName == NAME_None)
 	{
-		OutWarningMessage = LOCTEXT("MissingRoot", "Missing root.");
+		OutWarningMessage = LOCTEXT("MissingGoal", "Missing goal.");
 		return true;
 	}
 
