[IKRig] Major refactor part 2 of 3.

- Replaced UIKRigProcessorand UIKRetargetProcessor with struct based processors
- Ported all retarget ops to struct based equivalents with a PostLoad patch
- Modified the retargeting UI to reflect the new Op-based workflow
- Fixed crash when undo'ing details panel edits

diff --git a/Engine/Plugins/Animation/IKRig/Source/IKRig/Private/AnimNodes/AnimNode_IKRig.cpp b/Engine/Plugins/Animation/IKRig/Source/IKRig/Private/AnimNodes/AnimNode_IKRig.cpp
--- a/Engine/Plugins/Animation/IKRig/Source/IKRig/Private/AnimNodes/AnimNode_IKRig.cpp
+++ b/Engine/Plugins/Animation/IKRig/Source/IKRig/Private/AnimNodes/AnimNode_IKRig.cpp
@@ -2,19 +2,19 @@
 
 #include "AnimNodes/AnimNode_IKRig.h"
 
+#include "Engine/SkeletalMesh.h"
 #include "Animation/AnimInstance.h"
 #include "Animation/AnimTrace.h"
+#include "Animation/AnimInstanceProxy.h"
+#include "Animation/AnimStats.h"
+
 #include "Rig/IKRigProcessor.h"
 #include "Rig/IKRigDefinition.h"
-#include "Rig/Solvers/IKRigSolverBase.h"
 
 #include "Components/SkeletalMeshComponent.h"
 #include "ActorComponents/IKRigInterface.h"
-#include "Animation/AnimInstanceProxy.h"
-#include "Animation/AnimStats.h"
-#include "Algo/ForEach.h"
+
 #include "PrimitiveDrawingUtils.h"
-#include "Engine/SkeletalMesh.h"
 #include "Logging/LogVerbosity.h"
 #include "VisualLogger/VisualLogger.h"
 
@@ -40,22 +40,18 @@ void FAnimNode_IKRig::Evaluate_AnyThread(FPoseContext& Output)
 	{
 		Source.Evaluate(Output);
 	}
-
-	if (!(RigDefinitionAsset && IKRigProcessor))
+	
+	if (!FAnimWeight::IsRelevant(ActualAlpha))
 	{
 		return;
 	}
-
+	
+	// initialize the IK Rig (will only try once on the current version of the rig asset)
 	if (!IKRigProcessor->IsInitialized())
 	{
-		return;
+		IKRigProcessor->Initialize(RigDefinitionAsset, Output.AnimInstanceProxy->GetSkelMeshComponent()->GetSkeletalMeshAsset());
 	}
-
-	if (!FAnimWeight::IsRelevant(ActualAlpha))
-	{
-		return;
-	}
-
+	
 	// copy input pose to solver stack
 	CopyInputPoseToSolver(Output.Pose);
 	// update target goal transforms
@@ -107,7 +103,6 @@ void FAnimNode_IKRig::AssignGoalTargets()
 	// NOTE: as the transaction when undoing/redoing can be called on the preview scene before the editor, the processor
 	// might not have been reinitialized, resulting in data being desynchronized. Thus, we must wait until the
 	// transaction has been fully processed.
-
 	if (bDriveWithSourceAsset && !GIsTransacting)
 	{
 		IKRigProcessor->CopyAllSettingsFromAsset(RigDefinitionAsset);
@@ -184,37 +179,43 @@ void FAnimNode_IKRig::Initialize_AnyThread(const FAnimationInitializeContext& Co
 
 	// Initial update of the node, so we dont have a frame-delay on setup
 	GetEvaluateGraphExposedInputs().Execute(Context);
-	
-	CreateIKRigProcessorIfNeeded(Context.AnimInstanceProxy->GetSkelMeshComponent());
 
+	// initialize all the FProperty pointers to copy dynamic anim node inputs
 	InitializeProperties(Context.GetAnimInstanceObject(), GetTargetClass());
+
+	// trigger another search for goal creators on the main thread
+	bSearchedForGoalCreators = false;
+
+	// NOTE: this must be instantiated inside Initialize()
+	IKRigProcessor = MakeShared<FIKRigProcessor>();
 }
 
 void FAnimNode_IKRig::Update_AnyThread(const FAnimationUpdateContext& Context)
 {
-	ActualAlpha = 0.f;
-
 	GetEvaluateGraphExposedInputs().Execute(Context);
 
-	// alpha handlers
-	switch (AlphaInputType)
+	// update actual alpha
 	{
-	case EAnimAlphaInputType::Float:
-		ActualAlpha = AlphaScaleBias.ApplyTo(AlphaScaleBiasClamp.ApplyTo(Alpha, Context.GetDeltaTime()));
-		break;
-	case EAnimAlphaInputType::Bool:
-		ActualAlpha = AlphaBoolBlend.ApplyTo(bAlphaBoolEnabled, Context.GetDeltaTime());
-		break;
-	case EAnimAlphaInputType::Curve:
-		if (UAnimInstance* AnimInstance = Cast<UAnimInstance>(Context.AnimInstanceProxy->GetAnimInstanceObject()))
-		{
-			ActualAlpha = AlphaScaleBiasClamp.ApplyTo(AnimInstance->GetCurveValue(AlphaCurveName), Context.GetDeltaTime());
-		}
-		break;
-	};
+		ActualAlpha = 0.f;
 
-	// Make sure Alpha is clamped between 0 and 1.
-	ActualAlpha = FMath::Clamp<float>(ActualAlpha, 0.f, 1.f);
+		switch (AlphaInputType)
+		{
+		case EAnimAlphaInputType::Float:
+			ActualAlpha = AlphaScaleBias.ApplyTo(AlphaScaleBiasClamp.ApplyTo(Alpha, Context.GetDeltaTime()));
+			break;
+		case EAnimAlphaInputType::Bool:
+			ActualAlpha = AlphaBoolBlend.ApplyTo(bAlphaBoolEnabled, Context.GetDeltaTime());
+			break;
+		case EAnimAlphaInputType::Curve:
+			if (UAnimInstance* AnimInstance = Cast<UAnimInstance>(Context.AnimInstanceProxy->GetAnimInstanceObject()))
+			{
+				ActualAlpha = AlphaScaleBiasClamp.ApplyTo(AnimInstance->GetCurveValue(AlphaCurveName), Context.GetDeltaTime());
+			}
+			break;
+		};
+		
+		ActualAlpha = FMath::Clamp<float>(ActualAlpha, 0.f, 1.f);
+	}
 
 	PropagateInputProperties(Context.AnimInstanceProxy->GetAnimInstanceObject());
 
@@ -224,7 +225,7 @@ void FAnimNode_IKRig::Update_AnyThread(const FAnimationUpdateContext& Context)
 #if ENABLE_VISUAL_LOG
 #if WITH_EDITORONLY_DATA
 	// is node setup?
-	if (IKRigProcessor && IKRigProcessor->IsInitialized() && FVisualLogger::IsRecording())
+	if (IKRigProcessor->IsInitialized() && FVisualLogger::IsRecording())
 	{
 		static const FBox UnitBox(FVector(-1, -1, -1), FVector(1, 1, 1));
 
@@ -246,38 +247,26 @@ void FAnimNode_IKRig::Update_AnyThread(const FAnimationUpdateContext& Context)
 
 void FAnimNode_IKRig::PreUpdate(const UAnimInstance* InAnimInstance)
 {
-	if (!IsValid(RigDefinitionAsset))
-	{
-		return;
-	}
-	
-	if (!IsValid(IKRigProcessor))
-	{
-		CreateIKRigProcessorIfNeeded(InAnimInstance->GetOwningComponent());
-	}
-	
-	// initialize the IK Rig (will only try once on the current version of the rig asset)
-	if (!IKRigProcessor->IsInitialized())
-	{
- 		IKRigProcessor->Initialize(RigDefinitionAsset, InAnimInstance->GetSkelMeshComponent()->GetSkeletalMeshAsset());
-	}
-	
-	// cache list of goal creator components on the actor
-	// TODO tried doing this in Initialize_AnyThread but it would miss some GoalCreator components
-	// so it was moved here to be more robust, but we need to profile this and make sure it's not hurting perf
-	// (it may be enough to run this once and then never again...needs testing)
-	GoalCreators.Reset();
-	USkeletalMeshComponent* SkelMeshComponent = InAnimInstance->GetSkelMeshComponent();
-	AActor* OwningActor = SkelMeshComponent->GetOwner();
-	TArray<UActorComponent*> GoalCreatorComponents =  OwningActor->GetComponentsByInterface( UIKGoalCreatorInterface::StaticClass() );
-	for (UActorComponent* GoalCreatorComponent : GoalCreatorComponents)
+	// cache list of goal creator components on the actor (this is only done once)
+	// NOTE: tried doing this in Initialize_AnyThread but it would miss some GoalCreator components
+	// so it was moved here to be more robust
+	if (!bSearchedForGoalCreators)
 	{
-		IIKGoalCreatorInterface* GoalCreator = Cast<IIKGoalCreatorInterface>(GoalCreatorComponent);
-		if (!ensureMsgf(GoalCreator, TEXT("Goal creator component failed cast to IIKGoalCreatorInterface.")))
+		bSearchedForGoalCreators = true;
+
+		GoalCreators.Reset();
+		USkeletalMeshComponent* SkelMeshComponent = InAnimInstance->GetSkelMeshComponent();
+		AActor* OwningActor = SkelMeshComponent->GetOwner();
+		TArray<UActorComponent*> GoalCreatorComponents =  OwningActor->GetComponentsByInterface( UIKGoalCreatorInterface::StaticClass() );
+		for (UActorComponent* GoalCreatorComponent : GoalCreatorComponents)
 		{
-			continue;
+			IIKGoalCreatorInterface* GoalCreator = Cast<IIKGoalCreatorInterface>(GoalCreatorComponent);
+			if (!ensureMsgf(GoalCreator, TEXT("Goal creator component failed cast to IIKGoalCreatorInterface.")))
+			{
+				continue;
+			}
+			GoalCreators.Add(GoalCreator);
 		}
-		GoalCreators.Add(GoalCreator);
 	}
 	
 	// pull all the goals out of any goal creators on the owning actor
@@ -296,7 +285,7 @@ void FAnimNode_IKRig::OnInitializeAnimInstance(const FAnimInstanceProxy* InProxy
 	InitializeProperties(InAnimInstance, GetTargetClass());
 }
 
-void FAnimNode_IKRig::SetProcessorNeedsInitialized()
+void FAnimNode_IKRig::SetProcessorNeedsInitialized() const
 {
 	if (IKRigProcessor)
 	{
@@ -304,18 +293,9 @@ void FAnimNode_IKRig::SetProcessorNeedsInitialized()
 	}
 }
 
-void FAnimNode_IKRig::CreateIKRigProcessorIfNeeded(UObject* Outer)
-{
-	// ensure there is always a processor available
-	if (!IKRigProcessor && IsInGameThread())
-	{
-		IKRigProcessor = NewObject<UIKRigProcessor>(Outer);	
-	}
-}
-
-UIKRigProcessor* FAnimNode_IKRig::GetIKRigProcessor()
+FIKRigProcessor* FAnimNode_IKRig::GetIKRigProcessor()
 {
-	return IKRigProcessor;
+	return IKRigProcessor.Get();
 }
 
 void FAnimNode_IKRig::CacheBones_AnyThread(const FAnimationCacheBonesContext& Context)
@@ -447,36 +427,36 @@ void FAnimNode_IKRig::PropagateInputProperties(const UObject* InSourceInstance)
 	{
 		return;
 	}
-	
-	Algo::ForEach(UpdateFunctions, [InSourceInstance](const PropertyUpdateFunction& InFunc)
+
+	for (const TFunction<void(const UObject*)>& UpdateFunc : UpdateFunctions)
 	{
-		InFunc(InSourceInstance);
-	});
+		UpdateFunc(InSourceInstance);
+	}
 }
 
+#if WITH_EDITOR
 void FAnimNode_IKRig::ConditionalDebugDraw(
 	FPrimitiveDrawInterface* PDI,
 	USkeletalMeshComponent* PreviewSkelMeshComp) const
 {
-#if WITH_EDITOR
-
 	// is anim graph setup?
 	if (!(bEnableDebugDraw && PreviewSkelMeshComp && PreviewSkelMeshComp->GetWorld()))
 	{
 		return;
 	}
 
-	// is node setup?
-	if (!(RigDefinitionAsset && IKRigProcessor && IKRigProcessor->IsInitialized()))
+	// is processor running?
+	if (!IKRigProcessor->IsInitialized())
 	{
 		return;
 	}
 
+	// draw oriented box for each goal
 	const TArray<FIKRigGoal>& ProcessorGoals = IKRigProcessor->GetGoalContainer().GetGoalArray();
 	for (const FIKRigGoal& Goal : ProcessorGoals)
 	{
 		DrawOrientedWireBox(PDI, Goal.FinalBlendedPosition, FVector::XAxisVector, FVector::YAxisVector, FVector::ZAxisVector, FVector::One() * DebugScale, FLinearColor::Yellow, SDPG_World);
 		DrawCoordinateSystem(PDI, Goal.Position, Goal.FinalBlendedRotation.Rotator(), DebugScale, SDPG_World);
 	}
-#endif
 }
+#endif
diff --git a/Engine/Plugins/Animation/IKRig/Source/IKRig/Private/AnimNodes/AnimNode_RetargetPoseFromMesh.cpp b/Engine/Plugins/Animation/IKRig/Source/IKRig/Private/AnimNodes/AnimNode_RetargetPoseFromMesh.cpp
--- a/Engine/Plugins/Animation/IKRig/Source/IKRig/Private/AnimNodes/AnimNode_RetargetPoseFromMesh.cpp
+++ b/Engine/Plugins/Animation/IKRig/Source/IKRig/Private/AnimNodes/AnimNode_RetargetPoseFromMesh.cpp
@@ -26,6 +26,8 @@ void FAnimNode_RetargetPoseFromMesh::Initialize_AnyThread(const FAnimationInitia
 	{
 		Source.Initialize(Context);
 	}
+
+	Processor = MakeShared<FIKRetargetProcessor>();
 }
 
 void FAnimNode_RetargetPoseFromMesh::CacheBones_AnyThread(const FAnimationCacheBonesContext& Context)
@@ -80,10 +82,11 @@ void FAnimNode_RetargetPoseFromMesh::Evaluate_AnyThread(FPoseContext& Output)
 		{
 			return false; // must have retarget asset
 		}
-
-		if (!Processor)
+		
+		// ensure the processor is initialized
+		if (!EnsureProcessorIsInitialized(Output.AnimInstanceProxy->GetSkelMeshComponent()))
 		{
-			return false; // must have processor
+			return false; // must be initialized
 		}
 
 		const bool bCopyingFromOtherComponent = RetargetFrom != ERetargetSourceMode::SourcePosePin;
@@ -101,17 +104,6 @@ void FAnimNode_RetargetPoseFromMesh::Evaluate_AnyThread(FPoseContext& Output)
 		{
 			return false; // if copying from input pin, must be connected
 		}
-		
-		// ensure processor was initialized with the currently used assets (source/target meshes and retarget asset)
-		// if processor is not ready this tick, it will be next tick as this state will trigger re-initialization
-		// NOTE: do in-line reinitialization here once all runtime UObjects are removed (can't right now because UObject creation not allowed off the main thread)
-		const USkeletalMesh* TargetMesh = Output.AnimInstanceProxy->GetSkelMeshComponent()->GetSkeletalMeshAsset();
-		const USkeletalMesh* SourceMesh = bCopyingFromOtherComponent ? SourceMeshComponent->GetSkeletalMeshAsset() : TargetMesh;
-		const bool bIsProcessorReady = Processor->WasInitializedWithTheseAssets(SourceMesh, TargetMesh, IKRetargeterAsset);
-		if (!bIsProcessorReady)
-		{
-			return false;
-		}
 
 		if (!IsLODEnabled(Output.AnimInstanceProxy))
 		{
@@ -211,7 +203,7 @@ void FAnimNode_RetargetPoseFromMesh::Evaluate_AnyThread(FPoseContext& Output)
 		}
 	};
 
-	// remap curves for CurveRemapOp if one is present in the RetargetOp stack
+	// remap curves as required by any retarget ops that request it
 	auto CopyAndRemapCurvesFromSourceToTarget = [this](FBlendedCurve& OutputCurves)
 	{
 		if (!bCopyCurves)
@@ -219,7 +211,7 @@ void FAnimNode_RetargetPoseFromMesh::Evaluate_AnyThread(FPoseContext& Output)
 			return;
 		}
 	
-		if (!(Processor && Processor->IsInitialized()))
+		if (!Processor->IsInitialized())
 		{
 			return;
 		}
@@ -227,32 +219,28 @@ void FAnimNode_RetargetPoseFromMesh::Evaluate_AnyThread(FPoseContext& Output)
 		// copy curves over to same name on target (if it exists)
 		OutputCurves.CopyFrom(SourceCurves);
 	
-		// now remap curves that have different names using the CurveRemapOp data (if there are any)
+		// now query retarget ops for any curves that need remapped from source to target
 		FBlendedCurve RemapedCurves;
-		const TArray<TObjectPtr<URetargetOpBase>>& RetargetOps = Processor->GetRetargetOps();
-		for (URetargetOpBase* RetargetOp : RetargetOps)
+		TArray<FCurveRemapPair> CurvesToRetarget;
+		const TArray<TSharedPtr<FIKRetargetOpBase>>& RetargetOps = Processor->GetRetargetOps();
+		for (const TSharedPtr<FIKRetargetOpBase>& RetargetOp : RetargetOps)
 		{
-			UCurveRemapOp* CurveRemapOp = Cast<UCurveRemapOp>(RetargetOp);
-			if (!CurveRemapOp)
-			{
-				continue;
-			}
-
-			if (!CurveRemapOp->bIsEnabled)
+			if (RetargetOp->IsEnabled())
 			{
-				continue;
+				RetargetOp->GetCurvesToRetarget(CurvesToRetarget);	
 			}
+		}
 		
-			for (const FCurveRemapPair& CurveToRemap : CurveRemapOp->CurvesToRemap)
+		for (const FCurveRemapPair& CurveToRemap : CurvesToRetarget)
+		{
+			bool bOutIsValid = false;
+			const float SourceValue = SourceCurves.Get(CurveToRemap.SourceCurve, bOutIsValid);
+			if (bOutIsValid)
 			{
-				bool bOutIsValid = false;
-				const float SourceValue = SourceCurves.Get(CurveToRemap.SourceCurve, bOutIsValid);
-				if (bOutIsValid)
-				{
-					RemapedCurves.Add(CurveToRemap.TargetCurve, SourceValue);
-				}
+				RemapedCurves.Add(CurveToRemap.TargetCurve, SourceValue);
 			}
 		}
+
 		OutputCurves.Combine(RemapedCurves);
 	};
 
@@ -406,16 +394,9 @@ void FAnimNode_RetargetPoseFromMesh::PreUpdate(const UAnimInstance* InAnimInstan
 		}
 	};
 
-	// initialization of the retarget processor must be on the main thread
-	// NOTE: once UObject creation has been refactored out of IK Rig and IK Retargeter we can do this in-line off the main thread
-	USkeletalMeshComponent* TargetMeshComponent = InAnimInstance->GetSkelMeshComponent();
-	if (!TryInitializeProcessor(TargetMeshComponent))
-	{
-		return;
-	}
-
 	// copy all the data from the source component
-	if (RetargetFrom != ERetargetSourceMode::SourcePosePin)
+	USkeletalMeshComponent* TargetMeshComponent = InAnimInstance->GetSkelMeshComponent();
+	if (Processor->IsInitialized() && RetargetFrom != ERetargetSourceMode::SourcePosePin)
 	{
 		GetPoseToRetargetFromSourceMesh(TargetMeshComponent);
 		CopyCurvesFromSourceMesh();
@@ -423,32 +404,18 @@ void FAnimNode_RetargetPoseFromMesh::PreUpdate(const UAnimInstance* InAnimInstan
 	}
 }
 
-void FAnimNode_RetargetPoseFromMesh::CreateRetargetProcessorIfNeeded(UObject* Outer)
+FIKRetargetProcessor* FAnimNode_RetargetPoseFromMesh::GetRetargetProcessor()
 {
-	if (!Processor && IsInGameThread())
-	{
-		Processor = NewObject<UIKRetargetProcessor>(Outer);	
-	}
-}
-
-UIKRetargetProcessor* FAnimNode_RetargetPoseFromMesh::GetRetargetProcessor() const
-{
-	return Processor;
+	return Processor.Get();
 }
 
-bool FAnimNode_RetargetPoseFromMesh::TryInitializeProcessor(USkeletalMeshComponent* TargetMeshComponent)
+bool FAnimNode_RetargetPoseFromMesh::EnsureProcessorIsInitialized(const USkeletalMeshComponent* TargetMeshComponent)
 {
-	if (!ensure(IsInGameThread()))
+	if (!ensure(TargetMeshComponent))
 	{
-		// NOTE: remove this when processor no longer needs game thread to initialize
 		return false;
 	}
 	
-	if (!Processor)
-	{
-		Processor = NewObject<UIKRetargetProcessor>(TargetMeshComponent);	
-	}
-	
 	// has user supplied a retargeter asset?
 	if (!IKRetargeterAsset)
 	{
@@ -495,7 +462,6 @@ bool FAnimNode_RetargetPoseFromMesh::TryInitializeProcessor(USkeletalMeshCompone
 	if (!Processor->WasInitializedWithTheseAssets(SourceMesh, TargetMesh, IKRetargeterAsset))
 	{
 		// initialize retarget processor with source and target skeletal meshes
-		// (asset is passed in as outer UObject for new UIKRigProcessor)
 		bool bForceIKOff = false;
 		FRetargetProfile RetargetProfileToUse = GetMergedRetargetProfile(bForceIKOff);
 		Processor->Initialize(SourceMesh, TargetMesh, IKRetargeterAsset, RetargetProfileToUse);
diff --git a/Engine/Plugins/Animation/IKRig/Source/IKRig/Private/IKRigLogger.cpp b/Engine/Plugins/Animation/IKRig/Source/IKRig/Private/IKRigLogger.cpp
--- a/Engine/Plugins/Animation/IKRig/Source/IKRig/Private/IKRigLogger.cpp
+++ b/Engine/Plugins/Animation/IKRig/Source/IKRig/Private/IKRigLogger.cpp
@@ -2,21 +2,40 @@
 
 #include "IKRigLogger.h"
 #include "Logging/MessageLog.h"
+#include "Retargeter/IKRetargeter.h"
+#include "Rig/IKRigDefinition.h"
 
 #if WITH_EDITOR
 #include "Modules/ModuleManager.h"
 #include "MessageLogModule.h"
 #endif
 
+#define LOCTEXT_NAMESPACE "FIKRigLogger"
+
 DEFINE_LOG_CATEGORY(LogIKRig);
 
-void FIKRigLogger::SetLogTarget(const FName InLogName, const FText& LogLabel)
+void FIKRigLogger::SetLogTarget(const UObject* InAsset)
 {
-	LogName = InLogName;
+	FText LogLabel;
+	if (const UIKRigDefinition* IKRig = Cast<UIKRigDefinition>(InAsset))
+	{
+		LogName = FName("IKRig_", IKRig->GetUniqueID());
+		LogLabel =  LOCTEXT("IKRigLogName", "IK Rig Log");
+	}
+	else if (const UIKRetargeter* Retargeter = Cast<UIKRetargeter>(InAsset))
+	{
+		LogName = FName("IKRetarget_", Retargeter->GetUniqueID());
+		LogLabel = LOCTEXT("IKRetargetLogName", "IK Retarget Log");
+	}
+	else
+	{
+		checkNoEntry();
+		return;
+	}
 
 #if WITH_EDITOR
 	FMessageLogModule& MessageLogModule = FModuleManager::LoadModuleChecked<FMessageLogModule>("MessageLog");
-	if (!MessageLogModule.IsRegisteredLogListing(InLogName))
+	if (!MessageLogModule.IsRegisteredLogListing(LogName))
 	{
 		FMessageLogInitializationOptions InitOptions;
 		InitOptions.bShowFilters = false;
@@ -24,7 +43,7 @@ void FIKRigLogger::SetLogTarget(const FName InLogName, const FText& LogLabel)
 		InitOptions.bAllowClear = false;
 		InitOptions.bShowInLogWindow = false;
 		InitOptions.bDiscardDuplicates = true;
-		MessageLogModule.RegisterLogListing(InLogName, LogLabel, InitOptions);
+		MessageLogModule.RegisterLogListing(LogName, LogLabel, InitOptions);
 	}
 #endif
 	
@@ -34,7 +53,7 @@ FName FIKRigLogger::GetLogTarget() const
 {
 	return LogName;
 }
-	
+
 void FIKRigLogger::LogError(const FText& Message) const
 {
 	FMessageLog(LogName).SuppressLoggingToOutputLog(true).Error(Message);
@@ -59,3 +78,5 @@ void FIKRigLogger::Clear() const
 	Warnings.Empty();
 	Messages.Empty();
 }
+
+#undef LOCTEXT_NAMESPACE
diff --git a/Engine/Plugins/Animation/IKRig/Source/IKRig/Private/Retargeter/IKRetargetProcessor.cpp b/Engine/Plugins/Animation/IKRig/Source/IKRig/Private/Retargeter/IKRetargetProcessor.cpp
--- a/Engine/Plugins/Animation/IKRig/Source/IKRig/Private/Retargeter/IKRetargetProcessor.cpp
+++ b/Engine/Plugins/Animation/IKRig/Source/IKRig/Private/Retargeter/IKRetargetProcessor.cpp
@@ -1931,19 +1931,15 @@ FVector FPoleVectorMatcher::GetChainAxisNormalized(
 	return ChainAxis.GetSafeNormal();
 }
 
-UIKRetargetProcessor::UIKRetargetProcessor()
-{
-	const FName LogName = FName("IKRetarget_",GetUniqueID());
-	Log.SetLogTarget(LogName, LOCTEXT("IKRetargetLogName", "IK Retarget Log"));
-}
-
-void UIKRetargetProcessor::Initialize(
+void FIKRetargetProcessor::Initialize(
 		const USkeletalMesh* SourceSkeletalMesh,
 		const USkeletalMesh* TargetSkeletalMesh,
 		const UIKRetargeter* InRetargeterAsset,
 		const FRetargetProfile& Settings,
 		const bool bSuppressWarnings)
 {
+	check(InRetargeterAsset && SourceSkeletalMesh && TargetSkeletalMesh);
+	
 	// don't attempt reinitialization unless inputs have changed
 	const bool bSameTargetMesh = TargetSkeleton.SkeletalMesh == TargetSkeletalMesh;
 	const bool bSameSourceMesh = SourceSkeleton.SkeletalMesh == SourceSkeletalMesh;
@@ -1963,6 +1959,9 @@ void UIKRetargetProcessor::Initialize(
 	// record source asset
 	RetargeterAsset = InRetargeterAsset;
 
+	// set log target to unique ID for this asset
+	Log.SetLogTarget(RetargeterAsset);
+
 	const UIKRigDefinition* SourceIKRig = RetargeterAsset->GetIKRig(ERetargetSourceOrTarget::Source);
 	const UIKRigDefinition* TargetIKRig = RetargeterAsset->GetIKRig(ERetargetSourceOrTarget::Target);
 
@@ -2017,7 +2016,7 @@ void UIKRetargetProcessor::Initialize(
 	}
 
 	// initialize the IKRigProcessor for doing IK decoding
-	bIKRigInitialized = InitializeIKRig(this, TargetSkeletalMesh);
+	bIKRigInitialized = InitializeIKRig(TargetSkeletalMesh);
 	if (!bIKRigInitialized)
 	{
 		// couldn't initialize the IK Rig, we don't disable the retargeter in this case, just warn the user
@@ -2027,8 +2026,7 @@ void UIKRetargetProcessor::Initialize(
 	}
 
 	// initialize the post operations
-	const TArray<TObjectPtr<URetargetOpBase>>& OpsFromAsset = RetargeterAsset->GetPostSettingsUObject()->RetargetOps;
-	const bool bAllOpsInitialized = InitializeOpStack(OpsFromAsset);
+	const bool bAllOpsInitialized = InitializeOpStack(RetargeterAsset->GetRetargetOps());
 	if (!bAllOpsInitialized)
 	{
 		Log.LogWarning(FText(LOCTEXT("PostOperationInitError", "One of the retarget post operations has warnings. See output log for details.")));
@@ -2053,7 +2051,7 @@ void UIKRetargetProcessor::Initialize(
 #endif
 }
 
-bool UIKRetargetProcessor::InitializeRoots()
+bool FIKRetargetProcessor::InitializeRoots()
 {
 	// reset root data
 	RootRetargeter.Reset();
@@ -2081,7 +2079,7 @@ bool UIKRetargetProcessor::InitializeRoots()
 	return bRootEncoderInit && bRootDecoderInit;
 }
 
-bool UIKRetargetProcessor::InitializeBoneChainPairs(const TMap<FName, FTargetChainSettings>& ChainSettings)
+bool FIKRetargetProcessor::InitializeBoneChainPairs(const TMap<FName, FTargetChainSettings>& ChainSettings)
 {
 	ChainPairsFK.Reset();
 	ChainPairsIK.Reset();
@@ -2211,7 +2209,7 @@ bool UIKRetargetProcessor::InitializeBoneChainPairs(const TMap<FName, FTargetCha
 	return !(ChainPairsIK.IsEmpty() && ChainPairsFK.IsEmpty());
 }
 
-bool UIKRetargetProcessor::InitializeIKRig(UObject* Outer, const USkeletalMesh* InSkeletalMesh)
+bool FIKRetargetProcessor::InitializeIKRig(const USkeletalMesh* InSkeletalMesh)
 {
 	// get the target IK Rig asset
 	UIKRigDefinition* IKRig = RetargeterAsset->GetIKRigWriteable(ERetargetSourceOrTarget::Target);
@@ -2227,9 +2225,8 @@ bool UIKRetargetProcessor::InitializeIKRig(UObject* Outer, const USkeletalMesh*
 	}
 	
 	// initialize IK Rig runtime processor
-	IKRigProcessor = NewObject<UIKRigProcessor>(Outer);	
-	IKRigProcessor->Initialize(IKRig, InSkeletalMesh, GoalsToExclude);
-	if (!IKRigProcessor->IsInitialized())
+	IKRigProcessor.Initialize(IKRig, InSkeletalMesh, GoalsToExclude);
+	if (!IKRigProcessor.IsInitialized())
 	{
 		return false;
 	}
@@ -2244,7 +2241,7 @@ bool UIKRetargetProcessor::InitializeIKRig(UObject* Outer, const USkeletalMesh*
 		}
 		
 		// does the IK rig have the IK goal this bone chain requires?
-		if (!IKRigProcessor->GetGoalContainer().FindGoalByName(ChainPair.IKGoalName))
+		if (!IKRigProcessor.GetGoalContainer().FindGoalByName(ChainPair.IKGoalName))
 		{
 			Log.LogError( FText::Format(
 			LOCTEXT("TargetIKBoneNotInSolver", "IK Retargeter has target bone chain, {0} that references an IK Goal, {1} that is not present in any of the solvers in the IK Rig asset."),
@@ -2268,9 +2265,9 @@ bool UIKRetargetProcessor::InitializeIKRig(UObject* Outer, const USkeletalMesh*
 	};
 
 	// detect goals that are not being controlled by the retargeter and set them to "additive" mode
-	const TArray<FIKRigGoal>& GoalsInProcessor = IKRigProcessor->GetGoalContainer().GetGoalArray();
+	const TArray<FIKRigGoal>& GoalsInProcessor = IKRigProcessor.GetGoalContainer().GetGoalArray();
 	const TArray<FBoneChain>& ChainsInAsset = IKRig->GetRetargetChains();
-	const FIKRigSkeleton& IKRigSkeleton = IKRigProcessor->GetSkeleton();
+	const FIKRigSkeleton& IKRigSkeleton = IKRigProcessor.GetSkeleton();
 	for (const FIKRigGoal& GoalInProcessor : GoalsInProcessor)
 	{
 		if (GetRetargetChainByGoalName(ChainsInAsset, GoalInProcessor.Name))
@@ -2288,7 +2285,7 @@ bool UIKRetargetProcessor::InitializeIKRig(UObject* Outer, const USkeletalMesh*
 		GoalToModify.RotationSpace = EIKRigGoalSpace::Additive;
 
 		// but for this to work, the goal's position/rotation must first be converted to "additive" space
-		const FGoalBone* GoalBone = IKRigProcessor->GetGoalBone(GoalToModify.Name);
+		const FGoalBone* GoalBone = IKRigProcessor.GetGoalBone(GoalToModify.Name);
 		const int32& BoneIndex = GoalBone->BoneIndex;
 		check(BoneIndex != INDEX_NONE)
 
@@ -2307,7 +2304,7 @@ bool UIKRetargetProcessor::InitializeIKRig(UObject* Outer, const USkeletalMesh*
 		}
 		
 		// assign the new goal forced to be additive
-		IKRigProcessor->SetIKGoal(GoalToModify);
+		IKRigProcessor.SetIKGoal(GoalToModify);
 
 		// let the user know to avoid confusion
 		Log.LogWarning( FText::Format(
@@ -2318,41 +2315,44 @@ bool UIKRetargetProcessor::InitializeIKRig(UObject* Outer, const USkeletalMesh*
 	return true;
 }
 
-bool UIKRetargetProcessor::InitializeOpStack(const TArray<TObjectPtr<URetargetOpBase>>& OpStackFromAsset)
+bool FIKRetargetProcessor::InitializeOpStack(const TArray<FInstancedStruct>& OpsFromAsset)
 {
 	bool bAllOpsInitialized = true;
-	// create copies of all the operations in the retarget asset for runtime use
-	OpStack.Reset(OpStackFromAsset.Num());
-	int32 OpIndex = 0;
-	for (const URetargetOpBase* OpFromAsset : OpStackFromAsset)
+	
+	// create copies of all the ops in the asset
+	OpStack.Reset(OpsFromAsset.Num());
+	for (const FInstancedStruct& AssetOpStruct : OpsFromAsset)
 	{
-		if (!OpFromAsset)
+		if (!ensure(AssetOpStruct.IsValid()))
 		{
 			// this can happen if asset references deleted op type which should only happen during development (if at all)
-			Log.LogWarning(FText(LOCTEXT("UnknownPostOperation", "Retargeter has null/unknown post operation in it. Please remove it.")));
+			Log.LogWarning(LOCTEXT("UnknownOP", "IK Retargeter, '{0}' has null/unknown op in it. Please reload the asset to remove it."));
 			continue;
 		}
 
-		// create duplicate op instance with unique name
-		const FName BaseName = FName(*(OpFromAsset->GetName() + "_RetargetOpInstance_"));
-		const FName UniqueName = MakeUniqueObjectName(this, OpFromAsset->GetClass(), BaseName, EUniqueObjectNameOptions::GloballyUnique);
-		URetargetOpBase* NewOp = DuplicateObject(OpFromAsset, this, UniqueName);
+		// copy the op memory
+		const UScriptStruct* ScriptStruct = AssetOpStruct.GetScriptStruct();
+		FIKRetargetOpBase* NewOp = static_cast<FIKRetargetOpBase*>(FMemory::Malloc(ScriptStruct->GetStructureSize()));
+		ScriptStruct->InitializeStruct(NewOp);
+		ScriptStruct->CopyScriptStruct(NewOp, AssetOpStruct.GetMemory());
 		
-		// initialize it and store in the processor
-		NewOp->bIsInitialized = NewOp->Initialize(this, SourceSkeleton, TargetSkeleton, Log);
-		bAllOpsInitialized &= NewOp->bIsInitialized;
-		OpStack.Add(NewOp);
+		// initialize it 
+		NewOp->Initialize(*this, SourceSkeleton, TargetSkeleton, Log);
+		bAllOpsInitialized &= NewOp->IsInitialized();
 
+		// store it in the runtime stack
+		OpStack.Emplace(MakeShareable(NewOp));
+		
 		// asset instance gets reference to running editor instance
 		#if WITH_EDITOR
-		OpFromAsset->EditorInstance = NewOp;
+		AssetOpStruct.GetPtr<FIKRetargetOpBase>()->EditorInstance = NewOp;
 		#endif
 	}
 
 	return bAllOpsInitialized;
 }
 
-TArray<FTransform>&  UIKRetargetProcessor::RunRetargeter(
+TArray<FTransform>&  FIKRetargetProcessor::RunRetargeter(
 	TArray<FTransform>& InSourceGlobalPose,
 	const TMap<FName, float>& SpeedValuesFromCurves,
 	const float DeltaTime,
@@ -2362,14 +2362,24 @@ TArray<FTransform>&  UIKRetargetProcessor::RunRetargeter(
 
 	// apply the retargeting settings
 	ApplySettingsFromProfile(Settings);
+
 	// applying settings can cause the retargeter to require reinitialization (when enabling/disabling IK on a chain)
-	// since reinitialization can only currently be done on the main thread, we have to wait for the next frame
-	// in the meantime we return the last generated pose which we know the processor has because it wouldn't have made it past the check()
-	// unless it were already initialized once.
-	// NOTE: In the future hopefully we can reinitialize from any thread by avoiding UObject creation in UIKRetargetProcessor::Initialize()
+	// we should be able to safely reinitialize here because we are already initialized or it wouldn't pass the check() at the top
 	if (!IsInitialized())
 	{
-		return TargetSkeleton.OutputGlobalPose;
+		constexpr bool bSuppressWarnings = false;
+		Initialize(
+			SourceSkeleton.SkeletalMesh,
+			TargetSkeleton.SkeletalMesh,
+			RetargeterAsset,
+			Settings,
+			bSuppressWarnings);
+
+		if (!IsInitialized())
+		{
+			// something went wrong during reinit, bail out and return the input pose
+			return TargetSkeleton.OutputGlobalPose;	
+		}
 	}
 
 #if WITH_EDITOR
@@ -2430,7 +2440,7 @@ TArray<FTransform>&  UIKRetargetProcessor::RunRetargeter(
 	return TargetSkeleton.OutputGlobalPose;
 }
 
-void UIKRetargetProcessor::GenerateBasePoses(TArray<FTransform>& InSourceGlobalPose)
+void FIKRetargetProcessor::GenerateBasePoses(TArray<FTransform>& InSourceGlobalPose)
 {
 	// start from the retarget pose
 	TargetSkeleton.OutputGlobalPose = TargetSkeleton.RetargetGlobalPose;
@@ -2451,7 +2461,7 @@ void UIKRetargetProcessor::GenerateBasePoses(TArray<FTransform>& InSourceGlobalP
 	}
 }
 
-void UIKRetargetProcessor::RunRootRetarget(
+void FIKRetargetProcessor::RunRootRetarget(
 	const TArray<FTransform>& InGlobalTransforms,
     TArray<FTransform>& OutGlobalTransforms)
 {
@@ -2459,7 +2469,7 @@ void UIKRetargetProcessor::RunRootRetarget(
 	RootRetargeter.DecodePose(OutGlobalTransforms);
 }
 
-void UIKRetargetProcessor::RunFKRetarget(
+void FIKRetargetProcessor::RunFKRetarget(
 	const TArray<FTransform>& InGlobalTransforms,
     TArray<FTransform>& OutGlobalTransforms)
 {
@@ -2481,13 +2491,13 @@ void UIKRetargetProcessor::RunFKRetarget(
 	}
 }
 
-void UIKRetargetProcessor::RunIKRetarget(
+void FIKRetargetProcessor::RunIKRetarget(
 	const TArray<FTransform>& InSourceGlobalPose,
     TArray<FTransform>& OutTargetGlobalPose,
     const TMap<FName, float>& SpeedValuesFromCurves,
     const float DeltaTime)
 {
-	if (!IKRigProcessor->IsInitialized())
+	if (!IKRigProcessor.IsInitialized())
 	{
 		return;
 	}
@@ -2526,7 +2536,7 @@ void UIKRetargetProcessor::RunIKRetarget(
 			1.0f,
 			EIKRigGoalSpace::Component,
 			EIKRigGoalSpace::Component);
-		IKRigProcessor->SetIKGoal(Goal);
+		IKRigProcessor.SetIKGoal(Goal);
 
 		ChainPair.IKChainRetargeter.SaveDebugInfo(OutTargetGlobalPose);
 	}
@@ -2534,15 +2544,15 @@ void UIKRetargetProcessor::RunIKRetarget(
 	if (GlobalSettings.bEnableIK)
 	{
 		// copy input pose to start IK solve from
-		IKRigProcessor->SetInputPoseGlobal(OutTargetGlobalPose);
+		IKRigProcessor.SetInputPoseGlobal(OutTargetGlobalPose);
 		// run IK solve
-		IKRigProcessor->Solve();
+		IKRigProcessor.Solve();
 		// copy results of solve
-		IKRigProcessor->GetOutputPoseGlobal(OutTargetGlobalPose);
+		IKRigProcessor.GetOutputPoseGlobal(OutTargetGlobalPose);
 	}
 }
 
-void UIKRetargetProcessor::RunPoleVectorMatching(
+void FIKRetargetProcessor::RunPoleVectorMatching(
 	const TArray<FTransform>& InGlobalTransforms,
 	TArray<FTransform>& OutGlobalTransforms)
 {
@@ -2558,7 +2568,7 @@ void UIKRetargetProcessor::RunPoleVectorMatching(
 	}
 }
 
-void UIKRetargetProcessor::RunStrideWarping(const TArray<FTransform>& InTargetGlobalPose)
+void FIKRetargetProcessor::RunStrideWarping(const TArray<FTransform>& InTargetGlobalPose)
 {
 	// check if we should apply stride warping
 	if (!(GlobalSettings.bWarping && bRootsInitialized))
@@ -2737,20 +2747,20 @@ void UIKRetargetProcessor::RunStrideWarping(const TArray<FTransform>& InTargetGl
 #endif
 }
 
-void UIKRetargetProcessor::RunRetargetOps(
+void FIKRetargetProcessor::RunRetargetOps(
 	const TArray<FTransform>& InSourceGlobalPose,
 	TArray<FTransform>& OutTargetGlobalPose)
 {
-	for (URetargetOpBase* RetargetOp : OpStack)
+	for (TSharedPtr<FIKRetargetOpBase>& RetargetOp : OpStack)
 	{
-		if (RetargetOp->bIsEnabled && RetargetOp->bIsInitialized)
+		if (RetargetOp->IsEnabled() && RetargetOp->IsInitialized())
 		{
-			RetargetOp->Run(this, InSourceGlobalPose, OutTargetGlobalPose);
+			RetargetOp->Run(*this, InSourceGlobalPose, OutTargetGlobalPose);
 		}
 	}
 }
 
-void UIKRetargetProcessor::ResetPlanting()
+void FIKRetargetProcessor::ResetPlanting()
 {
 	for (FRetargetChainPairIK& ChainPair : ChainPairsIK)
 	{
@@ -2758,7 +2768,7 @@ void UIKRetargetProcessor::ResetPlanting()
 	}
 }
 
-bool UIKRetargetProcessor::WasInitializedWithTheseAssets(
+bool FIKRetargetProcessor::WasInitializedWithTheseAssets(
 	const USkeletalMesh* InSourceMesh,
 	const USkeletalMesh* InTargetMesh,
 	const UIKRetargeter* InRetargetAsset) const
@@ -2801,7 +2811,7 @@ bool UIKRetargetProcessor::WasInitializedWithTheseAssets(
 	return true;
 }
 
-void UIKRetargetProcessor::SetNeedsInitialized()
+void FIKRetargetProcessor::SetNeedsInitialized()
 {
 	bIsInitialized = false;
 
@@ -2810,16 +2820,13 @@ void UIKRetargetProcessor::SetNeedsInitialized()
 		RetargeterAsset->IncrementVersion(); // triggers re-init
 	}
 	
-	if (IKRigProcessor)
-	{
-		// may not be initialized yet (during setup as prerequisites are being created)
-		IKRigProcessor->SetNeedsInitialized();
-	}
+	// may not be initialized yet (during setup as prerequisites are being created)
+	IKRigProcessor.SetNeedsInitialized();
 }
 
 #if WITH_EDITOR
 
-bool UIKRetargetProcessor::IsBoneRetargeted(
+bool FIKRetargetProcessor::IsBoneRetargeted(
 	const FName BoneName,
 	const ERetargetSourceOrTarget SourceOrTarget) const
 {
@@ -2862,14 +2869,14 @@ bool UIKRetargetProcessor::IsBoneRetargeted(
 	return false;
 }
 
-int32 UIKRetargetProcessor::GetBoneIndexFromName(
+int32 FIKRetargetProcessor::GetBoneIndexFromName(
 	const FName BoneName,
 	const ERetargetSourceOrTarget SourceOrTarget) const
 {
 	return GetSkeleton(SourceOrTarget).FindBoneIndexByName(BoneName);
 }
 
-FName UIKRetargetProcessor::GetChainNameForBone(const FName BoneName, const ERetargetSourceOrTarget SourceOrTarget) const
+FName FIKRetargetProcessor::GetChainNameForBone(const FName BoneName, const ERetargetSourceOrTarget SourceOrTarget) const
 {
 	const bool bUseSource = SourceOrTarget == ERetargetSourceOrTarget::Source;
 	const FRetargetSkeleton& Skeleton = bUseSource ? SourceSkeleton : TargetSkeleton;
@@ -2883,7 +2890,7 @@ FName UIKRetargetProcessor::GetChainNameForBone(const FName BoneName, const ERet
 	return Skeleton.ChainThatContainsBone[BoneIndex];
 }
 
-FTransform UIKRetargetProcessor::GetGlobalRetargetPoseAtParam(
+FTransform FIKRetargetProcessor::GetGlobalRetargetPoseAtParam(
 	const FName InChainName,
 	const float Param,
 	const ERetargetSourceOrTarget SourceOrTarget) const
@@ -2897,7 +2904,7 @@ FTransform UIKRetargetProcessor::GetGlobalRetargetPoseAtParam(
 	return FTransform::Identity;
 }
 
-FTransform UIKRetargetProcessor::GetRetargetPoseBoneTransform(
+FTransform FIKRetargetProcessor::GetRetargetPoseBoneTransform(
 	const FName InBoneName,
 	const ERetargetSourceOrTarget SourceOrTarget,
 	ERetargetBoneSpace BoneSpace) const
@@ -2912,7 +2919,7 @@ FTransform UIKRetargetProcessor::GetRetargetPoseBoneTransform(
 	return BoneSpace == ERetargetBoneSpace::Global ? Skeleton.RetargetGlobalPose[BoneIndex] : Skeleton.RetargetLocalPose[BoneIndex];
 }
 
-const FChainFK* UIKRetargetProcessor::GetChain(const FName InChainName, const ERetargetSourceOrTarget SourceOrTarget) const
+const FChainFK* FIKRetargetProcessor::GetChain(const FName InChainName, const ERetargetSourceOrTarget SourceOrTarget) const
 {
 	const bool bIsSource = SourceOrTarget == ERetargetSourceOrTarget::Source;
 	for (const FRetargetChainPairFK& ChainPair : ChainPairsFK)
@@ -2929,7 +2936,7 @@ const FChainFK* UIKRetargetProcessor::GetChain(const FName InChainName, const ER
 	return nullptr;
 }
 
-float UIKRetargetProcessor::GetParamOfBoneInChain(
+float FIKRetargetProcessor::GetParamOfBoneInChain(
 	const FName InBoneName,
 	const ERetargetSourceOrTarget SourceOrTarget) const
 {
@@ -2945,7 +2952,7 @@ float UIKRetargetProcessor::GetParamOfBoneInChain(
 	return Chain->Params[BoneIndexInChain];
 }
 
-FName UIKRetargetProcessor::GetClosestBoneToParam(
+FName FIKRetargetProcessor::GetClosestBoneToParam(
 	const FName InChainName, 
 	const float InParam,
 	const ERetargetSourceOrTarget SourceOrTarget) const
@@ -2972,7 +2979,7 @@ FName UIKRetargetProcessor::GetClosestBoneToParam(
 	return Skeleton.BoneNames[Chain->BoneIndices[ChainIndexOfClosestBone]];
 }
 
-FName UIKRetargetProcessor::GetMappedChainName(
+FName FIKRetargetProcessor::GetMappedChainName(
 	const FName InChainName,
 	const ERetargetSourceOrTarget SourceOrTarget)
 {
@@ -2991,26 +2998,26 @@ FName UIKRetargetProcessor::GetMappedChainName(
 
 #endif
 
-void UIKRetargetProcessor::CopyIKRigSettingsFromAsset() const
+void FIKRetargetProcessor::CopyIKRigSettingsFromAsset()
 {
 	// copy IK Rig settings
 	if (const UIKRigDefinition* TargetIKRig = RetargeterAsset->GetIKRig(ERetargetSourceOrTarget::Target))
 	{
-		IKRigProcessor->CopyAllSettingsFromAsset(TargetIKRig);
+		IKRigProcessor.CopyAllSettingsFromAsset(TargetIKRig);
 	}
 }
 
-const FRetargetSkeleton& UIKRetargetProcessor::GetSkeleton(ERetargetSourceOrTarget SourceOrTarget) const
+const FRetargetSkeleton& FIKRetargetProcessor::GetSkeleton(ERetargetSourceOrTarget SourceOrTarget) const
 {
 	return SourceOrTarget == ERetargetSourceOrTarget::Source ? SourceSkeleton : TargetSkeleton;
 }
 
-FName UIKRetargetProcessor::GetRetargetRoot(ERetargetSourceOrTarget SourceOrTarget) const
+FName FIKRetargetProcessor::GetRetargetRoot(ERetargetSourceOrTarget SourceOrTarget) const
 {
 	return SourceOrTarget == ERetargetSourceOrTarget::Source ? RootRetargeter.Source.BoneName : RootRetargeter.Target.BoneName;
 }
 
-void UIKRetargetProcessor::ApplySettingsFromProfile(const FRetargetProfile& Profile)
+void FIKRetargetProcessor::ApplySettingsFromProfile(const FRetargetProfile& Profile)
 {
 	// assign global settings
 	// NOTE: must be done before updating retarget pose because it uses the scale factors in the global settings
@@ -3030,7 +3037,7 @@ void UIKRetargetProcessor::ApplySettingsFromProfile(const FRetargetProfile& Prof
 	RootRetargeter.Settings = Profile.RootSettings;
 }
 
-void UIKRetargetProcessor::UpdateRetargetPoseAtRuntime(
+void FIKRetargetProcessor::UpdateRetargetPoseAtRuntime(
 	const FName NewRetargetPoseName,
 	ERetargetSourceOrTarget SourceOrTarget)
 {
@@ -3085,7 +3092,7 @@ void UIKRetargetProcessor::UpdateRetargetPoseAtRuntime(
 	RootRetargeter.InitializeTarget(RootRetargeter.Target.BoneName, TargetSkeleton, Log);
 }
 
-void UIKRetargetProcessor::UpdateChainSettingsAtRuntime(const FName ChainName, const FTargetChainSettings& NewChainSettings)
+void FIKRetargetProcessor::UpdateChainSettingsAtRuntime(const FName ChainName, const FTargetChainSettings& NewChainSettings)
 {
 	// update FK chain
 	for (FRetargetChainPairFK& Chain : ChainPairsFK)
diff --git a/Engine/Plugins/Animation/IKRig/Source/IKRig/Private/Retargeter/IKRetargeter.cpp b/Engine/Plugins/Animation/IKRig/Source/IKRig/Private/Retargeter/IKRetargeter.cpp
--- a/Engine/Plugins/Animation/IKRig/Source/IKRig/Private/Retargeter/IKRetargeter.cpp
+++ b/Engine/Plugins/Animation/IKRig/Source/IKRig/Private/Retargeter/IKRetargeter.cpp
@@ -36,8 +36,10 @@ UIKRetargeter::UIKRetargeter(const FObjectInitializer& ObjectInitializer)
 	GlobalSettings = CreateDefaultSubobject<UIKRetargetGlobalSettings>(TEXT("GlobalSettings"));
 	GlobalSettings->SetFlags(RF_Transactional);
 
-	OpStack = CreateDefaultSubobject<URetargetOpStack>(TEXT("PostSettings"));
-	OpStack->SetFlags(RF_Transactional);
+	PRAGMA_DISABLE_DEPRECATION_WARNINGS
+	OpStack_DEPRECATED = CreateDefaultSubobject<UDEPRECATED_RetargetOpStack>(TEXT("PostSettings"));
+	OpStack_DEPRECATED->SetFlags(RF_Transactional);
+	PRAGMA_ENABLE_DEPRECATION_WARNINGS
 
 	CleanAndInitialize();
 }
@@ -188,6 +190,49 @@ void UIKRetargeter::PostLoad()
 	{
 		CurrentTargetRetargetPose = CurrentRetargetPose_DEPRECATED;
 	}
+
+	// validate the input FInstancedStruct is a valid subclass of FRetargetOpBase
+	auto DerivesFromBaseOpType = [](const FInstancedStruct& InConvertedSolver)
+	{
+		if (!InConvertedSolver.IsValid())
+		{
+			return false;
+		}
+		
+		bool bIsDerivedFromBase = false;
+		UStruct* CurrentSuperStruct = InConvertedSolver.GetScriptStruct()->GetSuperStruct();
+		while (CurrentSuperStruct)
+		{
+			if (CurrentSuperStruct == FIKRetargetOpBase::StaticStruct())
+			{
+				bIsDerivedFromBase = true;
+				break;
+			}
+			CurrentSuperStruct = CurrentSuperStruct->GetSuperStruct();
+		}
+		
+		return bIsDerivedFromBase;
+	};
+
+	// load old UObject-based op stack and convert to new struct-based types
+	for (TObjectPtr<UDEPRECATED_RetargetOpBase> DeprecatedOp : OpStack_DEPRECATED->RetargetOps_DEPRECATED)
+	{
+		if (DeprecatedOp == nullptr)
+		{
+			continue;
+		}
+		
+		FInstancedStruct ConvertedOp;
+		DeprecatedOp->ConvertToInstancedStruct(ConvertedOp);
+		if (DerivesFromBaseOpType(ConvertedOp))
+		{
+			RetargetOps.Add(ConvertedOp);
+		}
+		else
+		{
+			UE_LOG(LogTemp, Warning, TEXT("IK Retargeter: unable to load old UObject-based op type. Conversion failed for type:  %s"), *DeprecatedOp->GetName());
+		}
+	}
 	
 	PRAGMA_ENABLE_DEPRECATION_WARNINGS
 
@@ -204,9 +249,6 @@ void UIKRetargeter::DeclareConstructClasses(TArray<FTopLevelAssetPath>& OutConst
 
 void UIKRetargeter::CleanAndInitialize()
 {
-	// remove null retarget ops
-	OpStack->RetargetOps.Remove(nullptr);
-	
 	// remove null settings
 	ChainSettings.Remove(nullptr);
 
diff --git a/Engine/Plugins/Animation/IKRig/Source/IKRig/Private/Retargeter/RetargetOps/PinBoneOp.cpp b/Engine/Plugins/Animation/IKRig/Source/IKRig/Private/Retargeter/RetargetOps/PinBoneOp.cpp
--- a/Engine/Plugins/Animation/IKRig/Source/IKRig/Private/Retargeter/RetargetOps/PinBoneOp.cpp
+++ b/Engine/Plugins/Animation/IKRig/Source/IKRig/Private/Retargeter/RetargetOps/PinBoneOp.cpp
@@ -38,13 +38,13 @@
 	return Pose[ParentIndex];
 }
 
-void FPinBoneData::CachePinData(ERetargetSourceOrTarget InSkeletonToCopyFrom, const UIKRetargetProcessor* Processor)
+void FPinBoneData::CachePinData(ERetargetSourceOrTarget InSkeletonToCopyFrom, const FIKRetargetProcessor& Processor)
 {
-	SourceScaleFactor = Processor->GetGlobalSettings().SourceScaleFactor;
+	SourceScaleFactor = Processor.GetGlobalSettings().SourceScaleFactor;
 
     // get skeletons we are copying from/to
-    const FRetargetSkeleton& SkeletonToCopyFrom = Processor->GetSkeleton(InSkeletonToCopyFrom);
-    const FRetargetSkeleton& SkeletonToCopyTo = Processor->GetSkeleton(ERetargetSourceOrTarget::Target);
+    const FRetargetSkeleton& SkeletonToCopyFrom = Processor.GetSkeleton(InSkeletonToCopyFrom);
+    const FRetargetSkeleton& SkeletonToCopyTo = Processor.GetSkeleton(ERetargetSourceOrTarget::Target);
 
 	// get relevant transforms from the retarget poses
 	FTransform BoneToCopyFromParentRefPoseGlobal = GetParentTransform(BoneToCopyFrom.BoneIndex, SkeletonToCopyFrom, SkeletonToCopyFrom.RetargetGlobalPose);
@@ -60,18 +60,19 @@
     LocalRefPoseBoneToCopyTo = BoneToCopyToParentRefPoseGlobal.GetRelativeTransform(BoneToCopyToRefPoseGlobal);
 }
 
-bool UPinBoneOp::Initialize(
-	const UIKRetargetProcessor* Processor,
+bool FIKRetargetPinBoneOp::Initialize(
+	const FIKRetargetProcessor& Processor,
 	const FRetargetSkeleton& SourceSkeleton,
 	const FTargetSkeleton& TargetSkeleton,
 	FIKRigLogger& Log)
 {
 	bool bFullyInitialized = true;
-	for (FPinBoneData& PinData : BonesToPin)
+	
+	for (FPinBoneData& PinData : Settings.BonesToPin)
 	{
 		PinData.BoneToCopyTo.BoneIndex = TargetSkeleton.FindBoneIndexByName(PinData.BoneToCopyTo.BoneName);
 		
-		if (PinTo == ERetargetSourceOrTarget::Source)
+		if (Settings.SkeletonToCopyFrom == ERetargetSourceOrTarget::Source)
 		{
 			PinData.BoneToCopyFrom.BoneIndex = SourceSkeleton.FindBoneIndexByName(PinData.BoneToCopyFrom.BoneName);
 		}
@@ -105,7 +106,7 @@
 #if WITH_EDITOR
 	if (bFullyInitialized)
 	{
-		Message = FText::Format(LOCTEXT("ReadyToRun", "Pinning {0} bones."), FText::AsNumber(BonesToPin.Num()));
+		Message = FText::Format(LOCTEXT("ReadyToRun", "Pinning {0} bones."), FText::AsNumber(Settings.BonesToPin.Num()));
 	}
 	else
 	{
@@ -118,17 +119,18 @@
 #endif
 
 	// always treat this op as "initialized", individual pins will only execute if their prerequisites are met
+	bIsInitialized = true;
 	return true;
 }
 
-void UPinBoneOp::Run(
-	const UIKRetargetProcessor* Processor,
+void FIKRetargetPinBoneOp::Run(
+	const FIKRetargetProcessor& Processor,
 	const TArray<FTransform>& InSourceGlobalPose,
 	TArray<FTransform>& OutTargetGlobalPose)
 {
-	const TArray<FTransform>& PoseToCopyFrom = PinTo == ERetargetSourceOrTarget::Source ? InSourceGlobalPose : OutTargetGlobalPose;
+	const TArray<FTransform>& PoseToCopyFrom = Settings.SkeletonToCopyFrom == ERetargetSourceOrTarget::Source ? InSourceGlobalPose : OutTargetGlobalPose;
 	
-	for (FPinBoneData& PinData : BonesToPin)
+	for (FPinBoneData& PinData : Settings.BonesToPin)
 	{
 		bool bValidBoneToCopyFrom = PoseToCopyFrom.IsValidIndex(PinData.BoneToCopyFrom.BoneIndex);
 		bool bValidBoneToCopyTo = OutTargetGlobalPose.IsValidIndex(PinData.BoneToCopyTo.BoneIndex);
@@ -138,137 +140,82 @@
 		}
 
 		// recache the offsets if the source is scaled differently
-		if (!FMath::IsNearlyEqual(PinData.SourceScaleFactor, Processor->GetGlobalSettings().SourceScaleFactor))
+		if (!FMath::IsNearlyEqual(PinData.SourceScaleFactor, Processor.GetGlobalSettings().SourceScaleFactor))
 		{
-			PinData.CachePinData(PinTo, Processor);
+			PinData.CachePinData(Settings.SkeletonToCopyFrom, Processor);
 		}
 
 		// calculate new transform for bone to pin
 		const FTransform NewTransform = GetNewBoneTransform(PinData, Processor, InSourceGlobalPose, OutTargetGlobalPose);
 		// apply static local and global offsets
-		const FTransform Result = LocalOffset * (NewTransform * GlobalOffset);
+		const FTransform Result = Settings.LocalOffset * (NewTransform * Settings.GlobalOffset);
 
 		// store result in output pose
 		FTransform& BoneToCopyToTransform = OutTargetGlobalPose[PinData.BoneToCopyTo.BoneIndex];
-		if (bCopyTranslation)
+		if (Settings.bCopyTranslation)
 		{
 			BoneToCopyToTransform.SetTranslation(Result.GetLocation());
 		}
-		if (bCopyRotation)
+		if (Settings.bCopyRotation)
 		{
 			BoneToCopyToTransform.SetRotation(Result.GetRotation());
 		}
-		if (bCopyScale)
+		if (Settings.bCopyScale)
 		{
 			BoneToCopyToTransform.SetScale3D(Result.GetScale3D());
 		}
 	}
 }
 
-void UPinBoneOp::Serialize(FArchive& Ar)
+FIKRetargetOpSettingsBase* FIKRetargetPinBoneOp::GetSettings()
 {
-	Super::Serialize(Ar);
-	Ar.UsingCustomVersion(FIKRigObjectVersion::GUID);
-
-#if WITH_EDITORONLY_DATA
-	if (Ar.IsLoading())
-	{
-		
-		if (Ar.CustomVer(FIKRigObjectVersion::GUID) < FIKRigObjectVersion::PinBoneTypeAndOffsetsUpgraded)
-		{
-			PRAGMA_DISABLE_DEPRECATION_WARNINGS
-
-			// load the boolean "bMaintainOffset" and convert it into the equivalent translation and rotation offset modes
-			if (bMaintainOffset_DEPRECATED)
-			{
-				TranslationMode = EPinBoneTranslationMode::CopyGlobalPositionAndMaintainOffset;
-				RotationMode = EPinBoneRotationMode::MaintainOffsetFromBoneToCopyFrom;
-			}
-			else
-			{
-				TranslationMode = EPinBoneTranslationMode::CopyGlobalPosition;
-				RotationMode = EPinBoneRotationMode::CopyGlobalRotation;
-			}
-
-			// load the "PinType" enum and convert it into the equivalent trans/rot/scale toggles
-			switch (PinType_DEPRECATED)
-			{
-			case EPinBoneType::FullTransform:
-				{
-					bCopyTranslation = true;
-					bCopyRotation = true;
-					bCopyScale = true;
-					break;
-				}
-			case EPinBoneType::TranslateOnly:
-				{
-					bCopyTranslation = true;
-					bCopyRotation = false;
-					bCopyScale = false;
-					break;
-				}
-			case EPinBoneType::RotateOnly:
-				{
-					bCopyTranslation = false;
-					bCopyRotation = true;
-					bCopyScale = false;
-					break;
-				}
-			case EPinBoneType::ScaleOnly:
-				{
-					bCopyTranslation = false;
-					bCopyRotation = false;
-					bCopyScale = true;
-					break;
-				}
-			default:
-				break;
-			}
-			PRAGMA_ENABLE_DEPRECATION_WARNINGS
-		}
-	}
-#endif
+	return &Settings;
 }
 
-USkeleton* UPinBoneOp::GetSkeleton(bool& bInvalidSkeletonIsError, const IPropertyHandle* PropertyHandle)
+const UScriptStruct* FIKRetargetPinBoneOp::GetSettingsType() const
 {
+	return FIKRetargetPinBoneOpSettings::StaticStruct();
+}
+
+USkeleton* FIKRetargetPinBoneOp::GetSkeleton(bool& bInvalidSkeletonIsError, const IPropertyHandle* PropertyHandle)
+{
 #if WITH_EDITOR
-	if (UPinBoneOp* RunningInstance = Cast<UPinBoneOp>(EditorInstance))
+	if (FIKRetargetPinBoneOp* RunningInstance = static_cast<FIKRetargetPinBoneOp*>(EditorInstance))
 	{
 		FName PropertyName = PropertyHandle->GetProperty()->GetFName();
 		if (PropertyName == GET_MEMBER_NAME_CHECKED(FPinBoneData, BoneToCopyTo))
 		{
-			const USkeleton* Skeleton = RunningInstance->PinTo == ERetargetSourceOrTarget::Source ? RunningInstance->TargetSkeletonAsset : RunningInstance->SourceSkeletonAsset;
+			const USkeleton* Skeleton = RunningInstance->Settings.SkeletonToCopyFrom == ERetargetSourceOrTarget::Source ? RunningInstance->TargetSkeletonAsset : RunningInstance->SourceSkeletonAsset;
 			return const_cast<USkeleton*>(Skeleton);
 		}
 
 		if (PropertyName == GET_MEMBER_NAME_CHECKED(FPinBoneData, BoneToCopyFrom))
 		{
-			const USkeleton* Skeleton = RunningInstance->PinTo == ERetargetSourceOrTarget::Target ? RunningInstance->TargetSkeletonAsset : RunningInstance->SourceSkeletonAsset;
+			const USkeleton* Skeleton = RunningInstance->Settings.SkeletonToCopyFrom == ERetargetSourceOrTarget::Target ? RunningInstance->TargetSkeletonAsset : RunningInstance->SourceSkeletonAsset;
 			return const_cast<USkeleton*>(Skeleton);
 		}
 	}
 #endif
 	
-	ensureMsgf(false, TEXT("UPinBoneOp unable to get skeleton for UI widget."));
+	ensureMsgf(false, TEXT("PinBoneOp unable to get skeleton for UI widget."));
 	return nullptr;
 }
 
-FTransform UPinBoneOp::GetNewBoneTransform(
+FTransform FIKRetargetPinBoneOp::GetNewBoneTransform(
 	const FPinBoneData& PinData,
-	const UIKRetargetProcessor* Processor,
+	const FIKRetargetProcessor& Processor,
 	const TArray<FTransform>& InSourceGlobalPose,
 	TArray<FTransform>& OutTargetGlobalPose) const
 {
-	const TArray<FTransform>& PoseToCopyFrom = PinTo == ERetargetSourceOrTarget::Source ? InSourceGlobalPose : OutTargetGlobalPose;
-	const FRetargetSkeleton& SkeletonToCopyFrom = Processor->GetSkeleton(PinTo);
-	const FRetargetSkeleton& SkeletonToCopyTo = Processor->GetSkeleton(ERetargetSourceOrTarget::Target);
+	const TArray<FTransform>& PoseToCopyFrom = Settings.SkeletonToCopyFrom == ERetargetSourceOrTarget::Source ? InSourceGlobalPose : OutTargetGlobalPose;
+	const FRetargetSkeleton& CopyFrom = Processor.GetSkeleton(Settings.SkeletonToCopyFrom);
+	const FRetargetSkeleton& CopyTo = Processor.GetSkeleton(ERetargetSourceOrTarget::Target);
 
 	// get the current translational offset of BoneToCopyFrom relative to it's parent
-	auto GetBoneToCopyFromCurrentVector = [PoseToCopyFrom, SkeletonToCopyFrom, PinData]()
+	auto GetBoneToCopyFromCurrentVector = [PoseToCopyFrom, CopyFrom, PinData]()
 	{
 		const FTransform& BoneToCopyFromCurrent = PoseToCopyFrom[PinData.BoneToCopyFrom.BoneIndex];
-		const FTransform& ParentOfBoneToCopyFromCurrent = FPinBoneData::GetParentTransform(PinData.BoneToCopyFrom.BoneIndex, SkeletonToCopyFrom, PoseToCopyFrom);
+		const FTransform& ParentOfBoneToCopyFromCurrent = FPinBoneData::GetParentTransform(PinData.BoneToCopyFrom.BoneIndex, CopyFrom, PoseToCopyFrom);
 		return BoneToCopyFromCurrent.GetTranslation() - ParentOfBoneToCopyFromCurrent.GetTranslation();
 	};
 	
@@ -276,7 +223,7 @@
 	FTransform Result = FTransform::Identity;
 
 	// generate translational offset
-	switch (TranslationMode)
+	switch (Settings.TranslationMode)
 	{
 	case EPinBoneTranslationMode::CopyGlobalPosition:
 		{
@@ -294,7 +241,7 @@
 	case EPinBoneTranslationMode::CopyLocalPosition:
 		{
 			const FVector BoneToCopyFromCurrentVector = GetBoneToCopyFromCurrentVector();
-			const FTransform& ParentOfBoneToCopyToCurrent = FPinBoneData::GetParentTransform(PinData.BoneToCopyTo.BoneIndex, SkeletonToCopyTo, OutTargetGlobalPose);
+			const FTransform& ParentOfBoneToCopyToCurrent = FPinBoneData::GetParentTransform(PinData.BoneToCopyTo.BoneIndex, CopyTo, OutTargetGlobalPose);
 			Result.SetTranslation(ParentOfBoneToCopyToCurrent.GetTranslation() + BoneToCopyFromCurrentVector);
 			break;
 		}
@@ -308,7 +255,7 @@
 			BoneToCopyFromTranslationCurrent.ToDirectionAndLength(BoneToCopyFromTranslationNorm, CurrentLength);
 			const FVector BoneToCopyFromRelativeOffset = BoneToCopyFromTranslationNorm * (CurrentLength + RestPoseLengthDifference);
 		
-			const FTransform& ParentOfBoneToCopyToCurrent = FPinBoneData::GetParentTransform(PinData.BoneToCopyTo.BoneIndex, SkeletonToCopyTo, OutTargetGlobalPose);
+			const FTransform& ParentOfBoneToCopyToCurrent = FPinBoneData::GetParentTransform(PinData.BoneToCopyTo.BoneIndex, CopyTo, OutTargetGlobalPose);
 			Result.SetTranslation(ParentOfBoneToCopyToCurrent.GetTranslation() + BoneToCopyFromRelativeOffset);
 			break;
 		}
@@ -326,7 +273,7 @@
 			const double LengthOfBoneToCopyToInRefPose = PinData.LocalRefPoseBoneToCopyTo.GetTranslation().Size();
 			const FVector BoneToCopyFromRelativeOffset = BoneToCopyFromTranslationNorm * (LengthOfBoneToCopyToInRefPose * BoneLengthScaleFactor);
 		
-			const FTransform& ParentOfBoneToCopyToCurrent = FPinBoneData::GetParentTransform(PinData.BoneToCopyTo.BoneIndex, SkeletonToCopyTo, OutTargetGlobalPose);
+			const FTransform& ParentOfBoneToCopyToCurrent = FPinBoneData::GetParentTransform(PinData.BoneToCopyTo.BoneIndex, CopyTo, OutTargetGlobalPose);
 			Result.SetTranslation(ParentOfBoneToCopyToCurrent.GetTranslation() + BoneToCopyFromRelativeOffset);
 			
 			break;
@@ -338,7 +285,7 @@
 	}
 
 	// generate rotational offset
-	switch (RotationMode)
+	switch (Settings.RotationMode)
 	{
 	case EPinBoneRotationMode::CopyGlobalRotation:
 		{
@@ -347,8 +294,8 @@
 		}
 	case EPinBoneRotationMode::MaintainOffsetFromBoneToCopyFrom:
 		{
-			const FTransform& BoneToCopyFromRefPoseGlobal = SkeletonToCopyFrom.RetargetGlobalPose[PinData.BoneToCopyFrom.BoneIndex];
-			const FTransform& BoneToCopyToRefPoseGlobal = SkeletonToCopyTo.RetargetGlobalPose[PinData.BoneToCopyTo.BoneIndex];
+			const FTransform& BoneToCopyFromRefPoseGlobal = CopyFrom.RetargetGlobalPose[PinData.BoneToCopyFrom.BoneIndex];
+			const FTransform& BoneToCopyToRefPoseGlobal = CopyTo.RetargetGlobalPose[PinData.BoneToCopyTo.BoneIndex];
 			const FTransform Offset = BoneToCopyFromRefPoseGlobal.GetRelativeTransform(BoneToCopyToRefPoseGlobal);
 			const FTransform BoneToCopyFromTransform = PoseToCopyFrom[PinData.BoneToCopyFrom.BoneIndex];
 			Result.SetRotation((Offset * BoneToCopyFromTransform).GetRotation());
@@ -363,4 +310,67 @@
 	return Result;
 }
 
+PRAGMA_DISABLE_DEPRECATION_WARNINGS
+void UDEPRECATED_PinBoneOp::Serialize(FArchive& Ar)
+{
+	Super::Serialize(Ar);
+	Ar.UsingCustomVersion(FIKRigObjectVersion::GUID);
+
+#if WITH_EDITORONLY_DATA
+	if (Ar.IsLoading())
+	{
+		if (Ar.CustomVer(FIKRigObjectVersion::GUID) < FIKRigObjectVersion::PinBoneTypeAndOffsetsUpgraded)
+		{
+			// load the boolean "bMaintainOffset" and convert it into the equivalent translation and rotation offset modes
+			if (bMaintainOffset_DEPRECATED)
+			{
+				TranslationMode = EPinBoneTranslationMode::CopyGlobalPositionAndMaintainOffset;
+				RotationMode = EPinBoneRotationMode::MaintainOffsetFromBoneToCopyFrom;
+			}
+			else
+			{
+				TranslationMode = EPinBoneTranslationMode::CopyGlobalPosition;
+				RotationMode = EPinBoneRotationMode::CopyGlobalRotation;
+			}
+
+			// load the "PinType" enum and convert it into the equivalent trans/rot/scale toggles
+			switch (PinType_DEPRECATED)
+			{
+			case EPinBoneType::FullTransform:
+				{
+					bCopyTranslation = true;
+					bCopyRotation = true;
+					bCopyScale = true;
+					break;
+				}
+			case EPinBoneType::TranslateOnly:
+				{
+					bCopyTranslation = true;
+					bCopyRotation = false;
+					bCopyScale = false;
+					break;
+				}
+			case EPinBoneType::RotateOnly:
+				{
+					bCopyTranslation = false;
+					bCopyRotation = true;
+					bCopyScale = false;
+					break;
+				}
+			case EPinBoneType::ScaleOnly:
+				{
+					bCopyTranslation = false;
+					bCopyRotation = false;
+					bCopyScale = true;
+					break;
+				}
+			default:
+				break;
+			}
+		}
+	}
+#endif
+}
+PRAGMA_ENABLE_DEPRECATION_WARNINGS
+
 #undef LOCTEXT_NAMESPACE
diff --git a/Engine/Plugins/Animation/IKRig/Source/IKRig/Private/Retargeter/RetargetOps/RootMotionGeneratorOp.cpp b/Engine/Plugins/Animation/IKRig/Source/IKRig/Private/Retargeter/RetargetOps/RootMotionGeneratorOp.cpp
--- a/Engine/Plugins/Animation/IKRig/Source/IKRig/Private/Retargeter/RetargetOps/RootMotionGeneratorOp.cpp
+++ b/Engine/Plugins/Animation/IKRig/Source/IKRig/Private/Retargeter/RetargetOps/RootMotionGeneratorOp.cpp
@@ -6,60 +6,62 @@
 #define LOCTEXT_NAMESPACE "RootMotionGeneratorOp"
 
 #if WITH_EDITOR
-void URootMotionGeneratorOp::OnAddedToStack(const UIKRetargeter* Asset)
+void FIKRetargetRootMotionOp::OnAddedToStack(const UIKRetargeter* Asset)
 {
 	// get the root of the SOURCE skeleton
 	if (const USkeletalMesh* Mesh = Asset->GetPreviewMesh(ERetargetSourceOrTarget::Source))
 	{
-		SourceRootBone = Mesh->GetRefSkeleton().GetBoneName(0);
+		Settings.SourceRootBone = Mesh->GetRefSkeleton().GetBoneName(0);
 	}
 	
 	// get the root of the TARGET skeleton
 	if (const USkeletalMesh* Mesh = Asset->GetPreviewMesh(ERetargetSourceOrTarget::Target))
 	{
-		TargetRootBone = Mesh->GetRefSkeleton().GetBoneName(0);
+		Settings.TargetRootBone = Mesh->GetRefSkeleton().GetBoneName(0);
 	}
 
 	// get the retarget root bone
 	if (const UIKRigDefinition* IKRig = Asset->GetIKRig(ERetargetSourceOrTarget::Target))
 	{
-		TargetPelvisBone = IKRig->GetRetargetRoot();
+		Settings.TargetPelvisBone = IKRig->GetRetargetRoot();
 	}
 }
 #endif
 
-bool URootMotionGeneratorOp::Initialize(
-	const UIKRetargetProcessor* Processor,
+bool FIKRetargetRootMotionOp::Initialize(
+	const FIKRetargetProcessor& Processor,
 	const FRetargetSkeleton& SourceSkeleton,
 	const FTargetSkeleton& TargetSkeleton,
 	FIKRigLogger& Log)
 {
-	SourceRootIndex = SourceSkeleton.BoneNames.Find(SourceRootBone);
+	bIsInitialized = false;
+	
+	SourceRootIndex = SourceSkeleton.BoneNames.Find(Settings.SourceRootBone);
 	bool bHasAllPrerequisites = true;
 	if (SourceRootIndex == INDEX_NONE)
 	{
-		Log.LogWarning(FText::Format(LOCTEXT("MissingSourceRootBone", "Root Motion Remap Op, missing source root bone {0}."), FText::FromName(SourceRootBone)));
+		Log.LogWarning(FText::Format(LOCTEXT("MissingSourceRootBone", "Root Motion Remap Op, missing source root bone {0}."), FText::FromName(Settings.SourceRootBone)));
 		bHasAllPrerequisites = false;
 	}
 	
-	TargetRootIndex = TargetSkeleton.BoneNames.Find(TargetRootBone);
+	TargetRootIndex = TargetSkeleton.BoneNames.Find(Settings.TargetRootBone);
 	if (TargetRootIndex == INDEX_NONE)
 	{
-		Log.LogWarning(FText::Format(LOCTEXT("MissingTargetRootBone", "Root Motion Remap Op, missing target root bone {0}."), FText::FromName(TargetRootBone)));
+		Log.LogWarning(FText::Format(LOCTEXT("MissingTargetRootBone", "Root Motion Remap Op, missing target root bone {0}."), FText::FromName(Settings.TargetRootBone)));
 		bHasAllPrerequisites = false;
 	}
 
-	TargetPelvisIndex = TargetSkeleton.BoneNames.Find(TargetPelvisBone);
+	TargetPelvisIndex = TargetSkeleton.BoneNames.Find(Settings.TargetPelvisBone);
 	if (TargetPelvisIndex == INDEX_NONE)
 	{
-		Log.LogWarning(FText::Format(LOCTEXT("MissingPelvisBone", "Root Motion Remap Op, missing target pelvis bone {0}."), FText::FromName(TargetPelvisBone)));
+		Log.LogWarning(FText::Format(LOCTEXT("MissingPelvisBone", "Root Motion Remap Op, missing target pelvis bone {0}."), FText::FromName(Settings.TargetPelvisBone)));
 		bHasAllPrerequisites = false;
 	}
 
 #if WITH_EDITOR
 	if (bHasAllPrerequisites)
 	{
-		Message = FText::Format(LOCTEXT("ReadyToRun", "Generating root motion on {0}."), FText::FromName(TargetRootBone));
+		Message = FText::Format(LOCTEXT("ReadyToRun", "Generating root motion on {0}."), FText::FromName(Settings.TargetRootBone));
 	}
 	else
 	{
@@ -101,12 +103,13 @@ bool URootMotionGeneratorOp::Initialize(
 			NonRetargetedChildrenOfRoot.Add(BoneIndex);
 		}
 	}
-	
+
+	bIsInitialized = true;
 	return true;
 }
 
-void URootMotionGeneratorOp::Run(
-	const UIKRetargetProcessor* Processor,
+void FIKRetargetRootMotionOp::Run(
+	const FIKRetargetProcessor& Processor,
 	const TArray<FTransform>& InSourceGlobalPose,
 	TArray<FTransform>& OutTargetGlobalPose)
 {
@@ -114,7 +117,7 @@ void URootMotionGeneratorOp::Run(
 	FTransform NewRootTransform;
 	
 	// either generate motion from target pelvis, or copy from source root
-	if (RootMotionSource == ERootMotionSource::GenerateFromTargetPelvis)
+	if (Settings.RootMotionSource == ERootMotionSource::GenerateFromTargetPelvis)
 	{
 		GenerateRootMotionFromTargetPelvis(NewRootTransform, InSourceGlobalPose, OutTargetGlobalPose);
 	}
@@ -124,7 +127,7 @@ void URootMotionGeneratorOp::Run(
 	}
 
 	// optionally propagate delta to all non-retargeted children
-	if (bPropagateToNonRetargetedChildren)
+	if (Settings.bPropagateToNonRetargetedChildren)
 	{
 		const FTransform Delta = OutTargetGlobalPose[TargetRootIndex].Inverse() * NewRootTransform;
 		for (const int32 BoneIndex : NonRetargetedChildrenOfRoot)
@@ -134,10 +137,20 @@ void URootMotionGeneratorOp::Run(
 	}
 
 	// apply to the target root bone
-	OutTargetGlobalPose[TargetRootIndex] = NewRootTransform * GlobalOffset;
+	OutTargetGlobalPose[TargetRootIndex] = NewRootTransform * Settings.GlobalOffset;
+}
+
+FIKRetargetOpSettingsBase* FIKRetargetRootMotionOp::GetSettings()
+{
+	return &Settings;
+}
+
+const UScriptStruct* FIKRetargetRootMotionOp::GetSettingsType() const
+{
+	return FIKRetargetRootMotionOpSettings::StaticStruct();
 }
 
-void URootMotionGeneratorOp::GenerateRootMotionFromTargetPelvis(
+void FIKRetargetRootMotionOp::GenerateRootMotionFromTargetPelvis(
 	FTransform& OutRootTransform,
 	const TArray<FTransform>& InSourceGlobalPose,
 	const TArray<FTransform>& InTargetGlobalPose) const
@@ -145,7 +158,7 @@ void URootMotionGeneratorOp::GenerateRootMotionFromTargetPelvis(
 	// In this case, we are generating root motion "from scratch"
 	// using the target Pelvis bone as the source of the motion
 	
-	if (bMaintainOffsetFromPelvis)
+	if (Settings.bMaintainOffsetFromPelvis)
 	{
 		// set root to the relative offset from the pelvis (recorded from ref pose)
 		OutRootTransform = TargetPelvisRelativeToTargetRootRefPose * InTargetGlobalPose[TargetPelvisIndex];
@@ -157,19 +170,19 @@ void URootMotionGeneratorOp::GenerateRootMotionFromTargetPelvis(
 	}
 
 	// optionally remove all rotation (use static ref pose orientation)
-	if (!bRotateWithPelvis)
+	if (!Settings.bRotateWithPelvis)
 	{
 		OutRootTransform.SetRotation(TargetRootInRefPose.GetRotation());
 	}
 
 	// adjust height of root
-	if (RootHeightSource == ERootMotionHeightSource::SnapToGround)
+	if (Settings.RootHeightSource == ERootMotionHeightSource::SnapToGround)
 	{
 		// snap the root to the ground plane
 		FVector RootTranslation = OutRootTransform.GetTranslation();
 		RootTranslation.Z = 0.f;
 		OutRootTransform.SetTranslation(RootTranslation);
-	}else if (RootHeightSource == ERootMotionHeightSource::CopyHeightFromSource)
+	}else if (Settings.RootHeightSource == ERootMotionHeightSource::CopyHeightFromSource)
 	{
 		// snap the root to the height from the source
 		FVector RootTranslation = OutRootTransform.GetTranslation();
@@ -178,9 +191,9 @@ void URootMotionGeneratorOp::GenerateRootMotionFromTargetPelvis(
 	}
 }
 
-void URootMotionGeneratorOp::CopyRootMotionFromSourceRoot(
+void FIKRetargetRootMotionOp::CopyRootMotionFromSourceRoot(
 	FTransform& OutRootTransform,
-	const UIKRetargetProcessor* Processor,
+	const FIKRetargetProcessor& Processor,
 	const TArray<FTransform>& InSourceGlobalPose) const
 {
 	// In this case, we are copying root motion from the source root
@@ -193,11 +206,11 @@ void URootMotionGeneratorOp::CopyRootMotionFromSourceRoot(
 	// scale root translation by same scale factor applied to Pelvis (and modified by settings)
 	FVector NewRootLocation = InSourceGlobalPose[SourceRootIndex].GetLocation();
 	FVector RootTranslationDelta = NewRootLocation - SourceRootInRefPose.GetTranslation();
-	RootTranslationDelta *= Processor->GetRootRetargeter().GetGlobalScaleVector();
+	RootTranslationDelta *= Processor.GetRootRetargeter().GetGlobalScaleVector();
 	NewRootLocation = SourceRootInRefPose.GetTranslation() + RootTranslationDelta;
 
 	// optionally snap the root to the ground plane
-	if (RootHeightSource == ERootMotionHeightSource::SnapToGround)
+	if (Settings.RootHeightSource == ERootMotionHeightSource::SnapToGround)
 	{
 		NewRootLocation.Z = 0.f;
 	}
diff --git a/Engine/Plugins/Animation/IKRig/Source/IKRig/Private/Rig/IKRigDefinition.cpp b/Engine/Plugins/Animation/IKRig/Source/IKRig/Private/Rig/IKRigDefinition.cpp
--- a/Engine/Plugins/Animation/IKRig/Source/IKRig/Private/Rig/IKRigDefinition.cpp
+++ b/Engine/Plugins/Animation/IKRig/Source/IKRig/Private/Rig/IKRigDefinition.cpp
@@ -303,7 +303,6 @@ void UIKRigDefinition::Serialize(FArchive& Ar)
 
 void UIKRigDefinition::PostLoad()
 {
-	
 	Super::PostLoad();
 	
 	// very early versions of the asset may not have been set as standalone
@@ -336,6 +335,11 @@ void UIKRigDefinition::PostLoad()
 	PRAGMA_DISABLE_DEPRECATION_WARNINGS
 	for (TObjectPtr<UDEPRECATED_IKRigSolver> Solver : Solvers_DEPRECATED)
 	{
+		if (Solver == nullptr)
+		{
+			continue;
+		}
+		
 		FInstancedStruct ConvertedSolver;
 		Solver->ConvertToInstancedStruct(ConvertedSolver);
 		if (DerivesFromBaseSolverType(ConvertedSolver))
@@ -348,6 +352,12 @@ void UIKRigDefinition::PostLoad()
 		}
 	}
 	PRAGMA_ENABLE_DEPRECATION_WARNINGS
+	
+	// clean out null solvers
+	SolverStack.RemoveAll([](const FInstancedStruct& InSolver)
+	{
+		return !InSolver.IsValid();
+	});
 }
 
 const TArray<const FIKRigSolverBase*> UIKRigDefinition::GetSolverArray() const
@@ -367,7 +377,7 @@ const TArray<const FIKRigSolverBase*> UIKRigDefinition::GetSolverArray() const
 	return MoveTemp(Result);
 }
 
-const TArray<FInstancedStruct> UIKRigDefinition::GetSolverStructs() const
+const TArray<FInstancedStruct>& UIKRigDefinition::GetSolverStructs() const
 {
 	return SolverStack;
 }
diff --git a/Engine/Plugins/Animation/IKRig/Source/IKRig/Private/Rig/IKRigProcessor.cpp b/Engine/Plugins/Animation/IKRig/Source/IKRig/Private/Rig/IKRigProcessor.cpp
--- a/Engine/Plugins/Animation/IKRig/Source/IKRig/Private/Rig/IKRigProcessor.cpp
+++ b/Engine/Plugins/Animation/IKRig/Source/IKRig/Private/Rig/IKRigProcessor.cpp
@@ -10,22 +10,19 @@
 
 #define LOCTEXT_NAMESPACE "IKRigProcessor"
 
-UIKRigProcessor::UIKRigProcessor()
-{
-	const FName LogName = FName("IKRig_",GetUniqueID());
-	Log.SetLogTarget(LogName, LOCTEXT("IKRigLogName", "IK Rig Log"));
-}
 
-void UIKRigProcessor::Initialize(
+void FIKRigProcessor::Initialize(
 	const UIKRigDefinition* InRigAsset,
 	const USkeletalMesh* SkeletalMesh,
 	const TArray<FName>& ExcludedGoals)
 {
-	// we instantiate UObjects here which MUST be done on game thread...
-	check(IsInGameThread());
-	check(InRigAsset);
-	
 	bInitialized = false;
+	
+	// can't initialize without a rig definition
+	if (!ensure(InRigAsset))
+	{
+		return;
+	}
 
 	// bail out if we've already tried initializing with this exact version of the rig asset
 	if (bTriedToInitialize)
@@ -36,6 +33,9 @@ void UIKRigProcessor::Initialize(
 	// ok, lets try to initialize
 	bTriedToInitialize = true;
 
+	// set the log target to a unique id for this asset
+	Log.SetLogTarget(InRigAsset);
+
 	// copy skeleton data from the actual skeleton we want to run on
 	Skeleton.SetInputSkeleton(SkeletalMesh, InRigAsset->GetSkeleton().ExcludedBones);
 	
@@ -47,7 +47,7 @@ void UIKRigProcessor::Initialize(
 		return;
 	}
 
-	if (!UIKRigProcessor::IsIKRigCompatibleWithSkeleton(InRigAsset, SkeletalMesh, &Log))
+	if (!FIKRigProcessor::IsIKRigCompatibleWithSkeleton(InRigAsset, SkeletalMesh, &Log))
 	{
 		Log.LogError(FText::Format(
 			LOCTEXT("SkeletonMissingRequiredBones", "Trying to initialize IKRig, '{0}' with a Skeleton that is missing required bones. See prior warnings."),
@@ -146,7 +146,7 @@ void UIKRigProcessor::Initialize(
 		ScriptStruct->InitializeStruct(NewSolver);
 		ScriptStruct->CopyScriptStruct(NewSolver, AssetSolverStruct.GetMemory());
 		// store it in the runtime stack
-		Solvers.Emplace(TUniquePtr<FIKRigSolverBase>(NewSolver));
+		Solvers.Emplace(MakeShareable(NewSolver));
 
 		// remove excluded goals from the solver
 		for (const FName& ExcludedGoal : ExcludedGoals)
@@ -178,7 +178,7 @@ void UIKRigProcessor::Initialize(
 	bInitialized = true;
 }
 
-bool UIKRigProcessor::IsIKRigCompatibleWithSkeleton(
+bool FIKRigProcessor::IsIKRigCompatibleWithSkeleton(
 	const UIKRigDefinition* InRigAsset,
 	const FIKRigInputSkeleton& InputSkeleton,
 	const FIKRigLogger* Log)
@@ -276,7 +276,7 @@ bool UIKRigProcessor::IsIKRigCompatibleWithSkeleton(
 	return bAllParentsValid;
 }
 
-void UIKRigProcessor::SetInputPoseGlobal(const TArray<FTransform>& InGlobalBoneTransforms) 
+void FIKRigProcessor::SetInputPoseGlobal(const TArray<FTransform>& InGlobalBoneTransforms) 
 {
 	check(bInitialized);
 	check(InGlobalBoneTransforms.Num() == Skeleton.CurrentPoseGlobal.Num());
@@ -284,26 +284,26 @@ void UIKRigProcessor::SetInputPoseGlobal(const TArray<FTransform>& InGlobalBoneT
 	Skeleton.UpdateAllLocalTransformFromGlobal();
 }
 
-void UIKRigProcessor::SetInputPoseToRefPose()
+void FIKRigProcessor::SetInputPoseToRefPose()
 {
 	check(bInitialized);
 	Skeleton.CurrentPoseGlobal = Skeleton.RefPoseGlobal;
 	Skeleton.UpdateAllLocalTransformFromGlobal();
 }
 
-void UIKRigProcessor::SetIKGoal(const FIKRigGoal& InGoal)
+void FIKRigProcessor::SetIKGoal(const FIKRigGoal& InGoal)
 {
 	check(bInitialized);
 	GoalContainer.SetIKGoal(InGoal);
 }
 
-void UIKRigProcessor::SetIKGoal(const UIKRigEffectorGoal* InGoal)
+void FIKRigProcessor::SetIKGoal(const UIKRigEffectorGoal* InGoal)
 {
 	check(bInitialized);
 	GoalContainer.SetIKGoal(InGoal);
 }
 
-void UIKRigProcessor::Solve(const FTransform& ComponentToWorld)
+void FIKRigProcessor::Solve(const FTransform& ComponentToWorld)
 {
 	check(bInitialized);
 	
@@ -311,7 +311,7 @@ void UIKRigProcessor::Solve(const FTransform& ComponentToWorld)
 	ResolveFinalGoalTransforms(ComponentToWorld);
 
 	// run all the solvers
-	for (const TUniquePtr<FIKRigSolverBase>& Solver : Solvers)
+	for (const TSharedPtr<FIKRigSolverBase>& Solver : Solvers)
 	{
 		if (Solver->IsEnabled())
 		{
@@ -323,13 +323,13 @@ void UIKRigProcessor::Solve(const FTransform& ComponentToWorld)
 	Skeleton.NormalizeRotations(Skeleton.CurrentPoseGlobal);
 }
 
-void UIKRigProcessor::GetOutputPoseGlobal(TArray<FTransform>& OutputPoseGlobal) const
+void FIKRigProcessor::GetOutputPoseGlobal(TArray<FTransform>& OutputPoseGlobal) const
 {
 	// TODO, remove this copy
 	OutputPoseGlobal = Skeleton.CurrentPoseGlobal;
 }
 
-void UIKRigProcessor::Reset()
+void FIKRigProcessor::Reset()
 {
 	Solvers.Reset();
 	GoalContainer.Empty();
@@ -338,13 +338,13 @@ void UIKRigProcessor::Reset()
 	SetNeedsInitialized();
 }
 
-void UIKRigProcessor::SetNeedsInitialized()
+void FIKRigProcessor::SetNeedsInitialized()
 {
 	bInitialized = false;
 	bTriedToInitialize = false;
 };
 
-void UIKRigProcessor::CopyAllSettingsFromAsset(const UIKRigDefinition* SourceAsset)
+void FIKRigProcessor::CopyAllSettingsFromAsset(const UIKRigDefinition* SourceAsset)
 {
 	check(SourceAsset)
 	
@@ -371,28 +371,28 @@ void UIKRigProcessor::CopyAllSettingsFromAsset(const UIKRigDefinition* SourceAss
 	}
 }
 
-const FIKRigGoalContainer& UIKRigProcessor::GetGoalContainer() const
+const FIKRigGoalContainer& FIKRigProcessor::GetGoalContainer() const
 {
 	check(bInitialized);
 	return GoalContainer;
 }
 
-const FGoalBone* UIKRigProcessor::GetGoalBone(const FName& GoalName) const
+const FGoalBone* FIKRigProcessor::GetGoalBone(const FName& GoalName) const
 {
 	return GoalBones.Find(GoalName);
 }
 
-FIKRigSkeleton& UIKRigProcessor::GetSkeletonWriteable()
+FIKRigSkeleton& FIKRigProcessor::GetSkeletonWriteable()
 {
 	return Skeleton;
 }
 
-const FIKRigSkeleton& UIKRigProcessor::GetSkeleton() const
+const FIKRigSkeleton& FIKRigProcessor::GetSkeleton() const
 {
 	return Skeleton;
 }
 
-void UIKRigProcessor::ResolveFinalGoalTransforms(const FTransform& WorldToComponent)
+void FIKRigProcessor::ResolveFinalGoalTransforms(const FTransform& WorldToComponent)
 {
 	for (FIKRigGoal& Goal : GoalContainer.Goals)
 	{
diff --git a/Engine/Plugins/Animation/IKRig/Source/IKRig/Public/AnimNodes/AnimNode_IKRig.h b/Engine/Plugins/Animation/IKRig/Source/IKRig/Public/AnimNodes/AnimNode_IKRig.h
--- a/Engine/Plugins/Animation/IKRig/Source/IKRig/Public/AnimNodes/AnimNode_IKRig.h
+++ b/Engine/Plugins/Animation/IKRig/Source/IKRig/Public/AnimNodes/AnimNode_IKRig.h
@@ -9,9 +9,9 @@
 #include "Animation/InputScaleBias.h"
 #include <functional>
 
+#include "Rig/IKRigProcessor.h"
 #include "AnimNode_IKRig.generated.h"
 
-class UIKRigProcessor;
 class IIKGoalCreatorInterface;
 class UIKRigDefinition;
 
@@ -77,25 +77,6 @@ struct IKRIG_API FAnimNode_IKRig : public FAnimNode_CustomProperty
 	UPROPERTY(EditAnywhere, BlueprintReadWrite, Category = Alpha)
 	FInputScaleBiasClamp AlphaScaleBiasClamp;
 	
-private:
-
-	/** IK Rig runtime processor */
-	UPROPERTY(Transient)
-	TObjectPtr<UIKRigProcessor> IKRigProcessor = nullptr;
-
-	/** a cached list of components on the owning actor that implement the goal creator interface */
-	TArray<IIKGoalCreatorInterface*> GoalCreators;
-	TMap<FName, FIKRigGoal> GoalsFromGoalCreators;
-	
-	TMap<FCompactPoseBoneIndex, int32, FDefaultSetAllocator, TCompactPoseBoneIndexMapKeyFuncs<int32>> CompactPoseToRigIndices;
-
-	/** Cached functions used to update goals using custom properties to avoid looking for then when evaluating */
-	using PropertyUpdateFunction = TFunction<void(const UObject*)>;
-	TArray<PropertyUpdateFunction> UpdateFunctions;
-
-	UPROPERTY(Transient)
-	float ActualAlpha;
-	
 public:
 
 	// FAnimNode_Base interface
@@ -111,10 +92,34 @@ public:
 	// End of FAnimNode_Base interface
 
 	// force reinitialization
-	void SetProcessorNeedsInitialized();
+	void SetProcessorNeedsInitialized() const;
+	
+	FIKRigProcessor* GetIKRigProcessor();
+
+private:
+
+	/** IK Rig runtime processor */
+	TSharedPtr<FIKRigProcessor> IKRigProcessor;
+
+	/** a cached list of components on the owning actor that implement the goal creator interface */
+	TArray<IIKGoalCreatorInterface*> GoalCreators;
+	TMap<FName, FIKRigGoal> GoalsFromGoalCreators;
+	bool bSearchedForGoalCreators = false;
 	
-	void CreateIKRigProcessorIfNeeded(UObject* Outer);
-	UIKRigProcessor* GetIKRigProcessor();
+	TMap<FCompactPoseBoneIndex, int32, FDefaultSetAllocator, TCompactPoseBoneIndexMapKeyFuncs<int32>> CompactPoseToRigIndices;
+
+	/** Cached functions used to update goals using custom properties to avoid looking for then when evaluating */
+	using PropertyUpdateFunction = TFunction<void(const UObject*)>;
+	TArray<PropertyUpdateFunction> UpdateFunctions;
+	
+	float ActualAlpha;
+	
+	void CopyInputPoseToSolver(FCompactPose& InputPose);
+	void AssignGoalTargets();
+	void CopyOutputPoseToAnimGraph(FCompactPose& OutputPose);
+#if WITH_EDITOR
+	virtual void ConditionalDebugDraw(FPrimitiveDrawInterface* PDI, USkeletalMeshComponent* PreviewSkelMeshComp) const;
+#endif
 
 protected:
 
@@ -123,12 +128,6 @@ protected:
 	virtual void InitializeProperties(const UObject* InSourceInstance, UClass* InTargetClass) override;
 	virtual void PropagateInputProperties(const UObject* InSourceInstance) override;
 	
-private:
-	
-	void CopyInputPoseToSolver(FCompactPose& InputPose);
-	void AssignGoalTargets();
-	void CopyOutputPoseToAnimGraph(FCompactPose& OutputPose);	
-	virtual void ConditionalDebugDraw(FPrimitiveDrawInterface* PDI, USkeletalMeshComponent* PreviewSkelMeshComp) const;
 	
 	friend class UAnimGraphNode_IKRig;
 };
diff --git a/Engine/Plugins/Animation/IKRig/Source/IKRig/Public/AnimNodes/AnimNode_RetargetPoseFromMesh.h b/Engine/Plugins/Animation/IKRig/Source/IKRig/Public/AnimNodes/AnimNode_RetargetPoseFromMesh.h
--- a/Engine/Plugins/Animation/IKRig/Source/IKRig/Public/AnimNodes/AnimNode_RetargetPoseFromMesh.h
+++ b/Engine/Plugins/Animation/IKRig/Source/IKRig/Public/AnimNodes/AnimNode_RetargetPoseFromMesh.h
@@ -81,11 +81,10 @@ struct IKRIG_API FAnimNode_RetargetPoseFromMesh : public FAnimNode_Base
 	// End of FAnimNode_Base interface
 
 	// access to the runtime processor
-	UIKRetargetProcessor* GetRetargetProcessor() const;
-	void CreateRetargetProcessorIfNeeded(UObject* Outer);
+	FIKRetargetProcessor* GetRetargetProcessor();
 
 	// returns true if processor is setup and ready to go, false otherwise
-	bool TryInitializeProcessor(USkeletalMeshComponent* TargetMeshComponent);
+	bool EnsureProcessorIsInitialized(const USkeletalMeshComponent* TargetMeshComponent);
 
 	// load deprecated properties
 	bool Serialize(FArchive& Ar);
@@ -97,8 +96,7 @@ private:
 	FRetargetProfile GetMergedRetargetProfile(bool bEnableIK) const;
 
 	// the runtime processor used to run the retarget and generate new poses
-	UPROPERTY(Transient)
-	TObjectPtr<UIKRetargetProcessor> Processor = nullptr;
+	TSharedPtr<FIKRetargetProcessor> Processor;
 
 	// cached transforms, copied on the game thread
 	TArray<FTransform> PoseToRetargetFromComponentSpace;
diff --git a/Engine/Plugins/Animation/IKRig/Source/IKRig/Public/IKRigLogger.h b/Engine/Plugins/Animation/IKRig/Source/IKRig/Public/IKRigLogger.h
--- a/Engine/Plugins/Animation/IKRig/Source/IKRig/Public/IKRigLogger.h
+++ b/Engine/Plugins/Animation/IKRig/Source/IKRig/Public/IKRigLogger.h
@@ -9,8 +9,8 @@ DECLARE_LOG_CATEGORY_EXTERN(LogIKRig, Warning, All);
 
 struct IKRIG_API FIKRigLogger
 {
-	/** Set the name of the log to output messages to */
-	void SetLogTarget(const FName InLogName, const FText& LogLabel);
+	/** Assign an asset for this log to store messages for. */
+	void SetLogTarget(const UObject* InAsset);
 	/** Get the name this log is currently outputting to */
 	FName GetLogTarget() const;
 	/** Log a warning message to display to user. */
diff --git a/Engine/Plugins/Animation/IKRig/Source/IKRig/Public/Retargeter/IKRetargetOps.h b/Engine/Plugins/Animation/IKRig/Source/IKRig/Public/Retargeter/IKRetargetOps.h
--- a/Engine/Plugins/Animation/IKRig/Source/IKRig/Public/Retargeter/IKRetargetOps.h
+++ b/Engine/Plugins/Animation/IKRig/Source/IKRig/Public/Retargeter/IKRetargetOps.h
@@ -10,25 +10,129 @@
 class FIKRetargetEditorController;
 #endif
 
-UCLASS(abstract, hidecategories = UObject)
-class IKRIG_API URetargetOpBase : public UObject
+USTRUCT(BlueprintType)
+struct FCurveRemapPair
 {
 	GENERATED_BODY()
+	
+	// The curve name on the SOURCE skeletal mesh to copy animation data from.
+	UPROPERTY(EditAnywhere, Category=Settings)
+	FName SourceCurve;
+
+	// The curve name on the TARGET skeletal mesh to receive animation data.
+	UPROPERTY(EditAnywhere, Category=Settings)
+	FName TargetCurve;
+};
 
+/** This is the base class for defining editable settings for your custom retargeting operation.
+ * All user-configurable properties for your "op" should be stored in a subclass of this
+ * NOTE: the derived type must be returned by the op's GetSettingsType() and GetSettings()
+ */
+USTRUCT(BlueprintType)
+struct IKRIG_API FIKRetargetOpSettingsBase
+{
+	GENERATED_BODY()
+	
+	virtual ~FIKRetargetOpSettingsBase() {};
+};
+
+/**
+ * This is the base class for defining operations that live in the retargeter "op" stack.
+ * These operations are executed in order by called the virtual Run() function on each one in order.
+ * The Run() function takes an input pose on the source skeletal mesh and affects the output pose on the target mesh.
+ * NOTE: any user defined settings associated with a retarget op must be aggregated into a custom UStruct derived from FIKRetargetOpSettingsBase
+ * This ensures that the settings are user editable in the details panel with full undo/redo support and serialization.
+*/ 
+USTRUCT(BlueprintType)
+struct IKRIG_API FIKRetargetOpBase
+{
+	GENERATED_BODY()
+	
 public:
 
-	URetargetOpBase() { SetFlags(RF_Transactional); }
+	virtual ~FIKRetargetOpBase() {};
+
+	/** (optional) override to cache internal data when initializing the processor
+	 * NOTE: you must set bIsInitialized to true to inform the retargeter that this op is ok to execute. */
+	virtual bool Initialize(
+		const FIKRetargetProcessor& Processor,
+		const FRetargetSkeleton& SourceSkeleton,
+		const FTargetSkeleton& TargetSkeleton,
+		FIKRigLogger& Log) { bIsInitialized = true; return true; };
+
+	/** (optional) override to evaluate this operation and modify the output pose */
+	virtual void Run(
+		const FIKRetargetProcessor& Processor,
+		const TArray<FTransform>& InSourceGlobalPose,
+		TArray<FTransform>& OutTargetGlobalPose){};
+
+	/** (required) override and return a pointer to the settings struct used by this operation */ 
+	virtual FIKRetargetOpSettingsBase* GetSettings() { checkNoEntry(); return nullptr; };
+	
+	/** (required) override and return the type used to house the settings for this operation */
+	virtual const UScriptStruct* GetSettingsType() const { checkNoEntry(); return nullptr; };
+
+	/** (optional) override to specify curves you want to retarget from the source animation graph */
+	virtual void GetCurvesToRetarget(TArray<FCurveRemapPair>& OutCurvesToRetarget){};
+
+	/** get if this op is enabled */
+	bool IsEnabled() const { return bIsEnabled; };
+
+	/** turn this operation on/off (will be skipped during execution if disabled) */
+	void SetEnabled(const bool bEnabled){ bIsEnabled = bEnabled; };
+
+	/** return true if this op is initialized and ready to run */
+	bool IsInitialized() const { return bIsInitialized; };
+
+#if WITH_EDITOR
+	/** override to automate initial setup after being added to the stack */
+	virtual void OnAddedToStack(const UIKRetargeter* Asset) {};
+
+	/** override to give your operation a nice name to display in the UI */
+	virtual FText GetNiceName() const { return FText::FromString(TEXT("Default Op Name")); };
+
+	/** override to display a warning message in the op stack */
+	virtual FText GetWarningMessage() const { return FText::GetEmpty(); };
+	
+	/** the instance of this op currently running in the editor viewport */
+	mutable FIKRetargetOpBase* EditorInstance;
+#endif
+
+private:
+	
+	/** when false, execution of this op is skipped */
+	bool bIsEnabled = true;
+
+protected:
+	
+	/** set to true when the op is ready to run */
+	bool bIsInitialized = false;
+};
+
+
+//
+// BEGIN DEPRECATED UOBJECT-BASED RETARGET OP
+//
+
+UCLASS(Deprecated)
+class UE_DEPRECATED(5.6, "This type has been replaced by FRetargetOpBase. URetargetOpBase-based ops no longer work, please refactor into the new FRetargetOpBase struct.")
+IKRIG_API UDEPRECATED_RetargetOpBase : public UObject
+{
+	GENERATED_BODY()
+	
+	public:
+
+	// this is the deprecation upgrade path for solvers that inherit from UIKRigSolver to convert them to FIKRigSolverBase
+	// override this and supply a struct into OutInstancedStruct that derives from FIKRigSolverBase that implements your custom solver.
+	virtual void ConvertToInstancedStruct(FInstancedStruct& OutInstancedStruct){};
 	
-	// override to cache internal data when initializing the processor
 	virtual bool Initialize(
-	const UIKRetargetProcessor* Processor,
+	const FIKRetargetProcessor& Processor,
 		const FRetargetSkeleton& SourceSkeleton,
 		const FTargetSkeleton& TargetSkeleton,
 		FIKRigLogger& Log) { return false; };
-
-	// override to evaluate this operation and modify the output pose
 	virtual void Run(
-		const UIKRetargetProcessor* Processor,
+		const FIKRetargetProcessor& Processor,
 		const TArray<FTransform>& InOutSourceGlobalPose,
 		TArray<FTransform>& OutTargetGlobalPose){};
 
@@ -38,31 +142,23 @@ public:
 	bool bIsInitialized = false;
 	
 #if WITH_EDITOR
-	// override to automate initial setup after being added to the stack
 	virtual void OnAddedToStack(const UIKRetargeter* Asset) {};
-	// override to give your operation a nice name to display in the UI
 	virtual FText GetNiceName() const { return FText::FromString(TEXT("Default Op Name")); };
-	// override to display a warning message in the op stack
 	virtual FText WarningMessage() const { return FText::GetEmpty(); };
-	// the instance of this op currently running in the editor
 	mutable URetargetOpBase* EditorInstance;
 #endif
 };
 
-// wrapper around a TArray of Retarget Ops for display in details panel
-UCLASS(BlueprintType)
-class IKRIG_API URetargetOpStack: public UObject
+UCLASS(Deprecated)
+class UE_DEPRECATED(5.6, "This type is no longer in use except to load old stacks of UObject-based ops.")
+IKRIG_API UDEPRECATED_RetargetOpStack : public UObject
 {
 	GENERATED_BODY()
 	
 public:
 
-	// stack of operations to run AFTER the main retarget phases (Order is: Root/IK/FK/Post) 
+	PRAGMA_DISABLE_DEPRECATION_WARNINGS
 	UPROPERTY()
-	TArray<TObjectPtr<URetargetOpBase>> RetargetOps;
-
-	// pointer to editor for details customization
-#if WITH_EDITOR
-	TWeakPtr<FIKRetargetEditorController> EditorController;
-#endif
+	TArray<TObjectPtr<UDEPRECATED_RetargetOpBase>> RetargetOps_DEPRECATED;
+	PRAGMA_ENABLE_DEPRECATION_WARNINGS
 };
\ No newline at end of file
diff --git a/Engine/Plugins/Animation/IKRig/Source/IKRig/Public/Retargeter/IKRetargetProcessor.h b/Engine/Plugins/Animation/IKRig/Source/IKRig/Public/Retargeter/IKRetargetProcessor.h
--- a/Engine/Plugins/Animation/IKRig/Source/IKRig/Public/Retargeter/IKRetargetProcessor.h
+++ b/Engine/Plugins/Animation/IKRig/Source/IKRig/Public/Retargeter/IKRetargetProcessor.h
@@ -7,15 +7,16 @@
 #include "IKRetargetSettings.h"
 #include "IKRigLogger.h"
 #include "Kismet/KismetMathLibrary.h"
+#include "Rig/IKRigProcessor.h"
 
 #include "IKRetargetProcessor.generated.h"
 
-class UIKRetargetProcessor;
+struct FIKRetargetOpBase;
 enum class ERetargetTranslationMode : uint8;
 enum class ERetargetRotationMode : uint8;
 class URetargetChainSettings;
 class UIKRigDefinition;
-class UIKRigProcessor;
+struct FIKRigProcessor;
 struct FReferenceSkeleton;
 struct FBoneChain;
 struct FIKRetargetPose;
@@ -525,14 +526,10 @@ private:
  * 2. Call RunRetargeter and pass in a source pose as an array of global-space transforms
  * 3. RunRetargeter returns an array of global space transforms for the target skeleton.
  */
-UCLASS()
-class IKRIG_API UIKRetargetProcessor : public UObject
+struct IKRIG_API FIKRetargetProcessor
 {
-	GENERATED_BODY()
 
 public:
-
-	UIKRetargetProcessor();
 	
 	/**
 	* Initialize the retargeter to enable running it.
@@ -583,7 +580,7 @@ public:
 		const UIKRetargeter* InRetargetAsset) const;
 
 	/** Get the currently running IK Rig processor for the target */
-	UIKRigProcessor* GetTargetIKRigProcessor() const { return IKRigProcessor; };
+	const FIKRigProcessor& GetTargetIKRigProcessor() const { return IKRigProcessor; };
 	
 	/** Get read only access to the core IK retarget chains */
 	const TArray<FRetargetChainPairIK>& GetIKChainPairs() const { return ChainPairsIK; };
@@ -594,23 +591,23 @@ public:
 	/** Get read only access to the core root retargeter */
 	const FRootRetargeter& GetRootRetargeter() const { return RootRetargeter; };
 	
-	// Get read only access to the retarget ops currently running in processor
-	const TArray<TObjectPtr<URetargetOpBase>>& GetRetargetOps() const {return OpStack; };
+	/** Get read only access to the retarget ops currently running in processor */
+	const TArray<TSharedPtr<FIKRetargetOpBase>>& GetRetargetOps() const {return OpStack; };
 
-	// Get read only access to the current global settings
+	/** Get read only access to the current global settings */
 	const FRetargetGlobalSettings& GetGlobalSettings() const {return GlobalSettings; };
 	
 	/** Reset the IK planting state. */
 	void ResetPlanting();
 
-	/** logging system */
-	FIKRigLogger Log;
-
 	/** Set that this processor needs to be reinitialized. */
 	void SetNeedsInitialized();
 	
 	/** Apply the settings stored in the IKRig asset. */
-	void CopyIKRigSettingsFromAsset() const;
+	void CopyIKRigSettingsFromAsset();
+
+	/** The logging system */
+	FIKRigLogger Log;
 	
 #if WITH_EDITOR
 	
@@ -664,8 +661,7 @@ private:
 	FTargetSkeleton TargetSkeleton;
 
 	/** The IK Rig processor for running IK on the target */
-	UPROPERTY(Transient) // must be property to keep from being GC'd
-	TObjectPtr<UIKRigProcessor> IKRigProcessor = nullptr;
+	FIKRigProcessor IKRigProcessor;
 
 	/** The Source/Target pairs of Bone Chains retargeted using the FK algorithm */
 	TArray<FRetargetChainPairFK> ChainPairsFK;
@@ -683,8 +679,7 @@ private:
 	FRetargetGlobalSettings GlobalSettings;
 
 	/** The collection of operations to run in the final phase of retargeting */
-	UPROPERTY(Transient) // must be property to keep from being GC'd
-	TArray<TObjectPtr<URetargetOpBase>> OpStack;
+	TArray<TSharedPtr<FIKRetargetOpBase>> OpStack;
 
 	/** Apply the settings stored in a retarget profile. Called inside RunRetargeter(). */
 	void ApplySettingsFromProfile(const FRetargetProfile& Profile);
@@ -700,10 +695,10 @@ private:
 	bool InitializeBoneChainPairs(const TMap<FName, FTargetChainSettings>& ChainSettings);
 
 	/** Initializes the IK Rig that evaluates the IK solve for the target IK chains */
-	bool InitializeIKRig(UObject* Outer, const USkeletalMesh* InSkeletalMesh);
+	bool InitializeIKRig(const USkeletalMesh* InSkeletalMesh);
 
 	// Initialize the retarget op stack
-	bool InitializeOpStack(const TArray<TObjectPtr<URetargetOpBase>>& OpStackFromAsset);
+	bool InitializeOpStack(const TArray<FInstancedStruct>& OpsFromAsset);
 
 	/** Internal retarget phase that does simple bone-to-bone copying from source*/
 	void GenerateBasePoses(TArray<FTransform>& InSourceGlobalPose);
@@ -730,3 +725,148 @@ private:
 	/** Run all post process operations on the retargeted result. */
 	void RunRetargetOps(const TArray<FTransform>& InSourceGlobalPose, TArray<FTransform>& OutTargetGlobalPose);
 };
+
+//
+// BEGIN DEPRECATED UOBJECT-BASED PROCESSOR
+//
+// NOTE: As of 5.6, IK Retarget processor is no longer a UObject.
+// The UObject based processor is here for backwards compatibility.
+// If you have a system built off this processor, please convert it to the FIKRetargetProcessor as this will be removed in a future update.
+//
+
+UCLASS()
+class UE_DEPRECATED(5.6, "UIKRetargetProcessor has been replaced by FIKRetargetProcessor. UIKRetargetProcessor is now a thin wrapper around FIKRetargetProcessor and will be removed in a future version.")
+IKRIG_API UIKRetargetProcessor : public UObject
+{
+	GENERATED_BODY()
+	
+public:
+	
+	void Initialize(
+		const USkeletalMesh *SourceSkeleton,
+		const USkeletalMesh *TargetSkeleton,
+		const UIKRetargeter* InRetargeterAsset,
+		const FRetargetProfile& Settings,
+		const bool bSuppressWarnings=false)
+	{
+		Processor.Initialize(SourceSkeleton, TargetSkeleton, InRetargeterAsset, Settings, bSuppressWarnings);
+	}
+	
+	TArray<FTransform>& RunRetargeter(
+		TArray<FTransform>& InSourceGlobalPose,
+		const TMap<FName,float>& SpeedValuesFromCurves,
+		const float DeltaTime,
+		const FRetargetProfile& Profile)
+	{
+		return Processor.RunRetargeter(InSourceGlobalPose, SpeedValuesFromCurves, DeltaTime, Profile);
+	}
+
+	void UpdateRetargetPoseAtRuntime(
+		const FName NewRetargetPoseName,
+		ERetargetSourceOrTarget SourceOrTarget)
+	{
+		Processor.UpdateRetargetPoseAtRuntime(NewRetargetPoseName, SourceOrTarget);
+	}
+
+	const FRetargetSkeleton& GetSkeleton(ERetargetSourceOrTarget SourceOrTarget) const
+	{
+		return Processor.GetSkeleton(SourceOrTarget);
+	}
+	FName GetRetargetRoot(ERetargetSourceOrTarget SourceOrTarget) const
+	{
+		return Processor.GetRetargetRoot(SourceOrTarget);
+	}
+	bool IsInitialized() const
+	{
+		return Processor.IsInitialized();
+	}
+	bool WasInitializedWithTheseAssets(
+		const USkeletalMesh* InSourceMesh,
+		const USkeletalMesh* InTargetMesh,
+		const UIKRetargeter* InRetargetAsset) const
+	{
+		return Processor.WasInitializedWithTheseAssets(InSourceMesh, InTargetMesh, InRetargetAsset);
+	}
+	const FIKRigProcessor& GetTargetIKRigProcessor() const
+	{
+		return Processor.GetTargetIKRigProcessor();
+	}
+	const TArray<FRetargetChainPairIK>& GetIKChainPairs() const
+	{
+		return Processor.GetIKChainPairs();
+	}
+	const TArray<FRetargetChainPairFK>& GetFKChainPairs() const
+	{
+		return Processor.GetFKChainPairs();
+	}
+	const FRootRetargeter& GetRootRetargeter() const
+	{
+		return Processor.GetRootRetargeter();
+	};
+	const TArray<TObjectPtr<URetargetOpBase>>& GetRetargetOps() const
+	{
+		return DummyStack;
+	};
+	const FRetargetGlobalSettings& GetGlobalSettings() const
+	{
+		return Processor.GetGlobalSettings();
+	};
+	void ResetPlanting()
+	{
+		Processor.ResetPlanting();
+	};
+	void SetNeedsInitialized()
+	{
+		Processor.SetNeedsInitialized();
+	};
+	void CopyIKRigSettingsFromAsset()
+	{
+		Processor.CopyIKRigSettingsFromAsset();
+	};
+
+	FIKRigLogger Log;
+	
+#if WITH_EDITOR
+	bool IsBoneRetargeted(const FName BoneName, const ERetargetSourceOrTarget SourceOrTarget) const
+	{
+		return Processor.IsBoneRetargeted(BoneName, SourceOrTarget);
+	}
+	int32 GetBoneIndexFromName(const FName BoneName, const ERetargetSourceOrTarget SourceOrTarget) const
+	{
+		return Processor.GetBoneIndexFromName(BoneName, SourceOrTarget);
+	}
+	FName GetChainNameForBone(const FName BoneName, const ERetargetSourceOrTarget SourceOrTarget) const
+	{
+		return Processor.GetChainNameForBone(BoneName, SourceOrTarget);
+	}
+	FTransform GetGlobalRetargetPoseAtParam(const FName InChainName, const float Param, const ERetargetSourceOrTarget SourceOrTarget) const
+	{
+		return Processor.GetGlobalRetargetPoseAtParam(InChainName, Param, SourceOrTarget);
+	}
+	FTransform GetRetargetPoseBoneTransform(const FName InBoneName, const ERetargetSourceOrTarget SourceOrTarget, ERetargetBoneSpace BoneSpace) const
+	{
+		return Processor.GetRetargetPoseBoneTransform(InBoneName, SourceOrTarget, BoneSpace);
+	}
+	const FChainFK* GetChain(const FName InChainName, const ERetargetSourceOrTarget SourceOrTarget) const
+	{
+		return Processor.GetChain(InChainName, SourceOrTarget);
+	}
+	float GetParamOfBoneInChain(const FName InBoneName, const ERetargetSourceOrTarget SourceOrTarget) const
+	{
+		return Processor.GetParamOfBoneInChain(InBoneName, SourceOrTarget);
+	}
+	FName GetClosestBoneToParam(const FName InChainName, const float InParam, const ERetargetSourceOrTarget SourceOrTarget) const
+	{
+		return Processor.GetClosestBoneToParam(InChainName, InParam, SourceOrTarget);
+	}
+	FName GetMappedChainName(const FName InChainName, const ERetargetSourceOrTarget SourceOrTarget)
+	{
+		return Processor.GetMappedChainName(InChainName, SourceOrTarget);
+	}
+#endif
+
+
+private:
+	TArray<TObjectPtr<URetargetOpBase>> DummyStack;
+	FIKRetargetProcessor Processor;
+};
\ No newline at end of file
diff --git a/Engine/Plugins/Animation/IKRig/Source/IKRig/Public/Retargeter/IKRetargeter.h b/Engine/Plugins/Animation/IKRig/Source/IKRig/Public/Retargeter/IKRetargeter.h
--- a/Engine/Plugins/Animation/IKRig/Source/IKRig/Public/Retargeter/IKRetargeter.h
+++ b/Engine/Plugins/Animation/IKRig/Source/IKRig/Public/Retargeter/IKRetargeter.h
@@ -14,7 +14,7 @@ class FIKRetargetEditorController;
 struct FIKRetargetPose;
 class URetargetChainSettings;
 class URetargetRootSettings;
-class URetargetOpStack;
+class UDEPRECATED_RetargetOpStack;
 
 struct UE_DEPRECATED(5.1, "Use URetargetChainSettings instead.") FRetargetChainMap;
 USTRUCT()
@@ -239,8 +239,8 @@ public:
 	UIKRetargetGlobalSettings* GetGlobalSettingsUObject() const { return GlobalSettings; };
 	// Get access to the global settings itself 
 	const FRetargetGlobalSettings& GetGlobalSettings() const { return GlobalSettings->Settings; };
-	// Get access to the post settings uobject
-	URetargetOpStack* GetPostSettingsUObject() const { return OpStack; };
+	// Get access to the stack of retargeting operations
+	const TArray<FInstancedStruct>& GetRetargetOps() const { return RetargetOps; };
 
 	// Get read-only access to a retarget pose 
 	const FIKRetargetPose* GetCurrentRetargetPose(const ERetargetSourceOrTarget& SourceOrTarget) const;
@@ -504,13 +504,13 @@ private:
 	UPROPERTY()
 	TObjectPtr<URetargetRootSettings> RootSettings;
 
-	// the retarget root settings 
+	// the retarget global settings 
 	UPROPERTY()
 	TObjectPtr<UIKRetargetGlobalSettings> GlobalSettings;
 
-	// the stack of ops to run after the retarget
-	UPROPERTY()
-	TObjectPtr<URetargetOpStack> OpStack;
+	/** polymorphic stack of retargeting operations, executed in serial fashion where output of prior operation is input to the next */
+	UPROPERTY(meta = (ExcludeBaseStruct, BaseStruct = "/Script/IKRig.RetargetOpBase"))
+	TArray<FInstancedStruct> RetargetOps;
 
 	// settings profiles stored in this asset 
 	UPROPERTY()
@@ -540,5 +540,11 @@ private:
 	FName CurrentRetargetPose_DEPRECATED;
 	PRAGMA_ENABLE_DEPRECATION_WARNINGS
 
+	// 5.6 refactor to remove UObjects from retargeter
+	PRAGMA_DISABLE_DEPRECATION_WARNINGS
+	UPROPERTY(meta = (DeprecatedProperty, DeprecationMessage = "The pre 5.6 stack of UObject based solvers. Use RetargetOps instead."))
+	TObjectPtr<UDEPRECATED_RetargetOpStack> OpStack_DEPRECATED;
+	PRAGMA_ENABLE_DEPRECATION_WARNINGS
+
 	friend class UIKRetargeterController;
 };
diff --git a/Engine/Plugins/Animation/IKRig/Source/IKRig/Public/Retargeter/RetargetOps/CurveRemapOp.h b/Engine/Plugins/Animation/IKRig/Source/IKRig/Public/Retargeter/RetargetOps/CurveRemapOp.h
--- a/Engine/Plugins/Animation/IKRig/Source/IKRig/Public/Retargeter/RetargetOps/CurveRemapOp.h
+++ b/Engine/Plugins/Animation/IKRig/Source/IKRig/Public/Retargeter/RetargetOps/CurveRemapOp.h
@@ -7,36 +7,9 @@
 #define LOCTEXT_NAMESPACE "CurveRemapOp"
 
 USTRUCT(BlueprintType)
-struct FCurveRemapPair
+struct IKRIG_API FIKRetargetCurveRemapOpSettings : public FIKRetargetOpSettingsBase
 {
 	GENERATED_BODY()
-	
-	// The curve name on the SOURCE skeletal mesh to copy animation data from.
-	UPROPERTY(EditAnywhere, Category=Settings)
-	FName SourceCurve;
-
-	// The curve name on the TARGET skeletal mesh to receive animation data.
-	UPROPERTY(EditAnywhere, Category=Settings)
-	FName TargetCurve;
-};
-
-UCLASS(BlueprintType, EditInlineNew)
-class IKRIG_API UCurveRemapOp : public URetargetOpBase
-{
-	GENERATED_BODY()
-
-public:
-
-	// this op does not do anything in Initialize() or Run(), but is used by the Retarget Pose from Mesh node and the Batch Exporter
-	virtual bool Initialize(
-	const UIKRetargetProcessor* Processor,
-		const FRetargetSkeleton& SourceSkeleton,
-		const FTargetSkeleton& TargetSkeleton,
-		FIKRigLogger& Log) override { return true; };
-	virtual void Run(
-		const UIKRetargetProcessor* Processor,
-		const TArray<FTransform>& InSourceGlobalPose,
-		TArray<FTransform>& OutTargetGlobalPose) override {};
 
 	// Add pairs of Source/Target curve names to remap. While retargeting, the animation from the source curves
 	// will be redirected to the curves on the target skeletal meshes. Can be used to drive, blendshapes or other downstream systems.
@@ -50,13 +23,78 @@ public:
 	// False: only remapped curves are left on the target animation sequence asset.
 	UPROPERTY(EditAnywhere, Category="OpSettings")
 	bool bCopyAllSourceCurves = true;
+};
+
+USTRUCT(BlueprintType)
+struct IKRIG_API FIKRetargetCurveRemapOp : public FIKRetargetOpBase
+{
+	GENERATED_BODY()
+	
+	// NOTE: this op does not do anything in Initialize() or Run().
+	// It implements GetCurvesToRetarget() which the retargeting anim node calls and manages
+	// In the future we may remove this coupling and make the op itself do the work via callbacks
+	virtual bool Initialize(
+	const FIKRetargetProcessor& Processor,
+		const FRetargetSkeleton& SourceSkeleton,
+		const FTargetSkeleton& TargetSkeleton,
+		FIKRigLogger& Log) override { bIsInitialized = true; return true; };
+	virtual void Run(
+		const FIKRetargetProcessor& Processor,
+		const TArray<FTransform>& InSourceGlobalPose,
+		TArray<FTransform>& OutTargetGlobalPose) override {};
+
+	virtual FIKRetargetOpSettingsBase* GetSettings() override
+	{
+		return &Settings;
+	};
+	
+	virtual const UScriptStruct* GetSettingsType() const override
+	{
+		return FIKRetargetCurveRemapOpSettings::StaticStruct();
+	}
+
+	virtual void GetCurvesToRetarget(TArray<FCurveRemapPair>& OutCurvesToRetarget) override
+	{
+		OutCurvesToRetarget.Append(Settings.CurvesToRemap);
+	};
+
+	UPROPERTY()
+	FIKRetargetCurveRemapOpSettings Settings;
 	
 #if WITH_EDITOR
-	virtual FText GetNiceName() const override { return FText(LOCTEXT("OpName", "Remap Curves")); };
-	virtual FText WarningMessage() const override { return Message; };
+	virtual FText GetNiceName() const override { return FText(LOCTEXT("RemapOpName", "Remap Curves")); };
+	virtual FText GetWarningMessage() const override { return Message; };
 	FText Message;
 #endif
 };
 
+//
+// BEGIN DEPRECATED UOBJECT-based OP
+//
+
+PRAGMA_DISABLE_DEPRECATION_WARNINGS
+
+UCLASS(Deprecated)
+class
+UE_DEPRECATED(5.6, "This type has been replaced with FIKRetargetCurveRemapOp.")
+IKRIG_API UDEPRECATED_CurveRemapOp : public UDEPRECATED_RetargetOpBase
+{
+	GENERATED_BODY()
+	
+	virtual void ConvertToInstancedStruct(FInstancedStruct& OutInstancedStruct) override
+	{
+		OutInstancedStruct.InitializeAs(FIKRetargetCurveRemapOp::StaticStruct());
+		FIKRetargetCurveRemapOp& NewOp = OutInstancedStruct.GetMutable<FIKRetargetCurveRemapOp>();
+		NewOp.Settings.CurvesToRemap = CurvesToRemap;
+		NewOp.Settings.bCopyAllSourceCurves = bCopyAllSourceCurves;
+	};
+	
+	UPROPERTY()
+	TArray<FCurveRemapPair> CurvesToRemap;
+	UPROPERTY()
+	bool bCopyAllSourceCurves = true;
+};
+
+PRAGMA_ENABLE_DEPRECATION_WARNINGS
 
 #undef LOCTEXT_NAMESPACE
diff --git a/Engine/Plugins/Animation/IKRig/Source/IKRig/Public/Retargeter/RetargetOps/PinBoneOp.h b/Engine/Plugins/Animation/IKRig/Source/IKRig/Public/Retargeter/RetargetOps/PinBoneOp.h
--- a/Engine/Plugins/Animation/IKRig/Source/IKRig/Public/Retargeter/RetargetOps/PinBoneOp.h
+++ b/Engine/Plugins/Animation/IKRig/Source/IKRig/Public/Retargeter/RetargetOps/PinBoneOp.h
@@ -27,7 +27,7 @@ struct FPinBoneData
 	
 	static const FTransform& GetParentTransform(int32 BoneIndex, const FRetargetSkeleton& Skeleton, const TArray<FTransform>& Pose);
 
-	void CachePinData(ERetargetSourceOrTarget InSkeletonToCopyFrom, const UIKRetargetProcessor* Processor);
+	void CachePinData(ERetargetSourceOrTarget InSkeletonToCopyFrom, const FIKRetargetProcessor& Processor);
 
 	// The bone to copy FROM.
 	UPROPERTY(EditAnywhere, Category=Settings)
@@ -85,39 +85,18 @@ enum class EPinBoneRotationMode : uint8
 	MaintainOffsetFromBoneToCopyFrom
 };
 
-UCLASS(BlueprintType, EditInlineNew)
-class IKRIG_API UPinBoneOp : public URetargetOpBase, public IBoneReferenceSkeletonProvider
+USTRUCT(BlueprintType)
+struct IKRIG_API FIKRetargetPinBoneOpSettings : public FIKRetargetOpSettingsBase
 {
 	GENERATED_BODY()
-
-public:
-	
-	virtual bool Initialize(
-	const UIKRetargetProcessor* Processor,
-		const FRetargetSkeleton& SourceSkeleton,
-		const FTargetSkeleton& TargetSkeleton,
-		FIKRigLogger& Log) override;
 	
-	virtual void Run(
-		const UIKRetargetProcessor* Processor,
-		const TArray<FTransform>& InSourceGlobalPose,
-		TArray<FTransform>& OutTargetGlobalPose) override;
-
-	// UObject interface
-	virtual void Serialize(FArchive& Ar) override;
-	// ~END UObject interface
-
-	// IBoneReferenceSkeletonProvider overrides.
-	virtual USkeleton* GetSkeleton(bool& bInvalidSkeletonIsError, const IPropertyHandle* PropertyHandle) override;
-	// ~END IBoneReferenceSkeletonProvider overrides.
-
 	// A list of bone-pairs to copy transforms between
 	UPROPERTY(EditAnywhere, Category=Settings)
 	TArray<FPinBoneData> BonesToPin;
 
 	// Which skeleton to copy animation from.
-	UPROPERTY(EditAnywhere, Category=Settings, meta=(DisplayName="Skeleton To Copy From"))
-	ERetargetSourceOrTarget PinTo = ERetargetSourceOrTarget::Target;
+	UPROPERTY(EditAnywhere, Category=Settings)
+	ERetargetSourceOrTarget SkeletonToCopyFrom = ERetargetSourceOrTarget::Target;
 
 	// Copy the translation of the bone.
 	UPROPERTY(EditAnywhere, Category=Settings)
@@ -146,10 +125,38 @@ public:
 	// A manual offset to apply in local space
 	UPROPERTY(EditAnywhere, Category=Settings)
 	FTransform LocalOffset;
+};
+
+USTRUCT(BlueprintType)
+struct IKRIG_API FIKRetargetPinBoneOp : public FIKRetargetOpBase, public IBoneReferenceSkeletonProvider
+{
+	GENERATED_BODY()
+
+	virtual bool Initialize(
+	const FIKRetargetProcessor& Processor,
+		const FRetargetSkeleton& SourceSkeleton,
+		const FTargetSkeleton& TargetSkeleton,
+		FIKRigLogger& Log) override;
+	
+	virtual void Run(
+		const FIKRetargetProcessor& Processor,
+		const TArray<FTransform>& InSourceGlobalPose,
+		TArray<FTransform>& OutTargetGlobalPose) override;
+
+	virtual FIKRetargetOpSettingsBase* GetSettings() override;
+	
+	virtual const UScriptStruct* GetSettingsType() const override;
+
+	// IBoneReferenceSkeletonProvider overrides.
+	virtual USkeleton* GetSkeleton(bool& bInvalidSkeletonIsError, const IPropertyHandle* PropertyHandle) override;
+	// ~END IBoneReferenceSkeletonProvider overrides.
+
+	UPROPERTY()
+	FIKRetargetPinBoneOpSettings Settings;
 	
 #if WITH_EDITOR
 	virtual FText GetNiceName() const override { return FText(LOCTEXT("OpName", "Pin Bone")); };
-	virtual FText WarningMessage() const override { return Message; };
+	virtual FText GetWarningMessage() const override { return Message; };
 	FText Message;
 
 private:
@@ -161,19 +168,72 @@ private:
 
 	FTransform GetNewBoneTransform(
 		const FPinBoneData& PinData,
-		const UIKRetargetProcessor* Processor,
+		const FIKRetargetProcessor& Processor,
 		const TArray<FTransform>& InSourceGlobalPose,
 		TArray<FTransform>& OutTargetGlobalPose) const;
+};
+
+//
+// BEGIN DEPRECATED UOBJECT-based OP
+//
+PRAGMA_DISABLE_DEPRECATION_WARNINGS
+UCLASS(Deprecated)
+class
+UE_DEPRECATED(5.6, "This type has been replaced with FIKRetargetCurveRemapOp.")
+IKRIG_API UDEPRECATED_PinBoneOp : public UDEPRECATED_RetargetOpBase
+{
+	GENERATED_BODY()
+
+public:
+	
+	virtual void ConvertToInstancedStruct(FInstancedStruct& OutInstancedStruct) override
+	{
+		OutInstancedStruct.InitializeAs(FIKRetargetPinBoneOp::StaticStruct());
+		FIKRetargetPinBoneOp& NewOp = OutInstancedStruct.GetMutable<FIKRetargetPinBoneOp>();
+		NewOp.Settings.BonesToPin = BonesToPin;
+		NewOp.Settings.SkeletonToCopyFrom = PinTo;
+		NewOp.Settings.bCopyTranslation = bCopyTranslation;
+		NewOp.Settings.bCopyRotation = bCopyRotation;
+		NewOp.Settings.bCopyScale = bCopyScale;
+		NewOp.Settings.TranslationMode = TranslationMode;
+		NewOp.Settings.RotationMode = RotationMode;
+		NewOp.Settings.GlobalOffset = GlobalOffset;
+		NewOp.Settings.LocalOffset = LocalOffset;
+	};
+
+	// UObject interface
+	virtual void Serialize(FArchive& Ar) override;
+	// ~END UObject interface
+	
+	UPROPERTY()
+	TArray<FPinBoneData> BonesToPin;
+	UPROPERTY()
+	ERetargetSourceOrTarget PinTo = ERetargetSourceOrTarget::Target;
+	UPROPERTY()
+	bool bCopyTranslation = true;
+	UPROPERTY()
+	bool bCopyRotation = true;
+	UPROPERTY()
+	bool bCopyScale = true;
+	UPROPERTY()
+	EPinBoneTranslationMode TranslationMode;
+	UPROPERTY()
+	EPinBoneRotationMode RotationMode;
+	UPROPERTY()
+	FTransform GlobalOffset;
+	UPROPERTY()
+	FTransform LocalOffset;
+
+private:
 
 	UE_DEPRECATED(5.6, "Uses separate translation/rotation offset modes")
 	UPROPERTY()
 	bool bMaintainOffset_DEPRECATED = true;
-
-	PRAGMA_DISABLE_DEPRECATION_WARNINGS
+	
 	UE_DEPRECATED(5.6, "Uses separate translate/rotate/scale checkboxes")
 	UPROPERTY()
 	EPinBoneType PinType_DEPRECATED;
-	PRAGMA_ENABLE_DEPRECATION_WARNINGS
 };
+PRAGMA_ENABLE_DEPRECATION_WARNINGS
 
 #undef LOCTEXT_NAMESPACE
diff --git a/Engine/Plugins/Animation/IKRig/Source/IKRig/Public/Retargeter/RetargetOps/RootMotionGeneratorOp.h b/Engine/Plugins/Animation/IKRig/Source/IKRig/Public/Retargeter/RetargetOps/RootMotionGeneratorOp.h
--- a/Engine/Plugins/Animation/IKRig/Source/IKRig/Public/Retargeter/RetargetOps/RootMotionGeneratorOp.h
+++ b/Engine/Plugins/Animation/IKRig/Source/IKRig/Public/Retargeter/RetargetOps/RootMotionGeneratorOp.h
@@ -23,24 +23,11 @@ enum class ERootMotionHeightSource : uint8
 	SnapToGround,
 };
 
-UCLASS(BlueprintType, EditInlineNew)
-class IKRIG_API URootMotionGeneratorOp : public URetargetOpBase
+USTRUCT(BlueprintType)
+struct IKRIG_API FIKRetargetRootMotionOpSettings : public FIKRetargetOpSettingsBase
 {
 	GENERATED_BODY()
 
-public:
-	
-	virtual bool Initialize(
-		const UIKRetargetProcessor* Processor,
-		const FRetargetSkeleton& SourceSkeleton,
-		const FTargetSkeleton& TargetSkeleton,
-		FIKRigLogger& Log) override;
-	
-	virtual void Run(
-		const UIKRetargetProcessor* Processor,
-		const TArray<FTransform>& InSourceGlobalPose,
-		TArray<FTransform>& OutTargetGlobalPose) override;
-
 	// The root of the source skeleton.
 	UPROPERTY(EditAnywhere, Category="OpSettings")
 	FName SourceRootBone;
@@ -85,11 +72,35 @@ public:
 	// A manual offset to apply in global space to the root bone.
 	UPROPERTY(EditAnywhere, Category=Settings)
 	FTransform GlobalOffset;
+};
+
+USTRUCT(BlueprintType)
+struct IKRIG_API FIKRetargetRootMotionOp : public FIKRetargetOpBase
+{
+	GENERATED_BODY()
+
+	virtual bool Initialize(
+		const FIKRetargetProcessor& Processor,
+		const FRetargetSkeleton& SourceSkeleton,
+		const FTargetSkeleton& TargetSkeleton,
+		FIKRigLogger& Log) override;
+	
+	virtual void Run(
+		const FIKRetargetProcessor& Processor,
+		const TArray<FTransform>& InSourceGlobalPose,
+		TArray<FTransform>& OutTargetGlobalPose) override;
+	
+	virtual FIKRetargetOpSettingsBase* GetSettings() override;
+	
+	virtual const UScriptStruct* GetSettingsType() const override;
+
+	UPROPERTY()
+	FIKRetargetRootMotionOpSettings Settings;
 
 #if WITH_EDITOR
 	virtual void OnAddedToStack(const UIKRetargeter* Asset) override;
 	virtual FText GetNiceName() const override { return FText(LOCTEXT("OpName", "Root Motion Generator")); };
-	virtual FText WarningMessage() const override { return Message; };
+	virtual FText GetWarningMessage() const override { return Message; };
 	FText Message;
 #endif
 
@@ -102,7 +113,7 @@ private:
 	
 	void CopyRootMotionFromSourceRoot(
 		FTransform& OutRootTransform,
-		const UIKRetargetProcessor* Processor,
+		const FIKRetargetProcessor& Processor,
 		const TArray<FTransform>& InSourceGlobalPose) const;
 	
 	int32 SourceRootIndex;
@@ -117,4 +128,54 @@ private:
 	TArray<int32> NonRetargetedChildrenOfRoot;
 };
 
+//
+// BEGIN DEPRECATED UOBJECT-based OP
+//
+
+PRAGMA_DISABLE_DEPRECATION_WARNINGS
+UCLASS(Deprecated)
+class
+UE_DEPRECATED(5.6, "This type has been replaced with FIKRetargetCurveRemapOp.")
+IKRIG_API UDEPRECATED_RootMotionGeneratorOp : public UDEPRECATED_RetargetOpBase
+{
+	GENERATED_BODY()
+
+public:
+	
+	virtual void ConvertToInstancedStruct(FInstancedStruct& OutInstancedStruct) override
+	{
+		OutInstancedStruct.InitializeAs(FIKRetargetRootMotionOp::StaticStruct());
+		FIKRetargetRootMotionOp& NewOp = OutInstancedStruct.GetMutable<FIKRetargetRootMotionOp>();
+		NewOp.Settings.SourceRootBone = SourceRootBone;
+		NewOp.Settings.TargetRootBone = TargetRootBone;
+		NewOp.Settings.TargetPelvisBone = TargetPelvisBone;
+		NewOp.Settings.RootMotionSource = RootMotionSource;
+		NewOp.Settings.RootHeightSource = RootHeightSource;
+		NewOp.Settings.bPropagateToNonRetargetedChildren = bPropagateToNonRetargetedChildren;
+		NewOp.Settings.bMaintainOffsetFromPelvis = bMaintainOffsetFromPelvis;
+		NewOp.Settings.bRotateWithPelvis = bRotateWithPelvis;
+		NewOp.Settings.GlobalOffset = GlobalOffset;
+	};
+
+	UPROPERTY()
+	FName SourceRootBone;
+	UPROPERTY()
+	FName TargetRootBone;
+	UPROPERTY()
+	FName TargetPelvisBone;
+	UPROPERTY()
+	ERootMotionSource RootMotionSource;
+	UPROPERTY()
+	ERootMotionHeightSource RootHeightSource;
+	UPROPERTY()
+	bool bPropagateToNonRetargetedChildren = true;
+	UPROPERTY()
+	bool bMaintainOffsetFromPelvis = true;
+	UPROPERTY()
+	bool bRotateWithPelvis = false;
+	UPROPERTY()
+	FTransform GlobalOffset;
+};
+PRAGMA_ENABLE_DEPRECATION_WARNINGS
+
 #undef LOCTEXT_NAMESPACE
diff --git a/Engine/Plugins/Animation/IKRig/Source/IKRig/Public/Rig/IKRigDataTypes.h b/Engine/Plugins/Animation/IKRig/Source/IKRig/Public/Rig/IKRigDataTypes.h
--- a/Engine/Plugins/Animation/IKRig/Source/IKRig/Public/Rig/IKRigDataTypes.h
+++ b/Engine/Plugins/Animation/IKRig/Source/IKRig/Public/Rig/IKRigDataTypes.h
@@ -7,7 +7,7 @@
 
 #include "IKRigDataTypes.generated.h"
 
-class UIKRigProcessor;
+struct FIKRigProcessor;
 
 UENUM(BlueprintType)
 enum class EIKRigGoalSpace : uint8
@@ -177,5 +177,5 @@ private:
 	/** Array of Goals. Cannot contain duplicates (name is key). */
 	TArray<FIKRigGoal> Goals;
 
-	friend UIKRigProcessor;
+	friend FIKRigProcessor;
 };
diff --git a/Engine/Plugins/Animation/IKRig/Source/IKRig/Public/Rig/IKRigDefinition.h b/Engine/Plugins/Animation/IKRig/Source/IKRig/Public/Rig/IKRigDefinition.h
--- a/Engine/Plugins/Animation/IKRig/Source/IKRig/Public/Rig/IKRigDefinition.h
+++ b/Engine/Plugins/Animation/IKRig/Source/IKRig/Public/Rig/IKRigDefinition.h
@@ -205,7 +205,7 @@ public:
 	const TArray<const FIKRigSolverBase*> GetSolverArray() const;
 
 	/** runtime, read-only access to array of Solvers structs, all modifications must go through UIKRigController */
-	const TArray<FInstancedStruct> GetSolverStructs() const;
+	const TArray<FInstancedStruct>& GetSolverStructs() const;
 
 	/** runtime, read-only access to Retarget Definition, all modifications must go through UIKRigController */
 	const TArray<FBoneChain>& GetRetargetChains() const { return RetargetDefinition.BoneChains; };
diff --git a/Engine/Plugins/Animation/IKRig/Source/IKRig/Public/Rig/IKRigProcessor.h b/Engine/Plugins/Animation/IKRig/Source/IKRig/Public/Rig/IKRigProcessor.h
--- a/Engine/Plugins/Animation/IKRig/Source/IKRig/Public/Rig/IKRigProcessor.h
+++ b/Engine/Plugins/Animation/IKRig/Source/IKRig/Public/Rig/IKRigProcessor.h
@@ -22,26 +22,21 @@ struct FGoalBone
 	int32 OptSourceIndex = INDEX_NONE;
 };
 
-UCLASS()
-class IKRIG_API UIKRigProcessor : public UObject
+/** the runtime for an IKRig to convert an input pose into
+*   a solved output pose given a set of IK Rig Goals:
+*   
+* 1. Create an FIKRigProcessor in your animation system (ie anim node)
+* 2. Initialize() with an IKRigDefinition asset
+* 3. each tick, call SetIKGoal() and SetInputPoseGlobal()
+* 4. Call Solve()
+* 5. Copy output transforms with GetOutputPoseGlobal()
+* 
+*/
+struct IKRIG_API FIKRigProcessor
 {
-	GENERATED_BODY()
 	
 public:
 	
-	/** the runtime for an IKRig to convert an input pose into
-	*   a solved output pose given a set of IK Rig Goals:
-	*   
-	* 1. Create a new IKRigProcessor once using MakeNewIKRigProcessor()
-	* 2. Initialize() with an IKRigDefinition asset
-	* 3. each tick, call SetIKGoal() and SetInputPoseGlobal()
-	* 4. Call Solve()
-	* 5. Copy output transforms with GetOutputPoseGlobal()
-	* 
-	*/
-
-	UIKRigProcessor();
-	
 	/** setup a new processor to run the given IKRig asset
 	 *  NOTE!! this function creates new UObjects and consequently MUST be called from the main thread!!
 	 *  @param InRigAsset - the IK Rig defining the collection of solvers to execute and all the rig settings
@@ -101,19 +96,21 @@ public:
 		const FIKRigInputSkeleton& InputSkeleton,
 		const FIKRigLogger* Log);
 
+	/** return true if the processor is initialized with a rig and ready to run */
 	bool IsInitialized() const { return bInitialized; };
 
+	/** force the processor to reinitialize */
 	void SetNeedsInitialized();
 
-	/** logging system */
-	FIKRigLogger Log;
-
 	/** Get read-only access to currently running solvers */
-	const TArray<TUniquePtr<FIKRigSolverBase>>& GetSolvers(){ return Solvers; };
+	const TArray<TSharedPtr<FIKRigSolverBase>>& GetSolvers(){ return Solvers; };
 	
 	/** Used to propagate setting values from the source asset at runtime (settings that do not require re-initialization)
 	 * This allows the user to edit their IK Rig asset and see the effects in PIE. */
 	void CopyAllSettingsFromAsset(const UIKRigDefinition* SourceAsset);
+
+	/** logging system */
+	FIKRigLogger Log;
 	
 private:
 
@@ -121,7 +118,7 @@ private:
 	void ResolveFinalGoalTransforms(const FTransform& WorldToComponent);
 
 	/** the stack of solvers to run in order */
-	TArray<TUniquePtr<FIKRigSolverBase>> Solvers;
+	TArray<TSharedPtr<FIKRigSolverBase>> Solvers;
 
 	/** the named transforms that solvers use as end effectors */
 	FIKRigGoalContainer GoalContainer;
@@ -135,4 +132,96 @@ private:
 	/** solving disabled until this flag is true */
 	bool bInitialized = false;
 	bool bTriedToInitialize = false;
+};
+
+//
+// BEGIN DEPRECATED UOBJECT-BASED PROCESSOR
+//
+// NOTE: As of 5.6, IK Rig processors are no long UObjects.
+// The UObject based processor is here for backwards compatibility.
+// If you have a system built off this processor, please convert it to the FIKRigProcessor as this will be removed in a future update.
+//
+
+UCLASS()
+class UE_DEPRECATED(5.6, "UIKRigProcessor has been replaced by FIKRigProcessor. UIKRigProcessor is now a thin wrapper around FIKRigProcessor and will be removed in a future version.")
+IKRIG_API UIKRigProcessor : public UObject
+{
+	GENERATED_BODY()
+	
+	public:
+	
+	void Initialize(const UIKRigDefinition* InRigAsset, const USkeletalMesh* SkeletalMesh, const TArray<FName>& ExcludedGoals = TArray<FName>())
+	{
+		Processor.Initialize(InRigAsset, SkeletalMesh, ExcludedGoals);
+	};
+	void SetInputPoseGlobal(const TArray<FTransform>& InGlobalBoneTransforms)
+	{
+		Processor.SetInputPoseGlobal(InGlobalBoneTransforms);
+	};
+	void SetInputPoseToRefPose()
+	{
+		Processor.SetInputPoseToRefPose();
+	};
+	void SetIKGoal(const FIKRigGoal& Goal)
+	{
+		Processor.SetIKGoal(Goal);
+	};
+	void SetIKGoal(const UIKRigEffectorGoal* Goal)
+	{
+		Processor.SetIKGoal(Goal);
+	};
+	void Solve(const FTransform& WorldToComponent = FTransform::Identity)
+	{
+		Processor.Solve(WorldToComponent);
+	};
+	void GetOutputPoseGlobal(TArray<FTransform>& OutputPoseGlobal) const
+	{
+		Processor.GetOutputPoseGlobal(OutputPoseGlobal);
+	};
+	void Reset()
+	{
+		Processor.Reset();
+	};
+	const FIKRigGoalContainer& GetGoalContainer() const
+	{
+		return Processor.GetGoalContainer();
+	};
+	const FGoalBone* GetGoalBone(const FName& GoalName) const
+	{
+		return Processor.GetGoalBone(GoalName);
+	};
+	FIKRigSkeleton& GetSkeletonWriteable()
+	{
+		return Processor.GetSkeletonWriteable();
+	} ;
+	const FIKRigSkeleton& GetSkeleton() const
+	{
+		return Processor.GetSkeleton();
+	};
+	static bool IsIKRigCompatibleWithSkeleton(const UIKRigDefinition* InRigAsset, const FIKRigInputSkeleton& InputSkeleton, const FIKRigLogger* Log)
+	{
+		return FIKRigProcessor::IsIKRigCompatibleWithSkeleton(InRigAsset, InputSkeleton, Log);
+	};
+	bool IsInitialized() const
+	{
+		return Processor.IsInitialized();
+	};
+	void SetNeedsInitialized()
+	{
+		Processor.SetNeedsInitialized();
+	};
+	const TArray<TSharedPtr<FIKRigSolverBase>>& GetSolvers()
+	{
+		return Processor.GetSolvers();
+	};
+	void CopyAllSettingsFromAsset(const UIKRigDefinition* SourceAsset)
+	{
+		Processor.CopyAllSettingsFromAsset(SourceAsset);
+	};
+
+	/** NOTE: the deprecated logging system will no longer function. It's here to avoid compilation issues.*/
+	FIKRigLogger Log;
+
+private:
+	FIKRigProcessor Processor;
 };
\ No newline at end of file
diff --git a/Engine/Plugins/Animation/IKRig/Source/IKRigDeveloper/Private/AnimGraphNode_RetargetPoseFromMesh.cpp b/Engine/Plugins/Animation/IKRig/Source/IKRigDeveloper/Private/AnimGraphNode_RetargetPoseFromMesh.cpp
--- a/Engine/Plugins/Animation/IKRig/Source/IKRigDeveloper/Private/AnimGraphNode_RetargetPoseFromMesh.cpp
+++ b/Engine/Plugins/Animation/IKRig/Source/IKRigDeveloper/Private/AnimGraphNode_RetargetPoseFromMesh.cpp
@@ -114,7 +114,7 @@ void UAnimGraphNode_RetargetPoseFromMesh::ValidateAnimNodeDuringCompilation(USke
 	// pull messages out of the processor's log
 	if (!Node.bSuppressWarnings)
 	{
-		if (const UIKRetargetProcessor* Processor = Node.GetRetargetProcessor())
+		if (const FIKRetargetProcessor* Processor = Node.GetRetargetProcessor())
 		{
 			const TArray<FText>& Warnings = Processor->Log.GetWarnings();
 			for (const FText& Warning : Warnings)
diff --git a/Engine/Plugins/Animation/IKRig/Source/IKRigEditor/Private/IKRigEditor.cpp b/Engine/Plugins/Animation/IKRig/Source/IKRigEditor/Private/IKRigEditor.cpp
--- a/Engine/Plugins/Animation/IKRig/Source/IKRigEditor/Private/IKRigEditor.cpp
+++ b/Engine/Plugins/Animation/IKRig/Source/IKRigEditor/Private/IKRigEditor.cpp
@@ -78,9 +78,6 @@ void FIKRigEditor::StartupModule()
 	// custom retargeter global details
 	PropertyEditorModule.RegisterCustomClassLayout(UIKRetargetGlobalSettings::StaticClass()->GetFName(), FOnGetDetailCustomizationInstance::CreateStatic(&FRetargetGlobalSettingsCustomization::MakeInstance));
 	ClassesToUnregisterOnShutdown.Add(UIKRetargetGlobalSettings::StaticClass()->GetFName());
-	// custom retargeter op stack details
-	PropertyEditorModule.RegisterCustomClassLayout(URetargetOpStack::StaticClass()->GetFName(), FOnGetDetailCustomizationInstance::CreateStatic(&FRetargetOpStackCustomization::MakeInstance));
-	ClassesToUnregisterOnShutdown.Add(URetargetOpStack::StaticClass()->GetFName());
 
 	// register a thumbnail renderer for the assets
 	UThumbnailManager::Get().RegisterCustomRenderer(UIKRigDefinition::StaticClass(), UIKRigThumbnailRenderer::StaticClass());
diff --git a/Engine/Plugins/Animation/IKRig/Source/IKRigEditor/Private/RetargetEditor/IKRetargetAnimInstance.cpp b/Engine/Plugins/Animation/IKRig/Source/IKRigEditor/Private/RetargetEditor/IKRetargetAnimInstance.cpp
--- a/Engine/Plugins/Animation/IKRig/Source/IKRigEditor/Private/RetargetEditor/IKRetargetAnimInstance.cpp
+++ b/Engine/Plugins/Animation/IKRig/Source/IKRigEditor/Private/RetargetEditor/IKRetargetAnimInstance.cpp
@@ -151,18 +151,15 @@ void UIKRetargetAnimInstance::ConfigureAnimInstance(
 	Proxy.ConfigureAnimInstance(InSourceOrTarget, InAsset, InSourceMeshComponent);
 }
 
-UIKRetargetProcessor* UIKRetargetAnimInstance::GetRetargetProcessor()
+FIKRetargetProcessor* UIKRetargetAnimInstance::GetRetargetProcessor()
 {
-	// depending on initialization order, the anim node may not have created the processor yet, so optionally force it to do so now
-	RetargetNode.CreateRetargetProcessorIfNeeded(this);
-	
 	return RetargetNode.GetRetargetProcessor();
 }
 
 void UIKRetargetAnimInstance::ForceInitializeProcessor(USkeletalMeshComponent* TargetComponent)
 {
 	RetargetNode.IKRetargeterAsset->IncrementVersion();
-	RetargetNode.TryInitializeProcessor(TargetComponent);
+	RetargetNode.EnsureProcessorIsInitialized(TargetComponent);
 }
 
 FAnimInstanceProxy* UIKRetargetAnimInstance::CreateAnimInstanceProxy()
diff --git a/Engine/Plugins/Animation/IKRig/Source/IKRigEditor/Private/RetargetEditor/IKRetargetAnimInstanceProxy.cpp b/Engine/Plugins/Animation/IKRig/Source/IKRigEditor/Private/RetargetEditor/IKRetargetAnimInstanceProxy.cpp
--- a/Engine/Plugins/Animation/IKRig/Source/IKRigEditor/Private/RetargetEditor/IKRetargetAnimInstanceProxy.cpp
+++ b/Engine/Plugins/Animation/IKRig/Source/IKRigEditor/Private/RetargetEditor/IKRetargetAnimInstanceProxy.cpp
@@ -114,7 +114,7 @@ void FIKRetargetAnimInstanceProxy::ConfigureAnimInstance(
 		RetargetNode->IKRetargeterAsset = InIKRetargetAsset;
 		RetargetNode->RetargetFrom = ERetargetSourceMode::CustomSkeletalMeshComponent;
 		RetargetNode->SourceMeshComponent = InSourceMeshComponent;
-		if (UIKRetargetProcessor* Processor = RetargetNode->GetRetargetProcessor())
+		if (FIKRetargetProcessor* Processor = RetargetNode->GetRetargetProcessor())
 		{
 			Processor->SetNeedsInitialized();
 		}
@@ -125,7 +125,7 @@ void FIKRetargetAnimInstanceProxy::SetRetargetMode(const ERetargeterOutputMode&
 {
 	OutputMode = InOutputMode;
 	
-	if (UIKRetargetProcessor* Processor = RetargetNode->GetRetargetProcessor())
+	if (FIKRetargetProcessor* Processor = RetargetNode->GetRetargetProcessor())
 	{
 		Processor->SetNeedsInitialized();
 	}
diff --git a/Engine/Plugins/Animation/IKRig/Source/IKRigEditor/Private/RetargetEditor/IKRetargetApplicationMode.cpp b/Engine/Plugins/Animation/IKRig/Source/IKRigEditor/Private/RetargetEditor/IKRetargetApplicationMode.cpp
--- a/Engine/Plugins/Animation/IKRig/Source/IKRigEditor/Private/RetargetEditor/IKRetargetApplicationMode.cpp
+++ b/Engine/Plugins/Animation/IKRig/Source/IKRigEditor/Private/RetargetEditor/IKRetargetApplicationMode.cpp
@@ -11,6 +11,7 @@
 #include "RetargetEditor/IKRetargetChainTabSummoner.h"
 #include "RetargetEditor/IKRetargetEditor.h"
 #include "RetargetEditor/IKRetargetHierarchyTabSummoner.h"
+#include "RetargetEditor/IKRetargetOpStackTabSummoner.h"
 #include "RetargetEditor/IKRetargetOutputLogTabSummoner.h"
 
 #define LOCTEXT_NAMESPACE "IKRetargetMode"
@@ -42,9 +43,10 @@ FIKRetargetApplicationMode::FIKRetargetApplicationMode(
 	TabFactories.RegisterFactory(MakeShared<FIKRetargetAssetBrowserTabSummoner>(IKRetargetEditor));
 	TabFactories.RegisterFactory(MakeShared<FIKRetargetOutputLogTabSummoner>(IKRetargetEditor));
 	TabFactories.RegisterFactory(MakeShared<FIKRetargetHierarchyTabSummoner>(IKRetargetEditor));
+	TabFactories.RegisterFactory(MakeShared<FIKRetargetOpStackTabSummoner>(IKRetargetEditor));
 
 	// create tab layout
-	TabLayout = FTabManager::NewLayout("Standalone_IKRetargetEditor_Layout_v1.018")
+	TabLayout = FTabManager::NewLayout("Standalone_IKRetargetEditor_Layout_v1.019")
 		->AddArea
 		(
 			FTabManager::NewPrimaryArea()
@@ -61,10 +63,21 @@ FIKRetargetApplicationMode::FIKRetargetApplicationMode(
 					->SetOrientation(Orient_Horizontal)
 					->Split
 					(
-						FTabManager::NewStack()
+						FTabManager::NewSplitter()
 						->SetSizeCoefficient(0.2f)
-						->SetHideTabWell(true)
-						->AddTab(FIKRetargetHierarchyTabSummoner::TabID, ETabState::OpenedTab)
+						->SetOrientation(Orient_Vertical)
+						->Split
+						(
+							FTabManager::NewStack()
+							->SetSizeCoefficient(0.6f)
+							->AddTab(FIKRetargetOpStackTabSummoner::TabID, ETabState::OpenedTab)
+						)
+						->Split
+						(
+							FTabManager::NewStack()
+							->SetSizeCoefficient(0.4f)
+							->AddTab(FIKRetargetAssetBrowserTabSummoner::TabID, ETabState::OpenedTab)
+						)
 					)
 					->Split
 					(
@@ -96,6 +109,7 @@ FIKRetargetApplicationMode::FIKRetargetApplicationMode(
 					FTabManager::NewStack()
 						->SetSizeCoefficient(0.7f)
 						->AddTab(FPersonaTabs::DetailsID, ETabState::OpenedTab)
+						->AddTab(FIKRetargetHierarchyTabSummoner::TabID, ETabState::OpenedTab)
 						->AddTab(FPersonaTabs::AdvancedPreviewSceneSettingsID, ETabState::OpenedTab)
 						->SetForegroundTab(FPersonaTabs::DetailsID)
 					)
@@ -104,7 +118,6 @@ FIKRetargetApplicationMode::FIKRetargetApplicationMode(
 					FTabManager::NewStack()
 						->SetSizeCoefficient(0.3f)
 						->AddTab(FIKRetargetChainTabSummoner::TabID, ETabState::OpenedTab)
-						->AddTab(FIKRetargetAssetBrowserTabSummoner::TabID, ETabState::OpenedTab)
 						->SetForegroundTab(FIKRetargetChainTabSummoner::TabID)
 					)
 					
diff --git a/Engine/Plugins/Animation/IKRig/Source/IKRigEditor/Private/RetargetEditor/IKRetargetBatchOperation.cpp b/Engine/Plugins/Animation/IKRig/Source/IKRigEditor/Private/RetargetEditor/IKRetargetBatchOperation.cpp
--- a/Engine/Plugins/Animation/IKRig/Source/IKRigEditor/Private/RetargetEditor/IKRetargetBatchOperation.cpp
+++ b/Engine/Plugins/Animation/IKRig/Source/IKRigEditor/Private/RetargetEditor/IKRetargetBatchOperation.cpp
@@ -332,18 +332,18 @@ void UIKRetargetBatchOperation::ConvertAnimation(
 {
 	// initialize the retargeter
 	UObject* TransientOuter = Cast<UObject>(GetTransientPackage());
-	UIKRetargetProcessor* Processor = NewObject<UIKRetargetProcessor>(TransientOuter);
+	FIKRetargetProcessor Processor;
 	FRetargetProfile RetargetProfile;
 	Context.IKRetargetAsset->FillProfileWithAssetSettings(RetargetProfile);
-	Processor->Initialize(Context.SourceMesh, Context.TargetMesh, Context.IKRetargetAsset, RetargetProfile);
-	if (!Processor->IsInitialized())
+	Processor.Initialize(Context.SourceMesh, Context.TargetMesh, Context.IKRetargetAsset, RetargetProfile);
+	if (!Processor.IsInitialized())
 	{
 		UE_LOG(LogTemp, Warning, TEXT("Unable to initialize the IK Retargeter. Newly created animations were not retargeted!"));
 		return;
 	}
 
 	// target skeleton data
-	const FRetargetSkeleton& TargetSkeleton = Processor->GetSkeleton(ERetargetSourceOrTarget::Target);
+	const FRetargetSkeleton& TargetSkeleton = Processor.GetSkeleton(ERetargetSourceOrTarget::Target);
 	const TArray<FName>& TargetBoneNames = TargetSkeleton.BoneNames;
 	const int32 NumTargetBones = TargetBoneNames.Num();
 
@@ -352,7 +352,7 @@ void UIKRetargetBatchOperation::ConvertAnimation(
 	BoneTracks.SetNumZeroed(NumTargetBones);
 
 	// source skeleton data
-	const FRetargetSkeleton& SourceSkeleton = Processor->GetSkeleton(ERetargetSourceOrTarget::Source);
+	const FRetargetSkeleton& SourceSkeleton = Processor.GetSkeleton(ERetargetSourceOrTarget::Source);
 	const TArray<FName>& SourceBoneNames = SourceSkeleton.BoneNames;
 	const int32 NumSourceBones = SourceBoneNames.Num();
 
@@ -412,7 +412,7 @@ void UIKRetargetBatchOperation::ConvertAnimation(
 		EvaluationOptions.bIncorporateRootMotionIntoPose = true;
 
 		// reset the planting state
-		Processor->ResetPlanting();
+		Processor.ResetPlanting();
 		
 		// retarget each frame's pose from source to target
 		for (int32 FrameIndex=0; FrameIndex<NumFrames; ++FrameIndex)
@@ -442,7 +442,7 @@ void UIKRetargetBatchOperation::ConvertAnimation(
 			}
 
 			// update goals 
-			Processor->CopyIKRigSettingsFromAsset();
+			Processor.CopyIKRigSettingsFromAsset();
 			
 			// calculate the delta time
 			const float TimeAtCurrentFrame = SourceSequence->GetTimeAtFrame(FrameIndex);
@@ -473,7 +473,7 @@ void UIKRetargetBatchOperation::ConvertAnimation(
 			Context.IKRetargetAsset->FillProfileWithAssetSettings(SettingsProfile);
 
 			// run the retargeter
-			const TArray<FTransform>& TargetComponentPose = Processor->RunRetargeter(SourceComponentPose, SpeedCurveValues, DeltaTime, SettingsProfile);
+			const TArray<FTransform>& TargetComponentPose = Processor.RunRetargeter(SourceComponentPose, SpeedCurveValues, DeltaTime, SettingsProfile);
 
 			// convert to a local-space pose
 			TArray<FTransform> TargetLocalPose = TargetComponentPose;
@@ -524,28 +524,33 @@ void UIKRetargetBatchOperation::RemapCurves(const FIKRetargetBatchOperationConte
 	// get map of curves to remap (Source:Target)
 	bool bCopyAllSourceCurves = true;
 	TMap<FAnimationCurveIdentifier, FAnimationCurveIdentifier> CurvesToRemap;
-	URetargetOpStack* OpStack = Context.IKRetargetAsset->GetPostSettingsUObject();
-	for (const TObjectPtr<URetargetOpBase>& RetargetOp : OpStack->RetargetOps)
+	const TArray<FInstancedStruct>& OpStack = Context.IKRetargetAsset->GetRetargetOps();
+	for (const FInstancedStruct& OpStruct : OpStack)
 	{
-		UCurveRemapOp* CurveRemapOp = Cast<UCurveRemapOp>(RetargetOp);
+		if (OpStruct.GetScriptStruct() != FIKRetargetCurveRemapOp::StaticStruct())
+		{
+			continue;
+		}
+		
+		const FIKRetargetCurveRemapOp* CurveRemapOp = OpStruct.GetPtr<FIKRetargetCurveRemapOp>();
 		if (!CurveRemapOp)
 		{
 			continue;
 		}
 
-		if (!CurveRemapOp->bIsEnabled)
+		if (!CurveRemapOp->IsEnabled())
 		{
 			continue;
 		}
 
-		for (const FCurveRemapPair& CurveToRemap : CurveRemapOp->CurvesToRemap)
+		for (const FCurveRemapPair& CurveToRemap : CurveRemapOp->Settings.CurvesToRemap)
 		{
 			const FAnimationCurveIdentifier SourceCurveId = UAnimationCurveIdentifierExtensions::FindCurveIdentifier(SourceSkeleton, CurveToRemap.SourceCurve, ERawCurveTrackTypes::RCT_Float);
 			const FAnimationCurveIdentifier TargetCurveId = UAnimationCurveIdentifierExtensions::FindCurveIdentifier(TargetSkeleton, CurveToRemap.TargetCurve, ERawCurveTrackTypes::RCT_Float);
 			CurvesToRemap.Add(SourceCurveId, TargetCurveId);
 		}
 
-		bCopyAllSourceCurves &= CurveRemapOp->bCopyAllSourceCurves;
+		bCopyAllSourceCurves &= CurveRemapOp->Settings.bCopyAllSourceCurves;
 	}
 
 	// update progress bar
diff --git a/Engine/Plugins/Animation/IKRig/Source/IKRigEditor/Private/RetargetEditor/IKRetargetCommands.cpp b/Engine/Plugins/Animation/IKRig/Source/IKRigEditor/Private/RetargetEditor/IKRetargetCommands.cpp
--- a/Engine/Plugins/Animation/IKRig/Source/IKRigEditor/Private/RetargetEditor/IKRetargetCommands.cpp
+++ b/Engine/Plugins/Animation/IKRig/Source/IKRigEditor/Private/RetargetEditor/IKRetargetCommands.cpp
@@ -13,7 +13,6 @@ void FIKRetargetCommands::RegisterCommands()
 	UI_COMMAND(ShowAssetSettings, "Asset Settings", "Set the source or target IK Rigs, preview meshes. Adjust debug settings.", EUserInterfaceActionType::ToggleButton, FInputChord());
 	UI_COMMAND(ShowGlobalSettings, "Global Settings", "Access Stride Warping and retarget phases.", EUserInterfaceActionType::ToggleButton, FInputChord());
 	UI_COMMAND(ShowRootSettings, "Root Settings", "Adjust the motion of the Pelvis.", EUserInterfaceActionType::ToggleButton, FInputChord());
-	UI_COMMAND(ShowPostSettings, "Post Settings", "Create a stack of operations that run after the retarget.", EUserInterfaceActionType::ToggleButton, FInputChord());
 	
 	UI_COMMAND(ResetAllBones, "Reset All", "Sets the selected bones to the mesh reference pose.", EUserInterfaceActionType::Button, FInputChord());
 	UI_COMMAND(ResetSelectedBones, "Reset Selected Bones", "Sets the selected bones to the mesh reference pose.", EUserInterfaceActionType::Button, FInputChord());
diff --git a/Engine/Plugins/Animation/IKRig/Source/IKRigEditor/Private/RetargetEditor/IKRetargetDefaultMode.cpp b/Engine/Plugins/Animation/IKRig/Source/IKRigEditor/Private/RetargetEditor/IKRetargetDefaultMode.cpp
--- a/Engine/Plugins/Animation/IKRig/Source/IKRigEditor/Private/RetargetEditor/IKRetargetDefaultMode.cpp
+++ b/Engine/Plugins/Animation/IKRig/Source/IKRigEditor/Private/RetargetEditor/IKRetargetDefaultMode.cpp
@@ -70,7 +70,7 @@ void FIKRetargetDefaultMode::RenderDebugProxies(FPrimitiveDrawInterface* PDI, co
 		return;
 	}
 	
-	const UIKRetargetProcessor* RetargetProcessor = Controller->GetRetargetProcessor();
+	const FIKRetargetProcessor* RetargetProcessor = Controller->GetRetargetProcessor();
 	if (!(RetargetProcessor && RetargetProcessor->IsInitialized()))
 	{
 		return;
diff --git a/Engine/Plugins/Animation/IKRig/Source/IKRigEditor/Private/RetargetEditor/IKRetargetDetails.cpp b/Engine/Plugins/Animation/IKRig/Source/IKRigEditor/Private/RetargetEditor/IKRetargetDetails.cpp
--- a/Engine/Plugins/Animation/IKRig/Source/IKRigEditor/Private/RetargetEditor/IKRetargetDetails.cpp
+++ b/Engine/Plugins/Animation/IKRig/Source/IKRigEditor/Private/RetargetEditor/IKRetargetDetails.cpp
@@ -338,7 +338,7 @@ void UIKRetargetBoneDetails::OnPasteFromClipboard(
 	};
 
 	FIKRigLogger* Log = nullptr;
-	if (UIKRetargetProcessor* Processor = Controller->GetRetargetProcessor())
+	if (FIKRetargetProcessor* Processor = Controller->GetRetargetProcessor())
 	{
 		Log = &Processor->Log;
 	}
@@ -1480,110 +1480,5 @@ void FRetargetGlobalSettingsCustomization::CustomizeDetails(IDetailLayoutBuilder
 	WarpingGroup.AddPropertyRow(GetPropertyHandle(GET_MEMBER_NAME_STRING_CHECKED(FRetargetGlobalSettings, SidewaysOffset))).EditCondition(WarpingEditCondition,nullptr);
 }
 
-void FRetargetOpStackCustomization::CustomizeDetails(IDetailLayoutBuilder& DetailBuilder)
-{
-	TArray<TWeakObjectPtr<UObject>> ObjectsBeingCustomized = DetailBuilder.GetSelectedObjects();
-	if (ObjectsBeingCustomized.Num() < 0)
-	{
-		return;
-	}
-	
-	RetargetOpStackObject = Cast<URetargetOpStack>(ObjectsBeingCustomized[0].Get());
-	if (!RetargetOpStackObject.IsValid())
-	{
-		return;
-	}
-
-	Controller = RetargetOpStackObject->EditorController;
-	if (!Controller.IsValid())
-	{
-		return;
-	}
-
-	IDetailCategoryBuilder& OpStackCategory = DetailBuilder.EditCategory(TEXT("Retarget Ops Stack"));
-	
-	// add row to select source chain to map to
-	OpStackCategory.AddCustomRow(LOCTEXT("RetargetOps_Label", "Retarget Ops"))
-	.WholeRowWidget
-	[
-		SNew(SRetargetOpStack, Controller)
-	];
-
-	// add custom category with all op settings, filter visibility based on selection
-	IDetailCategoryBuilder& OpSettingCategory = DetailBuilder.EditCategory(TEXT("Op Settings"));
-	for (TObjectPtr<URetargetOpBase> Op : RetargetOpStackObject->RetargetOps)
-	{
-		for (TFieldIterator<FProperty> PropIt(Op->GetClass()); PropIt; ++PropIt)
-		{
-			const FProperty* Prop = *PropIt;
-			if (!Prop->HasAllPropertyFlags(CPF_Edit))
-			{
-				continue;
-			}
-			
-			const TSharedPtr<IPropertyHandle> PropertyHandle = DetailBuilder.AddObjectPropertyData({Op}, Prop->GetFName());
-			if (PropertyHandle && PropertyHandle->IsValidHandle())
-			{
-				PropertyHandle->SetOnPropertyValueChanged(FSimpleDelegate::CreateLambda([this]()
-				{
-					Controller.Pin().Get()->ReinitializeRetargeterNoUIRefresh();
-				}));
-				
-				PropertyHandle->SetOnChildPropertyValueChanged(FSimpleDelegate::CreateLambda([this]()
-				{
-					Controller.Pin().Get()->ReinitializeRetargeterNoUIRefresh();
-				}));
-				
-				OpSettingCategory.AddProperty(PropertyHandle)
-				.Visibility(MakeAttributeLambda([this, Op]()
-				{
-					return Controller.Pin().Get()->GetSelectedOp() == Op ? EVisibility::Visible : EVisibility::Hidden;
-				}));
-			}
-		}
-	}
-}
-
-TSharedRef<SWidget> FRetargetOpStackCustomization::CreateAddNewMenuWidget()
-{
-	constexpr bool bCloseMenuAfterSelection = true;
-	FMenuBuilder MenuBuilder(bCloseMenuAfterSelection, nullptr);
-
-	MenuBuilder.BeginSection("AddNewRetargetOp", LOCTEXT("AddOperations", "Add New Retarget Op"));
-
-	// add menu option to create each retarget op type
-	for(TObjectIterator<UClass> ClassIt; ClassIt; ++ClassIt)
-	{
-		UClass* Class = *ClassIt;
-		if(Class->IsChildOf(URetargetOpBase::StaticClass()) && !Class->HasAnyClassFlags(CLASS_Abstract))
-		{
-			const URetargetOpBase* OpCDO = Cast<URetargetOpBase>(Class->GetDefaultObject());
-			FUIAction Action = FUIAction( FExecuteAction::CreateSP(this, &FRetargetOpStackCustomization::AddNewRetargetOp, Class));
-			MenuBuilder.AddMenuEntry(FText::FromString(OpCDO->GetNiceName().ToString()), FText::GetEmpty(), FSlateIcon(), Action);
-		}
-	}
-
-	MenuBuilder.EndSection();
-
-	return MenuBuilder.MakeWidget();
-}
-
-void FRetargetOpStackCustomization::AddNewRetargetOp(UClass* Class)
-{
-	if (!Controller.IsValid())
-	{
-		return; 
-	}
-
-	const UIKRetargeterController* AssetController = Controller.Pin().Get()->AssetController;
-	if (!AssetController)
-	{
-		return;
-	}
-	
-	// add the op todo refresh UI
-	const int32 NewOpIndex = AssetController->AddRetargetOp(Class);
-}
-
 #undef LOCTEXT_NAMESPACE
 
diff --git a/Engine/Plugins/Animation/IKRig/Source/IKRigEditor/Private/RetargetEditor/IKRetargetEditor.cpp b/Engine/Plugins/Animation/IKRig/Source/IKRigEditor/Private/RetargetEditor/IKRetargetEditor.cpp
--- a/Engine/Plugins/Animation/IKRig/Source/IKRigEditor/Private/RetargetEditor/IKRetargetEditor.cpp
+++ b/Engine/Plugins/Animation/IKRig/Source/IKRigEditor/Private/RetargetEditor/IKRetargetEditor.cpp
@@ -166,15 +166,6 @@ void FIKRetargetEditor::BindCommands()
 			const URetargetRootSettings* RootSettings = EditorController->AssetController->GetAsset()->GetRootSettingsUObject();
 			return EditorController->IsObjectInDetailsView(RootSettings);	
 		}));
-	ToolkitCommands->MapAction(
-		Commands.ShowPostSettings,
-		FExecuteAction::CreateSP(EditorController, &FIKRetargetEditorController::ShowPostPhaseSettings),
-		FCanExecuteAction(),
-		FIsActionChecked::CreateLambda([this]() ->bool
-		{
-			const URetargetOpStack* PostSettings = EditorController->AssetController->GetAsset()->GetPostSettingsUObject();
-			return EditorController->IsObjectInDetailsView(PostSettings);	
-		}));
 
 	//
 	// Edit pose commands
@@ -340,13 +331,6 @@ void FIKRetargetEditor::FillToolbar(FToolBarBuilder& ToolbarBuilder)
 		TAttribute<FText>(),
 		TAttribute<FText>(),
 		FSlateIcon(FIKRetargetEditorStyle::Get().GetStyleSetName(),"IKRetarget.RootSettings"));
-
-		ToolbarBuilder.AddToolBarButton(
-		FIKRetargetCommands::Get().ShowPostSettings,
-		NAME_None,
-		TAttribute<FText>(),
-		TAttribute<FText>(),
-		FSlateIcon(FIKRetargetEditorStyle::Get().GetStyleSetName(),"IKRetarget.PostSettings"));
 	}
 	ToolbarBuilder.EndSection();
 }
@@ -569,7 +553,7 @@ void FIKRetargetEditor::HandleViewportCreated(const TSharedRef<class IPersonaVie
 	auto GetBorderColorAndOpacity = [this]()
 	{
 		// no processor or processor not initialized
-		const UIKRetargetProcessor* Processor = EditorController.Get().GetRetargetProcessor();
+		const FIKRetargetProcessor* Processor = EditorController.Get().GetRetargetProcessor();
 		if (!Processor || !Processor->IsInitialized() )
 		{
 			return FLinearColor::Red;
diff --git a/Engine/Plugins/Animation/IKRig/Source/IKRigEditor/Private/RetargetEditor/IKRetargetEditorController.cpp b/Engine/Plugins/Animation/IKRig/Source/IKRigEditor/Private/RetargetEditor/IKRetargetEditorController.cpp
--- a/Engine/Plugins/Animation/IKRig/Source/IKRigEditor/Private/RetargetEditor/IKRetargetEditorController.cpp
+++ b/Engine/Plugins/Animation/IKRig/Source/IKRigEditor/Private/RetargetEditor/IKRetargetEditorController.cpp
@@ -31,6 +31,7 @@
 #include "RigEditor/SIKRigOutputLog.h"
 #include "RigEditor/IKRigController.h"
 #include "ScopedTransaction.h"
+#include "RigEditor/IKRigDetailCustomizations.h"
 #include "Styling/AppStyle.h"
 
 #define LOCTEXT_NAMESPACE "IKRetargetEditorController"
@@ -475,7 +476,7 @@ bool FIKRetargetEditorController::IsEditingPoseWithAnyBoneSelected() const
 bool FIKRetargetEditorController::IsBoneRetargeted(const FName& BoneName, ERetargetSourceOrTarget SourceOrTarget) const
 {
 	// get an initialized processor
-	const UIKRetargetProcessor* Processor = GetRetargetProcessor();
+	const FIKRetargetProcessor* Processor = GetRetargetProcessor();
 	if (!(Processor && Processor->IsInitialized()))
 	{
 		return false;
@@ -488,7 +489,7 @@ bool FIKRetargetEditorController::IsBoneRetargeted(const FName& BoneName, ERetar
 FName FIKRetargetEditorController::GetChainNameFromBone(const FName& BoneName, ERetargetSourceOrTarget SourceOrTarget) const
 {
 	// get an initialized processor
-	const UIKRetargetProcessor* Processor = GetRetargetProcessor();
+	const FIKRetargetProcessor* Processor = GetRetargetProcessor();
 	if (!(Processor && Processor->IsInitialized()))
 	{
 		return NAME_None;
@@ -535,7 +536,7 @@ UDebugSkelMeshComponent* FIKRetargetEditorController::GetEditedSkeletalMesh() co
 	return CurrentlyEditingSourceOrTarget == ERetargetSourceOrTarget::Source ? SourceSkelMeshComponent : TargetSkelMeshComponent;
 }
 
-const FRetargetSkeleton& FIKRetargetEditorController::GetCurrentlyEditedSkeleton(const UIKRetargetProcessor& Processor) const
+const FRetargetSkeleton& FIKRetargetEditorController::GetCurrentlyEditedSkeleton(const FIKRetargetProcessor& Processor) const
 {
 	return Processor.GetSkeleton(CurrentlyEditingSourceOrTarget);
 }
@@ -599,7 +600,7 @@ void FIKRetargetEditorController::GetGlobalRetargetPoseOfImmediateChildren(
 	}
 }
 
-UIKRetargetProcessor* FIKRetargetEditorController::GetRetargetProcessor() const
+FIKRetargetProcessor* FIKRetargetEditorController::GetRetargetProcessor() const
 {	
 	if(UIKRetargetAnimInstance* AnimInstance = TargetAnimInstance.Get())
 	{
@@ -611,7 +612,7 @@ UIKRetargetProcessor* FIKRetargetEditorController::GetRetargetProcessor() const
 
 void FIKRetargetEditorController::ResetIKPlantingState() const
 {
-	if (UIKRetargetProcessor* Processor = GetRetargetProcessor())
+	if (FIKRetargetProcessor* Processor = GetRetargetProcessor())
 	{
 		Processor->ResetPlanting();
 	}
@@ -622,7 +623,7 @@ void FIKRetargetEditorController::ClearOutputLog() const
 	if (OutputLogView.IsValid())
 	{
 		OutputLogView.Get()->ClearLog();
-		if (const UIKRetargetProcessor* Processor = GetRetargetProcessor())
+		if (const FIKRetargetProcessor* Processor = GetRetargetProcessor())
 		{
 			Processor->Log.Clear();
 		}
@@ -730,6 +731,23 @@ void FIKRetargetEditorController::SetDetailsObjects(const TArray<UObject*>& Deta
 	}
 }
 
+void FIKRetargetEditorController::ShowDetailsForOp(const int32 OpIndex) const
+{
+	FInstancedStruct* InstancedStruct = AssetController->GetRetargetOpStructAtIndex(OpIndex);
+	UStructViewer* StructViewer = AssetController->GetStructViewer();
+	const FIKRetargetOpBase* OpToEdit = InstancedStruct->GetPtr<FIKRetargetOpBase>();
+	const UScriptStruct* TypeToEdit = OpToEdit->GetSettingsType();
+	auto MemoryProvider = [](FInstancedStruct& InInstancedStruct) -> uint8*
+	{
+		FIKRetargetOpBase* OpToEdit = InInstancedStruct.GetMutablePtr<FIKRetargetOpBase>();
+		return reinterpret_cast<uint8*>(OpToEdit->GetSettings());
+	};
+	StructViewer->SetStructToView(InstancedStruct, TypeToEdit, MemoryProvider, AssetController->GetAsset());
+	
+	constexpr bool bForceRefresh = true; // must forcibly refresh because struct viewer is a recycled UObject
+	DetailsView->SetObject(StructViewer, bForceRefresh);
+}
+
 float FIKRetargetEditorController::GetRetargetPoseAmount() const
 {
 	return RetargetPosePreviewBlend;
@@ -982,7 +1000,7 @@ void FIKRetargetEditorController::ClearSelection(const bool bKeepBoneSelection)
 	RefreshDetailsView();
 }
 
-URetargetOpBase* FIKRetargetEditorController::GetSelectedOp() const
+FIKRetargetOpBase* FIKRetargetEditorController::GetSelectedOp() const
 {
 	if (!OpStackView.IsValid())
 	{
@@ -1484,7 +1502,7 @@ void FIKRetargetEditorController::RenderSkeleton(FPrimitiveDrawInterface* PDI, E
 		const TArray<FName>& SelectedChainNames = GetSelectedChains();
 		const FLinearColor HighlightedColor = FLinearColor::Blue;
 		const bool bIsSource = InSourceOrTarget == ERetargetSourceOrTarget::Source;
-		if (const UIKRetargetProcessor* Processor = GetRetargetProcessor())
+		if (const FIKRetargetProcessor* Processor = GetRetargetProcessor())
 		{
 			const TArray<FRetargetChainPairFK>& FKChains = Processor->GetFKChainPairs();
 			for (const FRetargetChainPairFK& FKChain : FKChains)
@@ -1653,17 +1671,6 @@ void FIKRetargetEditorController::ShowGlobalSettings()
 	SetDetailsObject(GlobalSettings);
 }
 
-void FIKRetargetEditorController::ShowPostPhaseSettings()
-{
-	URetargetOpStack* PostSettings = AssetController->GetAsset()->GetPostSettingsUObject();
-	if (PostSettings->EditorController.Pin().Get() != this)
-	{
-		PostSettings->EditorController = SharedThis(this);
-	}
-	
-	SetDetailsObject(PostSettings);
-}
-
 void FIKRetargetEditorController::ShowRootSettings()
 {
 	URetargetRootSettings* RootSettings = AssetController->GetAsset()->GetRootSettingsUObject();
diff --git a/Engine/Plugins/Animation/IKRig/Source/IKRigEditor/Private/RetargetEditor/IKRetargetOpStackTabSummoner.cpp b/Engine/Plugins/Animation/IKRig/Source/IKRigEditor/Private/RetargetEditor/IKRetargetOpStackTabSummoner.cpp
new file mode 100644
--- /dev/null
+++ b/Engine/Plugins/Animation/IKRig/Source/IKRigEditor/Private/RetargetEditor/IKRetargetOpStackTabSummoner.cpp
@@ -0,0 +1,37 @@
+﻿// Copyright Epic Games, Inc. All Rights Reserved.
+
+#include "RetargetEditor/IKRetargetOpStackTabSummoner.h"
+
+#include "IDocumentation.h"
+#include "RigEditor/IKRigEditorStyle.h"
+#include "RetargetEditor/IKRetargetEditor.h"
+#include "RetargetEditor/SRetargetOpStack.h"
+
+#define LOCTEXT_NAMESPACE "IKRetargetOpStackTabSummoner"
+
+const FName FIKRetargetOpStackTabSummoner::TabID(TEXT("IKRetargeterOpStack"));
+
+FIKRetargetOpStackTabSummoner::FIKRetargetOpStackTabSummoner(const TSharedRef<FIKRetargetEditor>& InIKRetargetEditor)
+	: FWorkflowTabFactory(TabID, InIKRetargetEditor)
+	, RetargetEditor(InIKRetargetEditor)
+{
+	bIsSingleton = true; // only allow a single instance of this tab
+	
+	TabLabel = LOCTEXT("IKRetargetOpStackTabLabel", "Op Stack");
+	TabIcon = FSlateIcon(FIKRigEditorStyle::Get().GetStyleSetName(), "IKRig.SolverStack"); // TODO get new icon for this tab
+
+	ViewMenuDescription = LOCTEXT("IKRetargetOpStack_ViewMenu_Desc", "Retargeting Op Stack");
+	ViewMenuTooltip = LOCTEXT("IKRetargetOpStack_ViewMenu_ToolTip", "Show the Retargeting Op Stack Tab");
+}
+
+TSharedPtr<SToolTip> FIKRetargetOpStackTabSummoner::CreateTabToolTipWidget(const FWorkflowTabSpawnInfo& Info) const
+{
+	return  IDocumentation::Get()->CreateToolTip(LOCTEXT("RetargetOpStackTooltip", "A stack of operations executed from top to bottom to transfer animation from the source to the target."), NULL, TEXT("Shared/Editors/Persona"), TEXT("RetargetOpStack_Window"));
+}
+
+TSharedRef<SWidget> FIKRetargetOpStackTabSummoner::CreateTabBody(const FWorkflowTabSpawnInfo& Info) const
+{
+	return SNew(SRetargetOpStack, RetargetEditor.Pin()->GetController());
+}
+
+#undef LOCTEXT_NAMESPACE 
diff --git a/Engine/Plugins/Animation/IKRig/Source/IKRigEditor/Private/RetargetEditor/IKRetargetOutputLogTabSummoner.cpp b/Engine/Plugins/Animation/IKRig/Source/IKRigEditor/Private/RetargetEditor/IKRetargetOutputLogTabSummoner.cpp
--- a/Engine/Plugins/Animation/IKRig/Source/IKRigEditor/Private/RetargetEditor/IKRetargetOutputLogTabSummoner.cpp
+++ b/Engine/Plugins/Animation/IKRig/Source/IKRigEditor/Private/RetargetEditor/IKRetargetOutputLogTabSummoner.cpp
@@ -31,8 +31,12 @@ TSharedPtr<SToolTip> FIKRetargetOutputLogTabSummoner::CreateTabToolTipWidget(con
 
 TSharedRef<SWidget> FIKRetargetOutputLogTabSummoner::CreateTabBody(const FWorkflowTabSpawnInfo& Info) const
 {
+	// at this point the processor's log may not be bound to the asset, so forcibly do it now
 	const TSharedRef<FIKRetargetEditorController>& Controller = IKRetargetEditor.Pin()->GetController();
-	const FName LogName = Controller->GetRetargetProcessor()->Log.GetLogTarget();
+	FIKRetargetProcessor* Processor = Controller->GetRetargetProcessor();
+	Processor->Log.SetLogTarget(Controller->AssetController->GetAsset());
+	// now bind the log UI to the processor's log target
+	const FName LogName = Processor->Log.GetLogTarget();
 	TSharedRef<SIKRigOutputLog> LogView = SNew(SIKRigOutputLog, LogName);
 	Controller->SetOutputLogView(LogView);
 	return LogView;
diff --git a/Engine/Plugins/Animation/IKRig/Source/IKRigEditor/Private/RetargetEditor/IKRetargeterController.cpp b/Engine/Plugins/Animation/IKRig/Source/IKRigEditor/Private/RetargetEditor/IKRetargeterController.cpp
--- a/Engine/Plugins/Animation/IKRig/Source/IKRigEditor/Private/RetargetEditor/IKRetargeterController.cpp
+++ b/Engine/Plugins/Animation/IKRig/Source/IKRigEditor/Private/RetargetEditor/IKRetargeterController.cpp
@@ -11,6 +11,7 @@
 #include "IKRigEditor.h"
 #include "RetargetEditor/IKRetargeterPoseGenerator.h"
 #include "Retargeter/IKRetargetOps.h"
+#include "RigEditor/IKRigDetailCustomizations.h"
 
 #include UE_INLINE_GENERATED_CPP_BY_NAME(IKRetargeterController)
 
@@ -34,6 +35,11 @@ UIKRetargeterController* UIKRetargeterController::GetController(const UIKRetarge
 	return Cast<UIKRetargeterController>(InRetargeterAsset->Controller);
 }
 
+UIKRetargeterController::UIKRetargeterController()
+{
+	StructViewer = CreateDefaultSubobject<UStructViewer>(TEXT("RetargetSettingsViewer"));
+}
+
 void UIKRetargeterController::PostInitProperties()
 {
 	Super::PostInitProperties();
@@ -210,29 +216,23 @@ bool UIKRetargeterController::SetRetargetChainSettings(const FName& TargetChainN
 	return false;
 }
 
-int32 UIKRetargeterController::AddRetargetOp(TSubclassOf<URetargetOpBase> InOpClass) const
+int32 UIKRetargeterController::AddRetargetOp(const FString InIKRetargetOpType) const
 {
-	check(Asset)
-
-	if (!InOpClass)
+	UScriptStruct* OpType = FindObject<UScriptStruct>(nullptr, *InIKRetargetOpType);
+	if (!OpType)
 	{
-		UE_LOG(LogIKRigEditor, Warning, TEXT("Could not add Op to stack. Invalid Op class specified."));
+		UE_LOG(LogIKRigEditor, Warning, TEXT("Specified retarget op type was not found, %s."), *InIKRetargetOpType);
 		return INDEX_NONE;
 	}
 
-	FScopedTransaction Transaction(LOCTEXT("AddRetargetOp_Label", "Add Retarget Op"));
-	FScopedReinitializeIKRetargeter Reinitialize(this);
-	Asset->OpStack->Modify();
-	URetargetOpBase* NewOp = NewObject<URetargetOpBase>(Asset, InOpClass, NAME_None, RF_Transactional);
-	NewOp->OnAddedToStack(GetAsset());
-	return Asset->OpStack->RetargetOps.Add(NewOp);
+	return AddRetargetOp(OpType);
 }
 
 bool UIKRetargeterController::RemoveRetargetOp(const int32 OpIndex) const
 {
 	check(Asset)
 	
-	if (!Asset->OpStack->RetargetOps.IsValidIndex(OpIndex))
+	if (!Asset->RetargetOps.IsValidIndex(OpIndex))
 	{
 		UE_LOG(LogIKRigEditor, Warning, TEXT("Retarget Op not removed. Invalid index, %d."), OpIndex);
 		return false;
@@ -240,8 +240,8 @@ bool UIKRetargeterController::RemoveRetargetOp(const int32 OpIndex) const
 
 	FScopedTransaction Transaction(LOCTEXT("RemoveRetargetOp_Label", "Remove Retarget Op"));
 	FScopedReinitializeIKRetargeter Reinitialize(this);
-	Asset->OpStack->Modify();
-	Asset->OpStack->RetargetOps.RemoveAt(OpIndex);
+	Asset->Modify();
+	Asset->RetargetOps.RemoveAt(OpIndex);
 	return true;
 }
 
@@ -251,46 +251,70 @@ bool UIKRetargeterController::RemoveAllOps() const
 
 	FScopedTransaction Transaction(LOCTEXT("RemoveAllRetargetOps_Label", "Remove All Retarget Ops"));
 	FScopedReinitializeIKRetargeter Reinitialize(this);
-	Asset->OpStack->Modify();
-	Asset->OpStack->RetargetOps.Empty();
+	Asset->Modify();
+	Asset->RetargetOps.Empty();
 	return true;
 }
 
-URetargetOpBase* UIKRetargeterController::GetRetargetOpAtIndex(int32 Index) const
+FIKRetargetOpBase* UIKRetargeterController::GetRetargetOpAtIndex(int32 Index) const
 {
 	check(Asset)
 
-	if (Asset->OpStack->RetargetOps.IsValidIndex(Index))
+	if (Asset->RetargetOps.IsValidIndex(Index))
 	{
-		return Asset->OpStack->RetargetOps[Index];
+		return Asset->RetargetOps[Index].GetMutablePtr<FIKRetargetOpBase>();
 	}
 	
 	return nullptr;
 }
 
-int32 UIKRetargeterController::GetIndexOfRetargetOp(URetargetOpBase* RetargetOp) const
+int32 UIKRetargeterController::GetIndexOfRetargetOp(FIKRetargetOpBase* RetargetOp) const
 {
 	check(Asset)
-	return Asset->OpStack->RetargetOps.Find(RetargetOp);
+
+	if (!RetargetOp)
+	{
+		return INDEX_NONE;
+	}
+	
+	for (int32 OpIndex=0; OpIndex<Asset->RetargetOps.Num(); ++OpIndex)
+	{
+		const FIKRetargetOpBase* Op = Asset->RetargetOps[OpIndex].GetPtr<FIKRetargetOpBase>();
+		if (Op == RetargetOp)
+		{
+			return OpIndex;
+		}
+	}
+	
+	return INDEX_NONE;
+}
+
+FInstancedStruct* UIKRetargeterController::GetRetargetOpStructAtIndex(int32 Index) const
+{
+	check(Asset);
+	if (Asset->RetargetOps.IsValidIndex(Index))
+	{
+		return &Asset->RetargetOps[Index];
+	}
+	
+	return nullptr;
 }
 
 int32 UIKRetargeterController::GetNumRetargetOps() const
 {
 	check(Asset)
-	return Asset->OpStack->RetargetOps.Num();
+	return Asset->RetargetOps.Num();
 }
 
 bool UIKRetargeterController::MoveRetargetOpInStack(int32 OpToMoveIndex, int32 TargetIndex) const
 {
-	TArray<TObjectPtr<URetargetOpBase>>& RetargetOps = Asset->OpStack->RetargetOps;
-	
-	if (!RetargetOps.IsValidIndex(OpToMoveIndex))
+	if (!Asset->RetargetOps.IsValidIndex(OpToMoveIndex))
 	{
 		UE_LOG(LogIKRigEditor, Warning, TEXT("Retarget Op not moved. Invalid source index, %d."), OpToMoveIndex);
 		return false;
 	}
 
-	if (!RetargetOps.IsValidIndex(TargetIndex))
+	if (!Asset->RetargetOps.IsValidIndex(TargetIndex))
 	{
 		UE_LOG(LogIKRigEditor, Warning, TEXT("Retarget Op not moved. Invalid target index, %d."), TargetIndex);
 		return false;
@@ -304,17 +328,14 @@ bool UIKRetargeterController::MoveRetargetOpInStack(int32 OpToMoveIndex, int32 T
 
 	FScopedTransaction Transaction(LOCTEXT("ReorderRetargetOps_Label", "Reorder Retarget Ops"));
 	FScopedReinitializeIKRetargeter Reinitialize(this);
-	Asset->OpStack->Modify();
-	URetargetOpBase* OpToMove = RetargetOps[OpToMoveIndex];
-	RetargetOps.Insert(OpToMove, TargetIndex + 1);
-	const int32 OpToRemove = TargetIndex > OpToMoveIndex ? OpToMoveIndex : OpToMoveIndex + 1;
-	RetargetOps.RemoveAt(OpToRemove);
+	Asset->Modify();
+	Asset->RetargetOps.Swap(OpToMoveIndex, TargetIndex);
 	return true;
 }
 
 bool UIKRetargeterController::SetRetargetOpEnabled(int32 RetargetOpIndex, bool bIsEnabled) const
 {
-	if (!Asset->OpStack->RetargetOps.IsValidIndex(RetargetOpIndex))
+	if (!Asset->RetargetOps.IsValidIndex(RetargetOpIndex))
 	{
 		UE_LOG(LogIKRigEditor, Warning, TEXT("Retarget op not found. Invalid index, %d."), RetargetOpIndex);
 		return false;
@@ -322,21 +343,43 @@ bool UIKRetargeterController::SetRetargetOpEnabled(int32 RetargetOpIndex, bool b
 	
 	FScopedTransaction Transaction(LOCTEXT("SetRetargetOpEnabled_Label", "Enable/Disable Op"));
 	FScopedReinitializeIKRetargeter Reinitialize(this);
-	URetargetOpBase* OpToMove = Asset->OpStack->RetargetOps[RetargetOpIndex];
-	OpToMove->Modify();
-	OpToMove->bIsEnabled = bIsEnabled;
+	Asset->Modify();
+	GetRetargetOpAtIndex(RetargetOpIndex)->SetEnabled(bIsEnabled);
 	return true;
 }
 
 bool UIKRetargeterController::GetRetargetOpEnabled(int32 RetargetOpIndex) const
 {
-	if (!Asset->OpStack->RetargetOps.IsValidIndex(RetargetOpIndex))
+	if (!Asset->RetargetOps.IsValidIndex(RetargetOpIndex))
 	{
 		UE_LOG(LogIKRigEditor, Warning, TEXT("Invalid retarget op index, %d."), RetargetOpIndex);
 		return false;
 	}
 
-	return Asset->OpStack->RetargetOps[RetargetOpIndex]->bIsEnabled;
+	return GetRetargetOpAtIndex(RetargetOpIndex)->IsEnabled();
+}
+
+int32 UIKRetargeterController::AddRetargetOp(UScriptStruct* InRetargetOpType) const
+{
+	check(Asset)
+
+	if (!InRetargetOpType)
+	{
+		UE_LOG(LogIKRigEditor, Warning, TEXT("Could not add retarget operation. No type specified."));
+		return INDEX_NONE;
+	}
+
+	if (!InRetargetOpType->IsChildOf(FIKRetargetOpBase::StaticStruct()))
+	{
+		UE_LOG(LogIKRigEditor, Warning, TEXT("Could not add retarget operations. Invalid op-type specified. Must be child of FIKRetargetOpBase."));
+		return INDEX_NONE;
+	}
+
+	FScopedTransaction Transaction(LOCTEXT("AddRetargetOp_Label", "Add Retarget Op"));
+	FScopedReinitializeIKRetargeter Reinitialize(this);
+	Asset->Modify();
+	const int32 NewOpIndex = Asset->RetargetOps.Emplace(InRetargetOpType);
+	return NewOpIndex;
 }
 
 bool UIKRetargeterController::GetAskedToFixRootHeightForMesh(USkeletalMesh* Mesh) const
@@ -470,6 +513,12 @@ void UIKRetargeterController::CleanPoseList(const ERetargetSourceOrTarget Source
 	}
 }
 
+UStructViewer* UIKRetargeterController::GetStructViewer()
+{
+	StructViewer->Reset();
+	return StructViewer;
+}
+
 void UIKRetargeterController::AutoMapChains(const EAutoMapChainType AutoMapType, const bool bForceRemap) const
 {
 	FScopeLock Lock(&ControllerLock);
diff --git a/Engine/Plugins/Animation/IKRig/Source/IKRigEditor/Private/RetargetEditor/IKRetargeterPoseGenerator.cpp b/Engine/Plugins/Animation/IKRig/Source/IKRigEditor/Private/RetargetEditor/IKRetargeterPoseGenerator.cpp
--- a/Engine/Plugins/Animation/IKRig/Source/IKRigEditor/Private/RetargetEditor/IKRetargeterPoseGenerator.cpp
+++ b/Engine/Plugins/Animation/IKRig/Source/IKRigEditor/Private/RetargetEditor/IKRetargeterPoseGenerator.cpp
@@ -20,15 +20,9 @@
 FRetargetAutoPoseGenerator::FRetargetAutoPoseGenerator(const TWeakObjectPtr<UIKRetargeterController> InController)
 {
 	Controller = InController;
-
 	// we maintain our own processor separate from the editor, this allows the auto-pose generator
 	// to run "headless" and be accessed by BP/Python API in the asset controller
-	Processor = NewObject<UIKRetargetProcessor>();
-}
-
-void FRetargetAutoPoseGenerator::AddReferencedObjects(FReferenceCollector& Collector)
-{
-	Collector.AddReferencedObject(Processor);
+	Processor = MakeShared<FIKRetargetProcessor>();
 }
 
 void FRetargetAutoPoseGenerator::AlignBones(
diff --git a/Engine/Plugins/Animation/IKRig/Source/IKRigEditor/Private/RetargetEditor/SIKRetargetAssetBrowser.cpp b/Engine/Plugins/Animation/IKRig/Source/IKRigEditor/Private/RetargetEditor/SIKRetargetAssetBrowser.cpp
--- a/Engine/Plugins/Animation/IKRig/Source/IKRigEditor/Private/RetargetEditor/SIKRetargetAssetBrowser.cpp
+++ b/Engine/Plugins/Animation/IKRig/Source/IKRigEditor/Private/RetargetEditor/SIKRetargetAssetBrowser.cpp
@@ -198,7 +198,7 @@ bool SIKRetargetAssetBrowser::IsExportButtonEnabled() const
 		return false; // editor in bad state
 	}
 
-	const UIKRetargetProcessor* Processor = EditorController.Pin()->GetRetargetProcessor();
+	const FIKRetargetProcessor* Processor = EditorController.Pin()->GetRetargetProcessor();
 	if (!Processor)
 	{
 		return false; // no retargeter running
diff --git a/Engine/Plugins/Animation/IKRig/Source/IKRigEditor/Private/RetargetEditor/SRetargetAnimAssetsWindow.cpp b/Engine/Plugins/Animation/IKRig/Source/IKRigEditor/Private/RetargetEditor/SRetargetAnimAssetsWindow.cpp
--- a/Engine/Plugins/Animation/IKRig/Source/IKRigEditor/Private/RetargetEditor/SRetargetAnimAssetsWindow.cpp
+++ b/Engine/Plugins/Animation/IKRig/Source/IKRigEditor/Private/RetargetEditor/SRetargetAnimAssetsWindow.cpp
@@ -37,9 +37,6 @@
 
 #define LOCTEXT_NAMESPACE "RetargetAnimAssetWindow"
 
-const FName SRetargetAnimAssetsWindow::LogName = FName("BatchRetargetWindowLog");
-const FText SRetargetAnimAssetsWindow::LogLabel = LOCTEXT("BatchRetargetLogLabel", "Batch Retarget Animations Log");
-
 void SBatchExportPathDialog::Construct(const FArguments& InArgs)
 {
 	BatchContext = InArgs._BatchContext;
@@ -506,10 +503,10 @@ void SRetargetPoseViewport::Construct(const FArguments& InArgs)
 	PreviewScene.AddComponent(SourceComponent, FTransform::Identity);
 	PreviewScene.AddComponent(TargetComponent, FTransform::Identity);
 
-	UIKRetargetProcessor* Processor = TargetAnimInstance->GetRetargetProcessor();
+	FIKRetargetProcessor* Processor = TargetAnimInstance->GetRetargetProcessor();
 	if (ensure(Processor))
 	{
-		Processor->Log.SetLogTarget(SRetargetAnimAssetsWindow::LogName, SRetargetAnimAssetsWindow::LogLabel);
+		Processor->Log.SetLogTarget(BatchContext->IKRetargetAsset);
 	}
 
 	SetSkeletalMesh(BatchContext->SourceMesh, ERetargetSourceOrTarget::Source);
@@ -561,7 +558,7 @@ bool SRetargetPoseViewport::IsRetargeterValid()
 		return false;
 	}
 
-	UIKRetargetProcessor* Processor = TargetAnimInstance->GetRetargetProcessor();
+	FIKRetargetProcessor* Processor = TargetAnimInstance->GetRetargetProcessor();
 	if (!Processor)
 	{
 		return false;
@@ -703,35 +700,35 @@ void FProceduralRetargetAssets::AutoGenerateIKRetargetAsset()
 	if (TargetSkeletonRoot != TargetPelvis)
 	{
 		// add a root motion op to the post retarget op stack
-		const int32 RootOpIndex = RetargetController->AddRetargetOp(URootMotionGeneratorOp::StaticClass());
+		const int32 RootOpIndex = RetargetController->AddRetargetOp(FIKRetargetRootMotionOp::StaticStruct());
 
 		// if the source does not have a dedicated root bone, configure the root motion to be generated from the pelvis motion on target
 		const FName SourceSkeletonRoot = SourceIKRig->GetSkeleton().BoneNames[0];
 		const FName SourcePelvis = SourceCharacterizationResults.AutoRetargetDefinition.RetargetDefinition.RootBone;
 		if (SourceSkeletonRoot == SourcePelvis)
 		{
-			if (URootMotionGeneratorOp* RootOp = Cast<URootMotionGeneratorOp>(RetargetController->GetRetargetOpAtIndex(RootOpIndex)))
+			if (FIKRetargetRootMotionOp* RootOp = (FIKRetargetRootMotionOp*)(RetargetController->GetRetargetOpAtIndex(RootOpIndex)))
 			{
-				RootOp->RootMotionSource = ERootMotionSource::GenerateFromTargetPelvis;
-				RootOp->RootHeightSource = ERootMotionHeightSource::SnapToGround;
-				RootOp->bRotateWithPelvis = true;
-				RootOp->bMaintainOffsetFromPelvis = true;
+				RootOp->Settings.RootMotionSource = ERootMotionSource::GenerateFromTargetPelvis;
+				RootOp->Settings.RootHeightSource = ERootMotionHeightSource::SnapToGround;
+				RootOp->Settings.bRotateWithPelvis = true;
+				RootOp->Settings.bMaintainOffsetFromPelvis = true;
 			}
 		}
 	}
 	
 	// setup a "Pin Bone Op" to pin IK bones to the new target locations
-	const int32 PinBonesOpIndex = RetargetController->AddRetargetOp(UPinBoneOp::StaticClass());
-	UPinBoneOp* PinBoneOp = CastChecked<UPinBoneOp>(RetargetController->GetRetargetOpAtIndex(PinBonesOpIndex));
-	PinBoneOp->TranslationMode = EPinBoneTranslationMode::CopyGlobalPosition;
-	PinBoneOp->RotationMode = EPinBoneRotationMode::CopyGlobalRotation;
+	const int32 PinBonesOpIndex = RetargetController->AddRetargetOp(FIKRetargetPinBoneOp::StaticStruct());
+	FIKRetargetPinBoneOp* PinBoneOp = (FIKRetargetPinBoneOp*)RetargetController->GetRetargetOpAtIndex(PinBonesOpIndex);
+	PinBoneOp->Settings.TranslationMode = EPinBoneTranslationMode::CopyGlobalPosition;
+	PinBoneOp->Settings.RotationMode = EPinBoneRotationMode::CopyGlobalRotation;
 	const TArray<FBoneToPin>& AllBonesToPin = TargetCharacterizationResults.AutoRetargetDefinition.BonesToPin.GetBonesToPin();
 	for (const FBoneToPin& BonesToPin : AllBonesToPin)
 	{
 		FPinBoneData PinBoneData;
 		PinBoneData.BoneToCopyFrom = BonesToPin.BoneToPinTo;
 		PinBoneData.BoneToCopyTo = BonesToPin.BoneToPin;
-		PinBoneOp->BonesToPin.Add(PinBoneData);
+		PinBoneOp->Settings.BonesToPin.Add(PinBoneData);
 	}
 }
 
@@ -905,7 +902,7 @@ SRetargetAnimAssetsWindow::SRetargetAnimAssetsWindow()
 	BatchContext.IKRetargetAsset = ProceduralAssets.Retargeter;
 
 	// register log name
-	Log.SetLogTarget(LogName, LogLabel);
+	Log.SetLogTarget(ProceduralAssets.Retargeter);
 }
 
 TSharedPtr<SWindow> SRetargetAnimAssetsWindow::Window;
diff --git a/Engine/Plugins/Animation/IKRig/Source/IKRigEditor/Private/RetargetEditor/SRetargetOpStack.cpp b/Engine/Plugins/Animation/IKRig/Source/IKRigEditor/Private/RetargetEditor/SRetargetOpStack.cpp
--- a/Engine/Plugins/Animation/IKRig/Source/IKRigEditor/Private/RetargetEditor/SRetargetOpStack.cpp
+++ b/Engine/Plugins/Animation/IKRig/Source/IKRigEditor/Private/RetargetEditor/SRetargetOpStack.cpp
@@ -140,17 +140,17 @@ void SRetargetOpStackItem::Construct(
 
 bool SRetargetOpStackItem::GetWarningMessage(FText& Message) const
 {
-	const UIKRetargetProcessor* Processor = OpStackWidget.Pin()->EditorController.Pin()->GetRetargetProcessor();
+	const FIKRetargetProcessor* Processor = OpStackWidget.Pin()->EditorController.Pin()->GetRetargetProcessor();
 	if (!(Processor && Processor->IsInitialized()))
 	{
 		return false;
 	}
 
-	const TArray<TObjectPtr<URetargetOpBase>>& RetargetOps = Processor->GetRetargetOps();
+	const TArray<TSharedPtr<FIKRetargetOpBase>>& RetargetOps = Processor->GetRetargetOps();
 	const int32 OpIndex = StackElement.Pin()->IndexInStack;
 	if (RetargetOps.IsValidIndex(OpIndex))
 	{
-		Message = RetargetOps[OpIndex]->WarningMessage();
+		Message = RetargetOps[OpIndex]->GetWarningMessage();
 		return true;
 	}
 	
@@ -159,15 +159,15 @@ bool SRetargetOpStackItem::GetWarningMessage(FText& Message) const
 
 bool SRetargetOpStackItem::IsOpEnabled() const
 {
-	if (const URetargetOpBase* Op = GetRetargetOp())
+	if (const FIKRetargetOpBase* Op = GetRetargetOp())
 	{
-		return Op->bIsEnabled;
+		return Op->IsEnabled();
 	}
 	
 	return false;
 }
 
-URetargetOpBase* SRetargetOpStackItem::GetRetargetOp() const
+FIKRetargetOpBase* SRetargetOpStackItem::GetRetargetOp() const
 {
 	if (!(StackElement.IsValid() && OpStackWidget.IsValid()))
 	{
@@ -265,17 +265,23 @@ TSharedRef<SWidget> SRetargetOpStack::CreateAddNewMenuWidget()
 	MenuBuilder.BeginSection("AddNewRetargetOp", LOCTEXT("AddOperations", "Add New Retarget Op"));
 
 	// add menu option to create each retarget op type
-	for(TObjectIterator<UClass> ClassIt; ClassIt; ++ClassIt)
+	const UStruct* BaseStruct = FIKRigSolverBase::StaticStruct();
+	for(TObjectIterator<UStruct> StructIt; StructIt; ++StructIt)
 	{
-		UClass* Class = *ClassIt;
-		if(Class->IsChildOf(URetargetOpBase::StaticClass()) && !Class->HasAnyClassFlags(CLASS_Abstract))
+		UStruct* Struct = *StructIt;
+		if(!Struct->IsChildOf(BaseStruct) || Struct == BaseStruct)
 		{
-			const URetargetOpBase* OpCDO = Cast<URetargetOpBase>(Class->GetDefaultObject());
-			FUIAction Action = FUIAction( FExecuteAction::CreateSP(this, &SRetargetOpStack::AddNewRetargetOp, Class));
-			MenuBuilder.AddMenuEntry(FText::FromString(OpCDO->GetNiceName().ToString()), FText::GetEmpty(), FSlateIcon(), Action);
+			continue;
 		}
-	}
 
+		if (UScriptStruct* ScriptStruct = Cast<UScriptStruct>(Struct))
+		{
+			FString TypeName = ScriptStruct->GetName();
+			FUIAction Action = FUIAction( FExecuteAction::CreateSP(this, &SRetargetOpStack::AddNewRetargetOp, ScriptStruct));
+			MenuBuilder.AddMenuEntry(FText::FromString(TypeName), FText::GetEmpty(), FSlateIcon(), Action);
+		}
+	}
+	
 	MenuBuilder.EndSection();
 
 	return MenuBuilder.MakeWidget();
@@ -288,7 +294,7 @@ bool SRetargetOpStack::IsAddOpEnabled() const
 		return false; 
 	}
 
-	if (const UIKRetargetProcessor* Processor = EditorController.Pin()->GetRetargetProcessor())
+	if (const FIKRetargetProcessor* Processor = EditorController.Pin()->GetRetargetProcessor())
 	{
 		return Processor->IsInitialized();
 	}
@@ -296,7 +302,7 @@ bool SRetargetOpStack::IsAddOpEnabled() const
 	return false;
 }
 
-void SRetargetOpStack::AddNewRetargetOp(UClass* Class)
+void SRetargetOpStack::AddNewRetargetOp(UScriptStruct* ScriptStruct) const
 {
 	if (!EditorController.IsValid())
 	{
@@ -309,7 +315,7 @@ void SRetargetOpStack::AddNewRetargetOp(UClass* Class)
 		return;
 	}
 	
-	AssetController->AddRetargetOp(Class);
+	AssetController->AddRetargetOp(ScriptStruct);
 }
 
 void SRetargetOpStack::DeleteRetargetOp(TSharedPtr<FRetargetOpStackElement> OpToDelete)
@@ -336,6 +342,7 @@ void SRetargetOpStack::OnSelectionChanged(
 	if (SelectInfo != ESelectInfo::Direct && InItem.IsValid())
 	{
 		EditorController.Pin()->LastSelectedOpIndex = InItem->IndexInStack;
+		EditorController.Pin()->ShowDetailsForOp(InItem->IndexInStack);
 	}
 }
 
@@ -361,11 +368,11 @@ void SRetargetOpStack::RefreshStackView()
 	ListViewItems.Reset();
 	UIKRetargeterController* AssetController = Controller->AssetController;
 	const int32 NumOps = AssetController->GetNumRetargetOps();
-	for (int32 i=0; i<NumOps; ++i)
+	for (int32 OpIndex=0; OpIndex<NumOps; ++OpIndex)
 	{
-		const URetargetOpBase* Op = AssetController->GetRetargetOpAtIndex(i);
-		FString UniqueName =  FString::FromInt(i+1) + " - " + Op->GetNiceName().ToString();
-		TSharedPtr<FRetargetOpStackElement> StackElement = FRetargetOpStackElement::Make(FText::FromString(UniqueName), i);
+		const FIKRetargetOpBase* Op = AssetController->GetRetargetOpAtIndex(OpIndex);
+		FString UniqueName =  FString::FromInt(OpIndex+1) + " - " + Op->GetNiceName().ToString();
+		TSharedPtr<FRetargetOpStackElement> StackElement = FRetargetOpStackElement::Make(FText::FromString(UniqueName), OpIndex);
 		ListViewItems.Add(StackElement);
 	}
 
diff --git a/Engine/Plugins/Animation/IKRig/Source/IKRigEditor/Private/RigEditor/IKRigAnimInstance.cpp b/Engine/Plugins/Animation/IKRig/Source/IKRigEditor/Private/RigEditor/IKRigAnimInstance.cpp
--- a/Engine/Plugins/Animation/IKRig/Source/IKRigEditor/Private/RigEditor/IKRigAnimInstance.cpp
+++ b/Engine/Plugins/Animation/IKRig/Source/IKRigEditor/Private/RigEditor/IKRigAnimInstance.cpp
@@ -22,9 +22,8 @@ void UIKRigAnimInstance::SetProcessorNeedsInitialized()
 	IKRigNode.SetProcessorNeedsInitialized();
 }
 
-UIKRigProcessor* UIKRigAnimInstance::GetCurrentlyRunningProcessor()
+FIKRigProcessor* UIKRigAnimInstance::GetCurrentlyRunningProcessor()
 {
-	IKRigNode.CreateIKRigProcessorIfNeeded(this);
 	return IKRigNode.GetIKRigProcessor();
 }
 
diff --git a/Engine/Plugins/Animation/IKRig/Source/IKRigEditor/Private/RigEditor/IKRigController.cpp b/Engine/Plugins/Animation/IKRig/Source/IKRigEditor/Private/RigEditor/IKRigController.cpp
--- a/Engine/Plugins/Animation/IKRig/Source/IKRigEditor/Private/RigEditor/IKRigController.cpp
+++ b/Engine/Plugins/Animation/IKRig/Source/IKRigEditor/Private/RigEditor/IKRigController.cpp
@@ -688,7 +688,7 @@ bool UIKRigController::SetSkeletalMesh(USkeletalMesh* SkeletalMesh) const
 bool UIKRigController::IsSkeletalMeshCompatible(USkeletalMesh* SkeletalMeshToCheck) const
 {
 	const FIKRigInputSkeleton InputSkeleton = FIKRigInputSkeleton(SkeletalMeshToCheck);
-	return UIKRigProcessor::IsIKRigCompatibleWithSkeleton(Asset, InputSkeleton, nullptr);
+	return FIKRigProcessor::IsIKRigCompatibleWithSkeleton(Asset, InputSkeleton, nullptr);
 }
 
 USkeletalMesh* UIKRigController::GetSkeletalMesh() const
diff --git a/Engine/Plugins/Animation/IKRig/Source/IKRigEditor/Private/RigEditor/IKRigDetailCustomizations.cpp b/Engine/Plugins/Animation/IKRig/Source/IKRigEditor/Private/RigEditor/IKRigDetailCustomizations.cpp
--- a/Engine/Plugins/Animation/IKRig/Source/IKRigEditor/Private/RigEditor/IKRigDetailCustomizations.cpp
+++ b/Engine/Plugins/Animation/IKRig/Source/IKRigEditor/Private/RigEditor/IKRigDetailCustomizations.cpp
@@ -39,7 +39,7 @@ void FStructViewerCustomization::CustomizeDetails(IDetailLayoutBuilder& DetailBu
 	}
 
 	// show the struct in the details panel
-	TSharedPtr<FStructOnScope> StructOnScope = StructViewer->GetStructOnScope();
+	const TSharedPtr<FStructOnScope>& StructOnScope = StructViewer->GetStructOnScope();
 	FName StructTitle = StructViewer->GetStructName();
 	IDetailCategoryBuilder& CategoryBuilder = DetailBuilder.EditCategory(StructTitle);
 	TArray<IDetailPropertyRow*> OutRows;
diff --git a/Engine/Plugins/Animation/IKRig/Source/IKRigEditor/Private/RigEditor/IKRigEditMode.cpp b/Engine/Plugins/Animation/IKRig/Source/IKRigEditor/Private/RigEditor/IKRigEditMode.cpp
--- a/Engine/Plugins/Animation/IKRig/Source/IKRigEditor/Private/RigEditor/IKRigEditMode.cpp
+++ b/Engine/Plugins/Animation/IKRig/Source/IKRigEditor/Private/RigEditor/IKRigEditMode.cpp
@@ -132,13 +132,16 @@ void FIKRigEditMode::RenderBones(FPrimitiveDrawInterface* PDI)
 		return; 
 	}
 
-	// IKRig processor initialized and running?
-	const UIKRigProcessor* CurrentProcessor = Controller->GetIKRigProcessor();
-	if (!IsValid(CurrentProcessor))
+	// is the ik rig available?
+	const FIKRigProcessor* ProcessorPtr = Controller->GetIKRigProcessor();
+	if (!ProcessorPtr)
 	{
 		return;
 	}
-	if (!CurrentProcessor->IsInitialized())
+
+	// is the ik rig initialized?
+	const FIKRigProcessor& CurrentProcessor = *ProcessorPtr;
+	if (!CurrentProcessor.IsInitialized())
 	{
 		return;
 	}
@@ -208,7 +211,7 @@ void FIKRigEditMode::RenderBones(FPrimitiveDrawInterface* PDI)
 
 void FIKRigEditMode::GetBoneColors(
 	FIKRigEditorController* Controller,
-	const UIKRigProcessor* Processor,
+	const FIKRigProcessor& Processor,
 	const FReferenceSkeleton& RefSkeleton,
 	bool bUseMultiColorAsDefaultColor,
 	TArray<FLinearColor>& OutBoneColors) const
@@ -245,7 +248,7 @@ void FIKRigEditMode::GetBoneColors(
 			}
 
 			// record which bones in the skeleton are affected by this solver
-			const FIKRigSkeleton& Skeleton = Processor->GetSkeleton();
+			const FIKRigSkeleton& Skeleton = Processor.GetSkeleton();
 			const UIKRigController* AssetController = Controller->AssetController;
 			if(const FIKRigSolverBase* SelectedSolver = AssetController->GetSolverAtIndex(SelectedSolvers[0].Get()->IndexInStack))
 			{
@@ -269,7 +272,7 @@ void FIKRigEditMode::GetBoneColors(
 				return;
 			}
 
-			const FIKRigSkeleton& Skeleton = Processor->GetSkeleton();
+			const FIKRigSkeleton& Skeleton = Processor.GetSkeleton();
 			for (const FName& SelectedChainName : SelectedChainNames)
 			{
 				TSet<int32> OutChainBoneIndices;
diff --git a/Engine/Plugins/Animation/IKRig/Source/IKRigEditor/Private/RigEditor/IKRigEditorController.cpp b/Engine/Plugins/Animation/IKRig/Source/IKRigEditor/Private/RigEditor/IKRigEditorController.cpp
--- a/Engine/Plugins/Animation/IKRig/Source/IKRigEditor/Private/RigEditor/IKRigEditorController.cpp
+++ b/Engine/Plugins/Animation/IKRig/Source/IKRigEditor/Private/RigEditor/IKRigEditorController.cpp
@@ -349,7 +349,7 @@ void FIKRigEditorController::Close() const
 	AssetController->OnIKRigNeedsInitialized().Remove(ReinitializeDelegateHandle);
 }
 
-UIKRigProcessor* FIKRigEditorController::GetIKRigProcessor() const
+FIKRigProcessor* FIKRigEditorController::GetIKRigProcessor() const
 {
 	if (AnimInstance)
 	{
@@ -361,7 +361,7 @@ UIKRigProcessor* FIKRigEditorController::GetIKRigProcessor() const
 
 const FIKRigSkeleton* FIKRigEditorController::GetCurrentIKRigSkeleton() const
 {
-	if (const UIKRigProcessor* Processor = GetIKRigProcessor())
+	if (const FIKRigProcessor* Processor = GetIKRigProcessor())
 	{
 		if (Processor->IsInitialized())
 		{
@@ -494,7 +494,7 @@ void FIKRigEditorController::ClearOutputLog() const
 	if (OutputLogView.IsValid())
 	{
 		OutputLogView.Get()->ClearLog();
-		if (const UIKRigProcessor* Processor = GetIKRigProcessor())
+		if (const FIKRigProcessor* Processor = GetIKRigProcessor())
 		{
 			Processor->Log.Clear();
 		}
@@ -534,7 +534,7 @@ void FIKRigEditorController::AutoGenerateRetargetChains() const
 		FSlateNotificationManager::Get().AddNotification(Info);
 
 		// log all the differences between the template the skeleton
-		if (UIKRigProcessor* Processor = GetIKRigProcessor())
+		if (FIKRigProcessor* Processor = GetIKRigProcessor())
 		{
 			FIKRigLogger& Log = Processor->Log;
 
@@ -957,12 +957,17 @@ void FIKRigEditorController::ShowDetailsForBone(const FName BoneName) const
 
 void FIKRigEditorController::ShowDetailsForBoneSettings(const FName& BoneName, int32 SolverIndex) const
 {
-	FIKRigSolverBase* Solver = AssetController->GetSolverAtIndex(SolverIndex);
+	FInstancedStruct* InstancedStruct = AssetController->GetSolverStructAtIndex(SolverIndex);
+	const UScriptStruct* TypeToEdit = InstancedStruct->GetPtr<FIKRigSolverBase>()->GetBoneSettingsType();
+	auto MemoryProvider = [BoneName](FInstancedStruct& InInstancedStruct) -> uint8*
+	{
+		FIKRigSolverBase* BaseStruct = InInstancedStruct.GetMutablePtr<FIKRigSolverBase>();
+		return reinterpret_cast<uint8*>(BaseStruct->GetBoneSettings(BoneName));
+	};
+	
 	UStructViewer* StructViewer = AssetController->GetStructViewer();
-	StructViewer->SetStructToView(
-		Solver->GetBoneSettingsType(),
-		reinterpret_cast<uint8*>(Solver->GetBoneSettings(BoneName)),
-		AssetController->GetAsset());
+	StructViewer->SetStructToView(InstancedStruct, TypeToEdit, MemoryProvider, AssetController->GetAsset());
+	
 	ShowStructViewerInDetailsView(StructViewer);
 }
 
@@ -973,23 +978,33 @@ void FIKRigEditorController::ShowDetailsForGoal(const FName& GoalName) const
 
 void FIKRigEditorController::ShowDetailsForGoalSettings(const FName GoalName, const int32 SolverIndex) const
 {
-	FIKRigSolverBase* Solver = AssetController->GetSolverAtIndex(SolverIndex);
+	FInstancedStruct* InstancedStruct = AssetController->GetSolverStructAtIndex(SolverIndex);
+	const UScriptStruct* TypeToEdit = InstancedStruct->GetPtr<FIKRigSolverBase>()->GetGoalSettingsType();
+	auto MemoryProvider = [GoalName](FInstancedStruct& InInstancedStruct) -> uint8*
+	{
+		FIKRigSolverBase* BaseStruct = InInstancedStruct.GetMutablePtr<FIKRigSolverBase>();
+		return reinterpret_cast<uint8*>(BaseStruct->GetGoalSettings(GoalName));
+	};
+	
 	UStructViewer* StructViewer = AssetController->GetStructViewer();
-	StructViewer->SetStructToView(
-		Solver->GetGoalSettingsType(),
-		reinterpret_cast<uint8*>(Solver->GetGoalSettings(GoalName)),
-		AssetController->GetAsset());
+	StructViewer->SetStructToView(InstancedStruct, TypeToEdit, MemoryProvider, AssetController->GetAsset());
+	
 	ShowStructViewerInDetailsView(StructViewer);
 }
 
 void FIKRigEditorController::ShowDetailsForSolver(const int32 SolverIndex) const
 {
-	FIKRigSolverBase* Solver = AssetController->GetSolverAtIndex(SolverIndex);
+	FInstancedStruct* InstancedStruct = AssetController->GetSolverStructAtIndex(SolverIndex);
+	const UScriptStruct* TypeToEdit = InstancedStruct->GetPtr<FIKRigSolverBase>()->GetSolverSettingsType();
+	auto MemoryProvider = [](FInstancedStruct& InInstancedStruct) -> uint8*
+	{
+		FIKRigSolverBase* BaseStruct = InInstancedStruct.GetMutablePtr<FIKRigSolverBase>();
+		return reinterpret_cast<uint8*>(BaseStruct->GetSolverSettings());
+	};
+	
 	UStructViewer* StructViewer = AssetController->GetStructViewer();
-	StructViewer->SetStructToView(
-		Solver->GetSolverSettingsType(),
-		reinterpret_cast<uint8*>(Solver->GetSolverSettings()),
-		AssetController->GetAsset());
+	StructViewer->SetStructToView(InstancedStruct, TypeToEdit, MemoryProvider, AssetController->GetAsset());
+	
 	ShowStructViewerInDetailsView(StructViewer);
 }
 
diff --git a/Engine/Plugins/Animation/IKRig/Source/IKRigEditor/Private/RigEditor/IKRigOutputLogTabSummoner.cpp b/Engine/Plugins/Animation/IKRig/Source/IKRigEditor/Private/RigEditor/IKRigOutputLogTabSummoner.cpp
--- a/Engine/Plugins/Animation/IKRig/Source/IKRigEditor/Private/RigEditor/IKRigOutputLogTabSummoner.cpp
+++ b/Engine/Plugins/Animation/IKRig/Source/IKRigEditor/Private/RigEditor/IKRigOutputLogTabSummoner.cpp
@@ -31,8 +31,12 @@ TSharedPtr<SToolTip> FIKRigOutputLogTabSummoner::CreateTabToolTipWidget(const FW
 
 TSharedRef<SWidget> FIKRigOutputLogTabSummoner::CreateTabBody(const FWorkflowTabSpawnInfo& Info) const
 {
+	// at this point the processor's log may not be bound to the asset, so forcibly do it now
 	const TSharedRef<FIKRigEditorController>& Controller = IKRigEditor.Pin()->GetController();
-	const FName LogName = Controller->GetIKRigProcessor()->Log.GetLogTarget();
+	FIKRigProcessor* Processor = Controller->GetIKRigProcessor();
+	Processor->Log.SetLogTarget(Controller->AssetController->GetAsset());
+	// now bind the log UI to the processor's log target
+	const FName LogName = Processor->Log.GetLogTarget();
 	TSharedRef<SIKRigOutputLog> LogView = SNew(SIKRigOutputLog, LogName);
 	Controller->SetOutputLogView(LogView);
 	return LogView;
diff --git a/Engine/Plugins/Animation/IKRig/Source/IKRigEditor/Private/RigEditor/IKRigToolkit.cpp b/Engine/Plugins/Animation/IKRig/Source/IKRigEditor/Private/RigEditor/IKRigToolkit.cpp
--- a/Engine/Plugins/Animation/IKRig/Source/IKRigEditor/Private/RigEditor/IKRigToolkit.cpp
+++ b/Engine/Plugins/Animation/IKRig/Source/IKRigEditor/Private/RigEditor/IKRigToolkit.cpp
@@ -319,7 +319,7 @@ void FIKRigEditorToolkit::HandleViewportCreated(const TSharedRef<IPersonaViewpor
 	auto GetBorderColorAndOpacity = [this]()
 	{
 		// no processor or processor not initialized
-		const UIKRigProcessor* Processor = EditorController.Get().GetIKRigProcessor();
+		const FIKRigProcessor* Processor = EditorController.Get().GetIKRigProcessor();
 		if (!Processor || !Processor->IsInitialized() )
 		{
 			return FLinearColor::Red;
diff --git a/Engine/Plugins/Animation/IKRig/Source/IKRigEditor/Private/RigEditor/SIKRigSolverStack.cpp b/Engine/Plugins/Animation/IKRig/Source/IKRigEditor/Private/RigEditor/SIKRigSolverStack.cpp
--- a/Engine/Plugins/Animation/IKRig/Source/IKRigEditor/Private/RigEditor/SIKRigSolverStack.cpp
+++ b/Engine/Plugins/Animation/IKRig/Source/IKRigEditor/Private/RigEditor/SIKRigSolverStack.cpp
@@ -202,7 +202,7 @@ const FIKRigSolverBase* SIKRigSolverStackItem::GetSolver(const bool bFromAsset)
 	{
 		// get the currently running solver instance in the editor's runtime processor
 		// this is needed to report solver warnings dependent on being initialized
-		if (UIKRigProcessor* Processor = EditorController->GetIKRigProcessor())
+		if (FIKRigProcessor* Processor = EditorController->GetIKRigProcessor())
 		{
 			if (Processor->GetSolvers().IsValidIndex(SolverIndex))
 			{
@@ -325,7 +325,6 @@ void SIKRigSolverStack::BuildAddNewMenu(FMenuBuilder& MenuBuilder)
 
 		if (UScriptStruct* ScriptStruct = Cast<UScriptStruct>(Struct))
 		{
-			//UIKRigSolver* SolverCDO = Cast<UIKRigSolver>(Class->GetDefaultObject());
 			FString TypeName = ScriptStruct->GetName();
 			FUIAction Action = FUIAction( FExecuteAction::CreateSP(this, &SIKRigSolverStack::AddNewSolver, ScriptStruct));
 			MenuBuilder.AddMenuEntry(FText::FromString(TypeName), FText::GetEmpty(), FSlateIcon(), Action);
diff --git a/Engine/Plugins/Animation/IKRig/Source/IKRigEditor/Public/RetargetEditor/IKRetargetAnimInstance.h b/Engine/Plugins/Animation/IKRig/Source/IKRigEditor/Public/RetargetEditor/IKRetargetAnimInstance.h
--- a/Engine/Plugins/Animation/IKRig/Source/IKRigEditor/Public/RetargetEditor/IKRetargetAnimInstance.h
+++ b/Engine/Plugins/Animation/IKRig/Source/IKRigEditor/Public/RetargetEditor/IKRetargetAnimInstance.h
@@ -60,7 +60,7 @@ public:
 		UIKRetargeter* InIKRetargetAsset,
 		TWeakObjectPtr<USkeletalMeshComponent> InSourceMeshComponent);
 
-	UIKRetargetProcessor* GetRetargetProcessor();
+	FIKRetargetProcessor* GetRetargetProcessor();
 
 	void ForceInitializeProcessor(USkeletalMeshComponent* TargetComponent);
 
diff --git a/Engine/Plugins/Animation/IKRig/Source/IKRigEditor/Public/RetargetEditor/IKRetargetCommands.h b/Engine/Plugins/Animation/IKRig/Source/IKRigEditor/Public/RetargetEditor/IKRetargetCommands.h
--- a/Engine/Plugins/Animation/IKRig/Source/IKRigEditor/Public/RetargetEditor/IKRetargetCommands.h
+++ b/Engine/Plugins/Animation/IKRig/Source/IKRigEditor/Public/RetargetEditor/IKRetargetCommands.h
@@ -31,8 +31,6 @@ public:
 	TSharedPtr< FUICommandInfo > ShowGlobalSettings;
 	// open root settings
 	TSharedPtr< FUICommandInfo > ShowRootSettings;
-	// open post settings
-	TSharedPtr< FUICommandInfo > ShowPostSettings;
 
 	// enable Root retarget pass
 	TSharedPtr< FUICommandInfo > EnableRoot;
diff --git a/Engine/Plugins/Animation/IKRig/Source/IKRigEditor/Public/RetargetEditor/IKRetargetDefaultMode.h b/Engine/Plugins/Animation/IKRig/Source/IKRigEditor/Public/RetargetEditor/IKRetargetDefaultMode.h
--- a/Engine/Plugins/Animation/IKRig/Source/IKRigEditor/Public/RetargetEditor/IKRetargetDefaultMode.h
+++ b/Engine/Plugins/Animation/IKRig/Source/IKRigEditor/Public/RetargetEditor/IKRetargetDefaultMode.h
@@ -8,7 +8,6 @@
 #include "Retargeter/IKRetargeter.h"
 #include "IPersonaEditMode.h"
 
-class UIKRigProcessor;
 class FIKRetargetEditorController;
 class FIKRetargetEditor;
 class FIKRetargetPreviewScene;
diff --git a/Engine/Plugins/Animation/IKRig/Source/IKRigEditor/Public/RetargetEditor/IKRetargetDetails.h b/Engine/Plugins/Animation/IKRig/Source/IKRigEditor/Public/RetargetEditor/IKRetargetDetails.h
--- a/Engine/Plugins/Animation/IKRig/Source/IKRigEditor/Public/RetargetEditor/IKRetargetDetails.h
+++ b/Engine/Plugins/Animation/IKRig/Source/IKRigEditor/Public/RetargetEditor/IKRetargetDetails.h
@@ -241,27 +241,3 @@ private:
 
 	TArray<TSharedPtr<FString>> TargetChainOptions;
 };
-
-/** ------------------------------------- BEGIN POST DETAILS CUSTOMIZATION -------------*/
-
-class FRetargetOpStackCustomization : public IDetailCustomization
-{
-public:
-
-	static TSharedRef<IDetailCustomization> MakeInstance()
-	{
-		return MakeShareable(new FRetargetOpStackCustomization);
-	}
-
-	/** IDetailCustomization interface */
-	virtual void CustomizeDetails(IDetailLayoutBuilder& DetailBuilder) override;
-
-private:
-
-	TSharedRef<SWidget> CreateAddNewMenuWidget();
-
-	void AddNewRetargetOp(UClass* Class);
-	
-	TWeakObjectPtr<URetargetOpStack> RetargetOpStackObject;
-	TWeakPtr<FIKRetargetEditorController> Controller;
-};
diff --git a/Engine/Plugins/Animation/IKRig/Source/IKRigEditor/Public/RetargetEditor/IKRetargetEditPoseMode.h b/Engine/Plugins/Animation/IKRig/Source/IKRigEditor/Public/RetargetEditor/IKRetargetEditPoseMode.h
--- a/Engine/Plugins/Animation/IKRig/Source/IKRigEditor/Public/RetargetEditor/IKRetargetEditPoseMode.h
+++ b/Engine/Plugins/Animation/IKRig/Source/IKRigEditor/Public/RetargetEditor/IKRetargetEditPoseMode.h
@@ -9,7 +9,6 @@
 #include "IPersonaEditMode.h"
 #include "Retargeter/IKRetargetProcessor.h"
 
-class UIKRigProcessor;
 class FIKRetargetEditorController;
 class FIKRetargetEditor;
 class FIKRetargetPreviewScene;
diff --git a/Engine/Plugins/Animation/IKRig/Source/IKRigEditor/Public/RetargetEditor/IKRetargetEditorController.h b/Engine/Plugins/Animation/IKRig/Source/IKRigEditor/Public/RetargetEditor/IKRetargetEditorController.h
--- a/Engine/Plugins/Animation/IKRig/Source/IKRigEditor/Public/RetargetEditor/IKRetargetEditorController.h
+++ b/Engine/Plugins/Animation/IKRig/Source/IKRigEditor/Public/RetargetEditor/IKRetargetEditorController.h
@@ -15,7 +15,7 @@ class SRetargetOpStack;
 enum class ERetargetSourceOrTarget : uint8;
 class SIKRetargetHierarchy;
 class SIKRigOutputLog;
-class UIKRetargetProcessor;
+struct FIKRetargetProcessor;
 class SIKRetargetChainMapList;
 class UIKRetargetAnimInstance;
 class FIKRetargetEditor;
@@ -161,6 +161,7 @@ public:
 	void RefreshPoseList() const;
 	void SetDetailsObject(UObject* DetailsObject) const;
 	void SetDetailsObjects(const TArray<UObject*>& DetailsObjects) const;
+	void ShowDetailsForOp(const int32 OpIndex) const;
 
 	// retargeter state
 	bool IsReadyToRetarget() const;
@@ -170,7 +171,6 @@ public:
 	// display settings in the details panel
 	void ShowGlobalSettings();
 	void ShowRootSettings();
-	void ShowPostPhaseSettings();
 	FRetargetGlobalSettings& GetGlobalSettings() const;
 
 	// clear the output log
@@ -183,7 +183,7 @@ public:
 	// get currently edited debug skeletal mesh
 	UDebugSkelMeshComponent* GetEditedSkeletalMesh() const;
 	// get the currently edited retarget skeleton
-	const FRetargetSkeleton& GetCurrentlyEditedSkeleton(const UIKRetargetProcessor& Processor) const;
+	const FRetargetSkeleton& GetCurrentlyEditedSkeleton(const FIKRetargetProcessor& Processor) const;
 	
 	// get world space pose of a bone (with component scale / offset applied)
 	FTransform GetGlobalRetargetPoseOfBone(
@@ -202,7 +202,7 @@ public:
 		TArray<FVector>& OutChildrenPositions);
 
 	// get the retargeter that is running in the viewport (which is a duplicate of the source asset)
-	UIKRetargetProcessor* GetRetargetProcessor() const;
+	FIKRetargetProcessor* GetRetargetProcessor() const;
 	// Reset the planting state of the IK (when scrubbing or animation loops over)
 	void ResetIKPlantingState() const;
 
@@ -247,7 +247,7 @@ public:
 	ERetargetSelectionType GetLastSelectedItemType() const { return LastSelectedItem; };
 
 	// op stack selection
-	URetargetOpBase* GetSelectedOp() const;
+	FIKRetargetOpBase* GetSelectedOp() const;
 	int32 LastSelectedOpIndex = 0;
 
 	// to frame selection when pressing "f" in viewport
diff --git a/Engine/Plugins/Animation/IKRig/Source/IKRigEditor/Public/RetargetEditor/IKRetargetOpStackTabSummoner.h b/Engine/Plugins/Animation/IKRig/Source/IKRigEditor/Public/RetargetEditor/IKRetargetOpStackTabSummoner.h
new file mode 100644
--- /dev/null
+++ b/Engine/Plugins/Animation/IKRig/Source/IKRigEditor/Public/RetargetEditor/IKRetargetOpStackTabSummoner.h
@@ -0,0 +1,21 @@
+﻿// Copyright Epic Games, Inc. All Rights Reserved.
+
+#pragma once
+
+#include "CoreMinimal.h"
+#include "WorkflowOrientedApp/WorkflowTabFactory.h"
+
+class FIKRetargetEditor;
+
+struct FIKRetargetOpStackTabSummoner : public FWorkflowTabFactory
+{
+public:
+	static const FName TabID;
+	
+	FIKRetargetOpStackTabSummoner(const TSharedRef<FIKRetargetEditor>& InIKRetargetEditor);
+	virtual TSharedRef<SWidget> CreateTabBody(const FWorkflowTabSpawnInfo& Info) const override;
+	virtual TSharedPtr<SToolTip> CreateTabToolTipWidget(const FWorkflowTabSpawnInfo& Info) const override;
+
+protected:
+	TWeakPtr<FIKRetargetEditor> RetargetEditor;
+};
diff --git a/Engine/Plugins/Animation/IKRig/Source/IKRigEditor/Public/RetargetEditor/IKRetargeterController.h b/Engine/Plugins/Animation/IKRig/Source/IKRigEditor/Public/RetargetEditor/IKRetargeterController.h
--- a/Engine/Plugins/Animation/IKRig/Source/IKRigEditor/Public/RetargetEditor/IKRetargeterController.h
+++ b/Engine/Plugins/Animation/IKRig/Source/IKRigEditor/Public/RetargetEditor/IKRetargeterController.h
@@ -13,6 +13,7 @@ struct FRetargetGlobalSettings;
 struct FTargetRootSettings;
 struct FTargetChainSettings;
 struct FIKRetargetPose;
+struct FIKRetargetOpBase;
 enum class ERetargetSourceOrTarget : uint8;
 enum class ERetargetAutoAlignMethod : uint8;
 class FIKRetargetEditorController;
@@ -20,6 +21,7 @@ class URetargetChainSettings;
 class UIKRigDefinition;
 class UIKRetargeter;
 class USkeletalMesh;
+class UStructViewer;
 
 UENUM()
 enum class EAutoMapChainType : uint8
@@ -38,6 +40,8 @@ class IKRIGEDITOR_API UIKRetargeterController : public UObject
 
 public:
 
+	UIKRetargeterController();
+
 	// UObject
 	virtual void PostInitProperties() override;
 	
@@ -108,10 +112,11 @@ public:
 	//
 	// RETARGET OPS PUBLIC/SCRIPTING API
 	//
-	
-	// Add a new retarget op of the given type to the bottom of the stack. Returns the stack index.
-	UFUNCTION(BlueprintCallable, BlueprintPure=false, Category=RetargetOps)
-	int32 AddRetargetOp(TSubclassOf<URetargetOpBase> InOpClass) const;
+
+	// Add a new solver of the given type to the bottom of the stack. Returns the stack index.
+	// The type here is the full package path of the UStruct type. For example: /Script/IKRig.IKRetargetPinBoneOp
+	UFUNCTION(BlueprintCallable, BlueprintPure=false, Category=IKRig)
+	int32 AddRetargetOp(const FString InIKRetargetOpType) const;
 
 	// Remove the retarget op at the given stack index. 
 	UFUNCTION(BlueprintCallable, BlueprintPure=false, Category=RetargetOps)
@@ -121,14 +126,6 @@ public:
 	UFUNCTION(BlueprintCallable, Category=RetargetOps)
 	bool RemoveAllOps() const;
 
-	// Get access to the given retarget operation. 
-	UFUNCTION(BlueprintCallable, Category=RetargetOps)
-	URetargetOpBase* GetRetargetOpAtIndex(int32 Index) const;
-
-	// Get access to the given retarget operation. 
-	UFUNCTION(BlueprintCallable, Category=RetargetOps)
-	int32 GetIndexOfRetargetOp(URetargetOpBase* RetargetOp) const;
-
 	// Get the number of Ops in the stack.
 	UFUNCTION(BlueprintCallable, Category=RetargetOps)
 	int32 GetNumRetargetOps() const;
@@ -148,6 +145,18 @@ public:
 	//
 	// GENERAL C++ ONLY API
 	//
+	
+	// add a new retarget op to the stack
+	int32 AddRetargetOp(UScriptStruct* InIKRigSolverType) const;
+	
+	// Get access to the given retarget operation.
+	FIKRetargetOpBase* GetRetargetOpAtIndex(int32 Index) const;
+
+	// Get access to the given retarget operation. 
+	int32 GetIndexOfRetargetOp(FIKRetargetOpBase* RetargetOp) const;
+
+	// Get access to the given retarget operation.
+	FInstancedStruct* GetRetargetOpStructAtIndex(int32 Index) const;
 
 	// Ensures all internal data is compatible with assigned meshes and ready to edit.
 	// - Removes bones from retarget poses that are no longer in skeleton
@@ -352,6 +361,10 @@ private:
 	// prevent reinitializing from inner operations
 	mutable int32 ReinitializeScopeCounter = 0;
 	
+	// a dummy UObject to wrap UStructs for the details panel
+	UPROPERTY(transient)
+	TObjectPtr<UStructViewer> StructViewer = nullptr;
+	
 public:
 
 	// Attach a delegate to be notified whenever either the source or target Preview Mesh asset's are swapped out.
@@ -362,6 +375,9 @@ public:
 	
 	// Attach a delegate to be notified whenever the retargeter is modified in such a way that would require re-initialization of the processor.
 	FOnRetargeterNeedsInitialized& OnRetargeterNeedsInitialized(){ return RetargeterNeedsInitialized; };
+
+	// get the singleton struct viewer for editing structs belonging to the asset this controller controls
+	UStructViewer* GetStructViewer();
 	
 	friend class UIKRetargeter;
 	friend struct FScopedReinitializeIKRetargeter;
diff --git a/Engine/Plugins/Animation/IKRig/Source/IKRigEditor/Public/RetargetEditor/IKRetargeterPoseGenerator.h b/Engine/Plugins/Animation/IKRig/Source/IKRigEditor/Public/RetargetEditor/IKRetargeterPoseGenerator.h
--- a/Engine/Plugins/Animation/IKRig/Source/IKRigEditor/Public/RetargetEditor/IKRetargeterPoseGenerator.h
+++ b/Engine/Plugins/Animation/IKRig/Source/IKRigEditor/Public/RetargetEditor/IKRetargeterPoseGenerator.h
@@ -9,7 +9,7 @@
 
 class UIKRetargeterController;
 class FIKRetargetEditorController;
-class UIKRetargetProcessor;
+struct FIKRetargetProcessor;
 enum class ERetargetSourceOrTarget : uint8;
 
 // A Note on Determining the "Direction" of a Bone...
@@ -39,15 +39,10 @@ enum class ERetargetAutoAlignMethod : uint8
 };
 
 // a sub-system used by the IKRetargetEditorController for automatically generating a retarget pose
-struct FRetargetAutoPoseGenerator : public FGCObject
+struct FRetargetAutoPoseGenerator
 {
 	// created and owned by the retarget editor controller
 	FRetargetAutoPoseGenerator(const TWeakObjectPtr<UIKRetargeterController> InController);
-
-	// FGCObject interface
-	virtual void AddReferencedObjects(FReferenceCollector& Collector) override;
-	virtual FString GetReferencerName() const override { return TEXT("FRetargetAutoPoseGenerator"); }
-	// End FGCObject interface
 	
 	// align the specified bones on the current skeleton to the OTHER skeleton, using the specified method
 	// NOTE: if BonesToAlign is empty, it will align ALL bones in the skeleton
@@ -136,6 +131,6 @@ private:
 	bool CheckReadyToAlignBones() const;
 
 	TWeakObjectPtr<UIKRetargeterController> Controller;
-	TObjectPtr<UIKRetargetProcessor> Processor;
+	TSharedPtr<FIKRetargetProcessor> Processor;
 	static constexpr float ParamStepSize = 0.1f;
 };
diff --git a/Engine/Plugins/Animation/IKRig/Source/IKRigEditor/Public/RetargetEditor/SRetargetOpStack.h b/Engine/Plugins/Animation/IKRig/Source/IKRigEditor/Public/RetargetEditor/SRetargetOpStack.h
--- a/Engine/Plugins/Animation/IKRig/Source/IKRigEditor/Public/RetargetEditor/SRetargetOpStack.h
+++ b/Engine/Plugins/Animation/IKRig/Source/IKRigEditor/Public/RetargetEditor/SRetargetOpStack.h
@@ -11,7 +11,7 @@
 #include "Widgets/Views/SListView.h"
 #include "Framework/Commands/UICommandList.h"
 
-class URetargetOpBase;
+struct FIKRetargetOpBase;
 class FIKRetargetEditorController;
 class SRetargetOpStack;
 
@@ -55,7 +55,7 @@ public:
 	bool IsOpEnabled() const;
 
 private:
-	URetargetOpBase* GetRetargetOp() const;
+	FIKRetargetOpBase* GetRetargetOp() const;
 	TWeakPtr<FRetargetOpStackElement> StackElement;
 	TWeakPtr<SRetargetOpStack> OpStackWidget;
 };
@@ -97,7 +97,7 @@ private:
 	// END op stack menu
 	
 	// menu command callback for adding a new op
-	void AddNewRetargetOp(UClass* Class);
+	void AddNewRetargetOp(UScriptStruct* ScriptStruct) const;
 	// delete op from stack
 	void DeleteRetargetOp(TSharedPtr<FRetargetOpStackElement> OpToDelete);
 	// when an op is selected on in the stack view
diff --git a/Engine/Plugins/Animation/IKRig/Source/IKRigEditor/Public/RigEditor/IKRigAnimInstance.h b/Engine/Plugins/Animation/IKRig/Source/IKRigEditor/Public/RigEditor/IKRigAnimInstance.h
--- a/Engine/Plugins/Animation/IKRig/Source/IKRigEditor/Public/RigEditor/IKRigAnimInstance.h
+++ b/Engine/Plugins/Animation/IKRig/Source/IKRigEditor/Public/RigEditor/IKRigAnimInstance.h
@@ -21,7 +21,7 @@ public:
 
 	void SetProcessorNeedsInitialized();
 
-	UIKRigProcessor* GetCurrentlyRunningProcessor();
+	FIKRigProcessor* GetCurrentlyRunningProcessor();
 
 protected:
 	virtual FAnimInstanceProxy* CreateAnimInstanceProxy() override;
diff --git a/Engine/Plugins/Animation/IKRig/Source/IKRigEditor/Public/RigEditor/IKRigController.h b/Engine/Plugins/Animation/IKRig/Source/IKRigEditor/Public/RigEditor/IKRigController.h
--- a/Engine/Plugins/Animation/IKRig/Source/IKRigEditor/Public/RigEditor/IKRigController.h
+++ b/Engine/Plugins/Animation/IKRig/Source/IKRigEditor/Public/RigEditor/IKRigController.h
@@ -398,7 +398,7 @@ public:
 	// Optionally provide a runtime skeleton from an IKRigProcessor to get indices for a running instance (otherwise uses stored hierarchy in asset)
 	bool ValidateChain( const FName ChainName, const FIKRigSkeleton* OptionalSkeleton, TSet<int32>& OutChainIndices) const;
 
-	// 
+	// a singleton wrapper to display structs in the details panel
 	UStructViewer* GetStructViewer();
 
 private:
diff --git a/Engine/Plugins/Animation/IKRig/Source/IKRigEditor/Public/RigEditor/IKRigDetailCustomizations.h b/Engine/Plugins/Animation/IKRig/Source/IKRigEditor/Public/RigEditor/IKRigDetailCustomizations.h
--- a/Engine/Plugins/Animation/IKRig/Source/IKRigEditor/Public/RigEditor/IKRigDetailCustomizations.h
+++ b/Engine/Plugins/Animation/IKRig/Source/IKRigEditor/Public/RigEditor/IKRigDetailCustomizations.h
@@ -16,31 +16,47 @@ class UStructViewer : public UObject
 	GENERATED_BODY()
 
 public:
-
-	// assign a struct to view in the detail panel
-	// InStructOwner is the owning UObject to use for Undo transactions when editing the struct properties
-	void SetStructToView(const UScriptStruct* InStructType, uint8* InStructMemory, UObject* InStructOwner)
+	
+	/**
+	 * Configures an instanced struct to display a child struct property in the details panel with undo/redo support.
+	 * NOTE: we can't just pass in raw pointers to the struct memory because these can be destroyed after a transaction.
+	 * So instead we pass in TFunction callbacks that get the latest memory locations when the details panel is refreshed.
+	 *
+	 * @param InInstancedStruct A pointer to the instanced struct that contains the struct to be displayed, this pointer survives undo/redo
+	 * @param InScriptStruct The type that corresponds to the struct memory returned by the StructMemoryProvider
+	 * @param InStructMemoryProvider A callback function that retrieves the memory location for the struct to be edited
+	 * @param InStructOwner The UObject that owns the struct. This is what will be transacted when the property is edited.
+	 */
+	void SetStructToView(
+		FInstancedStruct* InInstancedStruct,
+		const UScriptStruct* InTypeToEdit,
+		TFunction<uint8*(FInstancedStruct& InInstancedStruct)> InStructMemoryProvider,
+		UObject* InStructOwner)
 	{
-		StructType = InStructType;
-		StructMemory = InStructMemory;
+		TopLevelInstancedStruct = InInstancedStruct;
+		StructType = InTypeToEdit;
+		StructMemoryProvider = MoveTemp(InStructMemoryProvider);
 		StructOwner = InStructOwner;
 	}
 
 	bool IsValid() const
 	{
-		return StructType && StructMemory;
+		return TopLevelInstancedStruct->IsValid() && StructType && StructOwner.IsValid();
 	}
 
 	void Reset()
 	{
+		TopLevelInstancedStruct = nullptr;
 		StructType = nullptr;
-		StructMemory = nullptr;
+		StructMemoryProvider.Reset();
 		StructOwner = nullptr;
 	}
 
-	TSharedPtr<FStructOnScope> GetStructOnScope()
+	TSharedPtr<FStructOnScope>& GetStructOnScope()
 	{
-		return MakeShared<FStructOnScope>(StructType, StructMemory);
+		uint8* StructMemory = StructMemoryProvider(*TopLevelInstancedStruct);
+		StructOnScope = MakeShared<FStructOnScope>(StructType, StructMemory);
+		return StructOnScope;
 	}
 
 	FName GetStructName() const
@@ -56,10 +72,16 @@ public:
 private:
 	
 	// we need the type reflection data and the memory to build a StructOnScope to display in details panel
-	const UScriptStruct* StructType = nullptr;
-	uint8* StructMemory = nullptr;
+	FInstancedStruct* TopLevelInstancedStruct;
+	// the type of struct we are editing
+	const UScriptStruct* StructType;
+	// provides the memory address of the struct to edit (refreshed after undo/redo)
+	TFunction<uint8*(FInstancedStruct& InInstancedStruct)> StructMemoryProvider;
+	
+	// a wrapper of the struct sent to the details panel
+	TSharedPtr<FStructOnScope> StructOnScope;
 
-	// object that owns the struct for Undo/Redo transactions
+	// a UObject that owns the struct for Undo/Redo transactions (ie, the asset we are editing)
 	TWeakObjectPtr<UObject> StructOwner = nullptr;
 };
 
diff --git a/Engine/Plugins/Animation/IKRig/Source/IKRigEditor/Public/RigEditor/IKRigEditMode.h b/Engine/Plugins/Animation/IKRig/Source/IKRigEditor/Public/RigEditor/IKRigEditMode.h
--- a/Engine/Plugins/Animation/IKRig/Source/IKRigEditor/Public/RigEditor/IKRigEditMode.h
+++ b/Engine/Plugins/Animation/IKRig/Source/IKRigEditor/Public/RigEditor/IKRigEditMode.h
@@ -12,7 +12,7 @@ class FIKRigEditorController;
 class UIKRigEffectorGoal;
 class FIKRigEditorToolkit;
 class FIKRigPreviewScene;
-class UIKRigProcessor;
+struct FIKRigProcessor;
 struct FGizmoState;
 
 class FIKRigEditMode : public IPersonaEditMode
@@ -60,7 +60,7 @@ private:
 	void RenderBones(FPrimitiveDrawInterface* PDI);
 	void GetBoneColors(
 		FIKRigEditorController* Controller,
-		const UIKRigProcessor* Processor,
+		const FIKRigProcessor& Processor,
 		const FReferenceSkeleton& RefSkeleton,
 		bool bUseMultiColorAsDefaultColor,
 		TArray<FLinearColor>& OutBoneColors) const;
diff --git a/Engine/Plugins/Animation/IKRig/Source/IKRigEditor/Public/RigEditor/IKRigEditorController.h b/Engine/Plugins/Animation/IKRig/Source/IKRigEditor/Public/RigEditor/IKRigEditorController.h
--- a/Engine/Plugins/Animation/IKRig/Source/IKRigEditor/Public/RigEditor/IKRigEditorController.h
+++ b/Engine/Plugins/Animation/IKRig/Source/IKRigEditor/Public/RigEditor/IKRigEditorController.h
@@ -22,7 +22,7 @@ class FIKRigTreeElement;
 class UAnimInstance;
 class UDebugSkelMeshComponent;
 class UStructViewer;
-class UIKRigProcessor;
+struct FIKRigProcessor;
 
 enum EIKRigSelectionType : int8
 {
@@ -116,7 +116,7 @@ public:
 	void Close() const;
 
 	// get the currently active processor running the IK Rig in the editor 
-	UIKRigProcessor* GetIKRigProcessor() const;
+	FIKRigProcessor* GetIKRigProcessor() const;
 	// get the currently running IKRig skeleton (if there is a running processor) 
 	const FIKRigSkeleton* GetCurrentIKRigSkeleton() const;
 
diff --git a/Engine/Plugins/Animation/PerformanceCaptureCore/Source/PerformanceCaptureCore/Private/RetargetInstanceProxy.cpp b/Engine/Plugins/Animation/PerformanceCaptureCore/Source/PerformanceCaptureCore/Private/RetargetInstanceProxy.cpp
--- a/Engine/Plugins/Animation/PerformanceCaptureCore/Source/PerformanceCaptureCore/Private/RetargetInstanceProxy.cpp
+++ b/Engine/Plugins/Animation/PerformanceCaptureCore/Source/PerformanceCaptureCore/Private/RetargetInstanceProxy.cpp
@@ -74,7 +74,7 @@ void FRetargetAnimInstanceProxy::ConfigureAnimInstanceProxy(UIKRetargeter* InIKR
 	RetargetNode->SourceMeshComponent = InSourceMeshComponent;
 	RetargetNode->CustomRetargetProfile = InRetargetProfile;
 	
-	if (UIKRetargetProcessor* Processor = RetargetNode->GetRetargetProcessor())
+	if (FIKRetargetProcessor* Processor = RetargetNode->GetRetargetProcessor())
 	{
 		Processor->SetNeedsInitialized();
 	}
