[IKRig] Fixing uninitialized variables in recently added structs.

diff --git a/Engine/Plugins/Animation/IKRig/Source/IKRig/Public/Retargeter/RetargetOps/PinBoneOp.h b/Engine/Plugins/Animation/IKRig/Source/IKRig/Public/Retargeter/RetargetOps/PinBoneOp.h
--- a/Engine/Plugins/Animation/IKRig/Source/IKRig/Public/Retargeter/RetargetOps/PinBoneOp.h
+++ b/Engine/Plugins/Animation/IKRig/Source/IKRig/Public/Retargeter/RetargetOps/PinBoneOp.h
@@ -112,11 +112,11 @@ struct IKRIG_API FIKRetargetPinBoneOpSettings : public FIKRetargetOpSettingsBase
 
 	// The method used to calculate the translation of the bone.
 	UPROPERTY(EditAnywhere, Category=Settings)
-	EPinBoneTranslationMode TranslationMode;
+	EPinBoneTranslationMode TranslationMode = EPinBoneTranslationMode::CopyGlobalPosition;
 
 	// The method used to calculate the rotation of the bone.
 	UPROPERTY(EditAnywhere, Category=Settings)
-	EPinBoneRotationMode RotationMode;
+	EPinBoneRotationMode RotationMode = EPinBoneRotationMode::CopyGlobalRotation;
 
 	// A manual offset to apply in global space
 	UPROPERTY(EditAnywhere, Category=Settings)
diff --git a/Engine/Plugins/Animation/IKRig/Source/IKRig/Public/Retargeter/RetargetOps/RootMotionGeneratorOp.h b/Engine/Plugins/Animation/IKRig/Source/IKRig/Public/Retargeter/RetargetOps/RootMotionGeneratorOp.h
--- a/Engine/Plugins/Animation/IKRig/Source/IKRig/Public/Retargeter/RetargetOps/RootMotionGeneratorOp.h
+++ b/Engine/Plugins/Animation/IKRig/Source/IKRig/Public/Retargeter/RetargetOps/RootMotionGeneratorOp.h
@@ -45,13 +45,13 @@ struct IKRIG_API FIKRetargetRootMotionOpSettings : public FIKRetargetOpSettingsB
 	// Generate From Target Pelvis: uses the retargeted Pelvis motion to generate root motion.
 	// NOTE: Generating root motion from the Pelvis 
 	UPROPERTY(EditAnywhere, Category="OpSettings")
-	ERootMotionSource RootMotionSource;
+	ERootMotionSource RootMotionSource = ERootMotionSource::CopyFromSourceRoot;
 	
 	// How to set the height of the root bone.
 	// Copy Height From Source: copies the height of the root bone on the source skeleton's root bone.
 	// Snap To Ground: sets the root bone height to the ground plane (Z=0).
 	UPROPERTY(EditAnywhere, Category="OpSettings")
-	ERootMotionHeightSource RootHeightSource;
+	ERootMotionHeightSource RootHeightSource = ERootMotionHeightSource::CopyHeightFromSource;
 
 	// Will transform all children of the target root that are not themselves part of a retarget chain.
 	UPROPERTY(EditAnywhere, Category="OpSettings")
diff --git a/Engine/Plugins/Cameras/GameplayCameras/Source/GameplayCameras/Public/Core/CameraContextDataTableAllocationInfo.h b/Engine/Plugins/Cameras/GameplayCameras/Source/GameplayCameras/Public/Core/CameraContextDataTableAllocationInfo.h
--- a/Engine/Plugins/Cameras/GameplayCameras/Source/GameplayCameras/Public/Core/CameraContextDataTableAllocationInfo.h
+++ b/Engine/Plugins/Cameras/GameplayCameras/Source/GameplayCameras/Public/Core/CameraContextDataTableAllocationInfo.h
@@ -25,7 +25,7 @@ struct FCameraContextDataDefinition
 
 	/** The type of the data. */
 	UPROPERTY()
-	ECameraContextDataType DataType;
+	ECameraContextDataType DataType = ECameraContextDataType::Name;
 
 	/** An extra type object for the data. */
 	UPROPERTY()
