[IKRig] Removing call to initialize log from inside processor initialize(). This had a Slate delegate callback that is not safe to call from the worker threads.

Also converted procressors from TSharedPtr to plain member variable.

diff --git a/Engine/Plugins/Animation/IKRig/Source/IKRig/Private/AnimNodes/AnimNode_IKRig.cpp b/Engine/Plugins/Animation/IKRig/Source/IKRig/Private/AnimNodes/AnimNode_IKRig.cpp
--- a/Engine/Plugins/Animation/IKRig/Source/IKRig/Private/AnimNodes/AnimNode_IKRig.cpp
+++ b/Engine/Plugins/Animation/IKRig/Source/IKRig/Private/AnimNodes/AnimNode_IKRig.cpp
@@ -47,9 +47,9 @@ void FAnimNode_IKRig::Evaluate_AnyThread(FPoseContext& Output)
 	}
 	
 	// initialize the IK Rig (will only try once on the current version of the rig asset)
-	if (!IKRigProcessor->IsInitialized())
+	if (!IKRigProcessor.IsInitialized())
 	{
-		IKRigProcessor->Initialize(RigDefinitionAsset, Output.AnimInstanceProxy->GetSkelMeshComponent()->GetSkeletalMeshAsset());
+		IKRigProcessor.Initialize(RigDefinitionAsset, Output.AnimInstanceProxy->GetSkelMeshComponent()->GetSkeletalMeshAsset());
 	}
 	
 	// copy input pose to solver stack
@@ -58,7 +58,7 @@ void FAnimNode_IKRig::Evaluate_AnyThread(FPoseContext& Output)
 	AssignGoalTargets();
 	// run stack of solvers,
 	const FTransform WorldToComponent =  Output.AnimInstanceProxy->GetComponentTransform().Inverse();
-	IKRigProcessor->Solve(WorldToComponent);
+	IKRigProcessor.Solve(WorldToComponent);
 	// updates transforms with new pose
 	CopyOutputPoseToAnimGraph(Output.Pose);
 }
@@ -68,13 +68,13 @@ void FAnimNode_IKRig::CopyInputPoseToSolver(FCompactPose& InputPose)
 	// start Solve() from REFERENCE pose
 	if (bStartFromRefPose)
 	{
-		IKRigProcessor->SetInputPoseToRefPose();
+		IKRigProcessor.SetInputPoseToRefPose();
 		return;
 	}
 	
 	// start Solve() from INPUT pose
 	// copy local bone transforms into IKRigProcessor skeleton
-	FIKRigSkeleton& IKRigSkeleton = IKRigProcessor->GetSkeletonWriteable();
+	FIKRigSkeleton& IKRigSkeleton = IKRigProcessor.GetSkeletonWriteable();
 	for (FCompactPoseBoneIndex CPIndex : InputPose.ForEachBoneIndex())
 	{
 		if (int32* Index = CompactPoseToRigIndices.Find(CPIndex))
@@ -105,7 +105,7 @@ void FAnimNode_IKRig::AssignGoalTargets()
 	// transaction has been fully processed.
 	if (bDriveWithSourceAsset && !GIsTransacting)
 	{
-		IKRigProcessor->CopyAllSettingsFromAsset(RigDefinitionAsset);
+		IKRigProcessor.CopyAllSettingsFromAsset(RigDefinitionAsset);
 		return;
 	}
 	#endif
@@ -113,19 +113,19 @@ void FAnimNode_IKRig::AssignGoalTargets()
 	// copy transforms from this anim node's goal pins from blueprint
 	for (const FIKRigGoal& Goal : Goals)
 	{
-		IKRigProcessor->SetIKGoal(Goal);
+		IKRigProcessor.SetIKGoal(Goal);
 	}
 
 	// override any goals that were manually set with goals from goal creator components (they take precedence)
 	for (const TPair<FName, FIKRigGoal>& GoalPair : GoalsFromGoalCreators)
 	{
-		IKRigProcessor->SetIKGoal(GoalPair.Value);
+		IKRigProcessor.SetIKGoal(GoalPair.Value);
 	}
 }
 
 void FAnimNode_IKRig::CopyOutputPoseToAnimGraph(FCompactPose& OutputPose)
 {
-	FIKRigSkeleton& IKRigSkeleton = IKRigProcessor->GetSkeletonWriteable();
+	FIKRigSkeleton& IKRigSkeleton = IKRigProcessor.GetSkeletonWriteable();
 	
 	// update local transforms of current IKRig pose
 	IKRigSkeleton.UpdateAllLocalTransformFromGlobal();
@@ -185,9 +185,6 @@ void FAnimNode_IKRig::Initialize_AnyThread(const FAnimationInitializeContext& Co
 
 	// trigger another search for goal creators on the main thread
 	bSearchedForGoalCreators = false;
-
-	// NOTE: this must be instantiated inside Initialize()
-	IKRigProcessor = MakeShared<FIKRigProcessor>();
 }
 
 void FAnimNode_IKRig::Update_AnyThread(const FAnimationUpdateContext& Context)
@@ -225,11 +222,11 @@ void FAnimNode_IKRig::Update_AnyThread(const FAnimationUpdateContext& Context)
 #if ENABLE_VISUAL_LOG
 #if WITH_EDITORONLY_DATA
 	// is node setup?
-	if (IKRigProcessor->IsInitialized() && FVisualLogger::IsRecording())
+	if (IKRigProcessor.IsInitialized() && FVisualLogger::IsRecording())
 	{
 		static const FBox UnitBox(FVector(-1, -1, -1), FVector(1, 1, 1));
 
-		const TArray<FIKRigGoal>& ProcessorGoals = IKRigProcessor->GetGoalContainer().GetGoalArray();
+		const TArray<FIKRigGoal>& ProcessorGoals = IKRigProcessor.GetGoalContainer().GetGoalArray();
 		for (const FIKRigGoal& Goal : ProcessorGoals)
 		{
 			FTransform GoalTransform(Goal.FinalBlendedRotation, Goal.FinalBlendedPosition, FVector(DebugScale, DebugScale, DebugScale));
@@ -285,17 +282,14 @@ void FAnimNode_IKRig::OnInitializeAnimInstance(const FAnimInstanceProxy* InProxy
 	InitializeProperties(InAnimInstance, GetTargetClass());
 }
 
-void FAnimNode_IKRig::SetProcessorNeedsInitialized() const
+void FAnimNode_IKRig::SetProcessorNeedsInitialized()
 {
-	if (IKRigProcessor)
-	{
-		IKRigProcessor->SetNeedsInitialized();
-	}
+	IKRigProcessor.SetNeedsInitialized();
 }
 
 FIKRigProcessor* FAnimNode_IKRig::GetIKRigProcessor()
 {
-	return IKRigProcessor.Get();
+	return &IKRigProcessor;
 }
 
 void FAnimNode_IKRig::CacheBones_AnyThread(const FAnimationCacheBonesContext& Context)
@@ -333,7 +327,7 @@ void FAnimNode_IKRig::CacheBones_AnyThread(const FAnimationCacheBonesContext& Co
 	}
 
 	// must reinitialize if bone count changes
-	IKRigProcessor->SetNeedsInitialized();
+	IKRigProcessor.SetNeedsInitialized();
 }
 
 void FAnimNode_IKRig::InitializeProperties(const UObject* InSourceInstance, UClass* InTargetClass)
@@ -446,13 +440,13 @@ void FAnimNode_IKRig::ConditionalDebugDraw(
 	}
 
 	// is processor running?
-	if (!IKRigProcessor->IsInitialized())
+	if (!IKRigProcessor.IsInitialized())
 	{
 		return;
 	}
 
 	// draw oriented box for each goal
-	const TArray<FIKRigGoal>& ProcessorGoals = IKRigProcessor->GetGoalContainer().GetGoalArray();
+	const TArray<FIKRigGoal>& ProcessorGoals = IKRigProcessor.GetGoalContainer().GetGoalArray();
 	for (const FIKRigGoal& Goal : ProcessorGoals)
 	{
 		DrawOrientedWireBox(PDI, Goal.FinalBlendedPosition, FVector::XAxisVector, FVector::YAxisVector, FVector::ZAxisVector, FVector::One() * DebugScale, FLinearColor::Yellow, SDPG_World);
diff --git a/Engine/Plugins/Animation/IKRig/Source/IKRig/Private/AnimNodes/AnimNode_RetargetPoseFromMesh.cpp b/Engine/Plugins/Animation/IKRig/Source/IKRig/Private/AnimNodes/AnimNode_RetargetPoseFromMesh.cpp
--- a/Engine/Plugins/Animation/IKRig/Source/IKRig/Private/AnimNodes/AnimNode_RetargetPoseFromMesh.cpp
+++ b/Engine/Plugins/Animation/IKRig/Source/IKRig/Private/AnimNodes/AnimNode_RetargetPoseFromMesh.cpp
@@ -26,8 +26,6 @@ void FAnimNode_RetargetPoseFromMesh::Initialize_AnyThread(const FAnimationInitia
 	{
 		Source.Initialize(Context);
 	}
-
-	Processor = MakeShared<FIKRetargetProcessor>();
 }
 
 void FAnimNode_RetargetPoseFromMesh::CacheBones_AnyThread(const FAnimationCacheBonesContext& Context)
@@ -211,7 +209,7 @@ void FAnimNode_RetargetPoseFromMesh::Evaluate_AnyThread(FPoseContext& Output)
 			return;
 		}
 	
-		if (!Processor->IsInitialized())
+		if (!Processor.IsInitialized())
 		{
 			return;
 		}
@@ -222,7 +220,7 @@ void FAnimNode_RetargetPoseFromMesh::Evaluate_AnyThread(FPoseContext& Output)
 		// now query retarget ops for any curves that need remapped from source to target
 		FBlendedCurve RemapedCurves;
 		TArray<FCurveRemapPair> CurvesToRetarget;
-		const TArray<TSharedPtr<FIKRetargetOpBase>>& RetargetOps = Processor->GetRetargetOps();
+		const TArray<TSharedPtr<FIKRetargetOpBase>>& RetargetOps = Processor.GetRetargetOps();
 		for (const TSharedPtr<FIKRetargetOpBase>& RetargetOp : RetargetOps)
 		{
 			if (RetargetOp->IsEnabled())
@@ -262,7 +260,7 @@ void FAnimNode_RetargetPoseFromMesh::Evaluate_AnyThread(FPoseContext& Output)
 	{
 		// live preview source asset settings in the retarget, editor only
 		// NOTE: this copies goal targets as well, but these are overwritten by IK chain goals
-		Processor->CopyIKRigSettingsFromAsset();
+		Processor.CopyIKRigSettingsFromAsset();
 	}
 	#endif
 
@@ -271,7 +269,7 @@ void FAnimNode_RetargetPoseFromMesh::Evaluate_AnyThread(FPoseContext& Output)
 	FRetargetProfile RetargetProfileToUse = GetMergedRetargetProfile(bForceIKOff);
 
 	// run the retargeter
-	const TArray<FTransform>& RetargetedPose = Processor->RunRetargeter(
+	const TArray<FTransform>& RetargetedPose = Processor.RunRetargeter(
 		PoseToRetargetFromComponentSpace,
 		SpeedValuesFromCurves,
 		DeltaTime,
@@ -396,7 +394,7 @@ void FAnimNode_RetargetPoseFromMesh::PreUpdate(const UAnimInstance* InAnimInstan
 
 	// copy all the data from the source component
 	USkeletalMeshComponent* TargetMeshComponent = InAnimInstance->GetSkelMeshComponent();
-	if (Processor->IsInitialized() && RetargetFrom != ERetargetSourceMode::SourcePosePin)
+	if (Processor.IsInitialized() && RetargetFrom != ERetargetSourceMode::SourcePosePin)
 	{
 		GetPoseToRetargetFromSourceMesh(TargetMeshComponent);
 		CopyCurvesFromSourceMesh();
@@ -406,7 +404,7 @@ void FAnimNode_RetargetPoseFromMesh::PreUpdate(const UAnimInstance* InAnimInstan
 
 FIKRetargetProcessor* FAnimNode_RetargetPoseFromMesh::GetRetargetProcessor()
 {
-	return Processor.Get();
+	return &Processor;
 }
 
 bool FAnimNode_RetargetPoseFromMesh::EnsureProcessorIsInitialized(const USkeletalMeshComponent* TargetMeshComponent)
@@ -459,15 +457,15 @@ bool FAnimNode_RetargetPoseFromMesh::EnsureProcessorIsInitialized(const USkeleta
 	}
 	
 	// try initializing the processor
-	if (!Processor->WasInitializedWithTheseAssets(SourceMesh, TargetMesh, IKRetargeterAsset))
+	if (!Processor.WasInitializedWithTheseAssets(SourceMesh, TargetMesh, IKRetargeterAsset))
 	{
 		// initialize retarget processor with source and target skeletal meshes
 		bool bForceIKOff = false;
 		FRetargetProfile RetargetProfileToUse = GetMergedRetargetProfile(bForceIKOff);
-		Processor->Initialize(SourceMesh, TargetMesh, IKRetargeterAsset, RetargetProfileToUse);
+		Processor.Initialize(SourceMesh, TargetMesh, IKRetargeterAsset, RetargetProfileToUse);
 	}
 
-	return Processor->IsInitialized();
+	return Processor.IsInitialized();
 }
 
 FRetargetProfile FAnimNode_RetargetPoseFromMesh::GetMergedRetargetProfile(bool bForceIKOff) const
diff --git a/Engine/Plugins/Animation/IKRig/Source/IKRig/Private/IKRigLogger.cpp b/Engine/Plugins/Animation/IKRig/Source/IKRig/Private/IKRigLogger.cpp
--- a/Engine/Plugins/Animation/IKRig/Source/IKRig/Private/IKRigLogger.cpp
+++ b/Engine/Plugins/Animation/IKRig/Source/IKRig/Private/IKRigLogger.cpp
@@ -56,19 +56,31 @@ FName FIKRigLogger::GetLogTarget() const
 
 void FIKRigLogger::LogError(const FText& Message) const
 {
-	FMessageLog(LogName).SuppressLoggingToOutputLog(true).Error(Message);
+	if (LogName != NAME_None)
+	{
+		FMessageLog(LogName).SuppressLoggingToOutputLog(true).Error(Message);	
+	}
+	
 	Errors.Add(Message);
 }
 
 void FIKRigLogger::LogWarning(const FText& Message) const
 {
-	FMessageLog(LogName).SuppressLoggingToOutputLog(true).Warning(Message);
+	if (LogName != NAME_None)
+	{
+		FMessageLog(LogName).SuppressLoggingToOutputLog(true).Warning(Message);	
+	}
+	
 	Warnings.Add(Message);
 }
 
 void FIKRigLogger::LogInfo(const FText& Message) const
 {
-	FMessageLog(LogName).SuppressLoggingToOutputLog(true).Info(Message);
+	if (LogName != NAME_None)
+	{
+		FMessageLog(LogName).SuppressLoggingToOutputLog(true).Info(Message);
+	}
+	
 	Messages.Add(Message);
 }
 
diff --git a/Engine/Plugins/Animation/IKRig/Source/IKRig/Private/Retargeter/IKRetargetProcessor.cpp b/Engine/Plugins/Animation/IKRig/Source/IKRig/Private/Retargeter/IKRetargetProcessor.cpp
--- a/Engine/Plugins/Animation/IKRig/Source/IKRig/Private/Retargeter/IKRetargetProcessor.cpp
+++ b/Engine/Plugins/Animation/IKRig/Source/IKRig/Private/Retargeter/IKRetargetProcessor.cpp
@@ -1959,9 +1959,6 @@ void FIKRetargetProcessor::Initialize(
 	// record source asset
 	RetargeterAsset = InRetargeterAsset;
 
-	// set log target to unique ID for this asset
-	Log.SetLogTarget(RetargeterAsset);
-
 	const UIKRigDefinition* SourceIKRig = RetargeterAsset->GetIKRig(ERetargetSourceOrTarget::Source);
 	const UIKRigDefinition* TargetIKRig = RetargeterAsset->GetIKRig(ERetargetSourceOrTarget::Target);
 
diff --git a/Engine/Plugins/Animation/IKRig/Source/IKRig/Private/Rig/IKRigProcessor.cpp b/Engine/Plugins/Animation/IKRig/Source/IKRig/Private/Rig/IKRigProcessor.cpp
--- a/Engine/Plugins/Animation/IKRig/Source/IKRig/Private/Rig/IKRigProcessor.cpp
+++ b/Engine/Plugins/Animation/IKRig/Source/IKRig/Private/Rig/IKRigProcessor.cpp
@@ -33,9 +33,6 @@ void FIKRigProcessor::Initialize(
 	// ok, lets try to initialize
 	bTriedToInitialize = true;
 
-	// set the log target to a unique id for this asset
-	Log.SetLogTarget(InRigAsset);
-
 	// copy skeleton data from the actual skeleton we want to run on
 	Skeleton.SetInputSkeleton(SkeletalMesh, InRigAsset->GetSkeleton().ExcludedBones);
 	
diff --git a/Engine/Plugins/Animation/IKRig/Source/IKRig/Public/AnimNodes/AnimNode_IKRig.h b/Engine/Plugins/Animation/IKRig/Source/IKRig/Public/AnimNodes/AnimNode_IKRig.h
--- a/Engine/Plugins/Animation/IKRig/Source/IKRig/Public/AnimNodes/AnimNode_IKRig.h
+++ b/Engine/Plugins/Animation/IKRig/Source/IKRig/Public/AnimNodes/AnimNode_IKRig.h
@@ -92,14 +92,14 @@ public:
 	// End of FAnimNode_Base interface
 
 	// force reinitialization
-	void SetProcessorNeedsInitialized() const;
+	void SetProcessorNeedsInitialized();
 	
 	FIKRigProcessor* GetIKRigProcessor();
 
 private:
 
 	/** IK Rig runtime processor */
-	TSharedPtr<FIKRigProcessor> IKRigProcessor;
+	FIKRigProcessor IKRigProcessor;
 
 	/** a cached list of components on the owning actor that implement the goal creator interface */
 	TArray<IIKGoalCreatorInterface*> GoalCreators;
diff --git a/Engine/Plugins/Animation/IKRig/Source/IKRig/Public/AnimNodes/AnimNode_RetargetPoseFromMesh.h b/Engine/Plugins/Animation/IKRig/Source/IKRig/Public/AnimNodes/AnimNode_RetargetPoseFromMesh.h
--- a/Engine/Plugins/Animation/IKRig/Source/IKRig/Public/AnimNodes/AnimNode_RetargetPoseFromMesh.h
+++ b/Engine/Plugins/Animation/IKRig/Source/IKRig/Public/AnimNodes/AnimNode_RetargetPoseFromMesh.h
@@ -96,7 +96,7 @@ private:
 	FRetargetProfile GetMergedRetargetProfile(bool bEnableIK) const;
 
 	// the runtime processor used to run the retarget and generate new poses
-	TSharedPtr<FIKRetargetProcessor> Processor;
+	FIKRetargetProcessor Processor;
 
 	// cached transforms, copied on the game thread
 	TArray<FTransform> PoseToRetargetFromComponentSpace;
