[IKRig] Fixes retarget op stack addition/removal.

Added support for FBoneReference UI in details panel.
Fixed bug drag/drop opstack reordering causes crash

diff --git a/Engine/Plugins/Animation/IKRig/Source/IKRig/Private/Retargeter/IKRetargetProcessor.cpp b/Engine/Plugins/Animation/IKRig/Source/IKRig/Private/Retargeter/IKRetargetProcessor.cpp
--- a/Engine/Plugins/Animation/IKRig/Source/IKRig/Private/Retargeter/IKRetargetProcessor.cpp
+++ b/Engine/Plugins/Animation/IKRig/Source/IKRig/Private/Retargeter/IKRetargetProcessor.cpp
@@ -2342,7 +2342,8 @@ bool FIKRetargetProcessor::InitializeOpStack(const TArray<FInstancedStruct>& Ops
 		
 		// asset instance gets reference to running editor instance
 		#if WITH_EDITOR
-		AssetOpStruct.GetPtr<FIKRetargetOpBase>()->EditorInstance = NewOp;
+		FIKRetargetOpBase* AssetOp = const_cast<FIKRetargetOpBase*>(AssetOpStruct.GetPtr<FIKRetargetOpBase>());
+		AssetOp->GetSettings()->EditorInstance = NewOp->GetSettings();
 		#endif
 	}
 
diff --git a/Engine/Plugins/Animation/IKRig/Source/IKRig/Private/Retargeter/RetargetOps/PinBoneOp.cpp b/Engine/Plugins/Animation/IKRig/Source/IKRig/Private/Retargeter/RetargetOps/PinBoneOp.cpp
--- a/Engine/Plugins/Animation/IKRig/Source/IKRig/Private/Retargeter/RetargetOps/PinBoneOp.cpp
+++ b/Engine/Plugins/Animation/IKRig/Source/IKRig/Private/Retargeter/RetargetOps/PinBoneOp.cpp
@@ -5,7 +5,8 @@
 
 #include "IKRigObjectVersion.h"
 
-#define LOCTEXT_NAMESPACE "UPinBoneOp"
+#define LOCTEXT_NAMESPACE "PinBoneOp"
+
 
 void FPinBoneData::PostSerialize(const FArchive& Ar)
 {
@@ -60,6 +61,33 @@ void FPinBoneData::CachePinData(ERetargetSourceOrTarget InSkeletonToCopyFrom, co
     LocalRefPoseBoneToCopyTo = BoneToCopyToParentRefPoseGlobal.GetRelativeTransform(BoneToCopyToRefPoseGlobal);
 }
 
+#if WITH_EDITOR
+USkeleton* FIKRetargetPinBoneOpSettings::GetSkeleton(const FName InPropertyName)
+{
+	if (!EditorInstance)
+	{
+		return nullptr;
+	}
+
+	FIKRetargetPinBoneOpSettings* EditorSettings = reinterpret_cast<FIKRetargetPinBoneOpSettings*>(EditorInstance); 
+	
+	if (InPropertyName == GET_MEMBER_NAME_CHECKED(FPinBoneData, BoneToCopyTo))
+	{
+		const USkeleton* Skeleton = SkeletonToCopyFrom == ERetargetSourceOrTarget::Source ? EditorSettings->TargetSkeletonAsset : EditorSettings->SourceSkeletonAsset;
+		return const_cast<USkeleton*>(Skeleton);
+	}
+
+	if (InPropertyName == GET_MEMBER_NAME_CHECKED(FPinBoneData, BoneToCopyFrom))
+	{
+		const USkeleton* Skeleton = SkeletonToCopyFrom == ERetargetSourceOrTarget::Target ? EditorSettings->TargetSkeletonAsset : EditorSettings->SourceSkeletonAsset;
+		return const_cast<USkeleton*>(Skeleton);
+	}
+	
+	ensureMsgf(false, TEXT("PinBoneOp unable to get skeleton for UI widget."));
+	return nullptr;
+}
+#endif
+
 bool FIKRetargetPinBoneOp::Initialize(
 	const FIKRetargetProcessor& Processor,
 	const FRetargetSkeleton& SourceSkeleton,
@@ -114,8 +142,8 @@ bool FIKRetargetPinBoneOp::Initialize(
 	}
 
 	// record skeletons for UI bone selector widgets
-	SourceSkeletonAsset = SourceSkeleton.SkeletalMesh->GetSkeleton();
-	TargetSkeletonAsset = TargetSkeleton.SkeletalMesh->GetSkeleton();
+	Settings.SourceSkeletonAsset = SourceSkeleton.SkeletalMesh->GetSkeleton();
+	Settings.TargetSkeletonAsset = TargetSkeleton.SkeletalMesh->GetSkeleton();
 #endif
 
 	// always treat this op as "initialized", individual pins will only execute if their prerequisites are met
@@ -177,30 +205,6 @@ const UScriptStruct* FIKRetargetPinBoneOp::GetSettingsType() const
 	return FIKRetargetPinBoneOpSettings::StaticStruct();
 }
 
-USkeleton* FIKRetargetPinBoneOp::GetSkeleton(bool& bInvalidSkeletonIsError, const IPropertyHandle* PropertyHandle)
-{
-#if WITH_EDITOR
-	if (FIKRetargetPinBoneOp* RunningInstance = static_cast<FIKRetargetPinBoneOp*>(EditorInstance))
-	{
-		FName PropertyName = PropertyHandle->GetProperty()->GetFName();
-		if (PropertyName == GET_MEMBER_NAME_CHECKED(FPinBoneData, BoneToCopyTo))
-		{
-			const USkeleton* Skeleton = RunningInstance->Settings.SkeletonToCopyFrom == ERetargetSourceOrTarget::Source ? RunningInstance->TargetSkeletonAsset : RunningInstance->SourceSkeletonAsset;
-			return const_cast<USkeleton*>(Skeleton);
-		}
-
-		if (PropertyName == GET_MEMBER_NAME_CHECKED(FPinBoneData, BoneToCopyFrom))
-		{
-			const USkeleton* Skeleton = RunningInstance->Settings.SkeletonToCopyFrom == ERetargetSourceOrTarget::Target ? RunningInstance->TargetSkeletonAsset : RunningInstance->SourceSkeletonAsset;
-			return const_cast<USkeleton*>(Skeleton);
-		}
-	}
-#endif
-	
-	ensureMsgf(false, TEXT("PinBoneOp unable to get skeleton for UI widget."));
-	return nullptr;
-}
-
 FTransform FIKRetargetPinBoneOp::GetNewBoneTransform(
 	const FPinBoneData& PinData,
 	const FIKRetargetProcessor& Processor,
diff --git a/Engine/Plugins/Animation/IKRig/Source/IKRig/Public/Retargeter/IKRetargetOps.h b/Engine/Plugins/Animation/IKRig/Source/IKRig/Public/Retargeter/IKRetargetOps.h
--- a/Engine/Plugins/Animation/IKRig/Source/IKRig/Public/Retargeter/IKRetargetOps.h
+++ b/Engine/Plugins/Animation/IKRig/Source/IKRig/Public/Retargeter/IKRetargetOps.h
@@ -34,6 +34,13 @@ struct IKRIG_API FIKRetargetOpSettingsBase
 	GENERATED_BODY()
 	
 	virtual ~FIKRetargetOpSettingsBase() {};
+
+#if WITH_EDITOR
+	/** the instance of this op currently running in the editor viewport */
+	mutable FIKRetargetOpSettingsBase* EditorInstance;
+	/** allow settings to provide a skeleton for any given FBoneReference widget */
+	virtual USkeleton* GetSkeleton(const FName InPropertyName) { return nullptr; };
+#endif
 };
 
 /**
@@ -93,9 +100,7 @@ public:
 
 	/** override to display a warning message in the op stack */
 	virtual FText GetWarningMessage() const { return FText::GetEmpty(); };
-	
-	/** the instance of this op currently running in the editor viewport */
-	mutable FIKRetargetOpBase* EditorInstance;
+
 #endif
 
 private:
diff --git a/Engine/Plugins/Animation/IKRig/Source/IKRig/Public/Retargeter/RetargetOps/PinBoneOp.h b/Engine/Plugins/Animation/IKRig/Source/IKRig/Public/Retargeter/RetargetOps/PinBoneOp.h
--- a/Engine/Plugins/Animation/IKRig/Source/IKRig/Public/Retargeter/RetargetOps/PinBoneOp.h
+++ b/Engine/Plugins/Animation/IKRig/Source/IKRig/Public/Retargeter/RetargetOps/PinBoneOp.h
@@ -7,7 +7,7 @@
 
 #include "PinBoneOp.generated.h"
 
-#define LOCTEXT_NAMESPACE "RetargetOps"
+#define LOCTEXT_NAMESPACE "PinBoneOp"
 
 UENUM()
 enum class EPinBoneType : uint8
@@ -125,10 +125,16 @@ struct IKRIG_API FIKRetargetPinBoneOpSettings : public FIKRetargetOpSettingsBase
 	// A manual offset to apply in local space
 	UPROPERTY(EditAnywhere, Category=Settings)
 	FTransform LocalOffset;
+
+#if WITH_EDITOR
+	virtual USkeleton* GetSkeleton(const FName InPropertyName) override;
+	const USkeleton* SourceSkeletonAsset;
+	const USkeleton* TargetSkeletonAsset;
+#endif
 };
 
 USTRUCT(BlueprintType)
-struct IKRIG_API FIKRetargetPinBoneOp : public FIKRetargetOpBase, public IBoneReferenceSkeletonProvider
+struct IKRIG_API FIKRetargetPinBoneOp : public FIKRetargetOpBase
 {
 	GENERATED_BODY()
 
@@ -147,10 +153,6 @@ struct IKRIG_API FIKRetargetPinBoneOp : public FIKRetargetOpBase, public IBoneRe
 	
 	virtual const UScriptStruct* GetSettingsType() const override;
 
-	// IBoneReferenceSkeletonProvider overrides.
-	virtual USkeleton* GetSkeleton(bool& bInvalidSkeletonIsError, const IPropertyHandle* PropertyHandle) override;
-	// ~END IBoneReferenceSkeletonProvider overrides.
-
 	UPROPERTY()
 	FIKRetargetPinBoneOpSettings Settings;
 	
@@ -158,10 +160,6 @@ struct IKRIG_API FIKRetargetPinBoneOp : public FIKRetargetOpBase, public IBoneRe
 	virtual FText GetNiceName() const override { return FText(LOCTEXT("OpName", "Pin Bone")); };
 	virtual FText GetWarningMessage() const override { return Message; };
 	FText Message;
-
-private:
-	const USkeleton* SourceSkeletonAsset;
-	const USkeleton* TargetSkeletonAsset;
 #endif
 
 private:
diff --git a/Engine/Plugins/Animation/IKRig/Source/IKRigEditor/Private/RetargetEditor/IKRetargetEditorController.cpp b/Engine/Plugins/Animation/IKRig/Source/IKRigEditor/Private/RetargetEditor/IKRetargetEditorController.cpp
--- a/Engine/Plugins/Animation/IKRig/Source/IKRigEditor/Private/RetargetEditor/IKRetargetEditorController.cpp
+++ b/Engine/Plugins/Animation/IKRig/Source/IKRigEditor/Private/RetargetEditor/IKRetargetEditorController.cpp
@@ -669,7 +669,8 @@ void FIKRetargetEditorController::RefreshDetailsView() const
 	// refresh the details panel, cannot assume tab is not closed
 	if (DetailsView.IsValid())
 	{
-		DetailsView->ForceRefresh();
+		int32 SelectedOp = OpStackView->GetSelectedItemIndex();
+		ShowDetailsForOp(SelectedOp);
 	}
 }
 
@@ -733,16 +734,24 @@ void FIKRetargetEditorController::SetDetailsObjects(const TArray<UObject*>& Deta
 
 void FIKRetargetEditorController::ShowDetailsForOp(const int32 OpIndex) const
 {
+	UIKRetargeter* Asset = AssetController->GetAsset();
+	auto MemoryProvider = [Asset, OpIndex]() -> uint8*
+	{
+		UIKRetargeterController* Controller = UIKRetargeterController::GetController(Asset);
+		if (FIKRetargetOpBase* OpToEdit = Controller->GetRetargetOpAtIndex(OpIndex))
+		{
+			return reinterpret_cast<uint8*>(OpToEdit->GetSettings());
+		}
+		
+		return nullptr;
+	};
+	
 	FInstancedStruct* InstancedStruct = AssetController->GetRetargetOpStructAtIndex(OpIndex);
-	UStructViewer* StructViewer = AssetController->GetStructViewer();
 	const FIKRetargetOpBase* OpToEdit = InstancedStruct->GetPtr<FIKRetargetOpBase>();
 	const UScriptStruct* TypeToEdit = OpToEdit->GetSettingsType();
-	auto MemoryProvider = [](FInstancedStruct& InInstancedStruct) -> uint8*
-	{
-		FIKRetargetOpBase* OpToEdit = InInstancedStruct.GetMutablePtr<FIKRetargetOpBase>();
-		return reinterpret_cast<uint8*>(OpToEdit->GetSettings());
-	};
-	StructViewer->SetStructToView(InstancedStruct, TypeToEdit, MemoryProvider, AssetController->GetAsset());
+
+	UStructViewer* StructViewer = AssetController->GetStructViewer();
+	StructViewer->SetStructToView(TypeToEdit, MemoryProvider, AssetController->GetAsset());
 	
 	constexpr bool bForceRefresh = true; // must forcibly refresh because struct viewer is a recycled UObject
 	DetailsView->SetObject(StructViewer, bForceRefresh);
diff --git a/Engine/Plugins/Animation/IKRig/Source/IKRigEditor/Private/RetargetEditor/SRetargetOpStack.cpp b/Engine/Plugins/Animation/IKRig/Source/IKRigEditor/Private/RetargetEditor/SRetargetOpStack.cpp
--- a/Engine/Plugins/Animation/IKRig/Source/IKRigEditor/Private/RetargetEditor/SRetargetOpStack.cpp
+++ b/Engine/Plugins/Animation/IKRig/Source/IKRigEditor/Private/RetargetEditor/SRetargetOpStack.cpp
@@ -265,7 +265,7 @@ TSharedRef<SWidget> SRetargetOpStack::CreateAddNewMenuWidget()
 	MenuBuilder.BeginSection("AddNewRetargetOp", LOCTEXT("AddOperations", "Add New Retarget Op"));
 
 	// add menu option to create each retarget op type
-	const UStruct* BaseStruct = FIKRigSolverBase::StaticStruct();
+	const UStruct* BaseStruct = FIKRetargetOpBase::StaticStruct();
 	for(TObjectIterator<UStruct> StructIt; StructIt; ++StructIt)
 	{
 		UStruct* Struct = *StructIt;
diff --git a/Engine/Plugins/Animation/IKRig/Source/IKRigEditor/Private/RigEditor/IKRigDetailCustomizations.cpp b/Engine/Plugins/Animation/IKRig/Source/IKRigEditor/Private/RigEditor/IKRigDetailCustomizations.cpp
--- a/Engine/Plugins/Animation/IKRig/Source/IKRigEditor/Private/RigEditor/IKRigDetailCustomizations.cpp
+++ b/Engine/Plugins/Animation/IKRig/Source/IKRigEditor/Private/RigEditor/IKRigDetailCustomizations.cpp
@@ -11,6 +11,7 @@
 #include "Widgets/SWidget.h"
 #include "Widgets/Input/SSegmentedControl.h"
 #include "Framework/MultiBox/MultiBoxBuilder.h"
+#include "Retargeter/RetargetOps/PinBoneOp.h"
 
 #define LOCTEXT_NAMESPACE "IKRigDetailCustomizations"
 
@@ -19,6 +20,28 @@ namespace IKRigDetailCustomizationsConstants
 	static const float ItemWidth = 125.0f;
 }
 
+USkeleton* UStructViewer::GetSkeleton(bool& bInvalidSkeletonIsError, const IPropertyHandle* PropertyHandle)
+{
+	if (!StructType)
+	{
+		return nullptr;
+	}
+
+#if WITH_EDITOR
+	// NOTE: it's not ideal that we are hardcoding supported types here, but because UStruct's do not support multiple
+	// inheritance we cannot use an interface to identify skeleton providers as we normally would.
+	if (StructType->IsChildOf(FIKRetargetOpSettingsBase::StaticStruct()))
+	{
+		FName PropertyName = PropertyHandle->GetProperty()->GetFName();
+		uint8* StructMemory = StructMemoryProvider();
+		FIKRetargetOpSettingsBase* SkeletonProvider = reinterpret_cast<FIKRetargetOpSettingsBase*>(StructMemory);
+		return SkeletonProvider->GetSkeleton(PropertyName);
+	}
+#endif
+	
+	return nullptr;
+}
+
 TSharedRef<IDetailCustomization> FStructViewerCustomization::MakeInstance()
 {
 	return MakeShareable(new FStructViewerCustomization);
diff --git a/Engine/Plugins/Animation/IKRig/Source/IKRigEditor/Private/RigEditor/IKRigEditorController.cpp b/Engine/Plugins/Animation/IKRig/Source/IKRigEditor/Private/RigEditor/IKRigEditorController.cpp
--- a/Engine/Plugins/Animation/IKRig/Source/IKRigEditor/Private/RigEditor/IKRigEditorController.cpp
+++ b/Engine/Plugins/Animation/IKRig/Source/IKRigEditor/Private/RigEditor/IKRigEditorController.cpp
@@ -957,16 +957,24 @@ void FIKRigEditorController::ShowDetailsForBone(const FName BoneName) const
 
 void FIKRigEditorController::ShowDetailsForBoneSettings(const FName& BoneName, int32 SolverIndex) const
 {
-	FInstancedStruct* InstancedStruct = AssetController->GetSolverStructAtIndex(SolverIndex);
-	const UScriptStruct* TypeToEdit = InstancedStruct->GetPtr<FIKRigSolverBase>()->GetBoneSettingsType();
-	auto MemoryProvider = [BoneName](FInstancedStruct& InInstancedStruct) -> uint8*
+	const UIKRigDefinition* Asset = AssetController->GetAsset();
+	
+	auto MemoryProvider = [Asset, SolverIndex, BoneName]() -> uint8*
 	{
-		FIKRigSolverBase* BaseStruct = InInstancedStruct.GetMutablePtr<FIKRigSolverBase>();
-		return reinterpret_cast<uint8*>(BaseStruct->GetBoneSettings(BoneName));
+		const UIKRigController* Controller = UIKRigController::GetController(Asset);
+		if (FIKRigSolverBase* Solver = Controller->GetSolverAtIndex(SolverIndex))
+		{
+			return reinterpret_cast<uint8*>(Solver->GetBoneSettings(BoneName));
+		}
+		
+		return nullptr;
 	};
+
+	FInstancedStruct* InstancedStruct = AssetController->GetSolverStructAtIndex(SolverIndex);
+	const UScriptStruct* TypeToEdit = InstancedStruct->GetPtr<FIKRigSolverBase>()->GetBoneSettingsType();
 	
 	UStructViewer* StructViewer = AssetController->GetStructViewer();
-	StructViewer->SetStructToView(InstancedStruct, TypeToEdit, MemoryProvider, AssetController->GetAsset());
+	StructViewer->SetStructToView(TypeToEdit, MemoryProvider, AssetController->GetAsset());
 	
 	ShowStructViewerInDetailsView(StructViewer);
 }
@@ -978,32 +986,48 @@ void FIKRigEditorController::ShowDetailsForGoal(const FName& GoalName) const
 
 void FIKRigEditorController::ShowDetailsForGoalSettings(const FName GoalName, const int32 SolverIndex) const
 {
-	FInstancedStruct* InstancedStruct = AssetController->GetSolverStructAtIndex(SolverIndex);
-	const UScriptStruct* TypeToEdit = InstancedStruct->GetPtr<FIKRigSolverBase>()->GetGoalSettingsType();
-	auto MemoryProvider = [GoalName](FInstancedStruct& InInstancedStruct) -> uint8*
+	const UIKRigDefinition* Asset = AssetController->GetAsset();
+	
+	auto MemoryProvider = [Asset, SolverIndex, GoalName]() -> uint8*
 	{
-		FIKRigSolverBase* BaseStruct = InInstancedStruct.GetMutablePtr<FIKRigSolverBase>();
-		return reinterpret_cast<uint8*>(BaseStruct->GetGoalSettings(GoalName));
+		const UIKRigController* Controller = UIKRigController::GetController(Asset);
+		if (FIKRigSolverBase* Solver = Controller->GetSolverAtIndex(SolverIndex))
+		{
+			return reinterpret_cast<uint8*>(Solver->GetGoalSettings(GoalName));
+		}
+		
+		return nullptr;
 	};
 	
+	FInstancedStruct* InstancedStruct = AssetController->GetSolverStructAtIndex(SolverIndex);
+	const UScriptStruct* TypeToEdit = InstancedStruct->GetPtr<FIKRigSolverBase>()->GetGoalSettingsType();
+	
 	UStructViewer* StructViewer = AssetController->GetStructViewer();
-	StructViewer->SetStructToView(InstancedStruct, TypeToEdit, MemoryProvider, AssetController->GetAsset());
+	StructViewer->SetStructToView(TypeToEdit, MemoryProvider, AssetController->GetAsset());
 	
 	ShowStructViewerInDetailsView(StructViewer);
 }
 
 void FIKRigEditorController::ShowDetailsForSolver(const int32 SolverIndex) const
 {
-	FInstancedStruct* InstancedStruct = AssetController->GetSolverStructAtIndex(SolverIndex);
-	const UScriptStruct* TypeToEdit = InstancedStruct->GetPtr<FIKRigSolverBase>()->GetSolverSettingsType();
-	auto MemoryProvider = [](FInstancedStruct& InInstancedStruct) -> uint8*
+	const UIKRigDefinition* Asset = AssetController->GetAsset();
+	
+	auto MemoryProvider = [Asset, SolverIndex]() -> uint8*
 	{
-		FIKRigSolverBase* BaseStruct = InInstancedStruct.GetMutablePtr<FIKRigSolverBase>();
-		return reinterpret_cast<uint8*>(BaseStruct->GetSolverSettings());
+		const UIKRigController* Controller = UIKRigController::GetController(Asset);
+		if (FIKRigSolverBase* Solver = Controller->GetSolverAtIndex(SolverIndex))
+		{
+			return reinterpret_cast<uint8*>(Solver->GetSolverSettings());
+		}
+		
+		return nullptr;
 	};
 	
+	FInstancedStruct* InstancedStruct = AssetController->GetSolverStructAtIndex(SolverIndex);
+	const UScriptStruct* TypeToEdit = InstancedStruct->GetPtr<FIKRigSolverBase>()->GetSolverSettingsType();
+	
 	UStructViewer* StructViewer = AssetController->GetStructViewer();
-	StructViewer->SetStructToView(InstancedStruct, TypeToEdit, MemoryProvider, AssetController->GetAsset());
+	StructViewer->SetStructToView(TypeToEdit, MemoryProvider, AssetController->GetAsset());
 	
 	ShowStructViewerInDetailsView(StructViewer);
 }
diff --git a/Engine/Plugins/Animation/IKRig/Source/IKRigEditor/Public/RigEditor/IKRigDetailCustomizations.h b/Engine/Plugins/Animation/IKRig/Source/IKRigEditor/Public/RigEditor/IKRigDetailCustomizations.h
--- a/Engine/Plugins/Animation/IKRig/Source/IKRigEditor/Public/RigEditor/IKRigDetailCustomizations.h
+++ b/Engine/Plugins/Animation/IKRig/Source/IKRigEditor/Public/RigEditor/IKRigDetailCustomizations.h
@@ -5,13 +5,15 @@
 #include "CoreMinimal.h"
 #include "IDetailCustomization.h"
 #include "ScopedTransaction.h"
+#include "Interfaces/Interface_BoneReferenceSkeletonProvider.h"
 #include "RigEditor/IKRigEditorController.h"
 
 #include "IKRigDetailCustomizations.generated.h"
 
+
 // a thin wrapper around UStruct data to display in a details panel
 UCLASS(Blueprintable)
-class UStructViewer : public UObject
+class UStructViewer : public UObject, public IBoneReferenceSkeletonProvider
 {
 	GENERATED_BODY()
 
@@ -22,18 +24,15 @@ public:
 	 * NOTE: we can't just pass in raw pointers to the struct memory because these can be destroyed after a transaction.
 	 * So instead we pass in TFunction callbacks that get the latest memory locations when the details panel is refreshed.
 	 *
-	 * @param InInstancedStruct A pointer to the instanced struct that contains the struct to be displayed, this pointer survives undo/redo
-	 * @param InScriptStruct The type that corresponds to the struct memory returned by the StructMemoryProvider
+	 * @param InTypeToEdit The type that corresponds to the struct memory returned by the StructMemoryProvider
 	 * @param InStructMemoryProvider A callback function that retrieves the memory location for the struct to be edited
 	 * @param InStructOwner The UObject that owns the struct. This is what will be transacted when the property is edited.
 	 */
 	void SetStructToView(
-		FInstancedStruct* InInstancedStruct,
 		const UScriptStruct* InTypeToEdit,
-		TFunction<uint8*(FInstancedStruct& InInstancedStruct)> InStructMemoryProvider,
+		TFunction<uint8*()> InStructMemoryProvider,
 		UObject* InStructOwner)
 	{
-		TopLevelInstancedStruct = InInstancedStruct;
 		StructType = InTypeToEdit;
 		StructMemoryProvider = MoveTemp(InStructMemoryProvider);
 		StructOwner = InStructOwner;
@@ -41,12 +40,12 @@ public:
 
 	bool IsValid() const
 	{
-		return TopLevelInstancedStruct->IsValid() && StructType && StructOwner.IsValid();
+		uint8* Memory = StructMemoryProvider();
+		return Memory != nullptr && StructType && StructOwner.IsValid();
 	}
 
 	void Reset()
 	{
-		TopLevelInstancedStruct = nullptr;
 		StructType = nullptr;
 		StructMemoryProvider.Reset();
 		StructOwner = nullptr;
@@ -54,8 +53,8 @@ public:
 
 	TSharedPtr<FStructOnScope>& GetStructOnScope()
 	{
-		uint8* StructMemory = StructMemoryProvider(*TopLevelInstancedStruct);
-		StructOnScope = MakeShared<FStructOnScope>(StructType, StructMemory);
+		uint8* Memory = StructMemoryProvider();
+		StructOnScope = MakeShared<FStructOnScope>(StructType, Memory);
 		return StructOnScope;
 	}
 
@@ -69,14 +68,16 @@ public:
 		return StructOwner.Get();
 	};
 
+	// IBoneReferenceSkeletonProvider overrides.
+	virtual USkeleton* GetSkeleton(bool& bInvalidSkeletonIsError, const IPropertyHandle* PropertyHandle) override;
+	// ~END IBoneReferenceSkeletonProvider overrides.
+
 private:
 	
-	// we need the type reflection data and the memory to build a StructOnScope to display in details panel
-	FInstancedStruct* TopLevelInstancedStruct;
 	// the type of struct we are editing
 	const UScriptStruct* StructType;
 	// provides the memory address of the struct to edit (refreshed after undo/redo)
-	TFunction<uint8*(FInstancedStruct& InInstancedStruct)> StructMemoryProvider;
+	TFunction<uint8*()> StructMemoryProvider;
 	
 	// a wrapper of the struct sent to the details panel
 	TSharedPtr<FStructOnScope> StructOnScope;
diff --git a/Engine/Source/Editor/Persona/Private/AnimGraphNodeDetails.cpp b/Engine/Source/Editor/Persona/Private/AnimGraphNodeDetails.cpp
--- a/Engine/Source/Editor/Persona/Private/AnimGraphNodeDetails.cpp
+++ b/Engine/Source/Editor/Persona/Private/AnimGraphNodeDetails.cpp
@@ -43,6 +43,7 @@
 #include "IPersonaToolkit.h"
 #include "Interfaces/Interface_BoneReferenceSkeletonProvider.h"
 #include "IPropertyAccessEditor.h"
+#include "IPropertyUtilities.h"
 #include "Algo/Accumulate.h"
 #include "Kismet2/BlueprintEditorUtils.h"
 #include "Widgets/Input/SEditableTextBox.h"
@@ -512,6 +513,10 @@ void FBoneReferenceCustomization::CustomizeHeader( TSharedRef<IPropertyHandle> S
 	
 	BoneNameProperty = FindStructMemberProperty(StructProperty, GET_MEMBER_NAME_CHECKED(FBoneReference, BoneName));
 	check(BoneNameProperty->IsValidHandle());
+
+	// record selected objects in case they are needed to find the skeleton
+	IPropertyUtilities* PropertyUtility = StructCustomizationUtils.GetPropertyUtilities().Get();
+	SelectedObjects = PropertyUtility->GetSelectedObjects();
 	
 	HeaderRow
 	.NameContent()
@@ -561,6 +566,15 @@ USkeleton* FBoneReferenceCustomization::GetSkeleton() const
 	TArray<UObject*> Objects;
 	StructProperty->GetOuterObjects(Objects);
 
+	// optionally add selected objects which may manage the skeleton provider
+	if (Objects.IsEmpty())
+	{
+		for (const TWeakObjectPtr<UObject>& SelectedObject : SelectedObjects)
+		{
+			Objects.Add(SelectedObject.Get());
+		}
+	}
+
 	auto FindSkeletonForObject = [this](UObject* InObject) -> USkeleton* 
 	{
 		for( ; InObject; InObject = InObject->GetOuter())
diff --git a/Engine/Source/Editor/Persona/Private/AnimGraphNodeDetails.h b/Engine/Source/Editor/Persona/Private/AnimGraphNodeDetails.h
--- a/Engine/Source/Editor/Persona/Private/AnimGraphNodeDetails.h
+++ b/Engine/Source/Editor/Persona/Private/AnimGraphNodeDetails.h
@@ -122,6 +122,8 @@ protected:
 	// Property to change after bone has been picked
 	TSharedPtr<IPropertyHandle> BoneNameProperty;
 
+	TArray<TWeakObjectPtr<UObject>> SelectedObjects;
+
 private:
 
 	// Bone tree widget delegates
