[IKRig] Fixes details panel crash in retargeter.

diff --git a/Engine/Plugins/Animation/IKRig/Source/IKRigEditor/Private/RetargetEditor/IKRetargetEditorController.cpp b/Engine/Plugins/Animation/IKRig/Source/IKRigEditor/Private/RetargetEditor/IKRetargetEditorController.cpp
--- a/Engine/Plugins/Animation/IKRig/Source/IKRigEditor/Private/RetargetEditor/IKRetargetEditorController.cpp
+++ b/Engine/Plugins/Animation/IKRig/Source/IKRigEditor/Private/RetargetEditor/IKRetargetEditorController.cpp
@@ -734,6 +734,12 @@ void FIKRetargetEditorController::SetDetailsObjects(const TArray<UObject*>& Deta
 
 void FIKRetargetEditorController::ShowDetailsForOp(const int32 OpIndex) const
 {
+	if (AssetController->GetRetargetOpAtIndex(OpIndex) == nullptr)
+	{
+		DetailsView->SetObject(AssetController->GetAsset());
+		return;
+	}
+	
 	UIKRetargeter* Asset = AssetController->GetAsset();
 	auto MemoryProvider = [Asset, OpIndex]() -> uint8*
 	{
