[IKRig] Major refactor part 3 of 3.

Converted all retargeting features into equivalent modular "ops".

diff --git a/Engine/Plugins/Animation/IKRig/Source/IKRig/Private/AnimNodes/AnimNode_IKRig.cpp b/Engine/Plugins/Animation/IKRig/Source/IKRig/Private/AnimNodes/AnimNode_IKRig.cpp
--- a/Engine/Plugins/Animation/IKRig/Source/IKRig/Private/AnimNodes/AnimNode_IKRig.cpp
+++ b/Engine/Plugins/Animation/IKRig/Source/IKRig/Private/AnimNodes/AnimNode_IKRig.cpp
@@ -49,7 +49,7 @@ void FAnimNode_IKRig::Evaluate_AnyThread(FPoseContext& Output)
 	// initialize the IK Rig (will only try once on the current version of the rig asset)
 	if (!IKRigProcessor.IsInitialized())
 	{
-		IKRigProcessor.Initialize(RigDefinitionAsset, Output.AnimInstanceProxy->GetSkelMeshComponent()->GetSkeletalMeshAsset());
+		IKRigProcessor.Initialize(RigDefinitionAsset, Output.AnimInstanceProxy->GetSkelMeshComponent()->GetSkeletalMeshAsset(),FIKRigGoalContainer());
 	}
 	
 	// copy input pose to solver stack
diff --git a/Engine/Plugins/Animation/IKRig/Source/IKRig/Private/AnimNodes/AnimNode_RetargetPoseFromMesh.cpp b/Engine/Plugins/Animation/IKRig/Source/IKRig/Private/AnimNodes/AnimNode_RetargetPoseFromMesh.cpp
--- a/Engine/Plugins/Animation/IKRig/Source/IKRig/Private/AnimNodes/AnimNode_RetargetPoseFromMesh.cpp
+++ b/Engine/Plugins/Animation/IKRig/Source/IKRig/Private/AnimNodes/AnimNode_RetargetPoseFromMesh.cpp
@@ -87,15 +87,24 @@ void FAnimNode_RetargetPoseFromMesh::Evaluate_AnyThread(FPoseContext& Output)
 			return false; // must be initialized
 		}
 
+		// validate source data from other component
 		const bool bCopyingFromOtherComponent = RetargetFrom != ERetargetSourceMode::SourcePosePin;
-		if (bCopyingFromOtherComponent && !SourceMeshComponent.IsValid())
+		if (bCopyingFromOtherComponent)
 		{
-			return false; // if copying from another mesh, must have valid pointer to it
-		}
+			if (!SourceMeshComponent.IsValid())
+			{
+				return false; // copying from another mesh but no source component was found
+			}
 
-		if (bCopyingFromOtherComponent && PoseToRetargetFromComponentSpace.IsEmpty())
-		{
-			return false; // if copying from another mesh, pose must be copied already from the PreUpdate()
+			if (!SourceMeshComponent->GetSkeletalMeshAsset())
+			{
+				return false; // if copying from another mesh, must have valid pointer to it
+			}
+			
+			if (PoseToRetargetFromComponentSpace.IsEmpty())
+			{
+				return false; // if copying from another mesh, pose must be copied already from the PreUpdate()
+			}
 		}
 
 		if (!bCopyingFromOtherComponent && !Source.GetLinkNode())
@@ -148,24 +157,8 @@ void FAnimNode_RetargetPoseFromMesh::Evaluate_AnyThread(FPoseContext& Output)
 		// convert to component space
 		FAnimationRuntime::FillUpComponentSpaceTransforms(TargetRefSkeleton, InputLocalTransforms, PoseToRetargetFromComponentSpace);
 
-		// update speed values from the main anim graph curves
-		{
-			SpeedValuesFromCurves.Reset();
-			const TArray<TObjectPtr<URetargetChainSettings>>& ChainSettings = IKRetargeterAsset->GetAllChainSettings();
-			for (const TObjectPtr<URetargetChainSettings>& ChainSetting : ChainSettings)
-			{
-				FName CurveName = ChainSetting->Settings.SpeedPlanting.SpeedCurveName;
-				if (CurveName == NAME_None)
-				{
-					continue;
-				}
-
-				// value will be negative if the curve was not found (retargeter ignores negative speeds)
-				bool bFoundCurve = false;
-				float CurveValue = Output.Curve.Get(CurveName, bFoundCurve, -1.0f);
-				SpeedValuesFromCurves.Add(CurveName, CurveValue);
-			}
-		}
+		// let the retargeter scale the input pose
+		Processor.ScaleSourcePose(PoseToRetargetFromComponentSpace);
 	};
 	
 	auto ApplyRetargetedPose = [this, &Output](const TArray<FTransform>& RetargetedPose)
@@ -201,47 +194,6 @@ void FAnimNode_RetargetPoseFromMesh::Evaluate_AnyThread(FPoseContext& Output)
 		}
 	};
 
-	// remap curves as required by any retarget ops that request it
-	auto CopyAndRemapCurvesFromSourceToTarget = [this](FBlendedCurve& OutputCurves)
-	{
-		if (!bCopyCurves)
-		{
-			return;
-		}
-	
-		if (!Processor.IsInitialized())
-		{
-			return;
-		}
-
-		// copy curves over to same name on target (if it exists)
-		OutputCurves.CopyFrom(SourceCurves);
-	
-		// now query retarget ops for any curves that need remapped from source to target
-		FBlendedCurve RemapedCurves;
-		TArray<FCurveRemapPair> CurvesToRetarget;
-		const TArray<TSharedPtr<FIKRetargetOpBase>>& RetargetOps = Processor.GetRetargetOps();
-		for (const TSharedPtr<FIKRetargetOpBase>& RetargetOp : RetargetOps)
-		{
-			if (RetargetOp->IsEnabled())
-			{
-				RetargetOp->GetCurvesToRetarget(CurvesToRetarget);	
-			}
-		}
-		
-		for (const FCurveRemapPair& CurveToRemap : CurvesToRetarget)
-		{
-			bool bOutIsValid = false;
-			const float SourceValue = SourceCurves.Get(CurveToRemap.SourceCurve, bOutIsValid);
-			if (bOutIsValid)
-			{
-				RemapedCurves.Add(CurveToRemap.TargetCurve, SourceValue);
-			}
-		}
-
-		OutputCurves.Combine(RemapedCurves);
-	};
-
 	// check the configuration of the node and output the reference pose unless it's ready to retarget
 	if (!IsReadyToRun())
 	{
@@ -255,34 +207,21 @@ void FAnimNode_RetargetPoseFromMesh::Evaluate_AnyThread(FPoseContext& Output)
 		FillPoseToRetargetFromInputPin(); // (also gets speed curves from anim graph)
 	}
 
-	#if WITH_EDITOR
-	if (GIsEditor)
-	{
-		// live preview source asset settings in the retarget, editor only
-		// NOTE: this copies goal targets as well, but these are overwritten by IK chain goals
-		Processor.CopyIKRigSettingsFromAsset();
-	}
-	#endif
-
 	// LOD off the IK pass
 	bool bForceIKOff = LODThresholdForIK != INDEX_NONE && Output.AnimInstanceProxy->GetLODLevel() > LODThresholdForIK;
 	FRetargetProfile RetargetProfileToUse = GetMergedRetargetProfile(bForceIKOff);
 
+	// give retarget ops a chance to access data from the anim graph before running
+	Processor.OnAnimGraphEvaluateAnyThread(Output);
+
 	// run the retargeter
-	const TArray<FTransform>& RetargetedPose = Processor.RunRetargeter(
-		PoseToRetargetFromComponentSpace,
-		SpeedValuesFromCurves,
-		DeltaTime,
-		RetargetProfileToUse);
+	const TArray<FTransform>& RetargetedPose = Processor.RunRetargeter(PoseToRetargetFromComponentSpace, RetargetProfileToUse, DeltaTime);
 
 	// reset delta time (it's accumulated)
 	DeltaTime = 0.f;
 
 	// apply the retargeted pose to the output pin
 	ApplyRetargetedPose(RetargetedPose);
-
-	// copy and/or remap curves from the source to the target skeletal mesh
-	CopyAndRemapCurvesFromSourceToTarget(Output.Curve);
 }
 
 void FAnimNode_RetargetPoseFromMesh::PreUpdate(const UAnimInstance* InAnimInstance)
@@ -343,53 +282,9 @@ void FAnimNode_RetargetPoseFromMesh::PreUpdate(const UAnimInstance* InAnimInstan
 		{
 			Transform.SetScale3D(FVector::OneVector);
 		}
-	};
-
-	auto CopyCurvesFromSourceMesh = [this]()
-	{
-		if(!bCopyCurves)
-		{
-			return;
-		}
-
-		const UAnimInstance* SourceAnimInstance = SourceMeshComponent->GetAnimInstance();
-		if (!SourceAnimInstance)
-		{
-			return;
-		}
 		
-		// Potential optimization/tradeoff: If we stored the curve results on the mesh component in non-editor scenarios, this would be
-		// much faster (but take more memory). As it is, we need to translate the map stored on the anim instance.
-		const TMap<FName, float>& AnimCurveList = SourceAnimInstance->GetAnimationCurveList(EAnimCurveType::AttributeCurve);
-		SourceCurves.Empty();
-		UE::Anim::FCurveUtils::BuildUnsorted(SourceCurves, AnimCurveList);
-	};
-
-	auto UpdateSpeedFromSourceMeshCurves = [this]()
-	{
-		SpeedValuesFromCurves.Reset();
-
-		const UAnimInstance* SourceAnimInstance = SourceMeshComponent->GetAnimInstance();
-		if (!SourceAnimInstance)
-		{
-			return;
-		}
-		
-		const TMap<FName, float>& AnimCurveList = SourceAnimInstance->GetAnimationCurveList(EAnimCurveType::AttributeCurve);
-		const TArray<TObjectPtr<URetargetChainSettings>>& ChainSettings = IKRetargeterAsset->GetAllChainSettings();
-		for (const TObjectPtr<URetargetChainSettings>& ChainSetting : ChainSettings)
-		{
-			FName CurveName = ChainSetting->Settings.SpeedPlanting.SpeedCurveName;
-			if (CurveName == NAME_None)
-			{
-				continue;
-			}
-			
-			const float* SourceValue = AnimCurveList.Find(CurveName);
-			constexpr float NotFoundValue = -1.0f; // value will be negative if the curve was not found (retargeter ignores negative speeds)
-			const float CurveValue = SourceValue ? *SourceValue : NotFoundValue;
-			SpeedValuesFromCurves.Add(CurveName, CurveValue);
-		}
+		// let the retargeter scale the input pose
+		Processor.ScaleSourcePose(PoseToRetargetFromComponentSpace);
 	};
 
 	// copy all the data from the source component
@@ -397,8 +292,8 @@ void FAnimNode_RetargetPoseFromMesh::PreUpdate(const UAnimInstance* InAnimInstan
 	if (Processor.IsInitialized() && RetargetFrom != ERetargetSourceMode::SourcePosePin)
 	{
 		GetPoseToRetargetFromSourceMesh(TargetMeshComponent);
-		CopyCurvesFromSourceMesh();
-		UpdateSpeedFromSourceMeshCurves();
+		
+		Processor.OnAnimGraphPreUpdateMainThread(*SourceMeshComponent.Get(), *TargetMeshComponent);
 	}
 }
 
@@ -472,16 +367,12 @@ FRetargetProfile FAnimNode_RetargetPoseFromMesh::GetMergedRetargetProfile(bool b
 {
 	// collect settings to retarget with starting with asset settings and overriding with custom profile
 	FRetargetProfile Profile;
-	IKRetargeterAsset->FillProfileWithAssetSettings(Profile);
+	Profile.FillProfileWithAssetSettings(IKRetargeterAsset);
 	// load custom profile plugged into the anim node
 	Profile.MergeWithOtherProfile(CustomRetargetProfile);
 	// force all IK off (skips IK solve)
-	if (bForceIKOff)
-	{
-		Profile.GlobalSettings.bEnableIK = false;
-	}
-
-	return Profile;
+	Profile.bForceAllIKOff = bForceIKOff;
+	return MoveTemp(Profile);
 }
 
 bool FAnimNode_RetargetPoseFromMesh::Serialize(FArchive& Ar)
diff --git a/Engine/Plugins/Animation/IKRig/Source/IKRig/Private/Retargeter/IKRetargetOps.cpp b/Engine/Plugins/Animation/IKRig/Source/IKRig/Private/Retargeter/IKRetargetOps.cpp
--- a/Engine/Plugins/Animation/IKRig/Source/IKRig/Private/Retargeter/IKRetargetOps.cpp
+++ b/Engine/Plugins/Animation/IKRig/Source/IKRig/Private/Retargeter/IKRetargetOps.cpp
@@ -1,3 +1,113 @@
 // Copyright Epic Games, Inc. All Rights Reserved.
 
 #include "Retargeter/IKRetargetOps.h"
+
+#define LOCTEXT_NAMESPACE "RetargetOpBase"
+
+void FIKRetargetOpBase::SetSettings(const FIKRetargetOpSettingsBase* InSettings)
+{
+	// copy everything by default
+	static TArray<FName> PropertiesToIgnore;
+	CopySettingsRaw(InSettings, PropertiesToIgnore);
+}
+
+#if WITH_EDITOR
+FText FIKRetargetOpBase::GetNiceName() const
+{
+	// calls virtual GetType() to get the derived type
+	const UScriptStruct* ScriptStruct = GetType();
+
+	// try to get the "nice name" from metadata
+	FString TypeName;
+	if (ScriptStruct->HasMetaData(TEXT("DisplayName")))
+	{
+		TypeName = ScriptStruct->GetMetaData(TEXT("DisplayName"));
+	}
+
+	// if no "DisplayName" metadata is found, fall back to the struct name
+	if (TypeName.IsEmpty())
+	{
+		TypeName = ScriptStruct->GetName();
+	}
+	
+	return FText::FromString(TypeName);
+}
+
+FText FIKRetargetOpBase::GetWarningMessage() const
+{
+	if (bIsInitialized)
+	{
+		return bIsEnabled ? LOCTEXT("OpReadyAndOn", "Running.") : LOCTEXT("OpReadyAndOff", "Ready, but disabled.");
+	}
+	
+	return LOCTEXT("OpNotReady", "Not initialized. See output log.");
+}
+#endif
+
+UIKRetargetOpControllerBase* FIKRetargetOpBase::GetController(UObject* Outer)
+{
+	return CreateControllerIfNeeded(Outer);
+}
+
+const UClass* FIKRetargetOpBase::GetControllerType() const
+{
+	return UIKRetargetOpControllerBase::StaticClass();
+}
+
+void FIKRetargetOpBase::CopySettingsRaw(const FIKRetargetOpSettingsBase* InSettings, const TArray<FName>& InPropertiesToIgnore)
+{
+	// copy everything except any property that is filtered
+	FIKRetargetOpBase::CopyStructProperties(
+		GetSettingsType(),
+		InSettings,
+		GetSettings(),
+		InPropertiesToIgnore);
+}
+
+void FIKRetargetOpBase::CopyStructProperties(
+	const UStruct* InStructType,
+	const void* InSrcStruct,
+	void* InOutDestStruct,
+	const TArray<FName>& InPropertiesToIgnore)
+{
+	if (!InStructType || !InSrcStruct || !InOutDestStruct)
+	{
+		return;
+	}
+
+	for (TFieldIterator<FProperty> PropIt(InStructType); PropIt; ++PropIt)
+	{
+		FProperty* Property = *PropIt;
+		if (!Property)
+		{
+			continue;
+		}
+
+		// check if this property should be ignored
+		if (InPropertiesToIgnore.Contains(Property->GetFName()))
+		{
+			continue;
+		}
+
+		// copy the property value from Src to Dest
+		Property->CopyCompleteValue(
+			Property->ContainerPtrToValuePtr<void>(InOutDestStruct),
+			Property->ContainerPtrToValuePtr<void>(InSrcStruct));
+	}
+}
+
+UIKRetargetOpControllerBase* FIKRetargetOpBase::CreateControllerIfNeeded(UObject* Outer)
+{
+	if (!Controller.IsValid())
+	{
+		const UClass* ClassType = GetControllerType();
+		if (ensure(ClassType && ClassType->IsChildOf(UIKRetargetOpControllerBase::StaticClass())))
+		{
+			Controller = TStrongObjectPtr(NewObject<UIKRetargetOpControllerBase>(Outer, ClassType));
+			Controller->OpToControl = this;
+		}
+	}
+	return Controller.Get();
+}
+
+#undef LOCTEXT_NAMESPACE
