[IKRig] Major refactor part 3 of 3.

diff --git a/Engine/Plugins/Animation/IKRig/Source/IKRig/Private/Retargeter/IKRetargetProcessor.cpp b/Engine/Plugins/Animation/IKRig/Source/IKRig/Private/Retargeter/IKRetargetProcessor.cpp
--- a/Engine/Plugins/Animation/IKRig/Source/IKRig/Private/Retargeter/IKRetargetProcessor.cpp
+++ b/Engine/Plugins/Animation/IKRig/Source/IKRig/Private/Retargeter/IKRetargetProcessor.cpp
@@ -4,15 +4,15 @@
 
 #include "IKRigLogger.h"
 #include "Rig/IKRigDefinition.h"
-#include "Rig/Solvers/PointsToRotation.h"
-#include "Rig/IKRigProcessor.h"
 
 #include "Retargeter/IKRetargeter.h"
 #include "Retargeter/IKRetargetProfile.h"
 #include "Retargeter/IKRetargetOps.h"
+#include "Retargeter/RetargetOps/PelvisMotionOp.h"
 
 #include "Engine/SkeletalMesh.h"
-#include "Kismet/KismetMathLibrary.h"
+#include "Retargeter/RetargetOps/CopyBasePoseOp.h"
+#include "Retargeter/RetargetOps/ScaleSourceOp.h"
 
 #include UE_INLINE_GENERATED_CPP_BY_NAME(IKRetargetProcessor)
 
@@ -27,7 +27,7 @@ void FRetargetSkeleton::Initialize(
 	const FName InRetargetPoseName,
 	const FIKRetargetPose* RetargetPose,
 	const FName RetargetRootBone,
-	const float ScaleFactor)
+	const double ScaleFactor)
 {
 	// reset all skeleton data
 	Reset();
@@ -43,21 +43,6 @@ void FRetargetSkeleton::Initialize(
 		ParentIndices.Add(RefSkeleton.GetParentIndex(BoneIndex));	
 	}
 
-	// determine set of bones referenced by one of the retarget bone chains
-	// this is the set of bones that will be affected by the retarget pose
-	ChainThatContainsBone.Init(NAME_None, BoneNames.Num());
-	for (const FBoneChain& BoneChain : BoneChains)
-	{
-		TArray<int32> BonesInChain;
-		if (FResolvedBoneChain(BoneChain, *this, BonesInChain).IsValid())
-		{
-			for (const int32 BoneInChain : BonesInChain)
-			{
-				ChainThatContainsBone[BoneInChain] = BoneChain.ChainName;
-			}
-		}
-	}
-
 	// initialize branch caching
 	CachedEndOfBranchIndices.Init(RETARGETSKELETON_INVALID_BRANCH_INDEX, ParentIndices.Num());
 
@@ -66,6 +51,23 @@ void FRetargetSkeleton::Initialize(
 	GenerateRetargetPose(InRetargetPoseName, RetargetPose, RetargetRootBone, ScaleFactor);
 }
 
+void FRetargetSkeleton::FillPerBoneChainNames(const TArray<FResolvedBoneChain>& InBoneChains)
+{
+	// record array of FName per-bone specifying which chain each bone belongs to
+
+	// set all to NAME_None by default
+	ChainThatContainsBone.Init(NAME_None, BoneNames.Num());
+
+	// spin through all bones on all the resolved chains on this skeleton
+	for (const FResolvedBoneChain& BoneChain : InBoneChains)
+	{
+		for (const int32 BoneIndex : BoneChain.BoneIndices)
+		{
+			ChainThatContainsBone[BoneIndex] = BoneChain.ChainName;
+		}
+	}
+}
+
 void FRetargetSkeleton::Reset()
 {
 	BoneNames.Reset();
@@ -79,7 +81,7 @@ void FRetargetSkeleton::GenerateRetargetPose(
 	const FName InRetargetPoseName,
 	const FIKRetargetPose* InRetargetPose,
 	const FName RetargetRootBone,
-	const float ScaleFactor)
+	const double ScaleFactor)
 {
 	// record the name of the retarget pose (prevents re-initialization if profile swaps it)
 	RetargetPoseName = InRetargetPoseName;
@@ -349,12 +351,13 @@ void FTargetSkeleton::Initialize(
 {
 	Reset();
 
-	constexpr float TargetSkeletonScaleFactor = 1.0f;
+	constexpr double TargetSkeletonScaleFactor = 1.0;
 	FRetargetSkeleton::Initialize(InSkeletalMesh, BoneChains, InRetargetPoseName, RetargetPose, RetargetRootBone, TargetSkeletonScaleFactor);
 
 	// make storage for per-bone "Is Retargeted" flag (used for hierarchy updates)
 	// these are bones that are in a target chain that is mapped to a source chain (ie, will actually be retargeted)
-	// these flags are actually set later in init phase when bone chains are mapped together
+	// these flags are actually set later between Op::Initialize() and Op::PostInitialize()
+	bIsMaskInitialized = false;
 	IsBoneRetargeted.Init(false, BoneNames.Num());
 
 	// initialize pose for base local pose (may be overridden by pose copier, but uses retarget pose by default)
@@ -368,119 +371,176 @@ void FTargetSkeleton::Reset()
 {
 	FRetargetSkeleton::Reset();
 	OutputGlobalPose.Reset();
-	IsBoneRetargeted.Reset();
 }
 
-void FTargetSkeleton::UpdateGlobalTransformsAllNonRetargetedBones(const TArray<FTransform>& LocalPose, TArray<FTransform>& InOutGlobalPose)
+void FTargetSkeleton::UpdateGlobalTransformsWithMask(
+	const TArray<bool>& InBoneMask,
+	const TArray<FTransform>& InLocalPose,
+	TArray<FTransform>& InOutGlobalPose) const
 {
-	check(IsBoneRetargeted.Num() == InOutGlobalPose.Num());
+	check(InBoneMask.Num() == InOutGlobalPose.Num());
 	
 	for (int32 BoneIndex=0; BoneIndex<InOutGlobalPose.Num(); ++BoneIndex)
 	{
-		if (!IsBoneRetargeted[BoneIndex])
+		if (!InBoneMask[BoneIndex])
 		{
-			UpdateGlobalTransformOfSingleBone(BoneIndex, LocalPose, InOutGlobalPose);
+			UpdateGlobalTransformOfSingleBone(BoneIndex, InLocalPose, InOutGlobalPose);
 		}
 	}
 }
 
+void FTargetSkeleton::SetRetargetedBones(const TSet<int32>& InRetargetedBones)
+{
+	IsBoneRetargeted.Init(false, BoneNames.Num());
+	for (int32 BoneIndex : InRetargetedBones)
+	{
+		IsBoneRetargeted[BoneIndex] = true;
+	}
+	bIsMaskInitialized = true;
+}
+
+bool FTargetSkeleton::GetIsBoneRetargeted(int32 InBoneIndex) const
+{
+	if (!ensure(bIsMaskInitialized))
+	{
+		return false;
+	}
+	
+	return IsBoneRetargeted[InBoneIndex];
+}
+
+const TArray<bool>& FTargetSkeleton::GetRetargetedBonesMask() const
+{
+	return IsBoneRetargeted;
+}
+
 FResolvedBoneChain::FResolvedBoneChain(
-	const FBoneChain& BoneChain,
-	const FRetargetSkeleton& Skeleton,
-	TArray<int32>& OutBoneIndices)
+	const FBoneChain& InBoneChain,
+	const FRetargetSkeleton& InSkeleton,
+	FIKRigLogger& InLog)
 {
+	// store the bone chain data
+	ChainName = InBoneChain.ChainName;
+	StartBone = InBoneChain.StartBone.BoneName;
+	EndBone = InBoneChain.EndBone.BoneName;
+	IKGoalName = InBoneChain.IKGoalName;
+	
 	// validate start and end bones exist and are not the root
-	const int32 StartIndex = Skeleton.FindBoneIndexByName(BoneChain.StartBone.BoneName);
-	const int32 EndIndex = Skeleton.FindBoneIndexByName(BoneChain.EndBone.BoneName);
+	const int32 StartIndex = InSkeleton.FindBoneIndexByName(StartBone);
+	const int32 EndIndex = InSkeleton.FindBoneIndexByName(EndBone);
 	bFoundStartBone = StartIndex > INDEX_NONE;
 	bFoundEndBone = EndIndex > INDEX_NONE;
 
 	// no need to build the chain if start/end indices are wrong 
 	const bool bIsWellFormed = bFoundStartBone && bFoundEndBone && EndIndex >= StartIndex;
-	if (bIsWellFormed)
+	if (!bIsWellFormed)
 	{
-		// init array with end bone 
-		OutBoneIndices = {EndIndex};
-
-		// if only one bone in the chain
-		if (EndIndex == StartIndex)
-		{
-			// validate end bone is child of start bone ?
-			bEndIsStartOrChildOfStart = true;
-			return;
-		}
+		return;
+	}
+	
+	// init array with end bone 
+	BoneIndices = {EndIndex};
 
-		// record all bones in chain while walking up the hierarchy (tip to root of chain)
-		int32 ParentIndex = Skeleton.GetParentIndex(EndIndex);
-		while (ParentIndex > INDEX_NONE && ParentIndex >= StartIndex)
-		{
-			OutBoneIndices.Add(ParentIndex);
-			ParentIndex = Skeleton.GetParentIndex(ParentIndex);
-		}
+	// record all bones in chain while walking up the hierarchy (tip to root of chain)
+	int32 ParentIndex = InSkeleton.GetParentIndex(EndIndex);
+	while (ParentIndex > INDEX_NONE && ParentIndex >= StartIndex)
+	{
+		BoneIndices.Add(ParentIndex);
+		ParentIndex = InSkeleton.GetParentIndex(ParentIndex);
+	}
 
-		// if we walked up till the start bone
-		if (OutBoneIndices.Last() == StartIndex)
-		{
-			// validate end bone is child of start bone
-			bEndIsStartOrChildOfStart = true;
-			// reverse the indices (we want root to tip order)
-			Algo::Reverse(OutBoneIndices);
-			return;
-		}
-      
-		// oops, we walked all the way up without finding the start bone
-		OutBoneIndices.Reset();
+	// did we walk all the way up without finding the start bone?
+	if (BoneIndices.Last() != StartIndex)
+	{
+		BoneIndices.Reset();
+		return;
 	}
+	
+	// confirm that end bone is child of start bone
+	bEndIsStartOrChildOfStart = true;
+	
+	// reverse the indices (we want root to tip order)
+	Algo::Reverse(BoneIndices);
+
+	// update the ref pose stored in the chain
+	UpdatePoseFromSkeleton(InSkeleton);
+	
+	// calculate parameter of each bone, normalized by the length of the bone chain
+	CalculateBoneParameters(InLog);
 }
 
-void FTargetSkeleton::SetBoneIsRetargeted(const int32 BoneIndex, const bool IsRetargeted)
+bool FResolvedBoneChain::IsValid() const
 {
-	check(IsBoneRetargeted.IsValidIndex(BoneIndex));
-	IsBoneRetargeted[BoneIndex] = IsRetargeted;
+	return bFoundStartBone && bFoundEndBone && bEndIsStartOrChildOfStart;
 }
 
-bool FChainFK::Initialize(
-	const FRetargetSkeleton& Skeleton,
-	const TArray<int32>& InBoneIndices,
-	const TArray<FTransform>& InitialGlobalPose,
-	FIKRigLogger& Log)
+void FResolvedBoneChain::UpdatePoseFromSkeleton(const FRetargetSkeleton& InSkeleton)
 {
-	check(!InBoneIndices.IsEmpty());
-
-	// store for debugging purposes
-	BoneIndices = InBoneIndices;
-
 	// store all the initial bone transforms in the bone chain
-	InitialGlobalTransforms.Reset();
+	RefPoseGlobalTransforms.Reset();
 	for (int32 Index=0; Index < BoneIndices.Num(); ++Index)
 	{
 		const int32 BoneIndex = BoneIndices[Index];
-		if (ensure(InitialGlobalPose.IsValidIndex(BoneIndex)))
+		if (ensure(InSkeleton.RetargetGlobalPose.IsValidIndex(BoneIndex)))
 		{
-			InitialGlobalTransforms.Emplace(InitialGlobalPose[BoneIndex]);
+			RefPoseGlobalTransforms.Emplace(InSkeleton.RetargetGlobalPose[BoneIndex]);
 		}
 	}
 
-	// initialize storage for current bones
-	CurrentGlobalTransforms = InitialGlobalTransforms;
-
 	// get the local space of the chain in retarget pose
-	InitialLocalTransforms.SetNum(InitialGlobalTransforms.Num());
-	FillTransformsWithLocalSpaceOfChain(Skeleton, InitialGlobalPose, BoneIndices, InitialLocalTransforms);
+	RefPoseLocalTransforms.SetNum(RefPoseGlobalTransforms.Num());
+	FillTransformsWithLocalSpaceOfChain(InSkeleton, InSkeleton.RetargetGlobalPose, BoneIndices, RefPoseLocalTransforms);
 
+	// initialize storage for current local transforms
+	CurrentLocalTransforms = RefPoseLocalTransforms;
+	
 	// store chain parent data
-	ChainParentBoneIndex = Skeleton.GetParentIndex(BoneIndices[0]);
+	ChainParentBoneIndex = InSkeleton.GetParentIndex(BoneIndices[0]);
 	ChainParentInitialGlobalTransform = FTransform::Identity;
 	if (ChainParentBoneIndex != INDEX_NONE)
 	{
-		ChainParentInitialGlobalTransform = InitialGlobalPose[ChainParentBoneIndex];
+		ChainParentInitialGlobalTransform = InSkeleton.RetargetGlobalPose[ChainParentBoneIndex];
+	}
+}
+
+void FResolvedBoneChain::GetWarnings(const FRetargetSkeleton& Skeleton, FIKRigLogger& Log) const
+{
+	// warn if START bone not found
+	if (!bFoundStartBone)
+	{
+		Log.LogWarning( FText::Format(
+			LOCTEXT("MissingStartBone", "IK Retargeter bone chain, {0}, could not find start bone, {1} in mesh {2}"),
+			FText::FromName(ChainName),
+			FText::FromName(StartBone),
+			FText::FromString(Skeleton.SkeletalMesh->GetName())));
+	}
+	
+	// warn if END bone not found
+	if (!bFoundEndBone)
+	{
+		Log.LogWarning( FText::Format(
+			LOCTEXT("MissingEndBone", "IK Retargeter bone chain, {0}, could not find end bone, {1} in mesh {2}"),
+			FText::FromName(ChainName), FText::FromName(EndBone), FText::FromString(Skeleton.SkeletalMesh->GetName())));
 	}
 
-	// calculate parameter of each bone, normalized by the length of the bone chain
-	return CalculateBoneParameters(Log);
+	// warn if END bone was not a child of START bone
+	if (bFoundEndBone && !bEndIsStartOrChildOfStart)
+	{
+		Log.LogWarning( FText::Format(
+			LOCTEXT("EndNotChildtOfStart", "IK Retargeter bone chain, {0}, end bone, '{1}' was not a child of the start bone '{2}'."),
+			FText::FromName(ChainName), FText::FromName(EndBone), FText::FromName(StartBone)));
+	}
+
+	// cannot retarget chain if all the bones are sitting directly on each other
+	if (InitialChainLength <= UE_KINDA_SMALL_NUMBER)
+	{
+		Log.LogWarning( FText::Format(
+			LOCTEXT("FailedParamOfChain", "IK Retargeter bone chain, {0}, was unable to be normalized. Chain too short."),
+			FText::FromName(ChainName)));
+	}
 }
 
-FTransform FChainFK::GetTransformAtChainParam(const TArray<FTransform>& Transforms, const float& Param) const
+FTransform FResolvedBoneChain::GetTransformAtChainParam(const TArray<FTransform>& Transforms, const double& Param) const
 {
 	check(Transforms.Num() == Params.Num());
 	
@@ -501,14 +561,14 @@ FTransform FChainFK::GetTransformAtChainParam(const TArray<FTransform>& Transfor
 
 	for (int32 ChainIndex=1; ChainIndex<Params.Num(); ++ChainIndex)
 	{
-		const float CurrentParam = Params[ChainIndex];
+		const double CurrentParam = Params[ChainIndex];
 		if (CurrentParam <= Param)
 		{
 			continue;
 		}
 		
-		const float PrevParam = Params[ChainIndex-1];
-		const float PercentBetweenParams = (Param - PrevParam) / (CurrentParam - PrevParam);
+		const double PrevParam = Params[ChainIndex-1];
+		const double PercentBetweenParams = (Param - PrevParam) / (CurrentParam - PrevParam);
 		const FTransform& Prev = Transforms[ChainIndex-1];
 		const FTransform& Next = Transforms[ChainIndex];
 		const FVector Position = FMath::Lerp(Prev.GetTranslation(), Next.GetTranslation(), PercentBetweenParams);
@@ -522,23 +582,23 @@ FTransform FChainFK::GetTransformAtChainParam(const TArray<FTransform>& Transfor
 	return FTransform::Identity;
 }
 
-float FChainFK::GetStretchAtParam(
+double FResolvedBoneChain::GetStretchAtParam(
 	const TArray<FTransform>& InitialTransforms,
 	const TArray<FTransform>& CurrentTransforms,
-	const double& Param)
+	const double& Param) const
 {
 	check(InitialTransforms.Num() == CurrentTransforms.Num() && InitialTransforms.Num() == Params.Num());
 
 	// chain only has 1 bone, so it can't stretch
 	if (Params.Num() <= 1)
 	{
-		return 1.0f;
+		return 1.0;
 	}
 
 	// start of chain cannot stretch
 	if (Param < UE_KINDA_SMALL_NUMBER)
 	{
-		return 1.0f;
+		return 1.0;
 	}
 	
 	// end of chain always uses last bone
@@ -575,47 +635,10 @@ float FChainFK::GetStretchAtParam(
 	// calculate the ratio of the current to the initial length of the closest bone
 	const double InitialLength = (InitialTransforms[BoneStartIndex].GetTranslation() - InitialTransforms[BoneEndIndex].GetTranslation()).Length();
 	const double CurrentLength = (CurrentTransforms[BoneStartIndex].GetTranslation() - CurrentTransforms[BoneEndIndex].GetTranslation()).Length();
-	return static_cast<float>(InitialLength < UE_KINDA_SMALL_NUMBER ? 1.0 : CurrentLength / InitialLength);
-}
-
-bool FChainFK::CalculateBoneParameters(FIKRigLogger& Log)
-{
-	Params.Reset();
-	
-	// special case, a single-bone chain
-	if (InitialGlobalTransforms.Num() == 1)
-	{
-		Params.Add(1.0f);
-		return true;
-	}
-
-	// calculate bone lengths in chain and accumulate total length
-	TArray<float> BoneDistances;
-	float TotalChainLength = 0.0f;
-	BoneDistances.Add(0.0f);
-	for (int32 i=1; i<InitialGlobalTransforms.Num(); ++i)
-	{
-		TotalChainLength += static_cast<float>((InitialGlobalTransforms[i].GetTranslation() - InitialGlobalTransforms[i-1].GetTranslation()).Size());
-		BoneDistances.Add(TotalChainLength);
-	}
-
-	// cannot retarget chain if all the bones are sitting directly on each other
-	if (TotalChainLength <= UE_KINDA_SMALL_NUMBER)
-	{
-		Log.LogWarning(LOCTEXT("TinyBoneChain", "IK Retargeter bone chain length is too small to reliably retarget."));
-		return false;
-	}
-
-	// calc each bone's param along length
-	for (int32 i=0; i<InitialGlobalTransforms.Num(); ++i)
-	{
-		Params.Add(BoneDistances[i] / TotalChainLength); 
-	}
-
-	return true;
+	return (InitialLength < UE_KINDA_SMALL_NUMBER ? 1.0 : CurrentLength / InitialLength);
 }
 
-double FChainFK::GetChainLength(const TArray<FTransform>& Transforms)
+double FResolvedBoneChain::GetChainLength(const TArray<FTransform>& Transforms)
 {
 	double TotalLength = 0.0f;
 	for (int32 ChainIndex = 1; ChainIndex < Transforms.Num(); ++ChainIndex)
@@ -627,17 +650,17 @@ double FChainFK::GetChainLength(const TArray<FTransform>& Transforms)
 	return TotalLength;
 }
 
-void FChainFK::FillTransformsWithLocalSpaceOfChain(
+void FResolvedBoneChain::FillTransformsWithLocalSpaceOfChain(
 	const FRetargetSkeleton& Skeleton,
 	const TArray<FTransform>& InGlobalPose,
-	const TArray<int32>& BoneIndices,
+	const TArray<int32>& InBoneIndices,
 	TArray<FTransform>& OutLocalTransforms)
 {
-	check(BoneIndices.Num() == OutLocalTransforms.Num())
+	check(InBoneIndices.Num() == OutLocalTransforms.Num())
 	
-	for (int32 ChainIndex=0; ChainIndex<BoneIndices.Num(); ++ChainIndex)
+	for (int32 ChainIndex=0; ChainIndex<InBoneIndices.Num(); ++ChainIndex)
 	{
-		const int32 BoneIndex = BoneIndices[ChainIndex];
+		const int32 BoneIndex = InBoneIndices[ChainIndex];
 		const int32 ParentIndex = Skeleton.GetParentIndex(BoneIndex);
 		if (ParentIndex == INDEX_NONE)
 		{
@@ -652,11 +675,14 @@ void FChainFK::FillTransformsWithLocalSpaceOfChain(
 	}
 }
 
-void FChainFK::PutCurrentTransformsInRefPose(
-	const TArray<int32>& InBoneIndices,
+void FResolvedBoneChain::FillTransformsWithGlobalRefPoseOfChain(
 	const FRetargetSkeleton& Skeleton,
-	const TArray<FTransform>& InCurrentGlobalPose)
+	const TArray<FTransform>& InGlobalPose,
+	const TArray<int32>& InBoneIndices,
+	TArray<FTransform>& OutGlobalTransforms)
 {
+	check(InBoneIndices.Num() == OutGlobalTransforms.Num())
+	
 	// update chain current transforms to the retarget pose in global space
 	for (int32 ChainIndex=0; ChainIndex<InBoneIndices.Num(); ++ChainIndex)
 	{
@@ -664,849 +690,221 @@ void FChainFK::PutCurrentTransformsInRefPose(
 		if (ChainIndex == 0)
 		{
 			const int32 BoneIndex = InBoneIndices[ChainIndex];
-			CurrentGlobalTransforms[ChainIndex] = Skeleton.GetGlobalRetargetPoseOfSingleBone(BoneIndex, InCurrentGlobalPose);
+			OutGlobalTransforms[ChainIndex] = Skeleton.GetGlobalRetargetPoseOfSingleBone(BoneIndex, InGlobalPose);
 		}
 		else
 		{
 			// all subsequent bones in chain are based on previous parent
 			const int32 BoneIndex = InBoneIndices[ChainIndex];
-			const FTransform& ParentGlobalTransform = CurrentGlobalTransforms[ChainIndex-1];
+			const FTransform& ParentGlobalTransform = OutGlobalTransforms[ChainIndex-1];
 			const FTransform& ChildLocalTransform = Skeleton.RetargetLocalPose[BoneIndex];
-			CurrentGlobalTransforms[ChainIndex] = ChildLocalTransform * ParentGlobalTransform;
+			OutGlobalTransforms[ChainIndex] = ChildLocalTransform * ParentGlobalTransform;
 		}
 	}
 }
 
-void FChainEncoderFK::EncodePose(
-	const FRetargetSkeleton& SourceSkeleton,
-	const TArray<int32>& SourceBoneIndices,
-    const TArray<FTransform> &InSourceGlobalPose)
+TArray<FTransform> FResolvedBoneChain::GetChainTransformsFromPose(const TArray<FTransform>& InPose) const
 {
-	check(SourceBoneIndices.Num() == CurrentGlobalTransforms.Num());
-	
-	// copy the global input pose for the chain
-	for (int32 ChainIndex=0; ChainIndex<SourceBoneIndices.Num(); ++ChainIndex)
+	TArray<FTransform> Transforms;
+	Transforms.Reset(BoneIndices.Num());
+	for (const int32 BoneIndex : BoneIndices)
 	{
-		const int32 BoneIndex = SourceBoneIndices[ChainIndex];
-		CurrentGlobalTransforms[ChainIndex] = InSourceGlobalPose[BoneIndex];
+		Transforms.Add(InPose[BoneIndex]);
 	}
-
-	CurrentLocalTransforms.SetNum(SourceBoneIndices.Num());
-	FillTransformsWithLocalSpaceOfChain(SourceSkeleton, InSourceGlobalPose, SourceBoneIndices, CurrentLocalTransforms);
-	
-	ChainParentCurrentGlobalTransform = ChainParentBoneIndex != INDEX_NONE ? InSourceGlobalPose[ChainParentBoneIndex] : FTransform::Identity;
+	return MoveTemp(Transforms);
 }
 
-void FChainEncoderFK::TransformCurrentChainTransforms(const FTransform& NewParentTransform)
+void FResolvedBoneChain::CalculateBoneParameters(FIKRigLogger& Log)
 {
-	for (int32 ChainIndex=0; ChainIndex<CurrentGlobalTransforms.Num(); ++ChainIndex)
+	Params.Reset();
+	
+	// special case, a single-bone chain
+	if (RefPoseGlobalTransforms.Num() == 1)
 	{
-		if (ChainIndex == 0)
-		{
-			CurrentGlobalTransforms[ChainIndex] = CurrentLocalTransforms[ChainIndex] * NewParentTransform;
-		}
-		else
-		{
-			CurrentGlobalTransforms[ChainIndex] = CurrentLocalTransforms[ChainIndex] * CurrentGlobalTransforms[ChainIndex-1];
-		}
+		Params.Add(1.0f);
+		return;
 	}
-}
-
-void FChainDecoderFK::DecodePose(
-	const FRootRetargeter& RootRetargeter,
-	const FTargetChainFKSettings& Settings,
-	const TArray<int32>& TargetBoneIndices,
-    FChainEncoderFK& SourceChain,
-    const FTargetSkeleton& TargetSkeleton,
-    TArray<FTransform> &InOutGlobalPose)
-{
-	check(TargetBoneIndices.Num() == CurrentGlobalTransforms.Num());
-	check(TargetBoneIndices.Num() == Params.Num());
 
-	// Before setting this chain pose, we need to ensure that any
-	// intermediate (between chains) NON-retargeted parent bones have had their
-	// global transforms updated.
-	// 
-	// For example, if this chain is retargeting a single head bone, AND the spine was
-	// retargeted in the prior step, then the neck bones will need updating first.
-	// Otherwise the neck bones will remain at their location prior to the spine update.
-	UpdateIntermediateParents(TargetSkeleton,InOutGlobalPose);
-	
-	// transform entire source chain from it's root to match target's current root orientation (maintaining offset from retarget pose)
-	// this ensures children are retargeted in a "local" manner free from skewing that will happen if source and target
-	// become misaligned as can happen if parent chains were not retargeted
-	FTransform SourceChainParentInitialDelta = SourceChain.ChainParentInitialGlobalTransform.GetRelativeTransform(ChainParentInitialGlobalTransform);
-	FTransform TargetChainParentCurrentGlobalTransform = ChainParentBoneIndex == INDEX_NONE ? FTransform::Identity : InOutGlobalPose[ChainParentBoneIndex]; 
-	FTransform SourceChainParentTransform = SourceChainParentInitialDelta * TargetChainParentCurrentGlobalTransform;
+	// calculate bone lengths in chain and accumulate total length
+	TArray<float> BoneDistances;
+	InitialChainLength = 0.0f;
+	BoneDistances.Add(0.0f);
+	for (int32 i=1; i<RefPoseGlobalTransforms.Num(); ++i)
+	{
+		InitialChainLength += static_cast<float>((RefPoseGlobalTransforms[i].GetTranslation() - RefPoseGlobalTransforms[i-1].GetTranslation()).Size());
+		BoneDistances.Add(InitialChainLength);
+	}
 
-	// apply delta to the source chain's current transforms before transferring rotations to the target
-	SourceChain.TransformCurrentChainTransforms(SourceChainParentTransform);
+	// very small chains will not retarget correctly, this will generate a warning
+	const float Divisor = InitialChainLength > UE_KINDA_SMALL_NUMBER ? InitialChainLength : UE_KINDA_SMALL_NUMBER;
 
-	// if FK retargeting has been disabled for this chain, then simply set it to the retarget pose
-	if (!Settings.EnableFK)
+	// calc each bone's param along length
+	for (int32 i=0; i<RefPoseGlobalTransforms.Num(); ++i)
 	{
-		// put the chain in the global ref pose (globally rotated by parent bone in it's currently retargeted state)
-		PutCurrentTransformsInRefPose(TargetBoneIndices, TargetSkeleton, InOutGlobalPose);
-		
-		for (int32 ChainIndex=0; ChainIndex<TargetBoneIndices.Num(); ++ChainIndex)
-		{
-			const int32 BoneIndex = TargetBoneIndices[ChainIndex];
-			InOutGlobalPose[BoneIndex] = CurrentGlobalTransforms[ChainIndex];
-		}
-
-		return;
+		Params.Add(BoneDistances[i] / Divisor); 
 	}
+}
 
-	const int32 NumBonesInSourceChain = SourceChain.CurrentGlobalTransforms.Num();
-	const int32 NumBonesInTargetChain = TargetBoneIndices.Num();
-	const int32 TargetStartIndex = FMath::Max(0, NumBonesInTargetChain - NumBonesInSourceChain);
-	const int32 SourceStartIndex = FMath::Max(0,NumBonesInSourceChain - NumBonesInTargetChain);
+bool FRetargeterBoneChains::Initialize(
+	const UIKRetargeter* InRetargetAsset,
+	const FRetargetSkeleton& InSourceSkeleton,
+	const FRetargetSkeleton& InTargetSkeleton,
+	FIKRigLogger& InLog)
+{
+	Reset();
 
-	// now retarget the pose of each bone in the chain, copying from source to target
-	for (int32 ChainIndex=0; ChainIndex<TargetBoneIndices.Num(); ++ChainIndex)
+	// get source and target IK Rig assets
+	const UIKRigDefinition* SourceIKRig = InRetargetAsset->GetIKRig(ERetargetSourceOrTarget::Source);
+	const UIKRigDefinition* TargetIKRig = InRetargetAsset->GetIKRig(ERetargetSourceOrTarget::Target);
+	if (!ensure(SourceIKRig && TargetIKRig))
 	{
-		const int32 BoneIndex = TargetBoneIndices[ChainIndex];
-		const FTransform& TargetInitialTransform = InitialGlobalTransforms[ChainIndex];
-		FTransform SourceCurrentTransform;
-		FTransform SourceInitialTransform;
-
-		// get source current / initial transforms for this bone
-		switch (Settings.RotationMode)
-		{
-			case ERetargetRotationMode::Interpolated:
-			case ERetargetRotationMode::MatchChain:
-			case ERetargetRotationMode::MatchScaledChain:
-			{
-				// get the initial and current transform of source chain at param
-				// this is the interpolated transform along the chain
-				const float Param = Params[ChainIndex];
-				SourceCurrentTransform = SourceChain.GetTransformAtChainParam(SourceChain.CurrentGlobalTransforms, Param);
-				SourceInitialTransform = SourceChain.GetTransformAtChainParam(SourceChain.InitialGlobalTransforms, Param);
-				break;
-			}
-			case ERetargetRotationMode::OneToOne:
-			{
-				if (ChainIndex < NumBonesInSourceChain)
-				{
-					SourceCurrentTransform = SourceChain.CurrentGlobalTransforms[ChainIndex];
-					SourceInitialTransform = SourceChain.InitialGlobalTransforms[ChainIndex];
-				}else
-				{
-					SourceCurrentTransform = SourceChain.CurrentGlobalTransforms.Last();
-					SourceInitialTransform = SourceChain.InitialGlobalTransforms.Last();
-				}
-				break;
-			}
-			case ERetargetRotationMode::OneToOneReversed:
-			{
-				if (ChainIndex < TargetStartIndex)
-				{
-					SourceCurrentTransform = SourceChain.InitialGlobalTransforms[0];
-					SourceInitialTransform = SourceChain.InitialGlobalTransforms[0];
-				}
-				else
-				{
-					const int32 SourceChainIndex = SourceStartIndex + (ChainIndex - TargetStartIndex);
-					SourceCurrentTransform = SourceChain.CurrentGlobalTransforms[SourceChainIndex];
-					SourceInitialTransform = SourceChain.InitialGlobalTransforms[SourceChainIndex];
-				}
-				break;
-			}
-			case ERetargetRotationMode::None:
-			{
-				// in order to induce no rotation on the FK chain, we rotate the chain rigidly from the root of the chain
-				SourceInitialTransform = SourceChain.InitialGlobalTransforms[0];
-				// use the current global space retarget pose as the "current" transform, so chain rotates with parent
-				SourceCurrentTransform = SourceChain.InitialLocalTransforms[0] * SourceChain.ChainParentCurrentGlobalTransform;
-				break;
-			}
-			default:
-				checkNoEntry();
-			break;
-		}
-		
-		// apply rotation offset to the initial target rotation
-		const FQuat SourceCurrentRotation = SourceCurrentTransform.GetRotation();
-		const FQuat SourceInitialRotation = SourceInitialTransform.GetRotation();
-		const FQuat RotationDelta = SourceCurrentRotation * SourceInitialRotation.Inverse();
-		const FQuat TargetInitialRotation = TargetInitialTransform.GetRotation();
-		const FQuat OutRotation = RotationDelta * TargetInitialRotation;
-
-		// calculate output POSITION based on translation mode setting
-		FTransform ParentGlobalTransform = FTransform::Identity;
-		const int32 ParentIndex = TargetSkeleton.ParentIndices[BoneIndex];
-		if (ParentIndex != INDEX_NONE)
-		{
-			ParentGlobalTransform = InOutGlobalPose[ParentIndex];
-		}
-		FVector OutPosition;
-		switch (Settings.TranslationMode)
+		InLog.LogError(LOCTEXT("MissingIKRig", "IK Retargeter either a source or target IK Rig asset."));
+		return false;
+	}
+	
+	// load ALL bone chains on source and target (ops may use them even if they aren't mapped)
+	auto LoadChains = [](
+		const FRetargetSkeleton& InRetargetSkeleton,
+		const TArray<FBoneChain>& InBoneChains,
+		FIKRigLogger& InLog,
+		TArray<FResolvedBoneChain>& OutResolvedChains)
+	{
+		for (const FBoneChain& BoneChain : InBoneChains)
 		{
-		case ERetargetTranslationMode::None:
-			{
-				const FVector InitialLocalOffset = TargetSkeleton.RetargetLocalPose[BoneIndex].GetTranslation();
-				OutPosition = ParentGlobalTransform.TransformPosition(InitialLocalOffset);
-				break;
-			}
-		case ERetargetTranslationMode::GloballyScaled:
+			FResolvedBoneChain NewChain(BoneChain, InRetargetSkeleton, InLog);
+			NewChain.GetWarnings(InRetargetSkeleton, InLog);
+			if (NewChain.IsValid())
 			{
-				OutPosition = SourceCurrentTransform.GetTranslation() * RootRetargeter.GetGlobalScaleVector();
-				break;
-			}
-		case ERetargetTranslationMode::Absolute:
-			{
-				OutPosition = SourceCurrentTransform.GetTranslation();
-				break;	
-			}
-		case ERetargetTranslationMode::StretchBoneLengthUniformly:
-			{
-				if (ChainIndex==0)
-				{
-					const FVector InitialLocalOffset = TargetSkeleton.RetargetLocalPose[BoneIndex].GetTranslation();
-					OutPosition = ParentGlobalTransform.TransformPosition(InitialLocalOffset);
-					break;
-				}
-				// initial chain length
-				double SourceChainLengthInitial = GetChainLength(SourceChain.InitialGlobalTransforms);
-				double SourceChainLengthCurrent = GetChainLength(SourceChain.CurrentGlobalTransforms);
-				double StretchRatio = SourceChainLengthInitial < UE_KINDA_SMALL_NUMBER ? 1.0f : SourceChainLengthCurrent / SourceChainLengthInitial;
-				// stretch local translation
-				const FVector InitialLocalOffset = TargetSkeleton.RetargetLocalPose[BoneIndex].GetTranslation();
-				OutPosition = ParentGlobalTransform.TransformPosition(InitialLocalOffset * StretchRatio);
-				break;
-			}
-		case ERetargetTranslationMode::StretchBoneLengthNonUniformly:
-			{
-				const double Param = Params[ChainIndex];
-				float StretchRatio = SourceChain.GetStretchAtParam( SourceChain.InitialGlobalTransforms, SourceChain.CurrentGlobalTransforms, Param);
-				// stretch local translation
-				const FVector InitialLocalOffset = TargetSkeleton.RetargetLocalPose[BoneIndex].GetTranslation();
-				OutPosition = ParentGlobalTransform.TransformPosition(InitialLocalOffset * StretchRatio);
-				break;
-			}
-			default:
-			{
-				const FVector InitialLocalOffset = TargetSkeleton.RetargetLocalPose[BoneIndex].GetTranslation();
-				OutPosition = ParentGlobalTransform.TransformPosition(InitialLocalOffset);
-				break;
+				OutResolvedChains.Add(NewChain);
 			}
 		}
+	};
+	LoadChains(InSourceSkeleton, SourceIKRig->GetRetargetChains(), InLog, SourceBoneChains);
+	LoadChains(InTargetSkeleton, TargetIKRig->GetRetargetChains(), InLog, TargetBoneChains);
 
-		// calculate output SCALE
-		const FVector SourceCurrentScale = SourceCurrentTransform.GetScale3D();
-		const FVector SourceInitialScale = SourceInitialTransform.GetScale3D();
-		const FVector TargetInitialScale = TargetInitialTransform.GetScale3D();
-		const FVector OutScale = SourceCurrentScale + (TargetInitialScale - SourceInitialScale);
-		
-		// apply output transform
-		CurrentGlobalTransforms[ChainIndex] = FTransform(OutRotation, OutPosition, OutScale);
-		InOutGlobalPose[BoneIndex] = CurrentGlobalTransforms[ChainIndex];
-	}
-
-	// apply match chain operation on-top of translated and rotated result
-	if (Settings.RotationMode == ERetargetRotationMode::MatchChain ||
-		Settings.RotationMode == ERetargetRotationMode::MatchScaledChain)
+	// sort the chains based on their StartBone's index
+	auto ChainsSorter = [this](const FResolvedBoneChain& A, const FResolvedBoneChain& B)
 	{
-		const bool bScaleSourceChain = Settings.RotationMode == ERetargetRotationMode::MatchScaledChain;
-		FVector TargetChainOrigin = InOutGlobalPose[TargetBoneIndices[0]].GetTranslation();
-		MatchChain(bScaleSourceChain, SourceChain, TargetBoneIndices, TargetChainOrigin);
-		
-		// update output pose
-		for (int32 ChainIndex=0; ChainIndex<TargetBoneIndices.Num(); ++ChainIndex)
+		const int32 IndexA = A.BoneIndices.Num() > 0 ? A.BoneIndices[0] : INDEX_NONE;
+		const int32 IndexB = B.BoneIndices.Num() > 0 ? B.BoneIndices[0] : INDEX_NONE;
+		if (IndexA == IndexB)
 		{
-			InOutGlobalPose[TargetBoneIndices[ChainIndex]] = CurrentGlobalTransforms[ChainIndex];
+			// fallback to sorting alphabetically
+			return A.ChainName.LexicalLess(B.ChainName);
 		}
-	}
+		return IndexA < IndexB;
+	};
+	SourceBoneChains.Sort(ChainsSorter);
+	TargetBoneChains.Sort(ChainsSorter);
 
-	// apply final blending between retarget pose of chain and newly retargeted pose
-	// blend must be done in local space, so we do it in a separate loop after full chain pose is generated
-	const bool bShouldBlendRotation = !FMath::IsNearlyEqual(Settings.RotationAlpha, 1.0f);
-	const bool bShouldBlendTranslation = !FMath::IsNearlyEqual(Settings.TranslationAlpha, 1.0f);
-	if (bShouldBlendRotation || bShouldBlendTranslation) // (skipped if the alphas are not near 1.0)
+	// load the chain mapping for each target chain (loads INDEX_NONE for any unmapped chains)
+	bool bFoundAtLeastOneMappedChain = true;
+	const FRetargetChainMapping& ChainNameMappings = InRetargetAsset->GetChainMapping();
+	for (int32 TargetBoneChainIndex = 0; TargetBoneChainIndex<TargetBoneChains.Num(); TargetBoneChainIndex++)
 	{
-		// generate local space pose of chain
-		TArray<FTransform> NewLocalTransforms;
-		NewLocalTransforms.SetNum(InitialLocalTransforms.Num());
-		FillTransformsWithLocalSpaceOfChain(TargetSkeleton, InOutGlobalPose, TargetBoneIndices, NewLocalTransforms);
-
-		// blend each bone in chain with the retarget pose
-		for (int32 ChainIndex=0; ChainIndex<InitialLocalTransforms.Num(); ++ChainIndex)
+		const FResolvedBoneChain& ResolvedTargetChain = TargetBoneChains[TargetBoneChainIndex];
+		FName SourceChainName = ChainNameMappings.GetChainMappedTo(ResolvedTargetChain.ChainName, ERetargetSourceOrTarget::Target);
+		if (SourceChainName == NAME_None)
 		{
-			// blend between current local pose and initial local pose
-			FTransform& NewLocalTransform = NewLocalTransforms[ChainIndex];
-			const FTransform& RefPoseLocalTransform = InitialLocalTransforms[ChainIndex];
-			NewLocalTransform.SetTranslation(FMath::Lerp(RefPoseLocalTransform.GetTranslation(), NewLocalTransform.GetTranslation(), Settings.TranslationAlpha));
-			NewLocalTransform.SetRotation(FQuat::FastLerp(RefPoseLocalTransform.GetRotation(), NewLocalTransform.GetRotation(), Settings.RotationAlpha).GetNormalized());
-
-			// put blended transforms back in global space and store in final output pose
-			const int32 BoneIndex = TargetBoneIndices[ChainIndex];
-			const int32 ParentIndex = TargetSkeleton.ParentIndices[BoneIndex];
-			const FTransform& ParentGlobalTransform = ParentIndex == INDEX_NONE ? FTransform::Identity : InOutGlobalPose[ParentIndex];
-			InOutGlobalPose[BoneIndex] = NewLocalTransform * ParentGlobalTransform;
+			ChainMapping.Add(INDEX_NONE);
+			continue;
 		}
-	}
-}
 
-void FChainDecoderFK::InitializeIntermediateParentIndices(
-	const int32 RetargetRootBoneIndex,
-	const int32 ChainRootBoneIndex,
-	const FTargetSkeleton& TargetSkeleton)
-{
-	IntermediateParentIndices.Reset();
-	int32 ParentBoneIndex = TargetSkeleton.ParentIndices[ChainRootBoneIndex];
-	while (true)
-	{
-		if (ParentBoneIndex < 0 || ParentBoneIndex == RetargetRootBoneIndex)
+		// find the index of the source chain to map to
+		int32 SourceChainIndex = SourceBoneChains.IndexOfByPredicate([&SourceChainName](const FResolvedBoneChain& Value)
 		{
-			break; // reached root of skeleton
-		}
+			return Value.ChainName == SourceChainName;
+		});
 
-		if (TargetSkeleton.IsBoneRetargeted[ParentBoneIndex])
+		// log a warning if the user has mapped a target chain to a source chain that doesn't exist (ill formed data)
+		if (SourceChainName!=NAME_None && SourceChainIndex==INDEX_NONE)
 		{
-			break; // reached the start of another retargeted chain
+			InLog.LogWarning( FText::Format(
+			LOCTEXT("MissingTargetChain", "IK Retargeter has a target bone chain: {0} that is mapped to a non-existant source chain."),
+			FText::FromString(ResolvedTargetChain.ChainName.ToString())));
 		}
 
-		IntermediateParentIndices.Add(ParentBoneIndex);
-		ParentBoneIndex = TargetSkeleton.ParentIndices[ParentBoneIndex];
+		bFoundAtLeastOneMappedChain &= SourceChainIndex != INDEX_NONE;
+		ChainMapping.Add(SourceChainIndex);
 	}
 
-	Algo::Reverse(IntermediateParentIndices);
+	bIsInitialized = bFoundAtLeastOneMappedChain;
+	return bIsInitialized;
 }
 
-void FChainDecoderFK::MatchChain(
-	const bool bScaleSourceChain,
-	const FChainEncoderFK& SourceChain,
-	const TArray<int32>& TargetBoneIndices,
-	const FVector& TargetChainOrigin)
+const TArray<FResolvedBoneChain>& FRetargeterBoneChains::GetResolvedBoneChains(ERetargetSourceOrTarget SourceOrTarget) const
 {
-	// "MatchChain" mode assumes the interpolated rotations (and any translation mode) as a starting point.
-	//
-	// The "spline-IK-like" method used below generates swing rotations, which fix-up the interpolated rotations
-	// such that they align the bone positions to lie on the source chain.
-	//
-	// Typically, spline-ik is not "twist aware" because aligning a joint chain with swing rotations does not twist,
-	// but by using the interpolated rotations/translation as a starting place we have a full twist/bend/stretch.
-	
-	// convert source chain into a linear spline
-	TArray<FVector> SourceSplinePoints;
-	for (const FTransform& SourceGlobalTransform : SourceChain.CurrentGlobalTransforms)
-	{
-		SourceSplinePoints.Add(SourceGlobalTransform.GetTranslation());
-	}
-		
-	// translate chain points to origin of target chain
-	FVector ChainOffset = TargetChainOrigin - SourceSplinePoints[0];
-	for (FVector& SplinePoint : SourceSplinePoints)
-	{
-		SplinePoint += ChainOffset;
-	}
+	return SourceOrTarget == ERetargetSourceOrTarget::Source ? SourceBoneChains : TargetBoneChains;
+}
 
-	// scale chain to match target length
-	if (bScaleSourceChain)
+int32 FRetargeterBoneChains::GetIndexOfBoneChainByName(const FName InChainName, const ERetargetSourceOrTarget SourceOrTarget) const
+{
+	const TArray<FResolvedBoneChain>& ResolvedBoneChains = GetResolvedBoneChains(SourceOrTarget);
+	return ResolvedBoneChains.IndexOfByPredicate([InChainName](const FResolvedBoneChain& Value)
 	{
-		double TargetChainLengthCurrent = GetChainLength(CurrentGlobalTransforms);
-		double SourceChainLengthCurrent = GetChainLength(SourceChain.CurrentGlobalTransforms);
-		double ScaleFactor = TargetChainLengthCurrent < UE_KINDA_SMALL_NUMBER ? 1.0f : SourceChainLengthCurrent / TargetChainLengthCurrent;
-		for (FVector& SplinePoint : SourceSplinePoints)
-		{
-			SplinePoint =  TargetChainOrigin + (SplinePoint - TargetChainOrigin) * ScaleFactor;
-		}
-	}
-	
-	// snap orient each bone to lie on the source spline
-	int32 OriginPointIndex = 0;
-	double OriginPointAlpha = 0.0;
-	for (int32 ChainIndex=0; ChainIndex<TargetBoneIndices.Num()-1; ++ChainIndex)
-	{
-		// generate an aim rotation from the current joint vector to the point on the spline
-		const FVector StartBonePosition = CurrentGlobalTransforms[ChainIndex].GetTranslation();
-		const FVector EndBonePosition = CurrentGlobalTransforms[ChainIndex+1].GetTranslation();
-		const FVector EndBoneLocalPosition = CurrentGlobalTransforms[ChainIndex].InverseTransformPosition(EndBonePosition);
-		FVector BoneNorm;
-		double BoneLength;
-		(EndBonePosition - StartBonePosition).ToDirectionAndLength(BoneNorm, BoneLength);
-
-		// get the first point along the spline that is "BoneLength" away from the joint location in a straight line
-		GetPointOnSplineDistanceFromPoint(
-			SourceSplinePoints,
-			OriginPointIndex,
-			OriginPointAlpha,
-			BoneLength,
-			OriginPointIndex,
-			OriginPointAlpha);
-		
-		// convert spline coordinates to euclidean
-		const FVector AimPointOnSpline = GetPointOnSplineFromIndexAndAlpha(SourceSplinePoints, OriginPointIndex, OriginPointAlpha);
-		
-		// generate from/to swing rotation to align child with spline
-		const FVector AimNorm = (AimPointOnSpline - StartBonePosition).GetSafeNormal();
-		const FQuat RotationOffset = FQuat::FindBetweenNormals(BoneNorm, AimNorm);
-
-		// rotate bone to contact spline
-		CurrentGlobalTransforms[ChainIndex].SetRotation( RotationOffset * CurrentGlobalTransforms[ChainIndex].GetRotation());
+		return Value.ChainName == InChainName;
+	});
+}
 
-		// propagate translational offset of end bone to all it's children
-		const FVector NewEndBonePosition = CurrentGlobalTransforms[ChainIndex].TransformPosition(EndBoneLocalPosition);
-		const FVector DeltaTranslation = NewEndBonePosition - EndBonePosition;
-		for (int32 ChildIndex=ChainIndex+1; ChildIndex<TargetBoneIndices.Num(); ++ChildIndex)
+const FResolvedBoneChain* FRetargeterBoneChains::GetSourceChainMappedToTargetChain(const FName InTargetChainName) const
+{
+	const int32 TargetChainIndex = GetIndexOfBoneChainByName(InTargetChainName, ERetargetSourceOrTarget::Target);
+	if (ChainMapping.IsValidIndex(TargetChainIndex))
+	{
+		const int32 SourceChainIndex = ChainMapping[TargetChainIndex];
+		if (SourceBoneChains.IsValidIndex(SourceChainIndex))
 		{
-			CurrentGlobalTransforms[ChildIndex].AddToTranslation(DeltaTranslation);
+			return &SourceBoneChains[SourceChainIndex];
 		}
 	}
+	return nullptr;
 }
 
-FVector FChainDecoderFK::GetPointOnSplineFromIndexAndAlpha(
-	const TArray<FVector>& InSplinePoints,
-	const int32 InPointIndex,
-	const double InSegmentAlpha)
+const FResolvedBoneChain* FRetargeterBoneChains::GetResolvedBoneChainByName(
+	const FName InChainName,
+	const ERetargetSourceOrTarget SourceOrTarget) const
 {
-	if (!ensure(InPointIndex>=0))
+	const TArray<FResolvedBoneChain>& ResolvedBoneChains = SourceOrTarget == ERetargetSourceOrTarget::Source ? SourceBoneChains : TargetBoneChains;
+	const FResolvedBoneChain* FoundChain = ResolvedBoneChains.FindByPredicate([InChainName](const FResolvedBoneChain& Value)
 	{
-		return InSplinePoints[0];
-	}
-	
-	if (InPointIndex == InSplinePoints.Num() - 1)
-	{
-		// extrapolate SegmentAlpha distance beyond the last segment
-		const FVector& LastSegmentStart = InSplinePoints[InSplinePoints.Num() - 2];
-		const FVector& LastSegmentEnd = InSplinePoints.Last();
-		const FVector LastSegmentDirection = (LastSegmentEnd - LastSegmentStart).GetSafeNormal();
-		return LastSegmentEnd + LastSegmentDirection * InSegmentAlpha;
-	}
-
-	// interpolate between point and next point by alpha
-	return FMath::Lerp(InSplinePoints[InPointIndex], InSplinePoints[InPointIndex+1], InSegmentAlpha);
+		return Value.ChainName == InChainName;
+	});
+	return FoundChain;
 }
 
-void FChainDecoderFK::GetPointOnSplineDistanceFromPoint(
-	const TArray<FVector>& InSplinePoints,
-	const int32 InOriginPointIndex,
-	const double InOriginPointAlpha,
-	const double InTargetDistanceFromOrigin,
-	int32& OutPointIndex,
-	double& OutPointAlpha)
+const FResolvedBoneChain* FRetargeterBoneChains::GetResolvedBoneChainByIndex(
+	const int32 InChainIndex,
+	const ERetargetSourceOrTarget SourceOrTarget) const
 {
-	// NOTE: this returns the point on the ray, in RayDirection, that is DistanceFromPoint away from Point
-	// It assumes that there is a point on the ray within TargetDistanceFromPointToRay
-	auto GetPointOnRayDistanceFromPoint = [](
-		const FVector& RayStart,
-		const FVector& RayDirection,
-		const FVector& Point,
-		const double TargetDistanceFromPointToRay)
-	{
-		const FVector RayStartToPoint = Point - RayStart;
-		const double ProjectionLength = FVector::DotProduct(RayStartToPoint, RayDirection);
-
-		// get the point projected onto the line formed by the ray
-		const FVector ProjectedPoint = RayStart + (RayDirection * ProjectionLength);
-		const double DistancePointToRay = (ProjectedPoint - Point).Size();
-
-		// check if point is further away from ray origin than TargetDistanceFromPointToRay (should not happen)
-		if (!ensure(DistancePointToRay < TargetDistanceFromPointToRay))
-		{
-			return (RayStart - Point).GetClampedToMaxSize(TargetDistanceFromPointToRay);
-		}
-	
-		// pythagorean theorem to find distance from projected point to point on ray
-		// NOTE: imagine the right angle triangle formed by Point, ProjectedPoint and the unknown point on the ray
-		const double DistanceFromProjPointToTargetPoint = FMath::Sqrt ((TargetDistanceFromPointToRay * TargetDistanceFromPointToRay) - (DistancePointToRay * DistancePointToRay));
-
-		// calculate the point on the ray that is DistanceFromPoint away from Point
-		return ProjectedPoint + RayDirection * DistanceFromProjPointToTargetPoint;
-	};
-
-	// convert spline coordinates to euclidean
-	const FVector OriginPoint = GetPointOnSplineFromIndexAndAlpha(InSplinePoints, InOriginPointIndex, InOriginPointAlpha);
-	
-	// iterate down the chain until we find a bone that is beyond TargetDistanceFromOrigin away from the origin point
-	for (OutPointIndex = InOriginPointIndex; OutPointIndex < InSplinePoints.Num()-1; ++OutPointIndex)
+	const TArray<FResolvedBoneChain>& ResolvedBoneChains = SourceOrTarget == ERetargetSourceOrTarget::Source ? SourceBoneChains : TargetBoneChains;
+	if (ResolvedBoneChains.IsValidIndex(InChainIndex))
 	{
-		const FVector& SegmentEnd = InSplinePoints[OutPointIndex+1];
-		const double DistanceOriginToSegmentEnd = (SegmentEnd - OriginPoint).Size();
-		if (DistanceOriginToSegmentEnd < InTargetDistanceFromOrigin)
-		{
-			continue;
-		}
-		
-		const FVector& RayStart = InSplinePoints[OutPointIndex];
-		FVector RayDirection;
-		double SegmentLength;
-		(SegmentEnd - RayStart).ToDirectionAndLength(RayDirection, SegmentLength);
-		const FVector PointOnRay = GetPointOnRayDistanceFromPoint(RayStart, RayDirection, OriginPoint, InTargetDistanceFromOrigin);
-		OutPointAlpha = (PointOnRay - RayStart).Size() / SegmentLength;
-		return;
+		return &ResolvedBoneChains[InChainIndex];
 	}
 
-	// spline is too short, so extrapolate the last spline segment and find the point on that segment that is DistanceFromOrigin away
-	const FVector& RayStart = InSplinePoints.Last();
-	const FVector RayDirection = (InSplinePoints.Last() - InSplinePoints[InSplinePoints.Num() - 2]).GetSafeNormal();
-	const FVector ExtrapolatedPointOnRay = GetPointOnRayDistanceFromPoint(RayStart, RayDirection, OriginPoint, InTargetDistanceFromOrigin);
-	OutPointAlpha = (ExtrapolatedPointOnRay - RayStart).Size();
-	OutPointIndex = InSplinePoints.Num() - 1;
+	return nullptr;
 }
 
-void FChainDecoderFK::UpdateIntermediateParents(
-	const FTargetSkeleton& TargetSkeleton,
-	TArray<FTransform>& InOutGlobalPose)
+const TArray<int32>& FRetargeterBoneChains::GetChainMappings() const
 {
-	for (const int32& ParentIndex : IntermediateParentIndices)
-	{
-		TargetSkeleton.UpdateGlobalTransformOfSingleBone(ParentIndex, TargetSkeleton.InputLocalPose, InOutGlobalPose);
-	}
+	return ChainMapping;
 }
 
-bool FChainRetargeterIK::InitializeSource(
-	const TArray<int32>& BoneIndices,
-	const TArray<FTransform>& SourceInitialGlobalPose,
-	FIKRigLogger& Log)
+void FRetargeterBoneChains::UpdatePoseFromSkeleton(const FRetargetSkeleton& InSkeleton, const ERetargetSourceOrTarget SourceOrTarget)
 {
-	if (BoneIndices.Num() < 2)
+	TArray<FResolvedBoneChain>& BoneChainsToUpdate = SourceOrTarget == ERetargetSourceOrTarget::Source ? SourceBoneChains : TargetBoneChains;
+	for (FResolvedBoneChain& BoneChain : BoneChainsToUpdate)
 	{
-		Log.LogWarning(LOCTEXT("SourceChainLessThanThree", "IK Retargeter trying to retarget source bone chain with IK but it has less than 2 joints."));
-		return false;
+		BoneChain.UpdatePoseFromSkeleton(InSkeleton);
 	}
-	
-	Source.StartBoneIndex = BoneIndices[0];
-	Source.EndBoneIndex = BoneIndices.Last();
-	
-	const FTransform& End = SourceInitialGlobalPose[Source.EndBoneIndex];
-	Source.PreviousEndPosition = End.GetTranslation();
-	Source.InitialEndPosition = End.GetTranslation();
-	Source.InitialEndRotation = End.GetRotation();
-
-	const FTransform& Start = SourceInitialGlobalPose[Source.StartBoneIndex];
-	const float Length = static_cast<float>((Start.GetTranslation() - Source.InitialEndPosition).Size());
+}
 
-	if (Length <= UE_KINDA_SMALL_NUMBER)
-	{
-		Log.LogWarning(LOCTEXT("SourceZeroLengthIK", "IK Retargeter trying to retarget source bone chain with IK, but it is zero length!"));
-    	return false;
-	}
-
-	Source.InvInitialLength = 1.0f / Length;
-	
-	return true;
-}
-
-void FChainRetargeterIK::EncodePose(const TArray<FTransform>& InSourceGlobalPose)
-{
-	const FVector Start = InSourceGlobalPose[Source.StartBoneIndex].GetTranslation();
-	const FVector End = InSourceGlobalPose[Source.EndBoneIndex].GetTranslation();
-
-    // get the normalized direction / length of the IK limb (how extended it is as percentage of original length)
-    const FVector ChainVector = End - Start;
-	float ChainLength;
-	FVector ChainDirection;
-	ChainVector.ToDirectionAndLength(ChainDirection, ChainLength);
-	const float NormalizedLimbLength = ChainLength * Source.InvInitialLength;
-
-	Source.PreviousEndPosition = Source.CurrentEndPosition;
-	Source.CurrentEndPosition = End;
-	Source.CurrentEndDirectionNormalized = ChainDirection * NormalizedLimbLength;
-	Source.CurrentEndRotation = InSourceGlobalPose[Source.EndBoneIndex].GetRotation();
-	Source.CurrentHeightFromGroundNormalized = static_cast<float>(End.Z - Source.InitialEndPosition.Z)  * Source.InvInitialLength;
-}
-
-bool FChainRetargeterIK::InitializeTarget(
-	const TArray<int32>& BoneIndices,
-	const TArray<FTransform> &TargetInitialGlobalPose,
-	FIKRigLogger& Log)
-{
-	if (BoneIndices.Num() < 3)
-	{
-		Log.LogWarning(LOCTEXT("TargetChainLessThanThree", "IK Retargeter trying to retarget target bone chain with IK but it has less than 3 joints."));
-		return false;
-	}
-	
-	Target.BoneIndexA = BoneIndices[0];
-	Target.BoneIndexC = BoneIndices.Last();
-	const FTransform& Last = TargetInitialGlobalPose[BoneIndices.Last()];
-	Target.PrevEndPosition = Last.GetLocation();
-	Target.InitialEndPosition = Last.GetTranslation();
-	Target.InitialEndRotation = Last.GetRotation();
-	Target.InitialLength = static_cast<float>((TargetInitialGlobalPose[Target.BoneIndexA].GetTranslation() - Last.GetTranslation()).Size());
-	ResetThisTick = true;
-
-	if (Target.InitialLength <= UE_KINDA_SMALL_NUMBER)
-	{
-		Log.LogWarning(LOCTEXT("TargetZeroLengthIK", "IK Retargeter trying to retarget target bone chain with IK, but it is zero length!"));
-		return false;
-	}
-	
-	return true;
-}
-	
-void FChainRetargeterIK::DecodePose(
-	const FTargetChainIKSettings& Settings,
-	const FTargetChainSpeedPlantSettings& SpeedPlantSettings,
-	const FRootRetargeter& RootRetargeter,
-	const TMap<FName, float>& SpeedValuesFromCurves,
-	const float DeltaTime,
-    const TArray<FTransform>& InGlobalPose)
-{
-	//
-	// calculate ROTATION of IK goal ...
-	//
-	
-	// apply delta rotation from input
-	const FQuat DeltaRotation = Source.CurrentEndRotation * Source.InitialEndRotation.Inverse();
-	FQuat GoalRotation = DeltaRotation * Target.InitialEndRotation;
-
-	// apply static rotation offset in the local space of the foot
-	GoalRotation = GoalRotation * Settings.StaticRotationOffset.Quaternion();
-
-	// blend to source rotation
-	const float BlendToSourceRotation = Settings.BlendToSource * Settings.BlendToSourceRotation;
-	if (BlendToSourceRotation > UE_KINDA_SMALL_NUMBER)
-	{
-		GoalRotation = FQuat::FastLerp(GoalRotation, Source.CurrentEndRotation, BlendToSourceRotation);
-		GoalRotation.Normalize();
-	}
-
-	//
-	// calculate POSITION of IK goal ...
-	//
-	
-	// set position to length-scaled direction from source limb
-	const FVector InvAffectIKWeights = FVector::OneVector - RootRetargeter.Settings.GetAffectIKWeightVector();
-	const FVector InvRootModification = RootRetargeter.Target.RootTranslationDelta * InvAffectIKWeights;
-	const FVector Start = InGlobalPose[Target.BoneIndexA].GetTranslation() - InvRootModification;
-	FVector GoalPosition = Start + (Source.CurrentEndDirectionNormalized * Target.InitialLength);
-
-	// blend to source location
-	const float BlendToSourceTranslation = Settings.BlendToSource * Settings.BlendToSourceTranslation;
-	if (BlendToSourceTranslation > UE_KINDA_SMALL_NUMBER)
-	{
-		const FVector RootModification = RootRetargeter.Target.RootTranslationDelta * RootRetargeter.Settings.GetAffectIKWeightVector();
-		const FVector Weight = BlendToSourceTranslation * Settings.BlendToSourceWeights;
-		const FVector SourceLocation = Source.CurrentEndPosition + RootModification;
-		GoalPosition.X = FMath::Lerp(GoalPosition.X, SourceLocation.X, Weight.X);
-		GoalPosition.Y = FMath::Lerp(GoalPosition.Y, SourceLocation.Y, Weight.Y);
-		GoalPosition.Z = FMath::Lerp(GoalPosition.Z, SourceLocation.Z, Weight.Z);
-	}
-
-	// apply global static offset
-	GoalPosition += Settings.StaticOffset;
-
-	// apply local static offset
-	GoalPosition += GoalRotation.RotateVector(Settings.StaticLocalOffset);
-
-	// apply vertical scale
-	GoalPosition.Z *= Settings.ScaleVertical;
-	
-	// apply extension
-	if (!FMath::IsNearlyEqual(Settings.Extension, 1.0f))
-	{
-		GoalPosition = Start + (GoalPosition - Start) * Settings.Extension;	
-	}
-	
-	// match velocity
-	if (!ResetThisTick && SpeedPlantSettings.EnableSpeedPlanting && SpeedValuesFromCurves.Contains(SpeedPlantSettings.SpeedCurveName))
-	{
-		const float SourceSpeed = SpeedValuesFromCurves[SpeedPlantSettings.SpeedCurveName];
-		if (SourceSpeed < 0.0f || SourceSpeed > SpeedPlantSettings.SpeedThreshold)
-		{
-			GoalPosition = UKismetMathLibrary::VectorSpringInterp(
-				Target.PrevEndPosition, GoalPosition, PlantingSpringState,
-				SpeedPlantSettings.UnplantStiffness,
-				SpeedPlantSettings.UnplantCriticalDamping,
-				DeltaTime, 1.0f, 0.0f);
-		}
-		else
-		{
-			PlantingSpringState.Reset();
-			GoalPosition = Target.PrevEndPosition;
-		}
-	}
-	
-	// output transform
-	Results.EndEffectorPosition = GoalPosition;
-	Results.EndEffectorRotation = GoalRotation;
-	Results.PoleVectorPosition = FVector::OneVector; // TODO calc pole vector position
-	Target.PrevEndPosition = GoalPosition;
-	ResetThisTick = false;
-}
-
-void FChainRetargeterIK::SaveDebugInfo(const TArray<FTransform>& InGlobalPose)
+void FRetargeterBoneChains::Reset()
 {
-#if WITH_EDITOR
-	DebugData.InputTransformStart = InGlobalPose[Target.BoneIndexA];
-	DebugData.InputTransformEnd = InGlobalPose[Target.BoneIndexC];
-	DebugData.OutputTransformEnd = FTransform(Results.EndEffectorRotation, Results.EndEffectorPosition);
-#endif
-}
-
-bool FRetargetChainPair::Initialize(
-    const FBoneChain& SourceBoneChain,
-    const FBoneChain& TargetBoneChain,
-    const FRetargetSkeleton& SourceSkeleton,
-    const FTargetSkeleton& TargetSkeleton,
-    FIKRigLogger& Log)
-{
-	// validate source bone chain is compatible with source skeletal mesh
-	const bool bIsSourceValid = ValidateBoneChainWithSkeletalMesh(true, SourceBoneChain, SourceSkeleton, Log);
-	if (!bIsSourceValid)
-	{
-		Log.LogWarning( FText::Format(
-			LOCTEXT("IncompatibleSourceChain", "IK Retargeter source bone chain, '{0}', is not compatible with Skeletal Mesh: '{1}'"),
-			FText::FromName(SourceBoneChain.ChainName), FText::FromString(SourceSkeleton.SkeletalMesh->GetName())));
-		return false;
-	}
-
-	// validate target bone chain is compatible with target skeletal mesh
-	const bool bIsTargetValid = ValidateBoneChainWithSkeletalMesh(false, TargetBoneChain, TargetSkeleton, Log);
-	if (!bIsTargetValid)
-    {
-		Log.LogWarning( FText::Format(
-			LOCTEXT("IncompatibleTargetChain", "IK Retargeter target bone chain, '{0}', is not compatible with Skeletal Mesh: '{1}'"),
-			FText::FromName(TargetBoneChain.ChainName), FText::FromString(TargetSkeleton.SkeletalMesh->GetName())));
-		return false;
-    }
-
-	// store attributes of chain
-	SourceBoneChainName = SourceBoneChain.ChainName;
-	TargetBoneChainName = TargetBoneChain.ChainName;
-	
-	return true;
-}
-
-bool FRetargetChainPair::ValidateBoneChainWithSkeletalMesh(
-    const bool IsSource,
-    const FBoneChain& BoneChain,
-    const FRetargetSkeleton& RetargetSkeleton,
-    FIKRigLogger& Log)
-{
-	// record the chain indices
-	TArray<int32>& BoneIndices = IsSource ? SourceBoneIndices : TargetBoneIndices;
-	
-	// resolve the bone bone to the skeleton
-	const FResolvedBoneChain ResolvedChain = FResolvedBoneChain(BoneChain, RetargetSkeleton, BoneIndices);
-	
-	// warn if START bone not found
-	if (!ResolvedChain.bFoundStartBone)
-	{
-		Log.LogWarning( FText::Format(
-			LOCTEXT("MissingStartBone", "IK Retargeter bone chain, {0}, could not find start bone, {1} in mesh {2}"),
-			FText::FromName(BoneChain.ChainName),
-			FText::FromName(BoneChain.StartBone.BoneName),
-			FText::FromString(RetargetSkeleton.SkeletalMesh->GetName())));
-	}
-	
-	// warn if END bone not found
-	if (!ResolvedChain.bFoundEndBone)
-	{
-		Log.LogWarning( FText::Format(
-			LOCTEXT("MissingEndBone", "IK Retargeter bone chain, {0}, could not find end bone, {1} in mesh {2}"),
-			FText::FromName(BoneChain.ChainName), FText::FromName(BoneChain.EndBone.BoneName), FText::FromString(RetargetSkeleton.SkeletalMesh->GetName())));
-	}
-
-	// warn if END bone was not a child of START bone
-	if (ResolvedChain.bFoundEndBone && !ResolvedChain.bEndIsStartOrChildOfStart)
-	{
-		Log.LogWarning( FText::Format(
-			LOCTEXT("EndNotChildtOfStart", "IK Retargeter bone chain, {0}, end bone, '{1}' was not a child of the start bone '{2}'."),
-			FText::FromName(BoneChain.ChainName), FText::FromName(BoneChain.EndBone.BoneName), FText::FromName(BoneChain.StartBone.BoneName)));
-	}
-	
-	return ResolvedChain.IsValid();
-}
-
-bool FRetargetChainPairFK::Initialize(
-	const FBoneChain& SourceBoneChain,
-	const FBoneChain& TargetBoneChain,
-	const FRetargetSkeleton& SourceSkeleton,
-	const FTargetSkeleton& TargetSkeleton,
-	FIKRigLogger& Log)
-{
-	const bool bChainInitialized = FRetargetChainPair::Initialize(SourceBoneChain, TargetBoneChain, SourceSkeleton, TargetSkeleton, Log);
-	if (!bChainInitialized)
-	{
-		return false;
-	}
-
-	// initialize SOURCE FK chain encoder with retarget pose
-	const bool bFKEncoderInitialized = FKEncoder.Initialize(SourceSkeleton, SourceBoneIndices, SourceSkeleton.RetargetGlobalPose, Log);
-	if (!bFKEncoderInitialized)
-	{
-		Log.LogWarning( FText::Format(
-			LOCTEXT("BadFKEncoder", "IK Retargeter failed to initialize FK encoder, '{0}', on Skeletal Mesh: '{1}'"),
-			FText::FromName(SourceBoneChainName), FText::FromString(SourceSkeleton.SkeletalMesh->GetName())));
-		return false;
-	}
-
-	// initialize TARGET FK chain decoder with retarget pose
-	const bool bFKDecoderInitialized = FKDecoder.Initialize(TargetSkeleton, TargetBoneIndices, TargetSkeleton.RetargetGlobalPose, Log);
-	if (!bFKDecoderInitialized)
-	{
-		Log.LogWarning( FText::Format(
-			LOCTEXT("BadFKDecoder", "IK Retargeter failed to initialize FK decoder, '{0}', on Skeletal Mesh: '{1}'"),
-			FText::FromName(TargetBoneChainName), FText::FromString(TargetSkeleton.SkeletalMesh->GetName())));
-		return false;
-	}
-
-	// initialize the pole vector matcher for this chain
-	const bool bPoleVectorMatcherInitialized = PoleVectorMatcher.Initialize(
-		SourceBoneIndices,
-		TargetBoneIndices,
-		SourceSkeleton.RetargetGlobalPose,
-		TargetSkeleton.RetargetGlobalPose,
-		TargetSkeleton);
-	if (!bPoleVectorMatcherInitialized)
-	{
-		Log.LogWarning( FText::Format(
-			LOCTEXT("BadPoleVectorMatcher", "IK Retargeter failed to initialize pole matching for chain, '{0}', on Skeletal Mesh: '{1}'"),
-			FText::FromName(TargetBoneChainName), FText::FromString(TargetSkeleton.SkeletalMesh->GetName())));
-		return false;
-	}
-
-	return true;
-}
-
-bool FRetargetChainPairIK::Initialize(
-	const FBoneChain& SourceBoneChain,
-	const FBoneChain& TargetBoneChain,
-	const FRetargetSkeleton& SourceSkeleton,
-	const FTargetSkeleton& TargetSkeleton,
-	FIKRigLogger& Log)
-{
-	// validate if this chain even uses an IK Goal
-	const bool bUsingIKGoal = TargetBoneChain.IKGoalName != NAME_None && TargetBoneChain.IKGoalName != "- None -";
-	if (!bUsingIKGoal)
-	{
-		return false;
-	}
-
-	// store target IK goal name
-	IKGoalName = TargetBoneChain.IKGoalName;
-
-	// initialize bone chains
-	const bool bChainInitialized = FRetargetChainPair::Initialize(SourceBoneChain, TargetBoneChain, SourceSkeleton, TargetSkeleton, Log);
-	if (!bChainInitialized)
-	{
-		return false;
-	}
-
-	// initialize SOURCE IK chain encoder with retarget pose
-	const bool bIKEncoderInitialized = IKChainRetargeter.InitializeSource(SourceBoneIndices, SourceSkeleton.RetargetGlobalPose, Log);
-	if (!bIKEncoderInitialized)
-	{
-		Log.LogWarning( FText::Format(
-			LOCTEXT("BadIKEncoder", "IK Retargeter failed to initialize IK encoder, '{0}', on Skeletal Mesh: '{1}'"),
-			FText::FromName(SourceBoneChainName), FText::FromString(SourceSkeleton.SkeletalMesh->GetName())));
-		return false;
-	}
-
-	// initialize TARGET IK chain decoder with retarget pose
-	const bool bIKDecoderInitialized = IKChainRetargeter.InitializeTarget(TargetBoneIndices, TargetSkeleton.RetargetGlobalPose, Log);
-	if (!bIKDecoderInitialized)
-	{
-		Log.LogWarning( FText::Format(
-			LOCTEXT("BadIKDecoder", "IK Retargeter failed to initialize IK decoder, '{0}', on Skeletal Mesh: '{1}'"),
-			FText::FromName(TargetBoneChainName), FText::FromString(TargetSkeleton.SkeletalMesh->GetName())));
-		return false;
-	}
-
-	return true;
+	SourceBoneChains.Reset();
+	TargetBoneChains.Reset();
+	ChainMapping.Reset();
+	bIsInitialized = false;
 }
 
 void FBasePoseCopier::CopyPoseFromSource(
@@ -1630,312 +1028,11 @@ void FBasePoseCopier::Initialize(
 	}
 }
 
-bool FRootRetargeter::InitializeSource(
-	const FName SourceRootBoneName,
-	const FRetargetSkeleton& SourceSkeleton,
-	FIKRigLogger& Log)
-{
-	// validate target root bone exists
-	Source.BoneName = SourceRootBoneName;
-	Source.BoneIndex = SourceSkeleton.FindBoneIndexByName(SourceRootBoneName);
-	if (Source.BoneIndex == INDEX_NONE)
-	{
-		Log.LogWarning( FText::Format(
-			LOCTEXT("MissingSourceRoot", "IK Retargeter could not find source root bone, {0} in mesh {1}"),
-			FText::FromName(SourceRootBoneName), FText::FromString(SourceSkeleton.SkeletalMesh->GetName())));
-		return false;
-	}
-	
-	// record initial root data
-	const FTransform InitialTransform = SourceSkeleton.RetargetGlobalPose[Source.BoneIndex]; 
-	float InitialHeight = static_cast<float>(InitialTransform.GetTranslation().Z);
-	Source.InitialRotation = InitialTransform.GetRotation();
-
-	// ensure root height is not at origin, this happens if user sets root to ACTUAL skeleton root and not pelvis
-	if (InitialHeight < UE_KINDA_SMALL_NUMBER)
-	{
-		// warn user and push it up slightly to avoid divide by zero
-		Log.LogError(LOCTEXT("BadRootHeight", "The source retarget root bone is very near the ground plane. This will cause the target to be moved very far. To resolve this, please create a retarget pose with the retarget root at the correct height off the ground."));
-		InitialHeight = 1.0f;
-	}
-
-	// invert height
-	Source.InitialHeightInverse = 1.0f / InitialHeight;
-
-	return true;
-}
-
-bool FRootRetargeter::InitializeTarget(
-	const FName TargetRootBoneName,
-	const FTargetSkeleton& TargetSkeleton,
-	FIKRigLogger& Log)
-{
-	// validate target root bone exists
-	Target.BoneName = TargetRootBoneName;
-	Target.BoneIndex = TargetSkeleton.FindBoneIndexByName(TargetRootBoneName);
-	if (Target.BoneIndex == INDEX_NONE)
-	{
-		Log.LogWarning( FText::Format(
-			LOCTEXT("CountNotFindRootBone", "IK Retargeter could not find target root bone, {0} in mesh {1}"),
-			FText::FromName(TargetRootBoneName), FText::FromString(TargetSkeleton.SkeletalMesh->GetName())));
-		return false;
-	}
-
-	const FTransform TargetInitialTransform = TargetSkeleton.RetargetGlobalPose[Target.BoneIndex];
-	Target.InitialHeight = static_cast<float>(TargetInitialTransform.GetTranslation().Z);
-	Target.InitialRotation = TargetInitialTransform.GetRotation();
-	Target.InitialPosition = TargetInitialTransform.GetTranslation();
-
-	// initialize the global scale factor
-	const float ScaleFactor = Source.InitialHeightInverse * Target.InitialHeight;
-	GlobalScaleFactor.Set(ScaleFactor, ScaleFactor, ScaleFactor);
-	
-	return true;
-}
-
-void FRootRetargeter::Reset()
-{
-	Source = FRootSource();
-	Target = FRootTarget();
-}
-
-void FRootRetargeter::EncodePose(const TArray<FTransform>& SourceGlobalPose)
-{
-	const FTransform& SourceTransform = SourceGlobalPose[Source.BoneIndex];
-	Source.CurrentPosition = SourceTransform.GetTranslation();
-	Source.CurrentPositionNormalized = Source.CurrentPosition * Source.InitialHeightInverse;
-	Source.CurrentRotation = SourceTransform.GetRotation();	
-}
-
-void FRootRetargeter::DecodePose(TArray<FTransform>& OutTargetGlobalPose)
-{
-	// retarget position
-	FVector Position;
-	{
-		// generate basic retarget root position by scaling the normalized position by root height
-		const FVector RetargetedPosition = Source.CurrentPositionNormalized * Target.InitialHeight;
-		
-		// blend the retarget root position towards the source retarget root position
-		Position = FMath::Lerp(RetargetedPosition, Source.CurrentPosition, Settings.BlendToSource*Settings.BlendToSourceWeights);
-
-		// apply vertical / horizontal scaling of motion
-		FVector ScaledRetargetedPosition = Position;
-		ScaledRetargetedPosition.Z *= Settings.ScaleVertical;
-		const FVector HorizontalOffset = (ScaledRetargetedPosition - Target.InitialPosition) * FVector(Settings.ScaleHorizontal, Settings.ScaleHorizontal, 1.0f);
-		Position = Target.InitialPosition + HorizontalOffset;
-		
-		// apply a static offset
-		Position += Settings.TranslationOffset;
-
-		// blend with alpha
-		Position = FMath::Lerp(Target.InitialPosition, Position, Settings.TranslationAlpha);
-
-		// record the delta created by all the modifications made to the root translation
-		Target.RootTranslationDelta = Position - RetargetedPosition;
-	}
-
-	// retarget rotation
-	FQuat Rotation;
-	{
-		// calc offset between initial source/target root rotations
-		const FQuat RotationDelta = Source.CurrentRotation * Source.InitialRotation.Inverse();
-		// add retarget pose delta to the current source rotation
-		const FQuat RetargetedRotation = RotationDelta * Target.InitialRotation;
-
-		// add static rotation offset
-		Rotation = RetargetedRotation * Settings.RotationOffset.Quaternion();
-
-		// blend with alpha
-		Rotation = FQuat::FastLerp(Target.InitialRotation, Rotation, Settings.RotationAlpha);
-		Rotation.Normalize();
-
-		// record the delta created by all the modifications made to the root rotation
-		Target.RootRotationDelta = RetargetedRotation * Target.InitialRotation.Inverse();
-	}
-
-	// apply to target
-	FTransform& TargetRootTransform = OutTargetGlobalPose[Target.BoneIndex];
-	TargetRootTransform.SetTranslation(Position);
-	TargetRootTransform.SetRotation(Rotation);
-}
-
-bool FPoleVectorMatcher::Initialize(
-	const TArray<int32>& SourceIndices,
-	const TArray<int32>& TargetIndices,
-	const TArray<FTransform>& SourceGlobalPose,
-	const TArray<FTransform>& TargetGlobalPose,
-	const FRetargetSkeleton& TargetSkeleton)
-{
-	SourcePoleAxis = CalculateBestPoleAxisForChain(SourceIndices, SourceGlobalPose);
-	TargetPoleAxis = CalculateBestPoleAxisForChain(TargetIndices, TargetGlobalPose);
-
-	const FVector SourcePoleVector = CalculatePoleVector(SourcePoleAxis, SourceIndices, SourceGlobalPose);
-	const FVector TargetPoleVector = CalculatePoleVector(TargetPoleAxis, TargetIndices, TargetGlobalPose);
-	
-	TargetToSourceAngularOffsetAtRefPose = static_cast<float>(FMath::Acos(FVector::DotProduct(TargetPoleVector, SourcePoleVector)));
-
-	// cache indices of bones in chain (and their children recursively) excluding children of the end bone
-	TArray<int32> AllChildrenIndices;
-	TargetSkeleton.GetChildrenIndicesRecursive(TargetIndices[0],AllChildrenIndices);
-	TArray<int32> ChildrenOfEndIndices = {TargetIndices.Last()};
-	TargetSkeleton.GetChildrenIndicesRecursive(TargetIndices.Last(),ChildrenOfEndIndices);
-	AllChildrenWithinChain.Reset();
-	for (int32 ChildIndex : AllChildrenIndices)
-	{
-		if (ChildrenOfEndIndices.Contains(ChildIndex))
-		{
-			continue;
-		}
-		AllChildrenWithinChain.Add(ChildIndex);
-	}
-
-	return true;
-}
-
-void FPoleVectorMatcher::MatchPoleVector(
-	const FTargetChainFKSettings& Settings,
-	const TArray<int32>& SourceIndices,
-	const TArray<int32>& TargetIndices,
-	const TArray<FTransform>& SourceGlobalPose,
-	TArray<FTransform>& OutTargetGlobalPose,
-	FRetargetSkeleton& TargetSkeleton)
-{
-	const bool bIsMatchingPoleVector = Settings.PoleVectorMatching > UE_KINDA_SMALL_NUMBER;
-	const bool bIsOffsetingPoleVector = !FMath::IsNearlyZero(Settings.PoleVectorOffset);
-	if (!(bIsMatchingPoleVector || bIsOffsetingPoleVector))
-	{
-		return;
-	}
-	
-	// update local spaces of all bones within chain
-	TArray<FTransform> LocalSpaces;
-	for (const int32 ChildIndex : AllChildrenWithinChain)
-	{
-		const int32 ParentIndex = TargetSkeleton.ParentIndices[ChildIndex];
-		if (ParentIndex == INDEX_NONE)
-		{
-			LocalSpaces.Add(OutTargetGlobalPose[ChildIndex]);
-		}
-		const FTransform& ChildGlobalTransform = OutTargetGlobalPose[ChildIndex];
-		const FTransform& ParentGlobalTransform = OutTargetGlobalPose[ParentIndex];
-		const FTransform LocalTransform = ChildGlobalTransform.GetRelativeTransform(ParentGlobalTransform);
-		LocalSpaces.Add(LocalTransform);
-	}
-
-	// normalized vector pointing from root to tip of chain
-	const FVector TargetChainAxisNorm = GetChainAxisNormalized(TargetIndices, OutTargetGlobalPose);
-
-	// calculate rotation to match the target to the source pole vector
-	FQuat MatchingRotation = FQuat::Identity;
-	if (bIsMatchingPoleVector)
-	{
-		const FVector SourcePoleVector = CalculatePoleVector(SourcePoleAxis, SourceIndices, SourceGlobalPose);
-		const FVector TargetPoleVector = CalculatePoleVector(TargetPoleAxis, TargetIndices, OutTargetGlobalPose);
-	
-		const float RotateTargetToSource = static_cast<float>(FMath::Acos(FVector::DotProduct(SourcePoleVector, TargetPoleVector)));
-		const float MatchPoleAngle = RotateTargetToSource - (Settings.PoleVectorMaintainOffset ? TargetToSourceAngularOffsetAtRefPose : 0);
-		
-		MatchingRotation = FQuat(TargetChainAxisNorm, MatchPoleAngle);
-		MatchingRotation = FQuat::FastLerp(FQuat::Identity, MatchingRotation, Settings.PoleVectorMatching).GetNormalized();
-	}
-
-	// manual offset rotation
-	FQuat OffsetRotation = FQuat::Identity;
-	if (bIsOffsetingPoleVector)
-	{
-		OffsetRotation = FQuat(TargetChainAxisNorm,FMath::DegreesToRadians(Settings.PoleVectorOffset));
-	}
-
-	// rotate the base of the chain to match the pole vectors
-	FTransform& BaseOfChain = OutTargetGlobalPose[TargetIndices[0]];
-	BaseOfChain.SetRotation(MatchingRotation * OffsetRotation * BaseOfChain.GetRotation());
-
-	// now update global pose of all bones within chain
-	for (int32 ChainIndex=0; ChainIndex < AllChildrenWithinChain.Num(); ++ChainIndex)
-	{
-		const int32 ChildIndex = AllChildrenWithinChain[ChainIndex];
-		const int32 ParentIndex = TargetSkeleton.ParentIndices[ChildIndex];
-		if (ParentIndex == INDEX_NONE)
-		{
-			// root always in global space already, no conversion required
-			OutTargetGlobalPose[ChildIndex] = LocalSpaces[ChainIndex];
-			return; 
-		}
-		const FTransform& ChildLocalTransform = LocalSpaces[ChainIndex];
-		const FTransform& ParentGlobalTransform = OutTargetGlobalPose[ParentIndex];
-		OutTargetGlobalPose[ChildIndex] = ChildLocalTransform * ParentGlobalTransform;
-	}
-}
-
-EAxis::Type FPoleVectorMatcher::CalculateBestPoleAxisForChain(
-	const TArray<int32>& BoneIndices,
-	const TArray<FTransform>& GlobalPose)
-{
-	check(!BoneIndices.IsEmpty());
-
-	TArray<EAxis::Type> PreferredAxes;
-
-	// handle 1-bone case, there is no "chain" to speak of just a single bone
-	// so we arbitrarily pick an axis and move on... realistically, this feature is fairly meaningless for a single bone
-	if (BoneIndices.Num() == 1)
-	{
-		return EAxis::Y;
-	}
-	
-	const FVector ChainOrigin = GlobalPose[BoneIndices[0]].GetLocation();
-	const FVector ChainAxisNormal = (GlobalPose[BoneIndices.Last()].GetLocation() - ChainOrigin).GetSafeNormal();
-	const EAxis::Type MostDifferentAxis = GetMostDifferentAxis(GlobalPose[BoneIndices[0]], ChainAxisNormal);
-	return MostDifferentAxis;
-}
-
-FVector FPoleVectorMatcher::CalculatePoleVector(
-	const EAxis::Type& PoleAxis,
-	const TArray<int32>& BoneIndices,
-	const TArray<FTransform>& GlobalPose)
-{
-	check(!BoneIndices.IsEmpty())
-
-	const FVector ChainNormal = GetChainAxisNormalized(BoneIndices, GlobalPose);
-	const FVector UnitPoleAxis = GlobalPose[BoneIndices[0]].GetUnitAxis(PoleAxis);
-	const FVector PoleVector = FVector::VectorPlaneProject(UnitPoleAxis, ChainNormal);
-	return PoleVector.GetSafeNormal();
-}
-
-EAxis::Type FPoleVectorMatcher::GetMostDifferentAxis(
-	const FTransform& Transform,
-	const FVector& InNormal)
-{
-	float MostDifferentDot = 2.0f;
-	EAxis::Type MostDifferentAxis = EAxis::Y;
-	const TArray<EAxis::Type> CardinalAxes = {EAxis::X, EAxis::Y, EAxis::Z};
-	for (const EAxis::Type Axis : CardinalAxes)
-	{
-		const FVector AxisVector = Transform.GetUnitAxis(Axis);
-		const float AbsAxisDotNormal = static_cast<float>(FMath::Abs(FVector::DotProduct(AxisVector, InNormal)));
-		if (AbsAxisDotNormal < MostDifferentDot)
-		{
-			MostDifferentDot = AbsAxisDotNormal;
-			MostDifferentAxis = Axis;
-		}
-	}
-
-	return MostDifferentAxis;
-}
-
-FVector FPoleVectorMatcher::GetChainAxisNormalized(
-	const TArray<int32>& BoneIndices,
-	const TArray<FTransform>& GlobalPose)
-{
-	const FVector ChainOrigin = GlobalPose[BoneIndices[0]].GetLocation();
-	const FVector ChainAxis = GlobalPose[BoneIndices.Last()].GetLocation() - ChainOrigin;
-	return ChainAxis.GetSafeNormal();
-}
-
 void FIKRetargetProcessor::Initialize(
 		const USkeletalMesh* SourceSkeletalMesh,
 		const USkeletalMesh* TargetSkeletalMesh,
 		const UIKRetargeter* InRetargeterAsset,
-		const FRetargetProfile& Settings,
+		const FRetargetProfile& InRetargetProfile,
 		const bool bSuppressWarnings)
 {
 	check(InRetargeterAsset && SourceSkeletalMesh && TargetSkeletalMesh);
@@ -1950,14 +1047,16 @@ void FIKRetargetProcessor::Initialize(
 		return;
 	}
 	
-	// reset all initialized flags
+	// assumed initialization fails unless we get to the bottom of this function
 	bIsInitialized = false;
-	bRootsInitialized = false;
-	bAtLeastOneValidBoneChainPair = false;
-	bIKRigInitialized = false;
-	
+
 	// record source asset
 	RetargeterAsset = InRetargeterAsset;
+	if (!RetargeterAsset)
+	{
+		Log.LogError(LOCTEXT("MissingRetargetAsset", "IK Retargeter unable to initialize because no IK Retargeter asset was supplied."));
+		return;
+	}
 
 	const UIKRigDefinition* SourceIKRig = RetargeterAsset->GetIKRig(ERetargetSourceOrTarget::Source);
 	const UIKRigDefinition* TargetIKRig = RetargeterAsset->GetIKRig(ERetargetSourceOrTarget::Target);
@@ -1984,14 +1083,16 @@ void FIKRetargetProcessor::Initialize(
 		return;
 	}
 	
-	// initialize skeleton data for source and target
+	// initialize source skeleton
 	SourceSkeleton.Initialize(
 		SourceSkeletalMesh,
 		SourceIKRig->GetRetargetChains(),
 		RetargeterAsset->GetCurrentRetargetPoseName(ERetargetSourceOrTarget::Source),
 		RetargeterAsset->GetCurrentRetargetPose(ERetargetSourceOrTarget::Source),
 		SourceIKRig->GetRetargetRoot(),
-		GlobalSettings.SourceScaleFactor);
+		GetSourceScaleFactor());
+	
+	// initialize target skeleton
 	TargetSkeleton.Initialize(
 		TargetSkeletalMesh,
 		TargetIKRig->GetRetargetChains(),
@@ -1999,367 +1100,143 @@ void FIKRetargetProcessor::Initialize(
 		RetargeterAsset->GetCurrentRetargetPose(ERetargetSourceOrTarget::Target),
 		TargetIKRig->GetRetargetRoot());
 
-	// initialize roots
-	bRootsInitialized = InitializeRoots();
+	// resolve all bone chains onto the skeletons
+	AllBoneChains.Initialize(RetargeterAsset, SourceSkeleton, TargetSkeleton, Log);
 
-	// initialize pairs of bone chains
-	bAtLeastOneValidBoneChainPair = InitializeBoneChainPairs(Settings.ChainSettings);
-	if (!bAtLeastOneValidBoneChainPair)
-	{
-		// couldn't match up any BoneChain pairs, no limb retargeting possible
-		Log.LogWarning( FText::Format(
-			LOCTEXT("NoMappedChains", "IK Retargeter unable to map any bone chains between source, {0} and target, {1}"),
-			FText::FromString(SourceSkeletalMesh->GetName()), FText::FromString(TargetSkeletalMesh->GetName())));
-	}
+	// record which chain each bone belongs to (must be done AFTER BoneChains.Initialize())
+	SourceSkeleton.FillPerBoneChainNames(AllBoneChains.GetResolvedBoneChains(ERetargetSourceOrTarget::Source));
+	TargetSkeleton.FillPerBoneChainNames(AllBoneChains.GetResolvedBoneChains(ERetargetSourceOrTarget::Target));
 
-	// initialize the IKRigProcessor for doing IK decoding
-	bIKRigInitialized = InitializeIKRig(TargetSkeletalMesh);
-	if (!bIKRigInitialized)
-	{
-		// couldn't initialize the IK Rig, we don't disable the retargeter in this case, just warn the user
-		Log.LogWarning( FText::Format(
-			LOCTEXT("CouldNotInitializeIKRig", "IK Retargeter was unable to initialize the IK Rig, {0} for the Skeletal Mesh {1}. See previous warnings."),
-			FText::FromString(TargetIKRig->GetName()), FText::FromString(TargetSkeletalMesh->GetName())));
-	}
-
-	// initialize the post operations
-	const bool bAllOpsInitialized = InitializeOpStack(RetargeterAsset->GetRetargetOps());
-	if (!bAllOpsInitialized)
-	{
-		Log.LogWarning(FText(LOCTEXT("PostOperationInitError", "One of the retarget post operations has warnings. See output log for details.")));
-	}
+	// setup the goal container
+	ResetIKRigGoalContainer();
 
-	// must have a mapped root bone OR at least a single mapped chain to be able to do any retargeting at all
-	if (bRootsInitialized && bAtLeastOneValidBoneChainPair)
-	{
-		// confirm for the user that the IK Rig was initialized successfully
-		Log.LogInfo(FText::Format(
-				LOCTEXT("SuccessfulInit", "Success! The IK Retargeter is ready to transfer animation from the source, {0} to the target, {1}"),
-				FText::FromString(SourceSkeletalMesh->GetName()), FText::FromString(TargetSkeletalMesh->GetName())));
-	}
+	// initialize the retargeting ops
+	InitializeOpStack(RetargeterAsset->GetRetargetOps(), InRetargetProfile);
 	
-	// copy the initial settings from the asset
-	CopyIKRigSettingsFromAsset();
+	// determine which bones are retargeted and run PostInitialize() on each op
+	PostInitializeOpStack();
+
+	// confirm for the user that the retargeter is initialized
+	Log.LogInfo(FText::Format(
+			LOCTEXT("SuccessfulInit", "Success! The IK Retargeter is ready to transfer animation from the source, {0} to the target, {1}"),
+			FText::FromString(SourceSkeletalMesh->GetName()), FText::FromString(TargetSkeletalMesh->GetName())));
 	
-	bIsInitialized = bRootsInitialized || bAtLeastOneValidBoneChainPair;
+	bIsInitialized = true;
 	AssetVersionInitializedWith = RetargeterAsset->GetVersion();
 #if WITH_EDITOR
 	RetargeterInitialized.Broadcast();
 #endif
 }
 
-bool FIKRetargetProcessor::InitializeRoots()
+void FIKRetargetProcessor::ResetIKRigGoalContainer()
 {
-	// reset root data
-	RootRetargeter.Reset();
-	
-	// initialize root encoder
-	const FName SourceRootBoneName = RetargeterAsset->GetIKRig(ERetargetSourceOrTarget::Source)->GetRetargetRoot();
-	const bool bRootEncoderInit = RootRetargeter.InitializeSource(SourceRootBoneName, SourceSkeleton, Log);
-	if (!bRootEncoderInit)
-	{
-		Log.LogWarning( FText::Format(
-			LOCTEXT("NoSourceRoot", "IK Retargeter unable to initialize source root, '{0}' on skeletal mesh: '{1}'"),
-			FText::FromName(SourceRootBoneName), FText::FromString(SourceSkeleton.SkeletalMesh->GetName())));
-	}
-
-	// initialize root decoder
-	const FName TargetRootBoneName =RetargeterAsset->GetIKRig(ERetargetSourceOrTarget::Target)->GetRetargetRoot();
-	const bool bRootDecoderInit = RootRetargeter.InitializeTarget(TargetRootBoneName, TargetSkeleton, Log);
-	if (!bRootDecoderInit)
-	{
-		Log.LogWarning( FText::Format(
-			LOCTEXT("NoTargetRoot", "IK Retargeter unable to initialize target root, '{0}' on skeletal mesh: '{1}'"),
-			FText::FromName(TargetRootBoneName), FText::FromString(TargetSkeleton.SkeletalMesh->GetName())));
-	}
-
-	return bRootEncoderInit && bRootDecoderInit;
-}
-
-bool FIKRetargetProcessor::InitializeBoneChainPairs(const TMap<FName, FTargetChainSettings>& ChainSettings)
-{
-	ChainPairsFK.Reset();
-	ChainPairsIK.Reset();
-
-	const UIKRigDefinition* SourceIKRig = RetargeterAsset->GetIKRig(ERetargetSourceOrTarget::Source);
+	// fill the goal container with all the goals in the IK Rig
 	const UIKRigDefinition* TargetIKRig = RetargeterAsset->GetIKRig(ERetargetSourceOrTarget::Target);
-	check(SourceIKRig && TargetIKRig);
-
-	// check that chains are available in both IKRig assets before sorting them based on StartBone index
-	const TArray<TObjectPtr<URetargetChainSettings>>& ChainMapping = RetargeterAsset->GetAllChainSettings();	
-	for (const URetargetChainSettings* ChainMap : ChainMapping)
-	{
-		// get target bone chain
-		const FBoneChain* TargetBoneChain = TargetIKRig->GetRetargetChainByName(ChainMap->TargetChain);
-		if (!TargetBoneChain)
-		{
-			Log.LogWarning( FText::Format(
-			LOCTEXT("MissingTargetChain", "IK Retargeter missing target bone chain: {0}. Please update the mapping."),
-			FText::FromString(ChainMap->TargetChain.ToString())));
-			continue;
-		}
-		
-		// user opted to not map this to anything, we don't need to spam a warning about it
-		if (ChainMap->SourceChain == NAME_None)
-		{
-			continue; 
-		}
-		
-		// get source bone chain
-		const FBoneChain* SourceBoneChain = SourceIKRig->GetRetargetChainByName(ChainMap->SourceChain);
-		if (!SourceBoneChain)
-		{
-			Log.LogWarning( FText::Format(
-			LOCTEXT("MissingSourceChain", "IK Retargeter missing source bone chain: {0}"),
-			FText::FromString(ChainMap->SourceChain.ToString())));
-			continue;
-		}
-
-		// all chains are loaded as FK (giving IK better starting pose)
-		FRetargetChainPairFK ChainPair;
-		if (ChainPair.Initialize(*SourceBoneChain, *TargetBoneChain, SourceSkeleton, TargetSkeleton, Log))
-		{
-			ChainPairsFK.Add(ChainPair);
-		}
-		
-		// load IK chain
-		FRetargetChainPairIK ChainPairIK;
-		if (ChainPairIK.Initialize(*SourceBoneChain, *TargetBoneChain, SourceSkeleton, TargetSkeleton, Log))
-		{
-			ChainPairsIK.Add(ChainPairIK);
-		}
-
-		// warn user if IK goal is not on the END bone of the target chain. It will still work, but may produce bad results.
-		const TArray<UIKRigEffectorGoal*> AllGoals = TargetIKRig->GetGoalArray();
-		for (const UIKRigEffectorGoal*Goal : AllGoals)
-		{
-			if (Goal->GoalName == TargetBoneChain->IKGoalName)
-			{
-				if (Goal->BoneName != TargetBoneChain->EndBone.BoneName)
-				{
-					Log.LogWarning( FText::Format(
-				LOCTEXT("TargetIKNotOnEndBone", "Retarget chain, '{0}' has an IK goal that is not on the End Bone of the chain."),
-					FText::FromString(TargetBoneChain->ChainName.ToString())));
-				}
-				break;
-			}
-		}
-	}
-
-	// set initial chain settings
-	for (FRetargetChainPairFK& FKChain : ChainPairsFK)
-	{
-		if (const FTargetChainSettings* Settings = ChainSettings.Find(FKChain.TargetBoneChainName))
-		{
-			FKChain.Settings = Settings->FK;
-		}
-	}
-	for (FRetargetChainPairIK& IKChain : ChainPairsIK)
-	{
-		if (const FTargetChainSettings* Settings = ChainSettings.Find(IKChain.TargetBoneChainName))
-		{
-			IKChain.Settings = Settings->IK;
-		}
-	}
-
-	// sort the chains based on their StartBone's index
-	auto ChainsSorter = [this](const FRetargetChainPair& A, const FRetargetChainPair& B)
-	{
-		const int32 IndexA = A.TargetBoneIndices.Num() > 0 ? A.TargetBoneIndices[0] : INDEX_NONE;
-		const int32 IndexB = B.TargetBoneIndices.Num() > 0 ? B.TargetBoneIndices[0] : INDEX_NONE;
-		if (IndexA == IndexB)
-		{
-			// fallback to sorting alphabetically
-			return A.TargetBoneChainName.LexicalLess(B.TargetBoneChainName);
-		}
-		return IndexA < IndexB;
-	};
+	GoalContainer.FillWithGoalArray(TargetIKRig->GetGoalArray());
 	
-	ChainPairsFK.Sort(ChainsSorter);
-	ChainPairsIK.Sort(ChainsSorter);
-	
-	// record which bones in the target skeleton are being retargeted
-	for (const FRetargetChainPairFK& FKChainPair : ChainPairsFK)
-	{
-		for (const int32 BoneIndex : FKChainPair.TargetBoneIndices)
-		{
-			TargetSkeleton.SetBoneIsRetargeted(BoneIndex, true);
-		}
-	}
-
-	// record intermediate bones (non-retargeted bones located BETWEEN FK chains on the target skeleton)
-	for (FRetargetChainPairFK& FKChainPair : ChainPairsFK)
+	// set all goals to be additive (ops may adjust them)
+	for (FIKRigGoal& Goal : GoalContainer.GetGoalArray())
 	{
-		FKChainPair.FKDecoder.InitializeIntermediateParentIndices(
-			RootRetargeter.Target.BoneIndex,
-			FKChainPair.TargetBoneIndices[0],
-			TargetSkeleton);
+		Goal.PositionSpace = EIKRigGoalSpace::Additive;
+		Goal.RotationSpace = EIKRigGoalSpace::Additive;
+		Goal.Position = FVector::ZeroVector;
+		Goal.Rotation = FRotator::ZeroRotator;
 	}
-
-	// root is updated before IK as well
-	if (bRootsInitialized)
-	{
-		TargetSkeleton.SetBoneIsRetargeted(RootRetargeter.Target.BoneIndex, true);	
-	}
-
-	// return true if at least 1 pair of bone chains were initialized
-	return !(ChainPairsIK.IsEmpty() && ChainPairsFK.IsEmpty());
 }
 
-bool FIKRetargetProcessor::InitializeIKRig(const USkeletalMesh* InSkeletalMesh)
+bool FIKRetargetProcessor::InitializeOpStack(
+	const TArray<FInstancedStruct>& OpsFromAsset,
+	const FRetargetProfile& InRetargetProfile)
 {
-	// get the target IK Rig asset
-	UIKRigDefinition* IKRig = RetargeterAsset->GetIKRigWriteable(ERetargetSourceOrTarget::Target);
-	
-	// gather list of excluded goals based on any chain mapping that has it's IK disabled
-	TArray<FName> GoalsToExclude;
-	for (FRetargetChainPairIK ChainPairIK : ChainPairsIK)
-	{
-		if (!ChainPairIK.Settings.EnableIK)
-		{
-			GoalsToExclude.Add(ChainPairIK.IKGoalName);
-		}
-	}
+	bool bAllOpsInitialized = true;
 	
-	// initialize IK Rig runtime processor
-	IKRigProcessor.Initialize(IKRig, InSkeletalMesh, GoalsToExclude);
-	if (!IKRigProcessor.IsInitialized())
-	{
-		return false;
-	}
-
-	// validate that all IK bone chains have an associated Goal
-	for (FRetargetChainPairIK& ChainPair : ChainPairsIK)
+	// create copies of all the ops in the asset
+	OpStack.Reset(OpsFromAsset.Num());
+	for (const FInstancedStruct& AssetOpStruct : OpsFromAsset)
 	{
-		// skip chains that have had their IK disabled, these goals won't be in the IK Rig
-		if (GoalsToExclude.Contains(ChainPair.IKGoalName))
+		if (!ensure(AssetOpStruct.IsValid()))
 		{
+			// this can happen if asset references deleted op type which should only happen during development (if at all)
+			Log.LogWarning(LOCTEXT("UnknownOP", "IK Retargeter, '{0}' has null/unknown op in it. Please reload the asset to remove it."));
 			continue;
 		}
+
+		// create a copy of the op
+		int32 NewOpIndex = OpStack.Emplace(AssetOpStruct);
+		// apply the profile settings to the op (may be different than what was loaded from the asset)
+		InRetargetProfile.ApplyOpProfilesToOpStruct(OpStack[NewOpIndex], NewOpIndex, ECopyOpSettingsContext::PreInitialize);
 		
-		// does the IK rig have the IK goal this bone chain requires?
-		if (!IKRigProcessor.GetGoalContainer().FindGoalByName(ChainPair.IKGoalName))
-		{
-			Log.LogError( FText::Format(
-			LOCTEXT("TargetIKBoneNotInSolver", "IK Retargeter has target bone chain, {0} that references an IK Goal, {1} that is not present in any of the solvers in the IK Rig asset."),
-			FText::FromName(ChainPair.TargetBoneChainName), FText::FromName(ChainPair.IKGoalName)));
-			return false;
-		}
+		// initialize it
+		FIKRetargetOpBase& NewOp = OpStack[NewOpIndex].GetMutable<FIKRetargetOpBase>();
+		NewOp.Initialize(*this, SourceSkeleton, TargetSkeleton, Log);
+		
+		
+#if WITH_EDITOR
+		FIKRetargetOpBase* AssetOp = const_cast<FIKRetargetOpBase*>(AssetOpStruct.GetPtr<FIKRetargetOpBase>());
+		FIKRetargetOpSettingsBase* AssetOpSettings = AssetOp->GetSettings();
+		// asset instance gets reference to running editor instance
+		AssetOpSettings->EditorInstance = NewOp.GetSettings();
+		// references the skeletons are used for bone selector widgets in the details panel
+		AssetOpSettings->SourceSkeletonAsset = SourceSkeleton.SkeletalMesh->GetSkeleton();
+		AssetOpSettings->TargetSkeletonAsset = TargetSkeleton.SkeletalMesh->GetSkeleton();
+#endif
 	}
 
-	// get a single retarget chain using the specified goal (if there is one)
-	auto GetRetargetChainByGoalName = [](const TArray<FBoneChain>& Chains, const FName GoalName) -> const FBoneChain*
-	{
-		for (const FBoneChain& Chain : Chains)
-		{
-			if (Chain.IKGoalName == GoalName)
-			{
-				return &Chain;
-			}
-		}
-
-		return nullptr;
-	};
-
-	// detect goals that are not being controlled by the retargeter and set them to "additive" mode
-	const TArray<FIKRigGoal>& GoalsInProcessor = IKRigProcessor.GetGoalContainer().GetGoalArray();
-	const TArray<FBoneChain>& ChainsInAsset = IKRig->GetRetargetChains();
-	const FIKRigSkeleton& IKRigSkeleton = IKRigProcessor.GetSkeleton();
-	for (const FIKRigGoal& GoalInProcessor : GoalsInProcessor)
+	if (!bAllOpsInitialized)
 	{
-		if (GetRetargetChainByGoalName(ChainsInAsset, GoalInProcessor.Name))
-		{
-			// this goal belongs to a chain, so leave it alone
-			continue;
-		}
-
-		// this goal has no associated chain, so we set it to "additive" mode to prevent the goal
-		// from being left at the starting location (and thus completely breaking the pose)
-		// when in "additive" mode, any offset applied in the IK Rig will be maintained, but will be relative
-		// to the input pose, not in absolute/component space.
-		FIKRigGoal GoalToModify = GoalInProcessor;
-		GoalToModify.PositionSpace = EIKRigGoalSpace::Additive;
-		GoalToModify.RotationSpace = EIKRigGoalSpace::Additive;
-
-		// but for this to work, the goal's position/rotation must first be converted to "additive" space
-		const FGoalBone* GoalBone = IKRigProcessor.GetGoalBone(GoalToModify.Name);
-		const int32& BoneIndex = GoalBone->BoneIndex;
-		check(BoneIndex != INDEX_NONE)
-
-		const FTransform& GoalBoneGlobalTransform = IKRigSkeleton.RefPoseGlobal[BoneIndex];
-	
-		// convert position to additive space
-		if (GoalToModify.PositionSpace == EIKRigGoalSpace::World || GoalToModify.PositionSpace == EIKRigGoalSpace::Component)
-		{
-			GoalToModify.Position = GoalToModify.Position - GoalBoneGlobalTransform.GetLocation();
-		}
-
-		// convert rotation to additive space
-		if (GoalToModify.RotationSpace == EIKRigGoalSpace::World || GoalToModify.RotationSpace == EIKRigGoalSpace::Component)
-		{
-			GoalToModify.Rotation = (GoalToModify.Rotation.Quaternion() * GoalBoneGlobalTransform.GetRotation().Inverse()).Rotator();
-		}
-		
-		// assign the new goal forced to be additive
-		IKRigProcessor.SetIKGoal(GoalToModify);
-
-		// let the user know to avoid confusion
-		Log.LogWarning( FText::Format(
-		LOCTEXT("TargetIKGoalNotInChain", "The target IK Rig has a Goal, {0} that is not present in any of the target chains. It has been set to 'Additive' mode while retargeting."),
-		FText::FromName(GoalInProcessor.Name)));
+		Log.LogWarning(FText(LOCTEXT("PostOperationInitError", "One (or more) of the retarget operations has warnings. See output log for details.")));
 	}
 	
-	return true;
+	return bAllOpsInitialized;
 }
 
-bool FIKRetargetProcessor::InitializeOpStack(const TArray<FInstancedStruct>& OpsFromAsset)
+void FIKRetargetProcessor::PostInitializeOpStack()
 {
-	bool bAllOpsInitialized = true;
-	
-	// create copies of all the ops in the asset
-	OpStack.Reset(OpsFromAsset.Num());
-	for (const FInstancedStruct& AssetOpStruct : OpsFromAsset)
+	// gather retargeted bones from each op
+	TSet<int32> AllRetargetedBones;
+	for (FInstancedStruct& OpStruct : OpStack)
 	{
-		if (!ensure(AssetOpStruct.IsValid()))
+		FIKRetargetOpBase& Op = OpStruct.GetMutable<FIKRetargetOpBase>();
+		if (Op.IsEnabled() && Op.IsInitialized())
 		{
-			// this can happen if asset references deleted op type which should only happen during development (if at all)
-			Log.LogWarning(LOCTEXT("UnknownOP", "IK Retargeter, '{0}' has null/unknown op in it. Please reload the asset to remove it."));
-			continue;
+			Op.CollectRetargetedBones(AllRetargetedBones);
 		}
-
-		// copy the op memory
-		const UScriptStruct* ScriptStruct = AssetOpStruct.GetScriptStruct();
-		FIKRetargetOpBase* NewOp = static_cast<FIKRetargetOpBase*>(FMemory::Malloc(ScriptStruct->GetStructureSize()));
-		ScriptStruct->InitializeStruct(NewOp);
-		ScriptStruct->CopyScriptStruct(NewOp, AssetOpStruct.GetMemory());
-		
-		// initialize it 
-		NewOp->Initialize(*this, SourceSkeleton, TargetSkeleton, Log);
-		bAllOpsInitialized &= NewOp->IsInitialized();
-
-		// store it in the runtime stack
-		OpStack.Emplace(MakeShareable(NewOp));
-		
-		// asset instance gets reference to running editor instance
-		#if WITH_EDITOR
-		FIKRetargetOpBase* AssetOp = const_cast<FIKRetargetOpBase*>(AssetOpStruct.GetPtr<FIKRetargetOpBase>());
-		AssetOp->GetSettings()->EditorInstance = NewOp->GetSettings();
-		#endif
 	}
+	// store retargeted bone mask on target skeleton
+	TargetSkeleton.SetRetargetedBones(AllRetargetedBones);
 
-	return bAllOpsInitialized;
+	// run PostInitialize() on all ops
+	for (FInstancedStruct& OpStruct : OpStack)
+	{
+		FIKRetargetOpBase& Op = OpStruct.GetMutable<FIKRetargetOpBase>();
+		if (Op.IsEnabled() && Op.IsInitialized())
+		{
+			Op.PostInitialize(*this, SourceSkeleton, TargetSkeleton, Log);
+		}
+	}
 }
 
 TArray<FTransform>&  FIKRetargetProcessor::RunRetargeter(
 	TArray<FTransform>& InSourceGlobalPose,
-	const TMap<FName, float>& SpeedValuesFromCurves,
-	const float DeltaTime,
-	const FRetargetProfile& Settings)
+	const FRetargetProfile& InProfile,
+	const float InDeltaTime)
 {
-	check(bIsInitialized);
+	if (!ensure(bIsInitialized))
+	{
+		return InSourceGlobalPose;
+	}
+	
+#if WITH_EDITOR
+	// validate system running the retargeter has stripped all the scale out of the incoming pose
+	for (const FTransform& Transform : InSourceGlobalPose)
+	{
+		const bool bHasNoScale = Transform.GetScale3D().Equals(FVector::OneVector);
+		ensureMsgf(bHasNoScale, TEXT("Found scale values on incoming pose in retarget processor. Scale should be baked into translation and set to 1,1,1. "));
+	}
+#endif
 
 	// apply the retargeting settings
-	ApplySettingsFromProfile(Settings);
+	ApplyProfileAtRuntime(InProfile);
 
 	// applying settings can cause the retargeter to require reinitialization (when enabling/disabling IK on a chain)
 	// we should be able to safely reinitialize here because we are already initialized or it wouldn't pass the check() at the top
@@ -2370,7 +1247,7 @@ TArray<FTransform>&  FIKRetargetProcessor::RunRetargeter(
 			SourceSkeleton.SkeletalMesh,
 			TargetSkeleton.SkeletalMesh,
 			RetargeterAsset,
-			Settings,
+			InProfile,
 			bSuppressWarnings);
 
 		if (!IsInitialized())
@@ -2379,63 +1256,30 @@ TArray<FTransform>&  FIKRetargetProcessor::RunRetargeter(
 			return TargetSkeleton.OutputGlobalPose;	
 		}
 	}
-
-#if WITH_EDITOR
-	// validate system running the retargeter has stripped all the scale out of the incoming pose
-	for (const FTransform& Transform : InSourceGlobalPose)
-	{
-		const bool bHasNoScale = Transform.GetScale3D().Equals(FVector::OneVector);
-		ensureMsgf(bHasNoScale, TEXT("Found scale values on incoming pose in retarget processor. Scale should be baked into translation and set to 1,1,1. "));
-	}
-#endif
-
-	// scale the input pose
-	if (!FMath::IsNearlyEqual(GlobalSettings.SourceScaleFactor, 1.0f))
-	{
-		for (FTransform& SourceBoneTransform : InSourceGlobalPose)
-		{
-			SourceBoneTransform.SetTranslation(SourceBoneTransform.GetTranslation() * GlobalSettings.SourceScaleFactor);
-		}
-	}
 	
 	// generate the pose to start retargeting from
 	GenerateBasePoses(InSourceGlobalPose);
 
-	// ROOT retargeting
-	if (GlobalSettings.bEnableRoot && bRootsInitialized)
-	{
-		RunRootRetarget(InSourceGlobalPose, TargetSkeleton.OutputGlobalPose);
-		// update global transforms below root
-		TargetSkeleton.UpdateGlobalTransformsBelowBone(RootRetargeter.Target.BoneIndex, TargetSkeleton.InputLocalPose, TargetSkeleton.OutputGlobalPose);
-	}
+	// reset all goals (ops may pose / adjust them)
+	ResetIKRigGoalContainer();
 	
-	// FK CHAIN retargeting
-	if (GlobalSettings.bEnableFK && bAtLeastOneValidBoneChainPair)
-	{
-		RunFKRetarget(InSourceGlobalPose, TargetSkeleton.OutputGlobalPose);
-		// update all the bones that are not controlled by FK chains or root
-		TargetSkeleton.UpdateGlobalTransformsAllNonRetargetedBones(TargetSkeleton.InputLocalPose,TargetSkeleton.OutputGlobalPose);
-	}
+	// run the stack of retargeting operations
+	RunRetargetOps(InDeltaTime, InSourceGlobalPose, TargetSkeleton.OutputGlobalPose);
 	
-	// IK CHAIN retargeting
-	if (bAtLeastOneValidBoneChainPair && bIKRigInitialized)
-	{
-		RunIKRetarget(InSourceGlobalPose, TargetSkeleton.OutputGlobalPose, SpeedValuesFromCurves, DeltaTime);
-	}
-
-	// Pole Vector matching between source / target chains
-	if (GlobalSettings.bEnableFK && bAtLeastOneValidBoneChainPair)
-	{
-		RunPoleVectorMatching(InSourceGlobalPose, TargetSkeleton.OutputGlobalPose);
-	}
+	return TargetSkeleton.OutputGlobalPose;
+}
 
-	// POST processes operations
-	if (GlobalSettings.bEnablePost)
+void FIKRetargetProcessor::ScaleSourcePose(TArray<FTransform>& InOutSourceGlobalPose) const
+{
+	// scale the input pose
+	const double SourceScaleFactor = GetSourceScaleFactor();
+	if (!FMath::IsNearlyEqual(SourceScaleFactor, 1.0f))
 	{
-		RunRetargetOps(InSourceGlobalPose, TargetSkeleton.OutputGlobalPose);
+		for (FTransform& SourceBoneTransform : InOutSourceGlobalPose)
+		{
+			SourceBoneTransform.SetTranslation(SourceBoneTransform.GetTranslation() * SourceScaleFactor);
+		}
 	}
-
-	return TargetSkeleton.OutputGlobalPose;
 }
 
 void FIKRetargetProcessor::GenerateBasePoses(TArray<FTransform>& InSourceGlobalPose)
@@ -2445,11 +1289,12 @@ void FIKRetargetProcessor::GenerateBasePoses(TArray<FTransform>& InSourceGlobalP
 	TargetSkeleton.InputLocalPose = TargetSkeleton.RetargetLocalPose;
 
 	// partially override with input pose
-	if (GlobalSettings.bCopyBasePose)
+	FName CopyBasePoseRootBoneName;
+	if (ShouldCopyBasePose(CopyBasePoseRootBoneName))
 	{
 		// do a straight copy of bone transforms from the source (no retargeting!)
 		BasePoseCopier.CopyPoseFromSource(
-			GlobalSettings.CopyBasePoseRoot,
+			CopyBasePoseRootBoneName,
 			SourceSkeleton,
 			TargetSkeleton,
 			InSourceGlobalPose,
@@ -2459,310 +1304,63 @@ void FIKRetargetProcessor::GenerateBasePoses(TArray<FTransform>& InSourceGlobalP
 	}
 }
 
-void FIKRetargetProcessor::RunRootRetarget(
-	const TArray<FTransform>& InGlobalTransforms,
-    TArray<FTransform>& OutGlobalTransforms)
-{
-	RootRetargeter.EncodePose(InGlobalTransforms);
-	RootRetargeter.DecodePose(OutGlobalTransforms);
-}
-
-void FIKRetargetProcessor::RunFKRetarget(
-	const TArray<FTransform>& InGlobalTransforms,
-    TArray<FTransform>& OutGlobalTransforms)
-{
-	// spin through chains and encode/decode them all using the input pose
-	for (FRetargetChainPairFK& ChainPair : ChainPairsFK)
-	{
-		ChainPair.FKEncoder.EncodePose(
-			SourceSkeleton,
-			ChainPair.SourceBoneIndices,
-			InGlobalTransforms);
-		
-		ChainPair.FKDecoder.DecodePose(
-			RootRetargeter,
-			ChainPair.Settings,
-			ChainPair.TargetBoneIndices,
-			ChainPair.FKEncoder,
-			TargetSkeleton,
-			OutGlobalTransforms);
-	}
-}
-
-void FIKRetargetProcessor::RunIKRetarget(
+void FIKRetargetProcessor::RunRetargetOps(
+	const double InDeltaTime,
 	const TArray<FTransform>& InSourceGlobalPose,
-    TArray<FTransform>& OutTargetGlobalPose,
-    const TMap<FName, float>& SpeedValuesFromCurves,
-    const float DeltaTime)
-{
-	if (!IKRigProcessor.IsInitialized())
-	{
-		return;
-	}
-	
-	if (ChainPairsIK.IsEmpty())
-	{
-		return; // skip IK
-	}
-	
-	// retarget the IK goals to their new locations based on input pose
-	for (FRetargetChainPairIK& ChainPair : ChainPairsIK)
-	{
-		// encode them all using the input pose
-		ChainPair.IKChainRetargeter.EncodePose(InSourceGlobalPose);
-		// decode the IK goal and apply to IKRig
-		ChainPair.IKChainRetargeter.DecodePose(
-			ChainPair.Settings,
-			ChainPair.SpeedPlantSettings,
-			RootRetargeter,
-			SpeedValuesFromCurves,
-			DeltaTime,
-			OutTargetGlobalPose);
-	}
-
-	// apply stride warping to the IK goals
-	RunStrideWarping(OutTargetGlobalPose);
-
-	// set the goal transform on the IK Rig
-	for (FRetargetChainPairIK& ChainPair : ChainPairsIK)
-	{
-		FIKRigGoal Goal = FIKRigGoal(
-			ChainPair.IKGoalName,
-			ChainPair.IKChainRetargeter.Results.EndEffectorPosition,
-			ChainPair.IKChainRetargeter.Results.EndEffectorRotation,
-			1.0f,
-			1.0f,
-			EIKRigGoalSpace::Component,
-			EIKRigGoalSpace::Component);
-		IKRigProcessor.SetIKGoal(Goal);
-
-		ChainPair.IKChainRetargeter.SaveDebugInfo(OutTargetGlobalPose);
-	}
-
-	if (GlobalSettings.bEnableIK)
-	{
-		// copy input pose to start IK solve from
-		IKRigProcessor.SetInputPoseGlobal(OutTargetGlobalPose);
-		// run IK solve
-		IKRigProcessor.Solve();
-		// copy results of solve
-		IKRigProcessor.GetOutputPoseGlobal(OutTargetGlobalPose);
-	}
-}
-
-void FIKRetargetProcessor::RunPoleVectorMatching(
-	const TArray<FTransform>& InGlobalTransforms,
-	TArray<FTransform>& OutGlobalTransforms)
+	TArray<FTransform>& OutTargetGlobalPose)
 {
-	for (FRetargetChainPairFK& ChainPair : ChainPairsFK)
+	for (FInstancedStruct& OpStruct : OpStack)
 	{
-		ChainPair.PoleVectorMatcher.MatchPoleVector(
-			ChainPair.Settings,
-			ChainPair.SourceBoneIndices,
-			ChainPair.TargetBoneIndices,
-			InGlobalTransforms,
-			OutGlobalTransforms,
-			TargetSkeleton);
+		FIKRetargetOpBase& Op = OpStruct.GetMutable<FIKRetargetOpBase>();
+		if (Op.IsEnabled() && Op.IsInitialized())
+		{
+			Op.Run(*this, InDeltaTime, InSourceGlobalPose, OutTargetGlobalPose);
+		}
 	}
 }
 
-void FIKRetargetProcessor::RunStrideWarping(const TArray<FTransform>& InTargetGlobalPose)
+TArray<const FIKRetargetOpBase*> FIKRetargetProcessor::GetRetargetOpsByType(const UScriptStruct* OpType) const
 {
-	// check if we should apply stride warping
-	if (!(GlobalSettings.bWarping && bRootsInitialized))
-	{
-		return;
-	}
+	TArray<const FIKRetargetOpBase*> Ops;
 
-	auto ConvertForwardVectorsToRotation = [](FVector& ForwardOrig, FVector& ForwardCurrent) -> FQuat
+	for (const FInstancedStruct& OpStruct : OpStack)
 	{
-		// project on floor (we only want Yaw) and normalize
-		ForwardOrig.Z = 0.0f;
-		ForwardCurrent.Z = 0.0f;
-		ForwardOrig.Normalize();
-		ForwardCurrent.Normalize();
-		// rotate from orig facing direction to current
-		return FQuat::FindBetweenNormals(ForwardOrig, ForwardCurrent);
-	};
-
-	// we need to determine the character's forward and side directions for warping
-	FVector InitialBodyPosition;
-	FVector CurrentBodyPosition;
-	FQuat CurrentRotation;
-	
-	// get sample points to use for "best fit" global body rotation
-	switch(GlobalSettings.DirectionSource)
-	{
-	case EWarpingDirectionSource::Goals:
+		if (OpStruct.GetScriptStruct()->IsChildOf(OpType))
 		{
-			TArray<FVector> InitialPoints;
-			TArray<FVector> CurrentPoints;
-			
-			// use goals to determine the body's rotation
-			for (const FRetargetChainPairIK& ChainPair : ChainPairsIK)
-			{
-				if (!ChainPair.Settings.bAffectedByIKWarping)
-				{
-					continue;
-				}
-				FVector A = ChainPair.IKChainRetargeter.Target.InitialEndPosition;
-				FVector B = ChainPair.IKChainRetargeter.Results.EndEffectorPosition;
-
-				// flatten into 2D for more robust yaw construction (what really matters)
-				A.Z = 0.0f;
-				B.Z = 0.0f;
-		
-				InitialPoints.Add(A);
-				CurrentPoints.Add(B);
-			}
-
-			// calculate "best fit" global body rotation based on deformation of sample points
-			CurrentRotation = GetRotationFromDeformedPoints(
-				InitialPoints,
-				CurrentPoints,
-				InitialBodyPosition,
-				CurrentBodyPosition);
-			
-			break;
+			const FIKRetargetOpBase* Op = OpStruct.GetPtr<FIKRetargetOpBase>();
+			Ops.Add(Op);
 		}
-	case EWarpingDirectionSource::Chain:
-		{
-			// use chain to determine the body's rotation
-			const FRetargetChainPairFK* Chain = ChainPairsFK.FindByPredicate(
-					[this](const FRetargetChainPairFK& Other)
-					{
-						return Other.TargetBoneChainName == GlobalSettings.DirectionChain;
-					});
-			if (!Chain)
-			{
-				return;
-			}
-
-			const TArray<FTransform>& InitialChainTransforms = Chain->FKDecoder.InitialGlobalTransforms;
-			const TArray<FTransform>& CurrentChainTransforms = Chain->FKDecoder.CurrentGlobalTransforms;
-			if (!ensure(!InitialChainTransforms.IsEmpty()))
-			{
-				return;
-			}
-
-			// calculate initial and current centroids
-			InitialBodyPosition = FVector::ZeroVector;
-			CurrentBodyPosition = FVector::ZeroVector;
-			for (int32 TransformIndex=0; TransformIndex<InitialChainTransforms.Num(); ++TransformIndex)
-			{
-				InitialBodyPosition += InitialChainTransforms[TransformIndex].GetTranslation();
-				CurrentBodyPosition += CurrentChainTransforms[TransformIndex].GetTranslation();
-			}
-			const float InvNum = 1.0f / static_cast<float>(InitialChainTransforms.Num());
-			InitialBodyPosition *= InvNum;
-			CurrentBodyPosition *= InvNum;
-
-			// get the forward vectors of the chain
-			FVector ForwardOrig;
-			FVector ForwardCurrent;
-			
-			// get forward vectors: orig and current
-			if (InitialChainTransforms.Num() == 1)
-			{
-				// in case of single bone chain, we rotate the global forward vector with this bone and then project onto floor
-				const FQuat DeltaRotation = InitialChainTransforms[0].GetRotation() * CurrentChainTransforms[0].GetRotation().Inverse();
-				ForwardOrig = GlobalSettings.GetAxisVector(GlobalSettings.ForwardDirection);
-				ForwardCurrent = DeltaRotation.RotateVector(ForwardOrig);
-			}
-			else
-			{
-				// in case of multi-bone chain, we use the vector from the start to the end of the chain
-				ForwardOrig = InitialChainTransforms.Last().GetTranslation() - InitialChainTransforms[0].GetTranslation();
-				ForwardCurrent = CurrentChainTransforms.Last().GetTranslation() - CurrentChainTransforms[0].GetTranslation();
-			}
-
-			// get rotation from fwd vectors
-			CurrentRotation = ConvertForwardVectorsToRotation(ForwardOrig, ForwardCurrent);
-			break;
 	}
-	case EWarpingDirectionSource::RootBone:
-		{
-			// use the root bone
-			InitialBodyPosition = TargetSkeleton.RetargetGlobalPose[0].GetTranslation();
-			CurrentBodyPosition = TargetSkeleton.OutputGlobalPose[0].GetTranslation();
-
-			// get the forward vectors of the root bone
-			const FQuat DeltaRotation = TargetSkeleton.RetargetGlobalPose[0].GetRotation() * TargetSkeleton.OutputGlobalPose[0].GetRotation().Inverse();
-			FVector ForwardOrig = GlobalSettings.GetAxisVector(GlobalSettings.ForwardDirection);
-			FVector ForwardCurrent = DeltaRotation.RotateVector(ForwardOrig);
 
-			// get rotation from fwd vectors
-			CurrentRotation = ConvertForwardVectorsToRotation(ForwardOrig, ForwardCurrent);
-			break;
-		}
-	default:
-		{
-			checkNoEntry();
-			return;
-		}
-	}
-	
-	FTransform CurrentBodyTransform = FTransform(CurrentRotation, CurrentBodyPosition);
-	FTransform InitialBodyTransform = FTransform(FQuat::Identity, InitialBodyPosition);
-	FVector Fwd = CurrentBodyTransform.TransformVector(GlobalSettings.GetAxisVector(GlobalSettings.ForwardDirection));
-	FVector Side = FVector::CrossProduct(Fwd, FVector::ZAxisVector);
-	FVector SideOrig = FVector::CrossProduct(GlobalSettings.GetAxisVector(GlobalSettings.ForwardDirection), FVector::ZAxisVector);
+	return MoveTemp(Ops);
+}
 
-	// warp goal positions...
-	for (FRetargetChainPairIK& ChainPair : ChainPairsIK)
+void FIKRetargetProcessor::OnPlaybackReset()
+{
+	for (FInstancedStruct& OpStruct : OpStack)
 	{
-		if (!ChainPair.Settings.bAffectedByIKWarping)
-		{
-			continue;
-		}
-		
-		FChainRetargeterIK& ChainRetargeter = ChainPair.IKChainRetargeter;
-
-		// warp this goal position in various ways...
-		FVector& WarpedGoalPosition = ChainRetargeter.Results.EndEffectorPosition;
-		
-		// forward warping
-		const FVector InitialGoalInOrigSpace = InitialBodyTransform.InverseTransformPosition(ChainRetargeter.Target.InitialEndPosition);
-		const FVector IntialGoalInCurrentSpace = CurrentBodyTransform.TransformPosition(InitialGoalInOrigSpace);
-		const FPlane FwdPlane(IntialGoalInCurrentSpace, Fwd);
-		const FVector GoalProjOnFwdPlane = FwdPlane.PointPlaneProject(WarpedGoalPosition, FwdPlane);
-		WarpedGoalPosition = GoalProjOnFwdPlane + ((WarpedGoalPosition - GoalProjOnFwdPlane) * GlobalSettings.WarpForwards);
-
-		// sideways offset
-		// first determine which side the goal is on originally
-		float GoalSideMultiplier = static_cast<float>(FVector::DotProduct(ChainRetargeter.Target.InitialEndPosition.GetSafeNormal(), SideOrig));
-		// push goal by offset in the newly calculated sideways direction
-		WarpedGoalPosition += Side * GlobalSettings.SidewaysOffset * GoalSideMultiplier;
-
-		// splay warping
-		FVector SplayOrigin = CurrentBodyPosition;
-		SplayOrigin.Z = WarpedGoalPosition.Z;
-		WarpedGoalPosition = SplayOrigin + (WarpedGoalPosition - SplayOrigin) * GlobalSettings.WarpSplay; 
+		FIKRetargetOpBase& Op = OpStruct.GetMutable<FIKRetargetOpBase>();
+		Op.OnPlaybackReset();
 	}
-
-#if WITH_EDITOR
-	DebugData.StrideWarpingFrame = CurrentBodyTransform;
-#endif
 }
 
-void FIKRetargetProcessor::RunRetargetOps(
-	const TArray<FTransform>& InSourceGlobalPose,
-	TArray<FTransform>& OutTargetGlobalPose)
+void FIKRetargetProcessor::OnAnimGraphPreUpdateMainThread(
+	USkeletalMeshComponent& SourceMeshComponent,
+	USkeletalMeshComponent& TargetMeshComponent)
 {
-	for (TSharedPtr<FIKRetargetOpBase>& RetargetOp : OpStack)
+	for (FInstancedStruct& OpStruct : OpStack)
 	{
-		if (RetargetOp->IsEnabled() && RetargetOp->IsInitialized())
-		{
-			RetargetOp->Run(*this, InSourceGlobalPose, OutTargetGlobalPose);
-		}
+		FIKRetargetOpBase& Op = OpStruct.GetMutable<FIKRetargetOpBase>();
+		Op.AnimGraphPreUpdateMainThread(SourceMeshComponent, TargetMeshComponent);
 	}
 }
 
-void FIKRetargetProcessor::ResetPlanting()
+void FIKRetargetProcessor::OnAnimGraphEvaluateAnyThread(FPoseContext& Output)
 {
-	for (FRetargetChainPairIK& ChainPair : ChainPairsIK)
+	for (FInstancedStruct& OpStruct : OpStack)
 	{
-		ChainPair.IKChainRetargeter.ResetThisTick = true;
+		FIKRetargetOpBase& Op = OpStruct.GetMutable<FIKRetargetOpBase>();
+		Op.AnimGraphEvaluateAnyThread(Output);
 	}
 }
 
@@ -2809,22 +1407,9 @@ bool FIKRetargetProcessor::WasInitializedWithTheseAssets(
 	return true;
 }
 
-void FIKRetargetProcessor::SetNeedsInitialized()
-{
-	bIsInitialized = false;
-
-	if (RetargeterAsset)
-	{
-		RetargeterAsset->IncrementVersion(); // triggers re-init
-	}
-	
-	// may not be initialized yet (during setup as prerequisites are being created)
-	IKRigProcessor.SetNeedsInitialized();
-}
-
 #if WITH_EDITOR
 
-bool FIKRetargetProcessor::IsBoneRetargeted(
+bool FIKRetargetProcessor::IsBoneInAMappedChain(
 	const FName BoneName,
 	const ERetargetSourceOrTarget SourceOrTarget) const
 {
@@ -2832,39 +1417,45 @@ bool FIKRetargetProcessor::IsBoneRetargeted(
 	
 	// NOTE: cannot use Skeleton.IsBoneRetargeted because that only exists
 	// on the target skeleton and this function needs to work for either skeleton
-	
+
+	// bone must exist
 	const FRetargetSkeleton& Skeleton = GetSkeleton(SourceOrTarget);
 	const int32 BoneIndex = Skeleton.FindBoneIndexByName(BoneName);
 	if (BoneIndex == INDEX_NONE)
 	{
-		// bone doesn't exist, therefore not retargeted
 		return false;
 	}
-	
-	const FName RootBoneName = GetRetargetRoot(SourceOrTarget);
+
+	// retarget root is always retargeted
+	const FName RootBoneName = GetPelvisBone(SourceOrTarget);
 	if (BoneName == RootBoneName)
 	{
-		// retarget root is always retargeted
 		return true;
 	}
+	
+	// bones must be in a chain
 	const FName ChainThatContainsBone = Skeleton.ChainThatContainsBone[BoneIndex];
 	if (ChainThatContainsBone == NAME_None)
 	{
-		// bones not in a chain, therefore cannot be retargeted
+		
 		return false;
 	}
 	
-	// bone must be in a successfully resolved chain that is mapped and ready to retarget
-	for (const FRetargetChainPairFK& ChainPair : ChainPairsFK)
+	// bone must be in a successfully resolved chain
+	const FResolvedBoneChain* Chain = AllBoneChains.GetResolvedBoneChainByName(ChainThatContainsBone, SourceOrTarget);
+	if (!Chain)
 	{
-		const FName ChainName = SourceOrTarget == ERetargetSourceOrTarget::Source ? ChainPair.SourceBoneChainName : ChainPair.TargetBoneChainName;
-		if (ChainName == ChainThatContainsBone)
-		{
-			return true;
-		}
+		return false;
 	}
 
-	return false;
+	// bone's chain must be mapped to alternate chain (source-to-target or vice versa)
+	const FName MappedChain = GetChainMappedToChain(Chain->ChainName, SourceOrTarget);
+	if (MappedChain == NAME_None)
+	{
+		return false;
+	}
+	
+	return true;
 }
 
 int32 FIKRetargetProcessor::GetBoneIndexFromName(
@@ -2874,12 +1465,18 @@ int32 FIKRetargetProcessor::GetBoneIndexFromName(
 	return GetSkeleton(SourceOrTarget).FindBoneIndexByName(BoneName);
 }
 
+const FRetargeterBoneChains& FIKRetargetProcessor::GetBoneChains()
+{
+	ensure(bIsInitialized);
+	return AllBoneChains;
+}
+
 FName FIKRetargetProcessor::GetChainNameForBone(const FName BoneName, const ERetargetSourceOrTarget SourceOrTarget) const
 {
 	const bool bUseSource = SourceOrTarget == ERetargetSourceOrTarget::Source;
 	const FRetargetSkeleton& Skeleton = bUseSource ? SourceSkeleton : TargetSkeleton;
 	const int32 BoneIndex = GetBoneIndexFromName(BoneName, SourceOrTarget);
-	const FName RootBoneName = GetRetargetRoot(SourceOrTarget);
+	const FName RootBoneName = GetPelvisBone(SourceOrTarget);
 	if (BoneName == RootBoneName)
 	{
 		return FName("Retarget Root");
@@ -2893,9 +1490,9 @@ FTransform FIKRetargetProcessor::GetGlobalRetargetPoseAtParam(
 	const float Param,
 	const ERetargetSourceOrTarget SourceOrTarget) const
 {
-	if (const FChainFK* Chain = GetChain(InChainName, SourceOrTarget))
+	if (const FResolvedBoneChain* Chain = GetBoneChains().GetResolvedBoneChainByName(InChainName, SourceOrTarget))
 	{
-		return Chain->GetTransformAtChainParam(Chain->InitialGlobalTransforms, Param);
+		return Chain->GetTransformAtChainParam(Chain->RefPoseGlobalTransforms, Param);
 	}
 
 	checkNoEntry();
@@ -2917,29 +1514,12 @@ FTransform FIKRetargetProcessor::GetRetargetPoseBoneTransform(
 	return BoneSpace == ERetargetBoneSpace::Global ? Skeleton.RetargetGlobalPose[BoneIndex] : Skeleton.RetargetLocalPose[BoneIndex];
 }
 
-const FChainFK* FIKRetargetProcessor::GetChain(const FName InChainName, const ERetargetSourceOrTarget SourceOrTarget) const
-{
-	const bool bIsSource = SourceOrTarget == ERetargetSourceOrTarget::Source;
-	for (const FRetargetChainPairFK& ChainPair : ChainPairsFK)
-	{
-		const FName ChainName = bIsSource ? ChainPair.SourceBoneChainName : ChainPair.TargetBoneChainName;
-		if (ChainName == InChainName)
-		{
-			const FChainFK* Encoder = &ChainPair.FKEncoder;
-			const FChainFK* Decoder = &ChainPair.FKDecoder;
-			return bIsSource ? Encoder : Decoder;
-		}
-	}
-
-	return nullptr;
-}
-
 float FIKRetargetProcessor::GetParamOfBoneInChain(
 	const FName InBoneName,
 	const ERetargetSourceOrTarget SourceOrTarget) const
 {
 	const FName ChainName = GetChainNameForBone(InBoneName, SourceOrTarget);
-	const FChainFK* Chain = GetChain(ChainName, SourceOrTarget);
+	const FResolvedBoneChain* Chain = GetBoneChains().GetResolvedBoneChainByName(ChainName, SourceOrTarget);
 	if (!Chain)
 	{
 		return INDEX_NONE;
@@ -2955,7 +1535,12 @@ FName FIKRetargetProcessor::GetClosestBoneToParam(
 	const float InParam,
 	const ERetargetSourceOrTarget SourceOrTarget) const
 {
-	const FChainFK* Chain = GetChain(InChainName, SourceOrTarget);
+	const FResolvedBoneChain* Chain = GetBoneChains().GetResolvedBoneChainByName(InChainName, SourceOrTarget);
+	if (!ensure(Chain != nullptr))
+	{
+		return NAME_None;
+	}
+	
 	const FRetargetSkeleton& Skeleton = GetSkeleton(SourceOrTarget);
 	if (Chain->BoneIndices.Num() == 1 || InParam < 0.0f)
 	{
@@ -2977,62 +1562,119 @@ FName FIKRetargetProcessor::GetClosestBoneToParam(
 	return Skeleton.BoneNames[Chain->BoneIndices[ChainIndexOfClosestBone]];
 }
 
-FName FIKRetargetProcessor::GetMappedChainName(
+FName FIKRetargetProcessor::GetChainMappedToChain(
 	const FName InChainName,
-	const ERetargetSourceOrTarget SourceOrTarget)
+	const ERetargetSourceOrTarget InSourceOrTarget) const
 {
-	const bool bIsSource = SourceOrTarget == ERetargetSourceOrTarget::Source;
-	for (const FRetargetChainPairFK& ChainPair : ChainPairsFK)
+	int32 IndexOfChain = GetBoneChains().GetIndexOfBoneChainByName(InChainName, InSourceOrTarget);
+	const TArray<int32>& ChainMappings = GetBoneChains().GetChainMappings();
+	
+	if (InSourceOrTarget == ERetargetSourceOrTarget::Target)
+	{
+		// find the mapped source chain
+		const TArray<FResolvedBoneChain>& SourceChains = GetBoneChains().GetResolvedBoneChains(ERetargetSourceOrTarget::Source);
+		int32 IndexOfMappedSourceChain = ChainMappings[IndexOfChain];
+		return SourceChains.IsValidIndex(IndexOfMappedSourceChain) ? SourceChains[IndexOfMappedSourceChain].ChainName : NAME_None;
+	}
+	else
 	{
-		const FName ChainName = bIsSource ? ChainPair.SourceBoneChainName : ChainPair.TargetBoneChainName;
-		if (ChainName == InChainName)
+		// find the mapped target chain
+		const TArray<FResolvedBoneChain>& TargetChains = GetBoneChains().GetResolvedBoneChains(ERetargetSourceOrTarget::Target);
+		for (int32 TargetChainIndex = 0; TargetChainIndex < ChainMappings.Num(); ++TargetChainIndex)
 		{
-			return bIsSource ? ChainPair.TargetBoneChainName : ChainPair.SourceBoneChainName;
+			if (ChainMappings[TargetChainIndex] == IndexOfChain)
+			{
+				return TargetChains.IsValidIndex(TargetChainIndex) ? TargetChains[TargetChainIndex].ChainName : NAME_None;
+			}
 		}
 	}
 
 	return NAME_None;
 }
 
-#endif
+void FIKRetargetProcessor::DebugDrawAllOps(
+	FPrimitiveDrawInterface* InPDI,
+	const FTransform& InComponentTransform,
+	const FIKRetargetDebugDrawState& EditorState) const
+{
+	const float ComponentScale = static_cast<float>(InComponentTransform.GetScale3D().GetMax());
+	
+	for (const FInstancedStruct& OpStruct : OpStack)
+	{
+		const FIKRetargetOpBase* Op = OpStruct.GetPtr<FIKRetargetOpBase>();
+		if (Op->IsEnabled() && Op->IsInitialized())
+		{
+			Op->DebugDraw(InPDI, InComponentTransform, ComponentScale, EditorState);
+		}
+	}
+}
 
-void FIKRetargetProcessor::CopyIKRigSettingsFromAsset()
+void FIKRetargetProcessor::DebugDrawOpAtIndex(
+	const int32 InOpIndex,
+	FPrimitiveDrawInterface* InPDI,
+	const FTransform& InComponentTransform,
+	const FIKRetargetDebugDrawState& EditorState) const
 {
-	// copy IK Rig settings
-	if (const UIKRigDefinition* TargetIKRig = RetargeterAsset->GetIKRig(ERetargetSourceOrTarget::Target))
+	if (!OpStack.IsValidIndex(InOpIndex))
+	{
+		return;
+	}
+
+	const FIKRetargetOpBase* Op = OpStack[InOpIndex].GetPtr<FIKRetargetOpBase>();
+	if (Op->IsEnabled() && Op->IsInitialized())
 	{
-		IKRigProcessor.CopyAllSettingsFromAsset(TargetIKRig);
+		const float ComponentScale = static_cast<float>(InComponentTransform.GetScale3D().GetMax());
+		Op->DebugDraw(InPDI, InComponentTransform, ComponentScale, EditorState);
 	}
 }
 
+#endif
+
 const FRetargetSkeleton& FIKRetargetProcessor::GetSkeleton(ERetargetSourceOrTarget SourceOrTarget) const
 {
 	return SourceOrTarget == ERetargetSourceOrTarget::Source ? SourceSkeleton : TargetSkeleton;
 }
 
-FName FIKRetargetProcessor::GetRetargetRoot(ERetargetSourceOrTarget SourceOrTarget) const
+FTargetSkeleton& FIKRetargetProcessor::GetTargetSkeleton()
 {
-	return SourceOrTarget == ERetargetSourceOrTarget::Source ? RootRetargeter.Source.BoneName : RootRetargeter.Target.BoneName;
+	return TargetSkeleton;
 }
 
-void FIKRetargetProcessor::ApplySettingsFromProfile(const FRetargetProfile& Profile)
+const FTargetSkeleton& FIKRetargetProcessor::GetTargetSkeleton() const
 {
-	// assign global settings
-	// NOTE: must be done before updating retarget pose because it uses the scale factors in the global settings
-	GlobalSettings = Profile.GlobalSettings;
-	
-	// assign retarget poses specified in the profile
-	UpdateRetargetPoseAtRuntime(Profile.SourceRetargetPoseName, ERetargetSourceOrTarget::Source);
-	UpdateRetargetPoseAtRuntime(Profile.TargetRetargetPoseName, ERetargetSourceOrTarget::Target);
+	return TargetSkeleton;
+}
+
+bool FIKRetargetProcessor::IsInitialized() const
+{
+	return bIsInitialized;
+}
+
+void FIKRetargetProcessor::SetNeedsInitialized()
+{
+	bIsInitialized = false;
+
+	if (RetargeterAsset)
+	{
+		RetargeterAsset->IncrementVersion(); // triggers re-init
+	}
+}
 
-	// assign chain settings
-	for (const TTuple<FName, FTargetChainSettings>& ChainSettings : Profile.ChainSettings)
+void FIKRetargetProcessor::ApplyProfileAtRuntime(const FRetargetProfile& Profile)
+{
+	// apply op settings from the profile
+	for (int32 OpIndex=0; OpIndex < OpStack.Num(); ++OpIndex)
 	{
-		UpdateChainSettingsAtRuntime(ChainSettings.Key, ChainSettings.Value);
+		Profile.ApplyOpProfilesToOpStruct(OpStack[OpIndex], OpIndex, ECopyOpSettingsContext::Runtime);
 	}
 
-	// assign root settings
-	RootRetargeter.Settings = Profile.RootSettings;
+	// apply flag to force IK off
+	bIKForcedOff = Profile.bForceAllIKOff;
+	
+	// apply retarget poses specified in the profile
+	// NOTE: must be done AFTER applying op settings because it uses the SourceScaleFactor in the ScaleSourceOp settings
+	UpdateRetargetPoseAtRuntime(Profile.SourceRetargetPoseName, ERetargetSourceOrTarget::Source);
+	UpdateRetargetPoseAtRuntime(Profile.TargetRetargetPoseName, ERetargetSourceOrTarget::Target);
 }
 
 void FIKRetargetProcessor::UpdateRetargetPoseAtRuntime(
@@ -3049,74 +1691,87 @@ void FIKRetargetProcessor::UpdateRetargetPoseAtRuntime(
 	FRetargetSkeleton& RetargetSkeleton = bIsSource ? SourceSkeleton : TargetSkeleton;
 	const bool bSameRetargetPose = RetargetSkeleton.RetargetPoseName == NewRetargetPoseName;
 	const bool bSameVersion = RetargetSkeleton.RetargetPoseVersion == NewRetargetPose->GetVersion();
-	const bool bSameGlobalScale = bIsSource ? FMath::IsNearlyEqual(RetargetSkeleton.RetargetPoseScale, GlobalSettings.SourceScaleFactor) : true;
+	const double SourceScaleFactor = GetSourceScaleFactor();
+	const bool bSameGlobalScale = bIsSource ? FMath::IsNearlyEqual(RetargetSkeleton.RetargetPoseScale, SourceScaleFactor) : true;
 	if (bSameRetargetPose && bSameVersion && bSameGlobalScale)
 	{
 		return; // retarget pose has not changed since it was initialized
 	}
 	
-	const FName RootBoneName = bIsSource ? RootRetargeter.Source.BoneName : RootRetargeter.Target.BoneName;
-	if (RootBoneName == NAME_None)
+	const FName PelvisBoneName = GetPelvisBone(SourceOrTarget);
+	if (PelvisBoneName == NAME_None)
 	{
-		return; // cannot regenerate pose without a root bone
+		return; // cannot regenerate pose without a pelvis bone
 	}
 
 	// re-generate the retarget pose
-	const float ScaleFactor = SourceOrTarget == ERetargetSourceOrTarget::Source ? GlobalSettings.SourceScaleFactor : 1.0f;
-	RetargetSkeleton.GenerateRetargetPose(NewRetargetPoseName, NewRetargetPose, RootBoneName, ScaleFactor);
+	const double ScaleFactor = SourceOrTarget == ERetargetSourceOrTarget::Source ? SourceScaleFactor : 1.0;
+	RetargetSkeleton.GenerateRetargetPose(NewRetargetPoseName, NewRetargetPose, PelvisBoneName, ScaleFactor);
 
-	// re-initialize the bone chains using the newly generated retarget pose
-	const UIKRigDefinition* SourceIKRig = RetargeterAsset->GetIKRig(ERetargetSourceOrTarget::Source);
-	const UIKRigDefinition* TargetIKRig = RetargeterAsset->GetIKRig(ERetargetSourceOrTarget::Target);
-	if (!(SourceIKRig && TargetIKRig))
+	// re-load the updated pose into the bone chains
+	AllBoneChains.UpdatePoseFromSkeleton(RetargetSkeleton, SourceOrTarget);
+
+	// re-initialize the op stack (ops sometimes cache retarget poses
+	for (FInstancedStruct& OpStruct : OpStack)
 	{
-		return;
+		FIKRetargetOpBase* Op = OpStruct.GetMutablePtr<FIKRetargetOpBase>();
+		Op->Initialize(*this, SourceSkeleton, TargetSkeleton, Log);
 	}
-	for (FRetargetChainPairFK& FKChainPair : ChainPairsFK)
+}
+
+FName FIKRetargetProcessor::GetPelvisBone(ERetargetSourceOrTarget SourceOrTarget) const
+{
+	// first try get it from a pelvis of
+	const FIKRetargetPelvisMotionOp* PelvisMotionOp = GetFirstRetargetOpOfType<FIKRetargetPelvisMotionOp>();
+	if (PelvisMotionOp)
 	{
-		const FBoneChain* SourceBoneChain = SourceIKRig->GetRetargetChainByName(FKChainPair.SourceBoneChainName);
-		const FBoneChain* TargetBoneChain = TargetIKRig->GetRetargetChainByName(FKChainPair.TargetBoneChainName);
-		FKChainPair.Initialize(*SourceBoneChain, *TargetBoneChain, SourceSkeleton, TargetSkeleton, Log);
+		FName PelvisNameFromOp = PelvisMotionOp->GetPelvisBoneName(SourceOrTarget);
+		if (PelvisNameFromOp != NAME_None)
+		{
+			return PelvisNameFromOp;
+		}
 	}
-	for (FRetargetChainPairIK& IKChainPair : ChainPairsIK)
+
+	// fallback to going directly to the IK Rig
+	const UIKRigDefinition* IKRigAsset = GetRetargetAsset()->GetIKRig(SourceOrTarget);
+	if (IKRigAsset)
 	{
-		const FBoneChain* SourceBoneChain = SourceIKRig->GetRetargetChainByName(IKChainPair.SourceBoneChainName);
-		const FBoneChain* TargetBoneChain = TargetIKRig->GetRetargetChainByName(IKChainPair.TargetBoneChainName);
-		IKChainPair.Initialize(*SourceBoneChain, *TargetBoneChain, SourceSkeleton, TargetSkeleton, Log);
+		return IKRigAsset->GetRetargetRoot();
 	}
 
-	// re-initialize the root
-	RootRetargeter.InitializeSource(RootRetargeter.Source.BoneName, SourceSkeleton, Log);
-	RootRetargeter.InitializeTarget(RootRetargeter.Target.BoneName, TargetSkeleton, Log);
+	return NAME_None;
 }
 
-void FIKRetargetProcessor::UpdateChainSettingsAtRuntime(const FName ChainName, const FTargetChainSettings& NewChainSettings)
+bool FIKRetargetProcessor::ShouldCopyBasePose(FName& OutCopyBasePoseRoot) const
 {
-	// update FK chain
-	for (FRetargetChainPairFK& Chain : ChainPairsFK)
+	// first try get it from the processor
+	const FIKRetargetCopyBasePoseOp* CopyBasePoseOp = GetFirstRetargetOpOfType<FIKRetargetCopyBasePoseOp>(ERetargetOpsToSearch::ProcessorOps);
+	// if the processor doesn't have the op, check the asset
+	CopyBasePoseOp = CopyBasePoseOp ? CopyBasePoseOp : GetFirstRetargetOpOfType<FIKRetargetCopyBasePoseOp>(ERetargetOpsToSearch::AssetOps);
+	if (CopyBasePoseOp && CopyBasePoseOp->IsEnabled())
 	{
-		if (Chain.TargetBoneChainName == ChainName)
-		{
-			Chain.Settings = NewChainSettings.FK;
-			break;
-		}
+		OutCopyBasePoseRoot = CopyBasePoseOp->Settings.CopyBasePoseRoot;
+		return CopyBasePoseOp->Settings.bCopyBasePose;
 	}
 
-	// update IK chain
-	for (FRetargetChainPairIK& Chain : ChainPairsIK)
+	// off by default
+	OutCopyBasePoseRoot = NAME_None;
+	return false;
+}
+
+double FIKRetargetProcessor::GetSourceScaleFactor() const
+{
+	// first try get it from the processor
+	const FIKRetargetScaleSourceOp* ScaleSourceOp = GetFirstRetargetOpOfType<FIKRetargetScaleSourceOp>(ERetargetOpsToSearch::ProcessorOps);
+	// if the processor doesn't have the op, check the asset
+	ScaleSourceOp = ScaleSourceOp ? ScaleSourceOp : GetFirstRetargetOpOfType<FIKRetargetScaleSourceOp>(ERetargetOpsToSearch::AssetOps);
+	if (ScaleSourceOp && ScaleSourceOp->IsEnabled())
 	{
-		if (Chain.TargetBoneChainName == ChainName)
-		{
-			if (Chain.Settings.EnableIK != NewChainSettings.IK.EnableIK)
-			{
-				SetNeedsInitialized();
-			}
-			
-			Chain.Settings = NewChainSettings.IK;
-			Chain.SpeedPlantSettings = NewChainSettings.SpeedPlanting;
-			break;
-		}
+		return ScaleSourceOp->Settings.SourceScaleFactor;
 	}
+
+	// if no source scale op is present, then we don't do any scaling
+	return 1.0;
 }
 
 #undef LOCTEXT_NAMESPACE
