[IKRig] Major refactor part 3 of 3.

diff --git a/Engine/Plugins/Animation/IKRig/Source/IKRig/Private/Retargeter/IKRetargetSettings.cpp b/Engine/Plugins/Animation/IKRig/Source/IKRig/Private/Retargeter/IKRetargetSettings.cpp
--- a/Engine/Plugins/Animation/IKRig/Source/IKRig/Private/Retargeter/IKRetargetSettings.cpp
+++ b/Engine/Plugins/Animation/IKRig/Source/IKRig/Private/Retargeter/IKRetargetSettings.cpp
@@ -1,44 +1,13 @@
 // Copyright Epic Games, Inc. All Rights Reserved.
 
 #include "Retargeter/IKRetargetSettings.h"
-#include "Retargeter/IKRetargeter.h"
 
 #include UE_INLINE_GENERATED_CPP_BY_NAME(IKRetargetSettings)
 
-bool FTargetChainSpeedPlantSettings::operator==(const FTargetChainSpeedPlantSettings& Other) const
-{
-	return EnableSpeedPlanting == Other.EnableSpeedPlanting
-	&& SpeedCurveName == Other.SpeedCurveName
-	&& FMath::IsNearlyEqualByULP(SpeedThreshold, Other.SpeedThreshold)
-	&& FMath::IsNearlyEqualByULP(UnplantStiffness, Other.UnplantStiffness)
-	&& FMath::IsNearlyEqualByULP(UnplantCriticalDamping, Other.UnplantCriticalDamping);
-}
-
-bool FTargetChainFKSettings::operator==(const FTargetChainFKSettings& Other) const
-{
-	return EnableFK == Other.EnableFK
-	   && RotationMode == Other.RotationMode
-	   && FMath::IsNearlyEqualByULP(RotationAlpha,Other.RotationAlpha)
-	   && TranslationMode == Other.TranslationMode
-	   && FMath::IsNearlyEqualByULP(TranslationAlpha, Other.TranslationAlpha)
-	   && PoleVectorMatching == Other.PoleVectorMatching
-	   && FMath::IsNearlyEqualByULP(PoleVectorOffset, Other.PoleVectorOffset);
-}
-
-bool FTargetChainIKSettings::operator==(const FTargetChainIKSettings& Other) const
-{
-	return EnableIK == Other.EnableIK
-	&& FMath::IsNearlyEqualByULP(BlendToSource, Other.BlendToSource)
-	&& BlendToSourceWeights.Equals(Other.BlendToSourceWeights)
-	&& StaticOffset.Equals(Other.StaticOffset)
-	&& StaticLocalOffset.Equals(Other.StaticLocalOffset)
-	&& StaticRotationOffset.Equals(Other.StaticRotationOffset)
-	&& FMath::IsNearlyEqualByULP(ScaleVertical, Other.ScaleVertical)
-	&& FMath::IsNearlyEqualByULP(Extension, Other.Extension)
-	&& bAffectedByIKWarping == Other.bAffectedByIKWarping;
-}
-
-bool FTargetChainSettings::operator==(const FTargetChainSettings& Other) const
-{
-	return FK == Other.FK && IK == Other.IK && SpeedPlanting == Other.SpeedPlanting;
-}
+#if WITH_EDITOR
+FLinearColor FIKRetargetDebugDrawState::Muted = FLinearColor::Gray;
+FLinearColor FIKRetargetDebugDrawState::SourceColor = (FLinearColor::Gray * FLinearColor::Blue) * Muted;
+FLinearColor FIKRetargetDebugDrawState::GoalColor = FLinearColor::Yellow;
+FLinearColor FIKRetargetDebugDrawState::MainColor = FLinearColor::Green;
+FLinearColor FIKRetargetDebugDrawState::NonSelected = FLinearColor::Gray * 0.3f;
+#endif
diff --git a/Engine/Plugins/Animation/IKRig/Source/IKRig/Private/Retargeter/IKRetargeter.cpp b/Engine/Plugins/Animation/IKRig/Source/IKRig/Private/Retargeter/IKRetargeter.cpp
--- a/Engine/Plugins/Animation/IKRig/Source/IKRig/Private/Retargeter/IKRetargeter.cpp
+++ b/Engine/Plugins/Animation/IKRig/Source/IKRig/Private/Retargeter/IKRetargeter.cpp
@@ -6,6 +6,14 @@
 #include "Retargeter/IKRetargetOps.h"
 #include "Retargeter/IKRetargetProfile.h"
 #include "Engine/SkeletalMesh.h"
+#include "Retargeter/RetargetOps/AlignPoleVectorOp.h"
+#include "Retargeter/RetargetOps/FKChainsOp.h"
+#include "Retargeter/RetargetOps/IKChainsOp.h"
+#include "Retargeter/RetargetOps/PelvisMotionOp.h"
+#include "Retargeter/RetargetOps/RunIKRigOp.h"
+#include "Retargeter/RetargetOps/ScaleSourceOp.h"
+#include "Retargeter/RetargetOps/SpeedPlantingOp.h"
+#include "Retargeter/RetargetOps/StrideWarpingOp.h"
 
 #include UE_INLINE_GENERATED_CPP_BY_NAME(IKRetargeter)
 
@@ -14,29 +22,186 @@ const FName UIKRetargeter::GetSourceIKRigPropertyName() { return GET_MEMBER_NAME
 const FName UIKRetargeter::GetTargetIKRigPropertyName() { return GET_MEMBER_NAME_STRING_CHECKED(UIKRetargeter, TargetIKRigAsset); };
 const FName UIKRetargeter::GetSourcePreviewMeshPropertyName() { return GET_MEMBER_NAME_STRING_CHECKED(UIKRetargeter, SourcePreviewMesh); };
 const FName UIKRetargeter::GetTargetPreviewMeshPropertyName() { return GET_MEMBER_NAME_STRING_CHECKED(UIKRetargeter, TargetPreviewMesh); }
+#endif
 
-void UIKRetargeter::GetSpeedCurveNames(TArray<FName>& OutSpeedCurveNames) const
+bool FRetargetChainMapping::HasChain(const FName& InChainName, ERetargetSourceOrTarget InSourceOrTarget) const
 {
-	for (const URetargetChainSettings* ChainSetting : ChainSettings)
+	if (InSourceOrTarget == ERetargetSourceOrTarget::Target)
 	{
-		if (ChainSetting->Settings.SpeedPlanting.SpeedCurveName != NAME_None)
+		return ChainMap.ContainsByPredicate([InChainName](const FRetargetChainPair& InChainPair)
 		{
-			OutSpeedCurveNames.Add(ChainSetting->Settings.SpeedPlanting.SpeedCurveName);
+			return InChainPair.TargetChainName == InChainName;
+		});
+	}
+	else
+	{
+		return ChainMap.ContainsByPredicate([InChainName](const FRetargetChainPair& InChainPair)
+		{
+			return InChainPair.SourceChainName == InChainName;
+		});
+	}
+}
+
+FName FRetargetChainMapping::GetChainMappedTo(const FName InChainName, ERetargetSourceOrTarget InSourceOrTarget) const
+{
+	if (const FRetargetChainPair* FoundChainPair = FindChainPair(InChainName, InSourceOrTarget))
+	{
+		return InSourceOrTarget == ERetargetSourceOrTarget::Source ? FoundChainPair->TargetChainName : FoundChainPair->SourceChainName;
+	}
+	
+	return NAME_None;
+}
+
+FRetargetChainPair* FRetargetChainMapping::FindChainPair(const FName& InChainName, ERetargetSourceOrTarget InSourceOrTarget) const
+{
+	return const_cast<FRetargetChainPair*>(FindChainPairConst(InChainName, InSourceOrTarget));
+}
+
+const FRetargetChainPair* FRetargetChainMapping::FindChainPairConst(const FName& InChainName, ERetargetSourceOrTarget InSourceOrTarget) const
+{
+	for (const FRetargetChainPair& ChainPair : ChainMap)
+	{
+		const FName& NameToCompare = InSourceOrTarget == ERetargetSourceOrTarget::Source ? ChainPair.SourceChainName : ChainPair.TargetChainName;
+		if (NameToCompare == InChainName)
+		{
+			return &ChainPair;
 		}
 	}
+	return nullptr;
 }
-#endif
 
-UIKRetargeter::UIKRetargeter(const FObjectInitializer& ObjectInitializer)
-	: Super(ObjectInitializer)
+TArray<FRetargetChainPair>& FRetargetChainMapping::GetChainPairs()
 {
-	RootSettings = CreateDefaultSubobject<URetargetRootSettings>(TEXT("RootSettings"));
-	RootSettings->SetFlags(RF_Transactional);
+	return ChainMap;
+}
+
+const TArray<FRetargetChainPair>& FRetargetChainMapping::GetChainPairs() const
+{
+	return ChainMap;
+}
+
+void FRetargetChainMapping::SetChainMapping(const FName InTargetChainName, const FName InSourceChainName)
+{
+	if (FRetargetChainPair* FoundChainPair = FindChainPair(InTargetChainName, ERetargetSourceOrTarget::Target))
+	{
+		FoundChainPair->SourceChainName = InSourceChainName;
+		return;
+	}
+
+	ChainMap.Add(FRetargetChainPair(InTargetChainName, InSourceChainName));
+}
+
+void FRetargetChainMapping::CleanMapping(const UIKRigDefinition* SourceIKRig, const UIKRigDefinition* TargetIKRig)
+{
+	// remove null target chains
+	ChainMap.RemoveAll([](const FRetargetChainPair& ChainPair)
+	{
+		return ChainPair.TargetChainName == NAME_None;
+	});
+	
+	if (IsValid(TargetIKRig))
+	{
+		TArray<FName> TargetChainNames = TargetIKRig->GetRetargetChainNames();
+
+		// remove chains in the map that are not in the IK Rig
+		ChainMap.RemoveAll([&TargetChainNames](const FRetargetChainPair& ChainPair)
+		{
+			return !TargetChainNames.Contains(ChainPair.TargetChainName);
+		});
+
+		// add any chains in the rig that are not in the mapping
+		for (const FName& RequiredChain : TargetChainNames)
+		{
+			if (!HasChain(RequiredChain, ERetargetSourceOrTarget::Target))
+			{
+				SetChainMapping(RequiredChain, NAME_None);
+			}
+		}
+	}
+
+	if (IsValid(SourceIKRig))
+	{
+		TArray<FName> SourceChainNames = SourceIKRig->GetRetargetChainNames();
+	
+		// reset any sources that are no longer present to "None"
+		for (FRetargetChainPair& Pair : ChainMap)
+		{
+			if (!SourceChainNames.Contains(Pair.SourceChainName))
+			{
+				Pair.SourceChainName = NAME_None;
+			}
+		}
+	}
+
+	// enforce the chain order based on the StartBone index
+	SortMapping(TargetIKRig);
+}
+
+void FRetargetChainMapping::SortMapping(const UIKRigDefinition* TargetIKRig)
+{
+	if (!IsValid(TargetIKRig))
+	{
+		return;
+	}
+	
+	ChainMap.Sort([TargetIKRig](const FRetargetChainPair& A, const FRetargetChainPair& B)
+	{
+		const TArray<FBoneChain>& BoneChains = TargetIKRig->GetRetargetChains();
+		const FIKRigSkeleton& TargetSkeleton = TargetIKRig->GetSkeleton();
 
-	GlobalSettings = CreateDefaultSubobject<UIKRetargetGlobalSettings>(TEXT("GlobalSettings"));
-	GlobalSettings->SetFlags(RF_Transactional);
+		// look for chains
+		const int32 IndexA = BoneChains.IndexOfByPredicate([&A](const FBoneChain& Chain)
+		{
+			return A.TargetChainName == Chain.ChainName;
+		});
+
+		const int32 IndexB = BoneChains.IndexOfByPredicate([&B](const FBoneChain& Chain)
+		{
+			return B.TargetChainName == Chain.ChainName;
+		});
+
+		// compare their StartBone Index 
+		if (IndexA > INDEX_NONE && IndexB > INDEX_NONE)
+		{
+			const int32 StartBoneIndexA = TargetSkeleton.GetBoneIndexFromName(BoneChains[IndexA].StartBone.BoneName);
+			const int32 StartBoneIndexB = TargetSkeleton.GetBoneIndexFromName(BoneChains[IndexB].StartBone.BoneName);
+
+			if (StartBoneIndexA == StartBoneIndexB)
+			{
+				// fallback to sorting alphabetically
+				return BoneChains[IndexA].ChainName.LexicalLess(BoneChains[IndexB].ChainName);
+			}
+				
+			return StartBoneIndexA < StartBoneIndexB;	
+		}
+
+		// sort them according to the target ik rig if previously failed 
+		return IndexA < IndexB;
+	});
+}
 
+PRAGMA_DISABLE_DEPRECATION_WARNINGS
+void FRetargetChainMapping::LoadFromDeprecatedChainSettings(const TArray<URetargetChainSettings*>& InChainSettings)
+{
+	ChainMap.Empty();
+	
+	for (const URetargetChainSettings* ChainSetting : InChainSettings)
+	{
+		ChainMap.Add(FRetargetChainPair(ChainSetting->TargetChain, ChainSetting->SourceChain));
+	}
+}
+PRAGMA_ENABLE_DEPRECATION_WARNINGS
+
+UIKRetargeter::UIKRetargeter(const FObjectInitializer& ObjectInitializer)
+	: Super(ObjectInitializer)
+{
 	PRAGMA_DISABLE_DEPRECATION_WARNINGS
+	RootSettings_DEPRECATED = CreateDefaultSubobject<URetargetRootSettings>(TEXT("RootSettings"));
+	RootSettings_DEPRECATED->SetFlags(RF_Transactional);
+
+	GlobalSettings_DEPRECATED = CreateDefaultSubobject<UIKRetargetGlobalSettings>(TEXT("GlobalSettings"));
+	GlobalSettings_DEPRECATED->SetFlags(RF_Transactional);
+	
 	OpStack_DEPRECATED = CreateDefaultSubobject<UDEPRECATED_RetargetOpStack>(TEXT("PostSettings"));
 	OpStack_DEPRECATED->SetFlags(RF_Transactional);
 	PRAGMA_ENABLE_DEPRECATION_WARNINGS
@@ -121,76 +286,66 @@ void UIKRetargeter::PostLoad()
 	// very early versions of the asset may not have been set as standalone
 	SetFlags(RF_Standalone);
 
-	// load deprecated chain mapping (pre UStruct to UObject refactor)
+	PostLoadOldSettingsToNew();
+	
+	PostLoadOldOpsToNewStructOps();
+
+	PostLoadConvertEverythingToOps();
+	
+	CleanAndInitialize();
+}
+
+void UIKRetargeter::PostLoadOldSettingsToNew()
+{
 	PRAGMA_DISABLE_DEPRECATION_WARNINGS
-	if (!ChainMapping_DEPRECATED.IsEmpty())
+	
+#if WITH_EDITORONLY_DATA
+	// load deprecated target actor offset
+	if (!FMath::IsNearlyZero(TargetActorOffset_DEPRECATED))
 	{
-		for (const FRetargetChainMap& OldChainMap : ChainMapping_DEPRECATED)
-		{
-			if (OldChainMap.TargetChain == NAME_None)
-			{
-				continue;
-			}
-			
-			TObjectPtr<URetargetChainSettings>* MatchingChain = ChainSettings.FindByPredicate([&](const URetargetChainSettings* Chain)
-			{
-				return Chain ? Chain->TargetChain == OldChainMap.TargetChain : false;
-			});
-			
-			if (MatchingChain)
-			{
-				(*MatchingChain)->SourceChain = OldChainMap.SourceChain;
-			}
-			else
-			{
-				TObjectPtr<URetargetChainSettings> NewChainMap = NewObject<URetargetChainSettings>(this, URetargetChainSettings::StaticClass(), NAME_None, RF_Transactional);
-				NewChainMap->TargetChain = OldChainMap.TargetChain;
-				NewChainMap->SourceChain = OldChainMap.SourceChain;
-				ChainSettings.Add(NewChainMap);
-			}
-		}
+		TargetMeshOffset.X = TargetActorOffset_DEPRECATED;
 	}
 
-	#if WITH_EDITORONLY_DATA
-		// load deprecated target actor offset
-		if (!FMath::IsNearlyZero(TargetActorOffset_DEPRECATED))
-		{
-			TargetMeshOffset.X = TargetActorOffset_DEPRECATED;
-		}
-
-		// load deprecated target actor scale
-		if (!FMath::IsNearlyZero(TargetActorScale_DEPRECATED))
-		{
-			TargetMeshScale = TargetActorScale_DEPRECATED;
-		}
+	// load deprecated target actor scale
+	if (!FMath::IsNearlyZero(TargetActorScale_DEPRECATED))
+	{
+		TargetMeshScale = TargetActorScale_DEPRECATED;
+	}
 
-		// load deprecated global settings
-		if (!bRetargetRoot_DEPRECATED)
-		{
-			GlobalSettings->Settings.bEnableRoot = false;
-		}
-		if (!bRetargetFK_DEPRECATED)
-		{
-			GlobalSettings->Settings.bEnableFK = false;
-		}
-		if (!bRetargetIK_DEPRECATED)
-		{
-			GlobalSettings->Settings.bEnableIK = false;
-		}
-	#endif
+	// load deprecated global settings
+	if (!bRetargetRoot_DEPRECATED)
+	{
+		GlobalSettings_DEPRECATED->Settings.bEnableRoot = false;
+	}
+	if (!bRetargetFK_DEPRECATED)
+	{
+		GlobalSettings_DEPRECATED->Settings.bEnableFK = false;
+	}
+	if (!bRetargetIK_DEPRECATED)
+	{
+		GlobalSettings_DEPRECATED->Settings.bEnableIK = false;
+	}
+#endif
 
 	// load deprecated retarget poses (pre adding retarget poses for source)
-	if (!RetargetPoses_DEPRECATED.IsEmpty())
+	if (!RetargetPoses.IsEmpty())
 	{
-		TargetRetargetPoses = RetargetPoses_DEPRECATED;
+		TargetRetargetPoses = RetargetPoses;
 	}
 
 	// load deprecated current retarget pose (pre adding retarget poses for source)
-	if (CurrentRetargetPose_DEPRECATED != NAME_None)
+	if (CurrentRetargetPose != NAME_None)
 	{
-		CurrentTargetRetargetPose = CurrentRetargetPose_DEPRECATED;
+		CurrentTargetRetargetPose = CurrentRetargetPose;
 	}
 
+	PRAGMA_ENABLE_DEPRECATION_WARNINGS
+}
+
+void UIKRetargeter::PostLoadOldOpsToNewStructOps()
+{
+	PRAGMA_DISABLE_DEPRECATION_WARNINGS
+	
 	// validate the input FInstancedStruct is a valid subclass of FRetargetOpBase
 	auto DerivesFromBaseOpType = [](const FInstancedStruct& InConvertedSolver)
 	{
@@ -235,23 +390,212 @@ void UIKRetargeter::PostLoad()
 	}
 	
 	PRAGMA_ENABLE_DEPRECATION_WARNINGS
-
-	CleanAndInitialize();
 }
 
-#if WITH_EDITORONLY_DATA
-void UIKRetargeter::DeclareConstructClasses(TArray<FTopLevelAssetPath>& OutConstructClasses, const UClass* SpecificSubclass)
+void UIKRetargeter::PostLoadConvertEverythingToOps()
 {
-	Super::DeclareConstructClasses(OutConstructClasses, SpecificSubclass);
-	OutConstructClasses.Add(FTopLevelAssetPath(URetargetChainSettings::StaticClass()));
+	// only perform this if loading older package version
+	if (GetLinkerCustomVersion(FIKRigObjectVersion::GUID) >= FIKRigObjectVersion::ModularRetargeterOps)
+	{
+		return;
+	}
+
+	PRAGMA_DISABLE_DEPRECATION_WARNINGS
+
+	// load old chain settings into new chain map
+	ChainMap.LoadFromDeprecatedChainSettings(ChainSettings_DEPRECATED);
+	
+	// copy old post-ops to add to the end of the stack
+	TArray<FInstancedStruct> OldPostOps = RetargetOps;
+	RetargetOps.Empty();
+
+	// record all chains that have an IK goal assigned
+	TArray<FName> ChainsWithIK;
+	const TArray<FBoneChain>& RetargetBoneChains = TargetIKRigAsset->GetRetargetChains();
+	for (const FBoneChain& Chain : RetargetBoneChains)
+	{
+		if (Chain.IKGoalName != NAME_None)
+		{
+			ChainsWithIK.Add(Chain.ChainName);
+		}
+	}
+	
+	int32 OpIndex = 0;
+
+	// Source Scale Op from old "global" settings
+	{
+		if (!FMath::IsNearlyEqual(GlobalSettings_DEPRECATED->Settings.SourceScaleFactor, 1.0f))
+		{
+			RetargetOps.EmplaceAt(OpIndex,FIKRetargetScaleSourceOp::StaticStruct());
+			FIKRetargetOpBase* Op = RetargetOps[OpIndex].GetMutablePtr<FIKRetargetOpBase>();
+			FIKRetargetScaleSourceOpSettings* Settings = reinterpret_cast<FIKRetargetScaleSourceOpSettings*>(Op->GetSettings());
+			Settings->SourceScaleFactor = GlobalSettings_DEPRECATED->Settings.SourceScaleFactor;
+		}
+	}
+
+	// Pelvis Motion Op from old "root" settings
+	{
+		RetargetOps.EmplaceAt(OpIndex,FIKRetargetPelvisMotionOp::StaticStruct());
+		FIKRetargetOpBase* Op = RetargetOps[OpIndex].GetMutablePtr<FIKRetargetOpBase>();
+		Op->SetEnabled(GlobalSettings_DEPRECATED->Settings.bEnableRoot);
+		FIKRetargetPelvisMotionOpSettings* Settings = reinterpret_cast<FIKRetargetPelvisMotionOpSettings*>(Op->GetSettings());
+		Settings->SourcePelvisBone.BoneName = SourceIKRigAsset ? SourceIKRigAsset->GetRetargetRoot() : NAME_None;
+		Settings->TargetPelvisBone.BoneName = TargetIKRigAsset ? TargetIKRigAsset->GetRetargetRoot() : NAME_None;
+		Settings->RotationAlpha = RootSettings_DEPRECATED->Settings.RotationAlpha;
+		Settings->TranslationAlpha = RootSettings_DEPRECATED->Settings.TranslationAlpha;
+		Settings->BlendToSourceTranslation = RootSettings_DEPRECATED->Settings.BlendToSource;
+		Settings->BlendToSourceTranslationWeights = RootSettings_DEPRECATED->Settings.BlendToSourceWeights;
+		Settings->ScaleHorizontal = RootSettings_DEPRECATED->Settings.ScaleHorizontal;
+		Settings->ScaleVertical = RootSettings_DEPRECATED->Settings.ScaleVertical;
+		Settings->TranslationOffset = RootSettings_DEPRECATED->Settings.TranslationOffset;
+		Settings->RotationOffset = RootSettings_DEPRECATED->Settings.RotationOffset;
+		Settings->AffectIKHorizontal = RootSettings_DEPRECATED->Settings.AffectIKHorizontal;
+		Settings->AffectIKVertical = RootSettings_DEPRECATED->Settings.AffectIKVertical;
+	}
+
+	// FK Chains Op from old FK chain settings
+	{
+		++OpIndex;
+		RetargetOps.EmplaceAt(OpIndex,FIKRetargetFKChainsOp::StaticStruct());
+		FIKRetargetOpBase* Op = RetargetOps[OpIndex].GetMutablePtr<FIKRetargetOpBase>();
+		Op->SetEnabled(GlobalSettings_DEPRECATED->Settings.bEnableFK);
+		FIKRetargetFKChainsOpSettings* Settings = reinterpret_cast<FIKRetargetFKChainsOpSettings*>(Op->GetSettings());
+		for (URetargetChainSettings* Chain : ChainSettings_DEPRECATED)
+		{
+			FRetargetFKChainSettings ChainToRetarget;
+			ChainToRetarget.TargetChainName = Chain->TargetChain;
+			ChainToRetarget.EnableFK = Chain->Settings.FK.EnableFK;
+			ChainToRetarget.RotationMode = static_cast<EFKChainRotationMode>(Chain->Settings.FK.RotationMode);
+			ChainToRetarget.RotationAlpha = Chain->Settings.FK.RotationAlpha;
+			ChainToRetarget.TranslationMode = static_cast<EFKChainTranslationMode>(Chain->Settings.FK.TranslationMode);
+			ChainToRetarget.TranslationAlpha = Chain->Settings.FK.TranslationAlpha;
+			
+			Settings->ChainsToRetarget.Add(ChainToRetarget);
+		}
+	}
+
+	// IK Chains Op from old IK chain settings
+	{
+		++OpIndex;
+		RetargetOps.EmplaceAt(OpIndex,FIKRetargetIKChainsOp::StaticStruct());
+		FIKRetargetOpBase* Op = RetargetOps[OpIndex].GetMutablePtr<FIKRetargetOpBase>();
+		Op->SetEnabled(GlobalSettings_DEPRECATED->Settings.bEnableIK);
+		FIKRetargetIKChainsOpSettings* Settings = reinterpret_cast<FIKRetargetIKChainsOpSettings*>(Op->GetSettings());
+		for (URetargetChainSettings* Chain : ChainSettings_DEPRECATED)
+		{
+			if (!ChainsWithIK.Contains(Chain->TargetChain))
+			{
+				continue; // skip chains with no IK
+			}
+			FRetargetIKChainSettings ChainToRetarget;
+			ChainToRetarget.TargetChainName = Chain->TargetChain;
+			ChainToRetarget.EnableIK = Chain->Settings.IK.EnableIK;
+			ChainToRetarget.BlendToSource = Chain->Settings.IK.BlendToSource;
+			ChainToRetarget.BlendToSourceTranslation = Chain->Settings.IK.BlendToSourceTranslation;
+			ChainToRetarget.BlendToSourceRotation = Chain->Settings.IK.BlendToSourceRotation;
+			ChainToRetarget.BlendToSourceWeights = Chain->Settings.IK.BlendToSourceWeights;
+			ChainToRetarget.StaticOffset = Chain->Settings.IK.StaticOffset;
+			ChainToRetarget.StaticLocalOffset = Chain->Settings.IK.StaticLocalOffset;
+			ChainToRetarget.StaticRotationOffset = Chain->Settings.IK.StaticRotationOffset;
+			ChainToRetarget.ScaleVertical = Chain->Settings.IK.ScaleVertical;
+			ChainToRetarget.Extension = Chain->Settings.IK.Extension;
+			
+			Settings->ChainsToRetarget.Add(ChainToRetarget);
+		}
+	}
+
+	// Stride Warping Op from old "IK" chain settings and global settings
+	{
+		++OpIndex;
+		RetargetOps.EmplaceAt(OpIndex,FIKRetargetStrideWarpingOp::StaticStruct());
+		FIKRetargetOpBase* Op = RetargetOps[OpIndex].GetMutablePtr<FIKRetargetOpBase>();
+		Op->SetEnabled(GlobalSettings_DEPRECATED->Settings.bWarping);
+		FIKRetargetStrideWarpingOpSettings* Settings = reinterpret_cast<FIKRetargetStrideWarpingOpSettings*>(Op->GetSettings());
+		Settings->DirectionSource = GlobalSettings_DEPRECATED->Settings.DirectionSource;
+		Settings->ForwardDirection = GlobalSettings_DEPRECATED->Settings.ForwardDirection;
+		Settings->DirectionChain = GlobalSettings_DEPRECATED->Settings.DirectionChain;
+		Settings->WarpForwards = GlobalSettings_DEPRECATED->Settings.WarpForwards;
+		Settings->SidewaysOffset = GlobalSettings_DEPRECATED->Settings.SidewaysOffset;
+		Settings->WarpSplay = GlobalSettings_DEPRECATED->Settings.WarpSplay;
+		
+		for (URetargetChainSettings* Chain : ChainSettings_DEPRECATED)
+		{
+			if (!ChainsWithIK.Contains(Chain->TargetChain))
+			{
+				continue; // skip chains with no IK
+			}
+			
+			if (Chain->Settings.IK.bAffectedByIKWarping)
+			{
+				Settings->ChainsToStrideWarp.Add(Chain->TargetChain);
+			}
+		}
+	}
+	
+	// Speed Planting Op from old "IK" chain settings and global settings
+	{
+		++OpIndex;
+		RetargetOps.EmplaceAt(OpIndex,FIKRetargetSpeedPlantingOp::StaticStruct());
+		FIKRetargetOpBase* Op = RetargetOps[OpIndex].GetMutablePtr<FIKRetargetOpBase>();
+		FIKRetargetSpeedPlantingOpSettings* Settings = reinterpret_cast<FIKRetargetSpeedPlantingOpSettings*>(Op->GetSettings());
+		for (URetargetChainSettings* Chain : ChainSettings_DEPRECATED)
+		{
+			if (!ChainsWithIK.Contains(Chain->TargetChain))
+			{
+				continue; // skip chains with no IK
+			}
+			
+			if (Chain->Settings.SpeedPlanting.SpeedCurveName != NAME_None)
+			{
+				FRetargetSpeedPlantingSettings ChainToPlant;
+				ChainToPlant.TargetChainName = Chain->TargetChain;
+				ChainToPlant.SpeedCurveName = Chain->Settings.SpeedPlanting.SpeedCurveName;
+				Settings->ChainsToSpeedPlant.Add(ChainToPlant);
+
+				Settings->SpeedThreshold = Chain->Settings.SpeedPlanting.SpeedThreshold;
+				Settings->Stiffness = Chain->Settings.SpeedPlanting.UnplantStiffness;
+				Settings->CriticalDamping = Chain->Settings.SpeedPlanting.UnplantCriticalDamping;
+			}
+		}
+	}
+	
+	// IK Solve Op from old global settings
+	{
+		++OpIndex;
+		RetargetOps.EmplaceAt(OpIndex,FIKRetargetRunIKRigOp::StaticStruct());
+		FIKRetargetOpBase* Op = RetargetOps[OpIndex].GetMutablePtr<FIKRetargetOpBase>();
+		FIKRetargetRunIKRigOpSettings* Settings = reinterpret_cast<FIKRetargetRunIKRigOpSettings*>(Op->GetSettings());
+		Settings->IKRigAsset = TargetIKRigAsset;
+	}
+
+	// Pole Vector Op from old chains
+	{
+		++OpIndex;
+		RetargetOps.EmplaceAt(OpIndex,FIKRetargetAlignPoleVectorOp::StaticStruct());
+		FIKRetargetOpBase* Op = RetargetOps[OpIndex].GetMutablePtr<FIKRetargetOpBase>();
+		FIKRetargetAlignPoleVectorOpSettings* Settings = reinterpret_cast<FIKRetargetAlignPoleVectorOpSettings*>(Op->GetSettings());
+		for (URetargetChainSettings* Chain : ChainSettings_DEPRECATED)
+		{
+			if (Chain->Settings.FK.PoleVectorMatching > 0.0f)
+			{
+				FRetargetPoleVectorSettings ChainToAlign;
+				ChainToAlign.TargetChainName = Chain->TargetChain;
+				ChainToAlign.AlignAlpha = Chain->Settings.FK.PoleVectorMatching;
+				ChainToAlign.MaintainOffset = Chain->Settings.FK.PoleVectorMaintainOffset;
+				ChainToAlign.StaticAngularOffset = Chain->Settings.FK.PoleVectorOffset;
+				Settings->ChainsToAlign.Add(ChainToAlign);
+			}
+		}
+	}
+
+	// append old post-ops to the end of the stack
+	RetargetOps.Append(OldPostOps);
+	
+	PRAGMA_ENABLE_DEPRECATION_WARNINGS
 }
-#endif
 
 void UIKRetargeter::CleanAndInitialize()
 {
-	// remove null settings
-	ChainSettings.Remove(nullptr);
-
 	// use default pose as current pose unless set to something else
 	if (CurrentSourceRetargetPose == NAME_None)
 	{
@@ -283,63 +627,15 @@ void UIKRetargeter::CleanAndInitialize()
 	}
 };
 
-void URetargetChainSettings::Serialize(FArchive& Ar)
-{
-	Super::Serialize(Ar);
-	Ar.UsingCustomVersion(FIKRigObjectVersion::GUID);
-
 #if WITH_EDITORONLY_DATA
-	if (Ar.IsLoading())
-	{
-		// load the old chain settings into the new struct format
-		if (Ar.CustomVer(FIKRigObjectVersion::GUID) < FIKRigObjectVersion::ChainSettingsConvertedToStruct)
-		{
-			PRAGMA_DISABLE_DEPRECATION_WARNINGS
-			Settings.FK.EnableFK =  CopyPoseUsingFK_DEPRECATED;
-			Settings.FK.RotationMode =  RotationMode_DEPRECATED;
-			Settings.FK.RotationAlpha =  RotationAlpha_DEPRECATED;
-			Settings.FK.TranslationMode =  TranslationMode_DEPRECATED;
-			Settings.FK.TranslationAlpha =  TranslationAlpha_DEPRECATED;
-			Settings.IK.EnableIK =  DriveIKGoal_DEPRECATED;
-			Settings.IK.BlendToSource =  BlendToSource_DEPRECATED;
-			Settings.IK.BlendToSourceWeights =  BlendToSourceWeights_DEPRECATED;
-			Settings.IK.StaticOffset =  StaticOffset_DEPRECATED;
-			Settings.IK.StaticLocalOffset =  StaticLocalOffset_DEPRECATED;
-			Settings.IK.StaticRotationOffset =  StaticRotationOffset_DEPRECATED;
-			Settings.IK.Extension =  Extension_DEPRECATED;
-			Settings.SpeedPlanting.EnableSpeedPlanting =  UseSpeedCurveToPlantIK_DEPRECATED;
-			Settings.SpeedPlanting.SpeedCurveName =  SpeedCurveName_DEPRECATED;
-			Settings.SpeedPlanting.SpeedThreshold =  VelocityThreshold_DEPRECATED;
-			Settings.SpeedPlanting.UnplantStiffness =  UnplantStiffness_DEPRECATED;
-			Settings.SpeedPlanting.UnplantCriticalDamping =  UnplantCriticalDamping_DEPRECATED;
-			PRAGMA_ENABLE_DEPRECATION_WARNINGS
-		}
-	}
-#endif
-}
-
-void URetargetRootSettings::Serialize(FArchive& Ar)
+PRAGMA_DISABLE_DEPRECATION_WARNINGS
+void UIKRetargeter::DeclareConstructClasses(TArray<FTopLevelAssetPath>& OutConstructClasses, const UClass* SpecificSubclass)
 {
-	Super::Serialize(Ar);
-	Ar.UsingCustomVersion(FIKRigObjectVersion::GUID);
-
-#if WITH_EDITORONLY_DATA
-	if (Ar.IsLoading())
-	{
-		// load the old root settings into the new struct format
-		if (Ar.CustomVer(FIKRigObjectVersion::GUID) < FIKRigObjectVersion::ChainSettingsConvertedToStruct)
-		{
-			PRAGMA_DISABLE_DEPRECATION_WARNINGS
-			Settings.ScaleHorizontal = GlobalScaleHorizontal_DEPRECATED;
-			Settings.ScaleVertical = GlobalScaleVertical_DEPRECATED;
-			Settings.BlendToSource = static_cast<float>(BlendToSource_DEPRECATED.Size());
-			Settings.TranslationOffset = StaticOffset_DEPRECATED;
-			Settings.RotationOffset = StaticRotationOffset_DEPRECATED;
-			PRAGMA_ENABLE_DEPRECATION_WARNINGS
-		}
-	}
-#endif
+	Super::DeclareConstructClasses(OutConstructClasses, SpecificSubclass);
+	OutConstructClasses.Add(FTopLevelAssetPath(URetargetChainSettings::StaticClass()));
 }
+PRAGMA_ENABLE_DEPRECATION_WARNINGS
+#endif
 
 FQuat FIKRetargetPose::GetDeltaRotationForBone(const FName BoneName) const
 {
@@ -397,28 +693,6 @@ void FIKRetargetPose::SortHierarchically(const FIKRigSkeleton& Skeleton)
 	});
 }
 
-const TObjectPtr<URetargetChainSettings> UIKRetargeter::GetChainMapByName(const FName& TargetChainName) const
-{
-	const TObjectPtr<URetargetChainSettings>* ChainMap = ChainSettings.FindByPredicate(
-		[TargetChainName](const TObjectPtr<URetargetChainSettings> ChainMap)
-		{
-			return ChainMap->TargetChain == TargetChainName;
-		});
-	
-	return !ChainMap ? nullptr : ChainMap->Get();
-}
-
-const FTargetChainSettings* UIKRetargeter::GetChainSettingsByName(const FName& TargetChainName) const
-{
-	const TObjectPtr<URetargetChainSettings> ChainMap = GetChainMapByName(TargetChainName);
-	if (ChainMap)
-	{
-		return &ChainMap->Settings;
-	}
-
-	return nullptr;
-}
-
 const FIKRetargetPose* UIKRetargeter::GetCurrentRetargetPose(const ERetargetSourceOrTarget& SourceOrTarget) const
 {
 	return SourceOrTarget == ERetargetSourceOrTarget::Source ? &SourceRetargetPoses[CurrentSourceRetargetPose] : &TargetRetargetPoses[CurrentTargetRetargetPose];
@@ -442,30 +716,6 @@ const FName UIKRetargeter::GetDefaultPoseName()
 	return DefaultPoseName;
 }
 
-void UIKRetargeter::FillProfileWithAssetSettings(FRetargetProfile& InOutProfile) const
-{
-	// first copy all the asset settings into the profile
-	InOutProfile.bApplyTargetRetargetPose = true;
-	InOutProfile.TargetRetargetPoseName = GetCurrentRetargetPoseName(ERetargetSourceOrTarget::Target);
-	InOutProfile.bApplySourceRetargetPose = true;
-	InOutProfile.SourceRetargetPoseName = GetCurrentRetargetPoseName(ERetargetSourceOrTarget::Source);
-	InOutProfile.bApplyChainSettings = true;
-	for (const TObjectPtr<URetargetChainSettings>& Chain : ChainSettings)
-	{
-		InOutProfile.ChainSettings.Add(Chain->TargetChain, Chain->Settings);
-	}
-	InOutProfile.bApplyRootSettings = true;
-	InOutProfile.RootSettings = RootSettings->Settings;
-	InOutProfile.bApplyGlobalSettings = true;
-	InOutProfile.GlobalSettings = GlobalSettings->Settings;
-
-	// now override any settings in the asset's current profile
-	if (const FRetargetProfile* ProfileToUse = GetCurrentProfile())
-	{
-		InOutProfile.MergeWithOtherProfile(*ProfileToUse);
-	}
-}
-
 const FRetargetProfile* UIKRetargeter::GetCurrentProfile() const
 {
 	return GetProfileByName(CurrentProfile);
@@ -476,6 +726,32 @@ const FRetargetProfile* UIKRetargeter::GetProfileByName(const FName& ProfileName
 	return Profiles.Find(ProfileName);
 }
 
+//
+// BEGIN DEPRECATED API
+//
+PRAGMA_DISABLE_DEPRECATION_WARNINGS
+const TObjectPtr<URetargetChainSettings> UIKRetargeter::GetChainMapByName(const FName& TargetChainName) const
+{
+	const TObjectPtr<URetargetChainSettings>* AllChainMapSettings = ChainSettings_DEPRECATED.FindByPredicate(
+		[TargetChainName](const TObjectPtr<URetargetChainSettings>& AllChainMapSettings)
+		{
+			return AllChainMapSettings->TargetChain == TargetChainName;
+		});
+	
+	return !AllChainMapSettings ? nullptr : AllChainMapSettings->Get();
+}
+
+const FTargetChainSettings* UIKRetargeter::GetChainSettingsByName(const FName& TargetChainName) const
+{
+	const TObjectPtr<URetargetChainSettings> AllChainMaps = GetChainMapByName(TargetChainName);
+	if (AllChainMaps)
+	{
+		return &AllChainMaps->Settings;
+	}
+
+	return nullptr;
+}
+
 FTargetChainSettings UIKRetargeter::GetChainUsingGoalFromRetargetAsset(
 	const UIKRetargeter* RetargetAsset,
 	const FName IKGoalName)
@@ -686,3 +962,7 @@ void UIKRetargeter::SetChainSpeedPlantSettingsInRetargetProfile(
 	ChainSettings.SpeedPlanting = SpeedPlantSettings;
 	RetargetProfile.bApplyChainSettings = true;
 }
+PRAGMA_ENABLE_DEPRECATION_WARNINGS
+//
+// END DEPRECATED API
+//
diff --git a/Engine/Plugins/Animation/IKRig/Source/IKRig/Private/Retargeter/RetargetOps/CopyBasePoseOp.cpp b/Engine/Plugins/Animation/IKRig/Source/IKRig/Private/Retargeter/RetargetOps/CopyBasePoseOp.cpp
new file mode 100644
--- /dev/null
+++ b/Engine/Plugins/Animation/IKRig/Source/IKRig/Private/Retargeter/RetargetOps/CopyBasePoseOp.cpp
@@ -0,0 +1,48 @@
+﻿// Copyright Epic Games, Inc. All Rights Reserved.
+
+#include "Retargeter/RetargetOps/CopyBasePoseOp.h"
+
+bool FIKRetargetCopyBasePoseOp::Initialize(
+	const FIKRetargetProcessor& InProcessor,
+	const FRetargetSkeleton& InSourceSkeleton,
+	const FTargetSkeleton& InTargetSkeleton,
+	FIKRigLogger& Log)
+{
+	bIsInitialized = true;
+	return true;
+}
+
+FIKRetargetOpSettingsBase* FIKRetargetCopyBasePoseOp::GetSettings()
+{
+	return &Settings;
+}
+
+void FIKRetargetCopyBasePoseOp::SetSettings(const FIKRetargetOpSettingsBase* InSettings)
+{
+	Settings = *reinterpret_cast<const FIKRetargetCopyBasePoseOpSettings*>(InSettings);
+}
+
+const UScriptStruct* FIKRetargetCopyBasePoseOp::GetSettingsType() const
+{
+	return FIKRetargetCopyBasePoseOpSettings::StaticStruct();
+}
+
+const UScriptStruct* FIKRetargetCopyBasePoseOp::GetType() const
+{
+	return FIKRetargetCopyBasePoseOp::StaticStruct();
+}
+
+const UClass* FIKRetargetCopyBasePoseOp::GetControllerType() const
+{
+	return UIKRetargetCopyBasePoseController::StaticClass();
+}
+
+FIKRetargetCopyBasePoseOpSettings UIKRetargetCopyBasePoseController::GetSettings()
+{
+	return *reinterpret_cast<FIKRetargetCopyBasePoseOpSettings*>(OpToControl->GetSettings());
+}
+
+void UIKRetargetCopyBasePoseController::SetSettings(FIKRetargetCopyBasePoseOpSettings InSettings)
+{
+	OpToControl->SetSettings(&InSettings);
+}
diff --git a/Engine/Plugins/Animation/IKRig/Source/IKRig/Private/Retargeter/RetargetOps/IKChainsOp.cpp b/Engine/Plugins/Animation/IKRig/Source/IKRig/Private/Retargeter/RetargetOps/IKChainsOp.cpp
new file mode 100644
--- /dev/null
+++ b/Engine/Plugins/Animation/IKRig/Source/IKRig/Private/Retargeter/RetargetOps/IKChainsOp.cpp
@@ -0,0 +1,522 @@
+﻿// Copyright Epic Games, Inc. All Rights Reserved.
+
+#include "Retargeter/RetargetOps/IKChainsOp.h"
+
+#include "Retargeter/IKRetargetProcessor.h"
+#include "Retargeter/RetargetOps/FKChainsOp.h"
+#include "Retargeter/RetargetOps/PelvisMotionOp.h"
+#include "Retargeter/RetargetOps/StrideWarpingOp.h"
+
+#if WITH_EDITOR
+#include "PrimitiveDrawInterface.h"
+#include "PrimitiveDrawingUtils.h"
+#include "Engine/EngineTypes.h"
+#include "IKRigDebugRendering.h"
+#endif
+
+#define LOCTEXT_NAMESPACE "IKChainsOp"
+
+
+bool FSourceChainIK::Initialize(const FResolvedBoneChain& InSourceBoneChain, FIKRigLogger& InLog)
+{
+	if (InSourceBoneChain.BoneIndices.Num() < 2)
+	{
+		InLog.LogWarning(LOCTEXT("SourceChainLessThanThree", "IK Chains Op: trying to retarget source bone chain with IK but it has less than 2 joints."));
+		return false;
+	}
+
+	SourceBoneChain = &InSourceBoneChain;
+	StartBoneIndex = InSourceBoneChain.BoneIndices[0];
+	EndBoneIndex = InSourceBoneChain.BoneIndices.Last();
+	
+	const FTransform& End = InSourceBoneChain.RefPoseGlobalTransforms.Last();
+	PreviousEndPosition = End.GetTranslation();
+	InitialEndPosition = End.GetTranslation();
+	InitialEndRotation = End.GetRotation();
+
+	const FTransform& Start = InSourceBoneChain.RefPoseGlobalTransforms[0];
+	const float Length = static_cast<float>((Start.GetTranslation() - InitialEndPosition).Size());
+
+	if (Length <= UE_KINDA_SMALL_NUMBER)
+	{
+		InLog.LogWarning(LOCTEXT("SourceZeroLengthIK", "IK Chains Op: found zero-length source bone chain."));
+		return false;
+	}
+
+	InvInitialLength = 1.0f / Length;
+
+	return true;
+}
+
+bool FTargetChainIK::Initialize(const FResolvedBoneChain& InTargetBoneChain, FIKRigLogger& InLog)
+{
+	if (InTargetBoneChain.BoneIndices.Num() < 3)
+	{
+		InLog.LogWarning(LOCTEXT("TargetChainLessThanThree", "IK Chains Op: trying to retarget target bone chain with IK but it has less than 3 joints."));
+		return false;
+	}
+
+	TargetBoneChain = &InTargetBoneChain;
+	BoneIndexA = InTargetBoneChain.BoneIndices[0];
+	BoneIndexC = InTargetBoneChain.BoneIndices.Last();
+	const FTransform& Last = InTargetBoneChain.RefPoseGlobalTransforms.Last();
+	PrevEndPosition = Last.GetLocation();
+	InitialEndPosition = Last.GetTranslation();
+	InitialEndRotation = Last.GetRotation();
+	InitialLength = static_cast<float>((InTargetBoneChain.RefPoseGlobalTransforms[0].GetTranslation() - Last.GetTranslation()).Size());
+
+	if (InitialLength <= UE_KINDA_SMALL_NUMBER)
+	{
+		InLog.LogWarning(LOCTEXT("TargetZeroLengthIK", "IK Retargeter trying to retarget target bone chain with IK, but it is zero length!"));
+		return false;
+	}
+	
+	return true;
+}
+
+bool FIKChainRetargeter::Initialize(
+	const FResolvedBoneChain& InSourceBoneChain,
+	const FResolvedBoneChain& InTargetBoneChain,
+	const FRetargetIKChainSettings& InSettings,
+	FIKRigLogger& InLog)
+{
+	bool bIsInitialized = true;
+	Settings = &InSettings;
+	bIsInitialized &= Source.Initialize(InSourceBoneChain, InLog);
+	bIsInitialized &= Target.Initialize(InTargetBoneChain, InLog);
+	return bIsInitialized;
+}
+
+void FIKChainRetargeter::EncodePose(const TArray<FTransform>& InSourceGlobalPose)
+{
+	const FVector Start = InSourceGlobalPose[Source.StartBoneIndex].GetTranslation();
+	const FVector End = InSourceGlobalPose[Source.EndBoneIndex].GetTranslation();
+
+    // get the normalized direction / length of the IK limb (how extended it is as percentage of original length)
+    const FVector ChainVector = End - Start;
+	double ChainLength;
+	FVector ChainDirection;
+	ChainVector.ToDirectionAndLength(ChainDirection, ChainLength);
+	const double NormalizedLimbLength = ChainLength * Source.InvInitialLength;
+
+	Source.PreviousEndPosition = Source.CurrentEndPosition;
+	Source.CurrentEndPosition = End;
+	Source.CurrentEndDirectionNormalized = ChainDirection * NormalizedLimbLength;
+	Source.CurrentEndRotation = InSourceGlobalPose[Source.EndBoneIndex].GetRotation();
+	Source.CurrentHeightFromGroundNormalized = static_cast<float>(End.Z - Source.InitialEndPosition.Z)  * Source.InvInitialLength;
+}
+	
+void FIKChainRetargeter::DecodePose(
+	const FIKRetargetPelvisMotionOp* PelvisMotionOp,
+    const TArray<FTransform>& InGlobalPose)
+{
+	//
+	// calculate ROTATION of IK goal ...
+	//
+	
+	// apply delta rotation from input
+	const FQuat DeltaRotation = Source.CurrentEndRotation * Source.InitialEndRotation.Inverse();
+	FQuat GoalRotation = DeltaRotation * Target.InitialEndRotation;
+
+	// apply static rotation offset in the local space of the foot
+	GoalRotation = GoalRotation * Settings->StaticRotationOffset.Quaternion();
+
+	// blend to source rotation
+	const double BlendToSourceRotation = Settings->BlendToSource * Settings->BlendToSourceRotation;
+	if (BlendToSourceRotation > UE_KINDA_SMALL_NUMBER)
+	{
+		GoalRotation = FQuat::FastLerp(GoalRotation, Source.CurrentEndRotation, BlendToSourceRotation);
+		GoalRotation.Normalize();
+	}
+
+	//
+	// calculate POSITION of IK goal ...
+	//
+	
+	// set position to length-scaled direction from source limb
+	const FVector PelvisTranslationDelta = PelvisMotionOp ? PelvisMotionOp->GetPelvisTranslationOffset() : FVector::ZeroVector;
+	const FVector AffectIKWeights = PelvisMotionOp ? PelvisMotionOp->GetAffectIKWeightAsVector() : FVector::ZeroVector;
+	const FVector InvAffectIKWeights = FVector::OneVector - AffectIKWeights;
+	const FVector InvRootModification = PelvisTranslationDelta * InvAffectIKWeights;
+	const FVector Start = InGlobalPose[Target.BoneIndexA].GetTranslation() - InvRootModification;
+	FVector GoalPosition = Start + (Source.CurrentEndDirectionNormalized * Target.InitialLength);
+
+	// blend to source location
+	const double BlendToSourceTranslation = Settings->BlendToSource * Settings->BlendToSourceTranslation;
+	if (BlendToSourceTranslation > UE_KINDA_SMALL_NUMBER)
+	{
+		const FVector RootModification = PelvisTranslationDelta * AffectIKWeights;
+		const FVector Weight = BlendToSourceTranslation * Settings->BlendToSourceWeights;
+		const FVector SourceLocation = Source.CurrentEndPosition + RootModification;
+		GoalPosition.X = FMath::Lerp(GoalPosition.X, SourceLocation.X, Weight.X);
+		GoalPosition.Y = FMath::Lerp(GoalPosition.Y, SourceLocation.Y, Weight.Y);
+		GoalPosition.Z = FMath::Lerp(GoalPosition.Z, SourceLocation.Z, Weight.Z);
+	}
+
+	// apply global static offset
+	GoalPosition += Settings->StaticOffset;
+
+	// apply local static offset
+	GoalPosition += GoalRotation.RotateVector(Settings->StaticLocalOffset);
+
+	// apply vertical scale
+	GoalPosition.Z *= Settings->ScaleVertical;
+	
+	// apply extension
+	if (!FMath::IsNearlyEqual(Settings->Extension, 1.0f))
+	{
+		GoalPosition = Start + (GoalPosition - Start) * Settings->Extension;	
+	}
+	
+	// output transform
+	Results.EndEffectorPosition = GoalPosition;
+	Results.EndEffectorRotation = GoalRotation;
+	Target.PrevEndPosition = GoalPosition;
+}
+
+bool FRetargetIKChainSettings::operator==(const FRetargetIKChainSettings& Other) const
+{
+	return EnableIK == Other.EnableIK
+		&& FMath::IsNearlyEqualByULP(BlendToSource, Other.BlendToSource)
+		&& BlendToSourceWeights.Equals(Other.BlendToSourceWeights)
+		&& StaticOffset.Equals(Other.StaticOffset)
+		&& StaticLocalOffset.Equals(Other.StaticLocalOffset)
+		&& StaticRotationOffset.Equals(Other.StaticRotationOffset)
+		&& FMath::IsNearlyEqualByULP(ScaleVertical, Other.ScaleVertical)
+		&& FMath::IsNearlyEqualByULP(Extension, Other.Extension);
+}
+
+bool FIKRetargetIKChainsOp::Initialize(
+	const FIKRetargetProcessor& InProcessor,
+	const FRetargetSkeleton& InSourceSkeleton,
+	const FTargetSkeleton& InTargetSkeleton,
+	FIKRigLogger& InLog)
+{
+	bIsInitialized = false;
+
+	IKChainRetargeters.Reset();
+	
+	const UIKRigDefinition* TargetIKRig = InProcessor.GetRetargetAsset()->GetIKRig(ERetargetSourceOrTarget::Target);
+	const FRetargeterBoneChains& BoneChains = InProcessor.GetBoneChains();
+
+	for (FRetargetIKChainSettings& ChainSettings : Settings.ChainsToRetarget)
+	{
+		FName TargetChainName = ChainSettings.TargetChainName;
+		const FResolvedBoneChain* TargetBoneChain = BoneChains.GetResolvedBoneChainByName(TargetChainName, ERetargetSourceOrTarget::Target);
+
+		// validate that the chain even exists
+		if (TargetBoneChain == nullptr)
+		{
+			InLog.LogWarning( FText::Format(
+			LOCTEXT("IKChainOpMissingChain", "IK Chain Op: chain data is out of sync with IK Rig. Missing target chain, '{0}."),
+			FText::FromName(TargetChainName)));
+			continue;
+		}
+
+		// validate that the chain has IK applied to it
+		if (TargetBoneChain->IKGoalName == NAME_None)
+		{
+			InLog.LogWarning( FText::Format(
+			LOCTEXT("IKChainOpChainHasNoIK", "IK Chain Op: an IK chain was found with no IK goal assigned to it, '{0}."),
+			FText::FromName(TargetChainName)));
+			continue;
+		}
+		
+		// which source chain was this target chain mapped to?
+		const FResolvedBoneChain* SourceBoneChain = BoneChains.GetSourceChainMappedToTargetChain(TargetChainName);
+		if (SourceBoneChain == nullptr)
+		{
+			InLog.LogWarning( FText::Format(
+			LOCTEXT("IKChainOpChainNotMapped", "IK Chain Op: found IK chain that was not mapped to a source chain, '{0}."),
+			FText::FromName(TargetChainName)));
+			continue;
+		}
+		
+		// initialize the mapped pair of source/target bone chains
+		FIKChainRetargeter IKChainRetargeter;
+		const bool bChainInitialized = IKChainRetargeter.Initialize(*SourceBoneChain, *TargetBoneChain, ChainSettings, InLog);
+		if (!bChainInitialized)
+		{
+			InLog.LogWarning( FText::Format(
+			LOCTEXT("IKChainOpBadChain", "IK Chain Op: could not initialize a mapped retarget chain for IK, '{0}."), FText::FromName(TargetChainName)));
+			continue;
+		}
+
+		// warn user if IK goal is not on the END bone of the target chain. It will still work, but may produce bad results.
+		for (const UIKRigEffectorGoal* Goal : TargetIKRig->GetGoalArray())
+		{
+			if (Goal->GoalName != TargetBoneChain->IKGoalName)
+			{
+				continue;
+			}
+			
+			if (Goal->BoneName != TargetBoneChain->EndBone)
+			{
+				InLog.LogWarning( FText::Format(
+			LOCTEXT("TargetIKNotOnEndBone", "IK Chain Op: Retarget chain, '{0}' has an IK goal that is not on the End Bone of the chain."),
+				FText::FromString(TargetBoneChain->ChainName.ToString())));
+			}
+			break;
+		}
+
+		// store valid chain pair to be retargeted
+		IKChainRetargeters.Add(IKChainRetargeter);
+	}
+	
+	// consider initialized if at least 1 IK chain was initialized
+	bIsInitialized = !IKChainRetargeters.IsEmpty();
+	return bIsInitialized;
+}
+
+void FIKRetargetIKChainsOp::Run(
+	FIKRetargetProcessor& InProcessor,
+	const double InDeltaTime,
+	const TArray<FTransform>& InSourceGlobalPose,
+	TArray<FTransform>& OutTargetGlobalPose)
+{
+	if (InProcessor.IsIKForcedOff())
+	{
+		return; // skip this op when IK is off
+	}
+	
+	const FIKRetargetPelvisMotionOp* PelvisMotionOp = InProcessor.GetFirstRetargetOpOfType<FIKRetargetPelvisMotionOp>();
+	
+	// retarget the IK goals to their new locations based on input pose
+	for (FIKChainRetargeter& IKChainRetargeter : IKChainRetargeters)
+	{
+		// encode them all using the input pose
+		IKChainRetargeter.EncodePose(InSourceGlobalPose);
+		// decode the IK goal and apply to IKRig
+		IKChainRetargeter.DecodePose(PelvisMotionOp, OutTargetGlobalPose);
+	}
+
+	// set the goal transform on the IK Rig
+	FIKRigGoalContainer& GoalContainer = InProcessor.GetIKRigGoalContainer();
+	for (const FIKChainRetargeter& IKChain : IKChainRetargeters)
+	{
+		constexpr double PositionAlpha = 1.0;
+		constexpr double RotationAlpha = 1.0;
+		
+		FIKRigGoal Goal = FIKRigGoal(
+			IKChain.GetTargetChain().IKGoalName,
+			IKChain.GetTargetChain().EndBone,
+			IKChain.GetResults().EndEffectorPosition,
+			IKChain.GetResults().EndEffectorRotation,
+			PositionAlpha,
+			RotationAlpha,
+			EIKRigGoalSpace::Component,
+			EIKRigGoalSpace::Component,
+			IKChain.GetSettings()->EnableIK);
+		
+		GoalContainer.SetIKGoal(Goal);
+	}
+
+#if WITH_EDITOR
+	SaveDebugData(InProcessor, InSourceGlobalPose, OutTargetGlobalPose);
+#endif
+}
+
+FIKRetargetOpSettingsBase* FIKRetargetIKChainsOp::GetSettings()
+{
+	return &Settings;
+}
+
+void FIKRetargetIKChainsOp::SetSettings(const FIKRetargetOpSettingsBase* InSettings)
+{
+	const FIKRetargetIKChainsOpSettings* NewSettings = reinterpret_cast<const FIKRetargetIKChainsOpSettings*>(InSettings);
+	
+	// copies everything except the ChainsToRetarget array (those are copied below, only for already existing chains)
+	const TArray<FName> PropertiesToIgnore = {"ChainsToRetarget"};
+	CopySettingsRaw(InSettings, PropertiesToIgnore);
+	
+	// copy ChainsToRetarget only for chains are in-common
+	for (const FRetargetIKChainSettings& NewChainSettings : NewSettings->ChainsToRetarget)
+	{
+		for (FRetargetIKChainSettings& ChainSettings : Settings.ChainsToRetarget)
+		{
+			if (ChainSettings.TargetChainName == NewChainSettings.TargetChainName)
+			{
+				ChainSettings = NewChainSettings;
+			}
+		}
+	}
+}
+
+const UScriptStruct* FIKRetargetIKChainsOp::GetSettingsType() const
+{
+	return FIKRetargetIKChainsOpSettings::StaticStruct();
+}
+
+const UScriptStruct* FIKRetargetIKChainsOp::GetType() const
+{
+	return FIKRetargetIKChainsOp::StaticStruct();
+}
+
+const UClass* FIKRetargetIKChainsOp::GetControllerType() const
+{
+	return UIKRetargetIKChainsController::StaticClass();
+}
+
+#if WITH_EDITOR
+void FIKRetargetIKChainsOp::RunInitialSetup(const UIKRetargeter* InRetargetAsset)
+{
+	// copy the target chains that have IK
+	if (const UIKRigDefinition* TargetIKRig = InRetargetAsset->GetIKRig(ERetargetSourceOrTarget::Target))
+	{
+		const TArray<FBoneChain>& TargetChains = TargetIKRig->GetRetargetChains();
+		for (const FBoneChain& Chain : TargetChains)
+		{
+			if (Chain.IKGoalName == NAME_None)
+			{
+				continue;
+			}
+			
+			Settings.ChainsToRetarget.Emplace(Chain.ChainName);
+		}
+	}
+}
+
+FCriticalSection FIKRetargetIKChainsOp::DebugDataMutex;
+
+void FIKRetargetIKChainsOp::DebugDraw(
+	FPrimitiveDrawInterface* InPDI,
+	const FTransform& InComponentTransform,
+	const double InComponentScale,
+	const FIKRetargetDebugDrawState& InEditorState) const
+{
+	// draw IK goals on each IK chain
+	if (!(Settings.bDrawFinalGoals || Settings.bDrawSourceLocations))
+	{
+		return;
+	}
+
+	// locked because this is called from the main thread and debug data is modified on worker
+	FScopeLock ScopeLock(&DebugDataMutex);
+	
+	// spin through all IK chains
+	for (const FChainDebugData& ChainDebugData : AllChainsDebugData)
+	{
+		FTransform FinalTransform = ChainDebugData.OutputTransformEnd * InComponentTransform;
+		
+		bool bIsSelected = InEditorState.SelectedChains.Contains(ChainDebugData.TargetChainName);
+
+		InPDI->SetHitProxy(new HIKRetargetEditorChainProxy(ChainDebugData.TargetChainName));
+
+		if (Settings.bDrawFinalGoals)
+		{
+			FLinearColor GoalColor = bIsSelected ? InEditorState.GoalColor : InEditorState.GoalColor * InEditorState.NonSelected;
+			
+			IKRigDebugRendering::DrawWireCube(
+			InPDI,
+			FinalTransform,
+			GoalColor,
+			static_cast<float>(Settings.GoalDrawSize),
+			static_cast<float>(Settings.GoalDrawThickness * InComponentScale));
+		}
+	
+		if (Settings.bDrawSourceLocations)
+		{
+			FTransform SourceGoalTransform;
+			SourceGoalTransform.SetTranslation(ChainDebugData.SourceTransformEnd.GetLocation() + DebugRootModification);
+			SourceGoalTransform.SetRotation(ChainDebugData.SourceTransformEnd.GetRotation());
+			SourceGoalTransform *= InComponentTransform;
+
+			FLinearColor Color = bIsSelected ? InEditorState.SourceColor : InEditorState.SourceColor * InEditorState.NonSelected;
+
+			DrawWireSphere(
+				InPDI,
+				SourceGoalTransform,
+				Color,
+				Settings.GoalDrawThickness,
+				12,
+				SDPG_World,
+				0.0f,
+				0.001f,
+				false);
+
+			if (Settings.bDrawFinalGoals)
+			{
+				DrawDashedLine(
+					InPDI,
+					SourceGoalTransform.GetLocation(),
+					FinalTransform.GetLocation(),
+					Color,
+					1.0f,
+					SDPG_Foreground);
+			}
+		}
+
+		// done drawing chain proxies
+		InPDI->SetHitProxy(nullptr);
+	}
+}
+
+void FIKRetargetIKChainsOp::SaveDebugData(
+	const FIKRetargetProcessor& InProcessor,
+	const TArray<FTransform>& InSourceGlobalPose,
+	const TArray<FTransform>& OutTargetGlobalPose)
+{
+	FScopeLock ScopeLock(&DebugDataMutex);
+	
+	AllChainsDebugData.Reset();
+	for (const FIKChainRetargeter& IKChainPair : IKChainRetargeters)
+	{
+		FChainDebugData NewChainData;
+		const FSourceChainIK& Source = IKChainPair.GetSource();
+		const FTargetChainIK& Target = IKChainPair.GetTarget();
+		const FDecodedIKChain& Results = IKChainPair.GetResults();
+		NewChainData.TargetChainName = IKChainPair.GetSettings()->TargetChainName;
+		NewChainData.InputTransformStart = OutTargetGlobalPose[Target.BoneIndexA];
+		NewChainData.InputTransformEnd = OutTargetGlobalPose[Target.BoneIndexC];
+		NewChainData.OutputTransformEnd = FTransform(Results.EndEffectorRotation, Results.EndEffectorPosition);
+		NewChainData.SourceTransformEnd = InSourceGlobalPose[Source.EndBoneIndex];
+		AllChainsDebugData.Add(NewChainData);
+	}
+
+	// get the root modification
+	DebugRootModification = FVector::ZeroVector;
+	if (const FIKRetargetPelvisMotionOp* PelvisMotionOp = InProcessor.GetFirstRetargetOpOfType<FIKRetargetPelvisMotionOp>())
+	{
+		DebugRootModification = PelvisMotionOp->GetPelvisTranslationOffset() * PelvisMotionOp->GetAffectIKWeightAsVector();
+	}
+}
+
+void FIKRetargetIKChainsOp::SetChainSettingsToDefault(const FName& InChainName)
+{
+	for (FRetargetIKChainSettings& ChainToRetarget : Settings.ChainsToRetarget)
+	{
+		if (ChainToRetarget.TargetChainName == InChainName)
+		{
+			ChainToRetarget = FRetargetIKChainSettings(InChainName);
+			return;
+		}
+	}
+}
+
+bool FIKRetargetIKChainsOp::AreChainSettingsAtDefault(const FName& InChainName)
+{
+	for (FRetargetIKChainSettings& ChainToRetarget : Settings.ChainsToRetarget)
+	{
+		if (ChainToRetarget.TargetChainName == InChainName)
+		{
+			FRetargetIKChainSettings DefaultSettings = FRetargetIKChainSettings();
+			return ChainToRetarget == DefaultSettings;
+		}
+	}
+
+	return true;
+}
+
+#endif
+
+FIKRetargetIKChainsOpSettings UIKRetargetIKChainsController::GetSettings()
+{
+	return *reinterpret_cast<FIKRetargetIKChainsOpSettings*>(OpToControl->GetSettings());
+}
+
+void UIKRetargetIKChainsController::SetSettings(FIKRetargetIKChainsOpSettings InSettings)
+{
+	OpToControl->SetSettings(&InSettings);
+}
+
+#undef LOCTEXT_NAMESPACE
diff --git a/Engine/Plugins/Animation/IKRig/Source/IKRig/Private/Retargeter/RetargetOps/PinBoneOp.cpp b/Engine/Plugins/Animation/IKRig/Source/IKRig/Private/Retargeter/RetargetOps/PinBoneOp.cpp
--- a/Engine/Plugins/Animation/IKRig/Source/IKRig/Private/Retargeter/RetargetOps/PinBoneOp.cpp
+++ b/Engine/Plugins/Animation/IKRig/Source/IKRig/Private/Retargeter/RetargetOps/PinBoneOp.cpp
@@ -1,6 +1,7 @@
 ﻿// Copyright Epic Games, Inc. All Rights Reserved.
 
 #include "Retargeter/RetargetOps/PinBoneOp.h"
+#include "Retargeter/IKRetargetProcessor.h"
 #include "Engine/SkeletalMesh.h"
 
 #include "IKRigObjectVersion.h"
@@ -10,7 +11,6 @@
 
 void FPinBoneData::PostSerialize(const FArchive& Ar)
 {
-	
 #if WITH_EDITORONLY_DATA
 	
 	if (Ar.IsLoading())
@@ -25,7 +25,6 @@ void FPinBoneData::PostSerialize(const FArchive& Ar)
 	}
 	
 #endif // WITH_EDITORONLY_DATA
-	
 }
 
 const FTransform& FPinBoneData::GetParentTransform(int32 BoneIndex, const FRetargetSkeleton& Skeleton, const TArray<FTransform>& Pose)
@@ -41,7 +40,7 @@ const FTransform& FPinBoneData::GetParentTransform(int32 BoneIndex, const FRetar
 
 void FPinBoneData::CachePinData(ERetargetSourceOrTarget InSkeletonToCopyFrom, const FIKRetargetProcessor& Processor)
 {
-	SourceScaleFactor = Processor.GetGlobalSettings().SourceScaleFactor;
+	SourceScaleFactor = Processor.GetSourceScaleFactor();
 
     // get skeletons we are copying from/to
     const FRetargetSkeleton& SkeletonToCopyFrom = Processor.GetSkeleton(InSkeletonToCopyFrom);
@@ -62,24 +61,23 @@ void FPinBoneData::CachePinData(ERetargetSourceOrTarget InSkeletonToCopyFrom, co
 }
 
 #if WITH_EDITOR
+
 USkeleton* FIKRetargetPinBoneOpSettings::GetSkeleton(const FName InPropertyName)
 {
-	if (!EditorInstance)
+	if (!SourceSkeletonAsset || !TargetSkeletonAsset)
 	{
 		return nullptr;
 	}
 
-	FIKRetargetPinBoneOpSettings* EditorSettings = reinterpret_cast<FIKRetargetPinBoneOpSettings*>(EditorInstance); 
-	
 	if (InPropertyName == GET_MEMBER_NAME_CHECKED(FPinBoneData, BoneToCopyTo))
 	{
-		const USkeleton* Skeleton = SkeletonToCopyFrom == ERetargetSourceOrTarget::Source ? EditorSettings->TargetSkeletonAsset : EditorSettings->SourceSkeletonAsset;
+		const USkeleton* Skeleton = SkeletonToCopyFrom == ERetargetSourceOrTarget::Source ? TargetSkeletonAsset : SourceSkeletonAsset;
 		return const_cast<USkeleton*>(Skeleton);
 	}
 
 	if (InPropertyName == GET_MEMBER_NAME_CHECKED(FPinBoneData, BoneToCopyFrom))
 	{
-		const USkeleton* Skeleton = SkeletonToCopyFrom == ERetargetSourceOrTarget::Target ? EditorSettings->TargetSkeletonAsset : EditorSettings->SourceSkeletonAsset;
+		const USkeleton* Skeleton = SkeletonToCopyFrom == ERetargetSourceOrTarget::Target ? TargetSkeletonAsset : SourceSkeletonAsset;
 		return const_cast<USkeleton*>(Skeleton);
 	}
 	
@@ -94,7 +92,7 @@ bool FIKRetargetPinBoneOp::Initialize(
 	const FTargetSkeleton& TargetSkeleton,
 	FIKRigLogger& Log)
 {
-	bool bFullyInitialized = true;
+	bFoundAllBonesToPin = true;
 	
 	for (FPinBoneData& PinData : Settings.BonesToPin)
 	{
@@ -113,7 +111,7 @@ bool FIKRetargetPinBoneOp::Initialize(
 		const bool bFoundBoneToPinTo = PinData.BoneToCopyFrom.BoneIndex != INDEX_NONE;
 		if (!bFoundBoneToPin)
 		{
-			bFullyInitialized = false;
+			bFoundAllBonesToPin = false;
 			Log.LogWarning(FText::Format(
 				LOCTEXT("MissingSourceBone", "Pin Bone retarget op refers to non-existant bone to pin, {0}."),
 				FText::FromName(PinData.BoneToCopyTo.BoneName)));
@@ -121,7 +119,7 @@ bool FIKRetargetPinBoneOp::Initialize(
 
 		if (!bFoundBoneToPinTo)
 		{
-			bFullyInitialized = false;
+			bFoundAllBonesToPin = false;
 			Log.LogWarning(FText::Format(
 				LOCTEXT("MissingTargetBone", "Pin Bone retarget op refers to non-existant bone to pin to, {0}."),
 				FText::FromName(PinData.BoneToCopyFrom.BoneName)));
@@ -130,29 +128,15 @@ bool FIKRetargetPinBoneOp::Initialize(
 		// force reinitialization
 		PinData.SourceScaleFactor = -1.0f;
 	}
-
-#if WITH_EDITOR
-	if (bFullyInitialized)
-	{
-		Message = FText::Format(LOCTEXT("ReadyToRun", "Pinning {0} bones."), FText::AsNumber(Settings.BonesToPin.Num()));
-	}
-	else
-	{
-		Message = FText(LOCTEXT("MissingBonesWarning", "Bone(s) not found. See output log."));
-	}
-
-	// record skeletons for UI bone selector widgets
-	Settings.SourceSkeletonAsset = SourceSkeleton.SkeletalMesh->GetSkeleton();
-	Settings.TargetSkeletonAsset = TargetSkeleton.SkeletalMesh->GetSkeleton();
-#endif
-
+	
 	// always treat this op as "initialized", individual pins will only execute if their prerequisites are met
 	bIsInitialized = true;
 	return true;
 }
 
 void FIKRetargetPinBoneOp::Run(
-	const FIKRetargetProcessor& Processor,
+	FIKRetargetProcessor& Processor,
+	const double InDeltaTime,
 	const TArray<FTransform>& InSourceGlobalPose,
 	TArray<FTransform>& OutTargetGlobalPose)
 {
@@ -168,7 +152,7 @@ void FIKRetargetPinBoneOp::Run(
 		}
 
 		// recache the offsets if the source is scaled differently
-		if (!FMath::IsNearlyEqual(PinData.SourceScaleFactor, Processor.GetGlobalSettings().SourceScaleFactor))
+		if (!FMath::IsNearlyEqual(PinData.SourceScaleFactor, Processor.GetSourceScaleFactor()))
 		{
 			PinData.CachePinData(Settings.SkeletonToCopyFrom, Processor);
 		}
@@ -200,11 +184,47 @@ FIKRetargetOpSettingsBase* FIKRetargetPinBoneOp::GetSettings()
 	return &Settings;
 }
 
+void FIKRetargetPinBoneOp::SetSettings(const FIKRetargetOpSettingsBase* InSettings)
+{
+	// copies everything except the bones we are operating on (those require reinit)
+	const TArray<FName> PropertiesToIgnore = {"BonesToPin"};
+	CopySettingsRaw(InSettings, PropertiesToIgnore);
+}
+
 const UScriptStruct* FIKRetargetPinBoneOp::GetSettingsType() const
 {
 	return FIKRetargetPinBoneOpSettings::StaticStruct();
 }
 
+const UScriptStruct* FIKRetargetPinBoneOp::GetType() const
+{
+	return FIKRetargetPinBoneOp::StaticStruct();
+}
+
+const UClass* FIKRetargetPinBoneOp::GetControllerType() const
+{
+	return UIKRetargetPinBoneController::StaticClass();
+}
+
+#if WITH_EDITOR
+FText FIKRetargetPinBoneOp::GetWarningMessage() const
+{
+	if (IsInitialized() && IsEnabled())
+	{
+		if (bFoundAllBonesToPin)
+		{
+			return FText::Format(LOCTEXT("ReadyToRun", "Running on {0} bone(s)."), FText::AsNumber(Settings.BonesToPin.Num()));
+		}
+		else
+		{
+			return FText::Format(LOCTEXT("ReadyToRun", "Running, but missing bones. See log."), FText::AsNumber(Settings.BonesToPin.Num()));
+		}
+	}
+	
+	return FIKRetargetOpBase::GetWarningMessage();
+}
+#endif
+
 FTransform FIKRetargetPinBoneOp::GetNewBoneTransform(
 	const FPinBoneData& PinData,
 	const FIKRetargetProcessor& Processor,
@@ -377,4 +397,14 @@ void UDEPRECATED_PinBoneOp::Serialize(FArchive& Ar)
 }
 PRAGMA_ENABLE_DEPRECATION_WARNINGS
 
+FIKRetargetPinBoneOpSettings UIKRetargetPinBoneController::GetSettings()
+{
+	return *reinterpret_cast<FIKRetargetPinBoneOpSettings*>(OpToControl->GetSettings());
+}
+
+void UIKRetargetPinBoneController::SetSettings(FIKRetargetPinBoneOpSettings InSettings)
+{
+	OpToControl->SetSettings(&InSettings);
+}
+
 #undef LOCTEXT_NAMESPACE
diff --git a/Engine/Plugins/Animation/IKRig/Source/IKRig/Private/Retargeter/RetargetOps/RootMotionGeneratorOp.cpp b/Engine/Plugins/Animation/IKRig/Source/IKRig/Private/Retargeter/RetargetOps/RootMotionGeneratorOp.cpp
--- a/Engine/Plugins/Animation/IKRig/Source/IKRig/Private/Retargeter/RetargetOps/RootMotionGeneratorOp.cpp
+++ b/Engine/Plugins/Animation/IKRig/Source/IKRig/Private/Retargeter/RetargetOps/RootMotionGeneratorOp.cpp
@@ -1,33 +1,62 @@
 // Copyright Epic Games, Inc. All Rights Reserved.
 #include "Retargeter/RetargetOps/RootMotionGeneratorOp.h"
 
+#include "IKRigObjectVersion.h"
+#include "Retargeter/IKRetargeter.h"
 #include "Engine/SkeletalMesh.h"
+#include "Retargeter/IKRetargetProcessor.h"
+#include "Retargeter/RetargetOps/PelvisMotionOp.h"
+#include "StructUtils/InstancedStruct.h"
+#include "Animation/BoneReference.h"
 
 #define LOCTEXT_NAMESPACE "RootMotionGeneratorOp"
 
 #if WITH_EDITOR
-void FIKRetargetRootMotionOp::OnAddedToStack(const UIKRetargeter* Asset)
+IMPLEMENT_HIT_PROXY(HIKRetargetEditorRootProxy, HHitProxy);
+#endif
+
+#if WITH_EDITOR
+USkeleton* FIKRetargetRootMotionOpSettings::GetSkeleton(const FName InPropertyName)
 {
-	// get the root of the SOURCE skeleton
-	if (const USkeletalMesh* Mesh = Asset->GetPreviewMesh(ERetargetSourceOrTarget::Source))
+	if (InPropertyName == GET_MEMBER_NAME_CHECKED(FIKRetargetRootMotionOpSettings, SourceRoot))
 	{
-		Settings.SourceRootBone = Mesh->GetRefSkeleton().GetBoneName(0);
+		return const_cast<USkeleton*>(SourceSkeletonAsset);
 	}
-	
-	// get the root of the TARGET skeleton
-	if (const USkeletalMesh* Mesh = Asset->GetPreviewMesh(ERetargetSourceOrTarget::Target))
+
+	if (InPropertyName == GET_MEMBER_NAME_CHECKED(FIKRetargetRootMotionOpSettings, TargetRoot))
 	{
-		Settings.TargetRootBone = Mesh->GetRefSkeleton().GetBoneName(0);
+		return const_cast<USkeleton*>(TargetSkeletonAsset);
 	}
 
-	// get the retarget root bone
-	if (const UIKRigDefinition* IKRig = Asset->GetIKRig(ERetargetSourceOrTarget::Target))
+	if (InPropertyName == GET_MEMBER_NAME_CHECKED(FIKRetargetRootMotionOpSettings, TargetPelvis))
 	{
-		Settings.TargetPelvisBone = IKRig->GetRetargetRoot();
+		return const_cast<USkeleton*>(TargetSkeletonAsset);
 	}
+	
+	ensureMsgf(false, TEXT("Root motion op unable to get skeleton for UI widget."));
+	return nullptr;
 }
 #endif
 
+void FIKRetargetRootMotionOpSettings::PostSerialize(const FArchive& Ar)
+{
+#if WITH_EDITORONLY_DATA
+	
+	if (Ar.IsLoading())
+	{
+		if (Ar.CustomVer(FIKRigObjectVersion::GUID) < FIKRigObjectVersion::ModularRetargeterOps)
+		{
+			PRAGMA_DISABLE_DEPRECATION_WARNINGS
+			SourceRoot.BoneName = SourceRootBone_DEPRECATED;
+			TargetRoot.BoneName = TargetRootBone_DEPRECATED;
+			TargetPelvis.BoneName = TargetPelvisBone_DEPRECATED;
+			PRAGMA_ENABLE_DEPRECATION_WARNINGS
+		}
+	}
+	
+#endif // WITH_EDITORONLY_DATA
+}
+
 bool FIKRetargetRootMotionOp::Initialize(
 	const FIKRetargetProcessor& Processor,
 	const FRetargetSkeleton& SourceSkeleton,
@@ -35,40 +64,33 @@ bool FIKRetargetRootMotionOp::Initialize(
 	FIKRigLogger& Log)
 {
 	bIsInitialized = false;
+	Reset();
 	
-	SourceRootIndex = SourceSkeleton.BoneNames.Find(Settings.SourceRootBone);
+	SourceRootIndex = SourceSkeleton.BoneNames.Find(Settings.SourceRoot.BoneName);
 	bool bHasAllPrerequisites = true;
 	if (SourceRootIndex == INDEX_NONE)
 	{
-		Log.LogWarning(FText::Format(LOCTEXT("MissingSourceRootBone", "Root Motion Remap Op, missing source root bone {0}."), FText::FromName(Settings.SourceRootBone)));
+		Log.LogWarning(FText::Format(LOCTEXT("MissingSourceRootBone", "Root Motion Remap Op, missing source root bone {0}."),
+			FText::FromName(Settings.SourceRoot.BoneName)));
 		bHasAllPrerequisites = false;
 	}
 	
-	TargetRootIndex = TargetSkeleton.BoneNames.Find(Settings.TargetRootBone);
+	TargetRootIndex = TargetSkeleton.BoneNames.Find(Settings.TargetRoot.BoneName);
 	if (TargetRootIndex == INDEX_NONE)
 	{
-		Log.LogWarning(FText::Format(LOCTEXT("MissingTargetRootBone", "Root Motion Remap Op, missing target root bone {0}."), FText::FromName(Settings.TargetRootBone)));
+		Log.LogWarning(FText::Format(LOCTEXT("MissingTargetRootBone", "Root Motion Remap Op, missing target root bone {0}."),
+			FText::FromName(Settings.TargetRoot.BoneName)));
 		bHasAllPrerequisites = false;
 	}
 
-	TargetPelvisIndex = TargetSkeleton.BoneNames.Find(Settings.TargetPelvisBone);
+	TargetPelvisIndex = TargetSkeleton.BoneNames.Find(Settings.TargetPelvis.BoneName);
 	if (TargetPelvisIndex == INDEX_NONE)
 	{
-		Log.LogWarning(FText::Format(LOCTEXT("MissingPelvisBone", "Root Motion Remap Op, missing target pelvis bone {0}."), FText::FromName(Settings.TargetPelvisBone)));
+		Log.LogWarning(FText::Format(LOCTEXT("MissingPelvisBone", "Root Motion Remap Op, missing target pelvis bone {0}."),
+			FText::FromName(Settings.TargetPelvis.BoneName)));
 		bHasAllPrerequisites = false;
 	}
 
-#if WITH_EDITOR
-	if (bHasAllPrerequisites)
-	{
-		Message = FText::Format(LOCTEXT("ReadyToRun", "Generating root motion on {0}."), FText::FromName(Settings.TargetRootBone));
-	}
-	else
-	{
-		Message = FText(LOCTEXT("MissingBonesWarning", "Bone(s) not found. See output log."));
-	}
-#endif
-
 	// can't cache everything unless all prerequisites are met
 	if (!bHasAllPrerequisites)
 	{
@@ -80,36 +102,13 @@ bool FIKRetargetRootMotionOp::Initialize(
 	TargetRootInRefPose = TargetSkeleton.RetargetGlobalPose[TargetRootIndex];
 	TargetPelvisRelativeToTargetRootRefPose = TargetRootInRefPose.GetRelativeTransform(TargetPelvisInRefPose);
 
-	// generate list of non-retargeted child indices
-	const int32 NumTargetBones = TargetSkeleton.BoneNames.Num();
-	NonRetargetedChildrenOfRoot.Reset();
-	for (int32 BoneIndex=1; BoneIndex<NumTargetBones; ++BoneIndex)
-	{
-		bool bHasNoRetargetedParent = true;
-		
-		int32 ParentIndex = BoneIndex;
-		while(ParentIndex != INDEX_NONE)
-		{
-			if (TargetSkeleton.IsBoneRetargeted[ParentIndex])
-			{
-				bHasNoRetargetedParent = false;
-				break;
-			}
-			ParentIndex = TargetSkeleton.ParentIndices[ParentIndex];
-		}
-
-		if (bHasNoRetargetedParent)
-		{
-			NonRetargetedChildrenOfRoot.Add(BoneIndex);
-		}
-	}
-
 	bIsInitialized = true;
 	return true;
 }
 
 void FIKRetargetRootMotionOp::Run(
-	const FIKRetargetProcessor& Processor,
+	FIKRetargetProcessor& Processor,
+	const double InDeltaTime,
 	const TArray<FTransform>& InSourceGlobalPose,
 	TArray<FTransform>& OutTargetGlobalPose)
 {
@@ -140,16 +139,72 @@ void FIKRetargetRootMotionOp::Run(
 	OutTargetGlobalPose[TargetRootIndex] = NewRootTransform * Settings.GlobalOffset;
 }
 
+void FIKRetargetRootMotionOp::PostInitialize(
+	const FIKRetargetProcessor& Processor,
+	const FRetargetSkeleton& SourceSkeleton,
+	const FTargetSkeleton& TargetSkeleton,
+	FIKRigLogger& Log)
+{
+	// generate list of non-retargeted child indices
+	const int32 NumTargetBones = TargetSkeleton.BoneNames.Num();
+	NonRetargetedChildrenOfRoot.Reset();
+	for (int32 BoneIndex=1; BoneIndex<NumTargetBones; ++BoneIndex)
+	{
+		bool bHasNoRetargetedParent = true;
+		
+		int32 ParentIndex = BoneIndex;
+		while(ParentIndex != INDEX_NONE)
+		{
+			if (TargetSkeleton.GetIsBoneRetargeted(ParentIndex))
+			{
+				bHasNoRetargetedParent = false;
+				break;
+			}
+			ParentIndex = TargetSkeleton.ParentIndices[ParentIndex];
+		}
+
+		if (bHasNoRetargetedParent)
+		{
+			NonRetargetedChildrenOfRoot.Add(BoneIndex);
+		}
+	}
+}
+
 FIKRetargetOpSettingsBase* FIKRetargetRootMotionOp::GetSettings()
 {
 	return &Settings;
 }
 
+void FIKRetargetRootMotionOp::SetSettings(const FIKRetargetOpSettingsBase* InSettings)
+{
+	// copies everything except the bones we are operating on (those require reinit)
+	const TArray<FName> PropertiesToIgnore = {"SourceRootBone", "TargetRootBone", "TargetPelvisBone"};
+	CopySettingsRaw(InSettings, PropertiesToIgnore);
+}
+
 const UScriptStruct* FIKRetargetRootMotionOp::GetSettingsType() const
 {
 	return FIKRetargetRootMotionOpSettings::StaticStruct();
 }
 
+const UScriptStruct* FIKRetargetRootMotionOp::GetType() const
+{
+	return FIKRetargetRootMotionOp::StaticStruct();
+}
+
+const UClass* FIKRetargetRootMotionOp::GetControllerType() const
+{
+	return UIKRetargetRootMotionController::StaticClass();
+}
+
+void FIKRetargetRootMotionOp::Reset()
+{
+	SourceRootIndex = INDEX_NONE;
+	TargetRootIndex = INDEX_NONE;
+	TargetPelvisIndex = INDEX_NONE;
+	NonRetargetedChildrenOfRoot.Reset();
+}
+
 void FIKRetargetRootMotionOp::GenerateRootMotionFromTargetPelvis(
 	FTransform& OutRootTransform,
 	const TArray<FTransform>& InSourceGlobalPose,
@@ -206,7 +261,12 @@ void FIKRetargetRootMotionOp::CopyRootMotionFromSourceRoot(
 	// scale root translation by same scale factor applied to Pelvis (and modified by settings)
 	FVector NewRootLocation = InSourceGlobalPose[SourceRootIndex].GetLocation();
 	FVector RootTranslationDelta = NewRootLocation - SourceRootInRefPose.GetTranslation();
-	RootTranslationDelta *= Processor.GetRootRetargeter().GetGlobalScaleVector();
+	const FIKRetargetPelvisMotionOp* PelvisMotionOp = Processor.GetFirstRetargetOpOfType<FIKRetargetPelvisMotionOp>();
+	if (PelvisMotionOp)
+	{
+		RootTranslationDelta *= PelvisMotionOp->GetGlobalScaleVector();
+	}
+	
 	NewRootLocation = SourceRootInRefPose.GetTranslation() + RootTranslationDelta;
 
 	// optionally snap the root to the ground plane
@@ -219,4 +279,67 @@ void FIKRetargetRootMotionOp::CopyRootMotionFromSourceRoot(
 	OutRootTransform.SetTranslation(NewRootLocation);
 }
 
+#if WITH_EDITOR
+void FIKRetargetRootMotionOp::RunInitialSetup(const UIKRetargeter* InRetargetAsset)
+{
+	// get the root of the SOURCE skeleton
+	if (const USkeletalMesh* Mesh = InRetargetAsset->GetPreviewMesh(ERetargetSourceOrTarget::Source))
+	{
+		Settings.SourceRoot = Mesh->GetRefSkeleton().GetBoneName(0);
+	}
+	
+	// get the root of the TARGET skeleton
+	if (const USkeletalMesh* Mesh = InRetargetAsset->GetPreviewMesh(ERetargetSourceOrTarget::Target))
+	{
+		Settings.TargetRoot = Mesh->GetRefSkeleton().GetBoneName(0);
+	}
+
+	// get the retarget root bone
+	if (const UIKRigDefinition* IKRig = InRetargetAsset->GetIKRig(ERetargetSourceOrTarget::Target))
+	{
+		Settings.TargetPelvis = IKRig->GetRetargetRoot();
+	}
+}
+
+FText FIKRetargetRootMotionOp::GetWarningMessage() const
+{
+	if (SourceRootIndex == INDEX_NONE || TargetRootIndex == INDEX_NONE)
+	{
+		return FText::Format(LOCTEXT("MissingBothRoots", "Missing both root bones."), FText::FromName(Settings.TargetRoot.BoneName));
+	}
+
+	if (Settings.RootMotionSource == ERootMotionSource::GenerateFromTargetPelvis && TargetPelvisIndex == INDEX_NONE)
+	{
+		return FText::Format(LOCTEXT("MissingBothRoots", "Missing target pelvis bone."), FText::FromName(Settings.TargetRoot.BoneName));
+	}
+
+	return FIKRetargetOpBase::GetWarningMessage();
+}
+#endif
+
+void UDEPRECATED_RootMotionGeneratorOp::ConvertToInstancedStruct(FInstancedStruct& OutInstancedStruct)
+{
+	OutInstancedStruct.InitializeAs(FIKRetargetRootMotionOp::StaticStruct());
+	FIKRetargetRootMotionOp& NewOp = OutInstancedStruct.GetMutable<FIKRetargetRootMotionOp>();
+	NewOp.Settings.SourceRoot.BoneName = SourceRootBone;
+	NewOp.Settings.TargetRoot.BoneName = TargetRootBone;
+	NewOp.Settings.TargetPelvis.BoneName = TargetPelvisBone;
+	NewOp.Settings.RootMotionSource = RootMotionSource;
+	NewOp.Settings.RootHeightSource = RootHeightSource;
+	NewOp.Settings.bPropagateToNonRetargetedChildren = bPropagateToNonRetargetedChildren;
+	NewOp.Settings.bMaintainOffsetFromPelvis = bMaintainOffsetFromPelvis;
+	NewOp.Settings.bRotateWithPelvis = bRotateWithPelvis;
+	NewOp.Settings.GlobalOffset = GlobalOffset;
+};
+
+FIKRetargetRootMotionOpSettings UIKRetargetRootMotionController::GetSettings()
+{
+	return *reinterpret_cast<FIKRetargetRootMotionOpSettings*>(OpToControl->GetSettings());
+}
+
+void UIKRetargetRootMotionController::SetSettings(FIKRetargetRootMotionOpSettings InSettings)
+{
+	OpToControl->SetSettings(&InSettings);
+}
+
 #undef LOCTEXT_NAMESPACE
diff --git a/Engine/Plugins/Animation/IKRig/Source/IKRig/Private/Retargeter/RetargetOps/ScaleSourceOp.cpp b/Engine/Plugins/Animation/IKRig/Source/IKRig/Private/Retargeter/RetargetOps/ScaleSourceOp.cpp
new file mode 100644
--- /dev/null
+++ b/Engine/Plugins/Animation/IKRig/Source/IKRig/Private/Retargeter/RetargetOps/ScaleSourceOp.cpp
@@ -0,0 +1,48 @@
+﻿// Copyright Epic Games, Inc. All Rights Reserved.
+
+#include "Retargeter/RetargetOps/ScaleSourceOp.h"
+
+bool FIKRetargetScaleSourceOp::Initialize(
+	const FIKRetargetProcessor& InProcessor,
+	const FRetargetSkeleton& InSourceSkeleton,
+	const FTargetSkeleton& InTargetSkeleton,
+	FIKRigLogger& Log)
+{
+	bIsInitialized = true;
+	return true;
+}
+
+FIKRetargetOpSettingsBase* FIKRetargetScaleSourceOp::GetSettings()
+{
+	return &Settings;
+}
+
+void FIKRetargetScaleSourceOp::SetSettings(const FIKRetargetOpSettingsBase* InSettings)
+{
+	Settings = *reinterpret_cast<const FIKRetargetScaleSourceOpSettings*>(InSettings);
+}
+
+const UScriptStruct* FIKRetargetScaleSourceOp::GetSettingsType() const
+{
+	return FIKRetargetScaleSourceOpSettings::StaticStruct();
+}
+
+const UScriptStruct* FIKRetargetScaleSourceOp::GetType() const
+{
+	return FIKRetargetScaleSourceOp::StaticStruct();
+}
+
+const UClass* FIKRetargetScaleSourceOp::GetControllerType() const
+{
+	return UIKRetargetScaleSourceController::StaticClass();
+}
+
+FIKRetargetScaleSourceOpSettings UIKRetargetScaleSourceController::GetSettings()
+{
+	return *reinterpret_cast<FIKRetargetScaleSourceOpSettings*>(OpToControl->GetSettings());
+}
+
+void UIKRetargetScaleSourceController::SetSettings(FIKRetargetScaleSourceOpSettings InSettings)
+{
+	OpToControl->SetSettings(&InSettings);
+}
diff --git a/Engine/Plugins/Animation/IKRig/Source/IKRig/Private/Retargeter/RetargetOps/SpeedPlantingOp.cpp b/Engine/Plugins/Animation/IKRig/Source/IKRig/Private/Retargeter/RetargetOps/SpeedPlantingOp.cpp
new file mode 100644
--- /dev/null
+++ b/Engine/Plugins/Animation/IKRig/Source/IKRig/Private/Retargeter/RetargetOps/SpeedPlantingOp.cpp
@@ -0,0 +1,244 @@
+﻿// Copyright Epic Games, Inc. All Rights Reserved.
+
+#include "Retargeter/RetargetOps/SpeedPlantingOp.h"
+
+#include "Retargeter/IKRetargetProcessor.h"
+#include "Animation/AnimInstance.h"
+#include "Components/SkeletalMeshComponent.h"
+#include "Animation/AnimNodeBase.h"
+
+#define LOCTEXT_NAMESPACE "SpeedPlantingOp"
+
+bool FIKRetargetSpeedPlantingOp::Initialize(
+	const FIKRetargetProcessor& InProcessor,
+	const FRetargetSkeleton& InSourceSkeleton,
+	const FTargetSkeleton& InTargetSkeleton,
+	FIKRigLogger& InLog)
+{
+	bIsInitialized = false;
+
+	GoalsToPlant.Reset();
+
+	const FRetargeterBoneChains& BoneChains = InProcessor.GetBoneChains();
+	const FIKRigGoalContainer& GoalContainer = InProcessor.GetIKRigGoalContainer();
+	for (const FRetargetSpeedPlantingSettings& SettingsForChain : Settings.ChainsToSpeedPlant)
+	{
+		const FName TargetChainName = SettingsForChain.TargetChainName;
+		const FResolvedBoneChain* TargetBoneChain = BoneChains.GetResolvedBoneChainByName(TargetChainName, ERetargetSourceOrTarget::Target);
+		if (TargetBoneChain == nullptr)
+		{
+			InLog.LogWarning( FText::Format(
+			LOCTEXT("SpeedPlantingMissingChain", "Speed Planting Op: chain data is out of sync with IK Rig. Missing target chain, '{0}."),
+			FText::FromName(TargetChainName)));
+			continue;
+		}
+
+		if (TargetBoneChain->IKGoalName == NAME_None)
+		{
+			InLog.LogWarning( FText::Format(
+			LOCTEXT("SpeedPlantingNoGoal", "Speed Planting Op: specified chain does not have an IK goal. Cannot speed plant, '{0}."),
+			FText::FromName(TargetChainName)));
+			continue;
+		}
+
+		const FIKRigGoal* Goal = GoalContainer.FindGoalByName(TargetBoneChain->IKGoalName);
+		if (Goal == nullptr)
+		{
+			InLog.LogWarning( FText::Format(
+			LOCTEXT("SpeedPlantingNoGoal", "Speed Planting Op: target chain references a goal that is not present in the IK Rig, '{0}."),
+			FText::FromName(TargetBoneChain->IKGoalName)));
+			continue;
+		}
+		
+		GoalsToPlant.Emplace(TargetBoneChain->IKGoalName, &SettingsForChain, Goal->Position);
+	}
+
+	bIsInitialized = !GoalsToPlant.IsEmpty();
+	return bIsInitialized;
+}
+
+void FIKRetargetSpeedPlantingOp::Run(
+	FIKRetargetProcessor& InProcessor,
+	const double InDeltaTime,
+	const TArray<FTransform>& InSourceGlobalPose,
+	TArray<FTransform>& OutTargetGlobalPose)
+{
+	if (InProcessor.IsIKForcedOff())
+	{
+		return; // skip this op when IK is off
+	}
+	
+	FIKRigGoalContainer& GoalContainer = InProcessor.GetIKRigGoalContainer();
+	
+	for (FSpeedPlantingGoalState& GoalToPlant : GoalsToPlant)
+	{
+		FIKRigGoal* IKRigGoal = GoalContainer.FindGoalByName(GoalToPlant.GoalName);
+		if (!IKRigGoal)
+		{
+			continue; // goal excluded, just ignore it
+		}
+		
+		if (GoalToPlant.CurrentSpeedValue < 0.0f || GoalToPlant.CurrentSpeedValue > Settings.SpeedThreshold)
+		{
+			GoalToPlant.PrevGoalPosition = IKRigGoal->Position;
+			GoalToPlant.PositionSpring.Reset();
+			continue;
+		}
+
+		// run the goal position through a spring damper to un-plant it
+		IKRigGoal->Position = UKismetMathLibrary::VectorSpringInterp(
+				GoalToPlant.PrevGoalPosition,
+				IKRigGoal->Position,
+				GoalToPlant.PositionSpring,
+				static_cast<float>(Settings.Stiffness),
+				static_cast<float>(Settings.CriticalDamping),
+				static_cast<float>(InDeltaTime),
+				1.0f /*mass*/,
+				0.0f /*target velocity amount*/);
+	}
+}
+
+FIKRetargetOpSettingsBase* FIKRetargetSpeedPlantingOp::GetSettings()
+{
+	return &Settings;
+}
+
+void FIKRetargetSpeedPlantingOp::SetSettings(const FIKRetargetOpSettingsBase* InSettings)
+{
+	const FIKRetargetSpeedPlantingOpSettings* NewSettings = reinterpret_cast<const FIKRetargetSpeedPlantingOpSettings*>(InSettings);
+	
+	// copies everything except "ChainsToSpeedPlant"
+	const TArray<FName> PropertiesToIgnore = {"ChainsToSpeedPlant"};
+	CopySettingsRaw(InSettings, PropertiesToIgnore);
+	
+	// copy settings only for chains that the op has initialized
+	for (const FRetargetSpeedPlantingSettings& NewChainSettings : NewSettings->ChainsToSpeedPlant)
+	{
+		for (FRetargetSpeedPlantingSettings& ChainSettings : Settings.ChainsToSpeedPlant)
+		{
+			if (ChainSettings.TargetChainName == NewChainSettings.TargetChainName)
+			{
+				ChainSettings = NewChainSettings;
+			}
+		}
+	}
+}
+
+const UScriptStruct* FIKRetargetSpeedPlantingOp::GetSettingsType() const
+{
+	return FIKRetargetSpeedPlantingOpSettings::StaticStruct();
+}
+
+const UScriptStruct* FIKRetargetSpeedPlantingOp::GetType() const
+{
+	return FIKRetargetSpeedPlantingOp::StaticStruct();
+}
+
+const UClass* FIKRetargetSpeedPlantingOp::GetControllerType() const
+{
+	return UIKRetargetSpeedPlantingController::StaticClass();
+}
+
+void FIKRetargetSpeedPlantingOp::OnPlaybackReset()
+{
+	for (FSpeedPlantingGoalState& GoalToPlant : GoalsToPlant)
+	{
+		GoalToPlant.PositionSpring.Reset();
+	}
+	
+	ResetThisTick = true;
+}
+
+void FIKRetargetSpeedPlantingOp::AnimGraphPreUpdateMainThread(
+	USkeletalMeshComponent& SourceMeshComponent,
+	USkeletalMeshComponent& TargetMeshComponent)
+{
+	if (!bIsInitialized)
+	{
+		return;
+	}
+
+	const UAnimInstance* SourceAnimInstance = SourceMeshComponent.GetAnimInstance();
+	if (!SourceAnimInstance)
+	{
+		return;
+	}
+
+	// update speed values for each planted chain
+	// NOTE: these are values from curves running on the SOURCE skeletal mesh
+	// they will be overriden by any values coming from the target in AnimGraphEvaluateAnyThread
+	const TMap<FName, float>& AnimCurveList = SourceAnimInstance->GetAnimationCurveList(EAnimCurveType::AttributeCurve);
+	for (FSpeedPlantingGoalState& GoalToPlant : GoalsToPlant)
+	{
+		const float* SourceValue = AnimCurveList.Find(GoalToPlant.Settings->SpeedCurveName);
+		GoalToPlant.CurrentSpeedValue = SourceValue ? *SourceValue : GoalToPlant.CurrentSpeedValue;
+	}
+}
+
+void FIKRetargetSpeedPlantingOp::AnimGraphEvaluateAnyThread(FPoseContext& Output)
+{
+	if (!bIsInitialized)
+	{
+		return;
+	}
+
+	// update speed values for each planted chain
+	// NOTE: these are values from curves running on the TARGET skeletal mesh
+	// they will override any values coming from the source via AnimGraphPreUpdateMainThread()
+	for (FSpeedPlantingGoalState& GoalToPlant : GoalsToPlant)
+	{
+		bool bFoundCurve = false;
+		GoalToPlant.CurrentSpeedValue = Output.Curve.Get(GoalToPlant.Settings->SpeedCurveName, bFoundCurve, static_cast<float>(GoalToPlant.CurrentSpeedValue));
+	}
+}
+
+TArray<FName> FIKRetargetSpeedPlantingOp::GetRequiredSpeedCurves() const
+{
+	TArray<FName> OutSpeedCurveNames;
+	for (const FRetargetSpeedPlantingSettings& ChainToPlant : Settings.ChainsToSpeedPlant)
+	{
+		const FName SpeedCurveName = ChainToPlant.SpeedCurveName;
+		if (SpeedCurveName != NAME_None)
+		{
+			OutSpeedCurveNames.Add(SpeedCurveName);
+		}
+	}
+	return MoveTemp(OutSpeedCurveNames);
+}
+
+#if WITH_EDITOR
+FText FIKRetargetSpeedPlantingOp::GetWarningMessage() const
+{
+	// warn about missing curves
+	if (bIsInitialized)
+	{
+		int32 NumMissingCurves = 0;
+		for (const FSpeedPlantingGoalState& GoalToPlant : GoalsToPlant)
+		{
+			if (!GoalToPlant.bFoundCurveInSourceComponent && !GoalToPlant.bFoundCurveInTargetComponent)
+			{
+				NumMissingCurves++;
+			}
+		}
+
+		if (NumMissingCurves > 0)
+		{
+			return FText::Format(LOCTEXT("MissingSpeedCurves", "Running, but missing {0} speed curves."), FText::AsNumber(NumMissingCurves));
+		}
+	}
+
+	return FIKRetargetOpBase::GetWarningMessage();
+}
+#endif
+
+FIKRetargetSpeedPlantingOpSettings UIKRetargetSpeedPlantingController::GetSettings()
+{
+	return *reinterpret_cast<FIKRetargetSpeedPlantingOpSettings*>(OpToControl->GetSettings());
+}
+
+void UIKRetargetSpeedPlantingController::SetSettings(FIKRetargetSpeedPlantingOpSettings InSettings)
+{
+	OpToControl->SetSettings(&InSettings);
+}
+
+#undef LOCTEXT_NAMESPACE
\ No newline at end of file
diff --git a/Engine/Plugins/Animation/IKRig/Source/IKRig/Private/Rig/IKRigDataTypes.cpp b/Engine/Plugins/Animation/IKRig/Source/IKRig/Private/Rig/IKRigDataTypes.cpp
--- a/Engine/Plugins/Animation/IKRig/Source/IKRig/Private/Rig/IKRigDataTypes.cpp
+++ b/Engine/Plugins/Animation/IKRig/Source/IKRig/Private/Rig/IKRigDataTypes.cpp
@@ -12,9 +12,14 @@ void FIKRigGoalContainer::SetIKGoal(const FIKRigGoal& InGoal)
 	{
 		// container hasn't seen this goal before, create new one, copying the input goal
 		Goals.Emplace(InGoal);
+		// new goal added, rig needs initialized
+		bRigNeedsInitialized = true;
 		return;
 	}
 
+	// check if we need to reinitialize
+	bRigNeedsInitialized = Goal->bEnabled != InGoal.bEnabled ? true : bRigNeedsInitialized;
+
 	// copy settings to existing goal
 	*Goal = InGoal;
 }
@@ -26,6 +31,8 @@ void FIKRigGoalContainer::SetIKGoal(const UIKRigEffectorGoal* InEffectorGoal)
 	{
 		// container hasn't seen this goal before, create new one, copying the Effector goal
 		Goals.Emplace(InEffectorGoal);
+		// new goal added, rig needs initialized
+		bRigNeedsInitialized = true;
 		return;
 	}
 
@@ -56,7 +63,21 @@ void FIKRigGoalContainer::SetIKGoal(const UIKRigEffectorGoal* InEffectorGoal)
 
 const FIKRigGoal* FIKRigGoalContainer::FindGoalByName(const FName& GoalName) const
 {
-	return Goals.FindByPredicate([GoalName](const FIKRigGoal& Other)	{return Other.Name == GoalName;});
+	return Goals.FindByPredicate([GoalName](const FIKRigGoal& Other){return Other.Name == GoalName;});
+}
+
+FIKRigGoal* FIKRigGoalContainer::FindGoalByName(const FName& GoalName)
+{
+	return Goals.FindByPredicate([GoalName](const FIKRigGoal& Other){return Other.Name == GoalName;});
+}
+
+const void FIKRigGoalContainer::FillWithGoalArray(const TArray<UIKRigEffectorGoal*>& InGoals)
+{
+	Empty();
+	for (const UIKRigEffectorGoal* Goal : InGoals)
+	{
+		SetIKGoal(Goal);
+	}
 }
 
 FIKRigGoal* FIKRigGoalContainer::FindGoalWriteable(const FName& GoalName) const
diff --git a/Engine/Plugins/Animation/IKRig/Source/IKRig/Private/Rig/IKRigDefinition.cpp b/Engine/Plugins/Animation/IKRig/Source/IKRig/Private/Rig/IKRigDefinition.cpp
--- a/Engine/Plugins/Animation/IKRig/Source/IKRig/Private/Rig/IKRigDefinition.cpp
+++ b/Engine/Plugins/Animation/IKRig/Source/IKRig/Private/Rig/IKRigDefinition.cpp
@@ -395,6 +395,16 @@ const FBoneChain* UIKRigDefinition::GetRetargetChainByName(FName ChainName) cons
 	return nullptr;
 }
 
+TArray<FName> UIKRigDefinition::GetRetargetChainNames() const
+{
+	TArray<FName> OutChainNames;
+	for (const FBoneChain& Chain : RetargetDefinition.BoneChains)
+	{
+		OutChainNames.Add(Chain.ChainName);
+	}
+	return MoveTemp(OutChainNames);
+}
+
 /** IInterface_PreviewMeshProvider interface */
 void UIKRigDefinition::SetPreviewMesh(USkeletalMesh* PreviewMesh, bool bMarkAsDirty)
 {
diff --git a/Engine/Plugins/Animation/IKRig/Source/IKRig/Private/Rig/IKRigProcessor.cpp b/Engine/Plugins/Animation/IKRig/Source/IKRig/Private/Rig/IKRigProcessor.cpp
--- a/Engine/Plugins/Animation/IKRig/Source/IKRig/Private/Rig/IKRigProcessor.cpp
+++ b/Engine/Plugins/Animation/IKRig/Source/IKRig/Private/Rig/IKRigProcessor.cpp
@@ -13,8 +13,8 @@
 
 void FIKRigProcessor::Initialize(
 	const UIKRigDefinition* InRigAsset,
-	const USkeletalMesh* SkeletalMesh,
-	const TArray<FName>& ExcludedGoals)
+	const USkeletalMesh* InSkeletalMesh,
+	const FIKRigGoalContainer& InOptionalGoals)
 {
 	bInitialized = false;
 	
@@ -34,7 +34,7 @@ void FIKRigProcessor::Initialize(
 	bTriedToInitialize = true;
 
 	// copy skeleton data from the actual skeleton we want to run on
-	Skeleton.SetInputSkeleton(SkeletalMesh, InRigAsset->GetSkeleton().ExcludedBones);
+	Skeleton.SetInputSkeleton(InSkeletalMesh, InRigAsset->GetSkeleton().ExcludedBones);
 	
 	if (InRigAsset->GetSkeleton().BoneNames.IsEmpty())
 	{
@@ -44,7 +44,7 @@ void FIKRigProcessor::Initialize(
 		return;
 	}
 
-	if (!FIKRigProcessor::IsIKRigCompatibleWithSkeleton(InRigAsset, SkeletalMesh, &Log))
+	if (!FIKRigProcessor::IsIKRigCompatibleWithSkeleton(InRigAsset, InSkeletalMesh, &Log))
 	{
 		Log.LogError(FText::Format(
 			LOCTEXT("SkeletonMissingRequiredBones", "Trying to initialize IKRig, '{0}' with a Skeleton that is missing required bones. See prior warnings."),
@@ -52,54 +52,95 @@ void FIKRigProcessor::Initialize(
 		return;
 	}
 	
-	// initialize goals based on source asset
+	// initialize goal container
 	GoalContainer.Empty();
-	const TArray<UIKRigEffectorGoal*>& EffectorGoals = InRigAsset->GetGoalArray();
-	for (const UIKRigEffectorGoal* EffectorGoal : EffectorGoals)
+	if (InOptionalGoals.IsEmpty())
 	{
-		// skip excluded goals
-		if (ExcludedGoals.Contains(EffectorGoal->GoalName))
+		// use goals from IK Rig
+		const TArray<UIKRigEffectorGoal*>& EffectorGoals = InRigAsset->GetGoalArray();
+		GoalContainer.FillWithGoalArray(EffectorGoals);
+	}
+	else
+	{
+		// use the optional goals passed in
+		GoalContainer = InOptionalGoals;
+	}
+
+	// get the set of ALL required goals so we can validate that each goal is used by at least 1 solver
+	const TArray<const FIKRigSolverBase*> AssetSolvers = InRigAsset->GetSolverArray();
+	TSet<FName> AllRequiredGoals;
+	for (const FIKRigSolverBase* AssetSolver : AssetSolvers)
+	{
+		AssetSolver->GetRequiredGoals(AllRequiredGoals);
+	}
+
+	// generate list of excluded goals
+	// (these are goals that were not in the supplied GoalContainer but are used by at least 1 solver)
+	TSet<FName> ExcludedGoals;
+	for (FName RequiredGoalName : AllRequiredGoals)
+	{
+		const FIKRigGoal* RequiredGoal = GoalContainer.FindGoalByName(RequiredGoalName);
+		if (!RequiredGoal)
 		{
-			continue;
+			ExcludedGoals.Add(RequiredGoalName);
+		}
+	}
+	// exclude disabled goals
+	if (!InOptionalGoals.IsEmpty())
+	{
+		for (const FIKRigGoal& InputGoal : InOptionalGoals.GetGoalArray())
+		{
+			if (!InputGoal.bEnabled)
+			{
+				ExcludedGoals.Add(InputGoal.Name);
+			}
 		}
-		
-		// add a copy of the Goal to the container
-		GoalContainer.SetIKGoal(EffectorGoal);
 	}
 	
 	// initialize goal bones from asset
 	GoalBones.Reset();
-	for (const UIKRigEffectorGoal* EffectorGoal : EffectorGoals)
+	for (const FIKRigGoal& Goal : GoalContainer.GetGoalArray())
 	{
-		// skip excluded goals
-		if (ExcludedGoals.Contains(EffectorGoal->GoalName))
+		// default to using the bone name supplied by the goal
+		FName BoneNameToUse = Goal.BoneName;
+		
+		// if the goal itself didn't come with a bone name, try to find it from the asset
+		if (BoneNameToUse == NAME_None)
 		{
-			continue;
+			const TArray<UIKRigEffectorGoal*>& EffectorGoals = InRigAsset->GetGoalArray();
+			for (const UIKRigEffectorGoal* EffectorGoal : EffectorGoals)
+			{
+				if (EffectorGoal->GoalName == Goal.Name)
+				{
+					BoneNameToUse = EffectorGoal->BoneName;
+					break;
+				}
+			}
 		}
 		
 		FGoalBone NewGoalBone;
-		NewGoalBone.BoneName = EffectorGoal->BoneName;
-		NewGoalBone.BoneIndex = Skeleton.GetBoneIndexFromName(EffectorGoal->BoneName);
+		NewGoalBone.BoneName = BoneNameToUse;
+		NewGoalBone.BoneIndex = Skeleton.GetBoneIndexFromName(Goal.BoneName);
 
 		// validate that the skeleton we are trying to solve this goal on contains the bone the goal expects
 		if (NewGoalBone.BoneIndex == INDEX_NONE)
 		{
 			Log.LogError(FText::Format(
 				LOCTEXT("MissingGoalBone", "IK Rig, '{0}' has a Goal, '{1}' that references an unknown bone, '{2}'. Cannot evaluate."),
-				FText::FromString(InRigAsset->GetName()), FText::FromName(EffectorGoal->GoalName), FText::FromName(EffectorGoal->BoneName) ));
+				FText::FromString(InRigAsset->GetName()), FText::FromName(Goal.Name), FText::FromName(Goal.BoneName) ));
 			return;
 		}
 
 		// validate that there is not already a different goal, with the same name, that is using a different bone
 		// (all goals with the same name must reference the same bone within a single IK Rig)
-		if (const FGoalBone* Bone = GoalBones.Find(EffectorGoal->GoalName))
+		if (const FGoalBone* Bone = GoalBones.Find(Goal.Name))
 		{
 			if (Bone->BoneName != NewGoalBone.BoneName)
 			{
 				Log.LogError(FText::Format(
 				LOCTEXT("DuplicateGoal", "IK Rig, '{0}' has a Goal, '{1}' that references different bones in different solvers, '{2}' and '{3}'. Cannot evaluate."),
                 FText::FromString(InRigAsset->GetName()),
-                FText::FromName(EffectorGoal->GoalName),
+                FText::FromName(Goal.Name),
                 FText::FromName(Bone->BoneName), 
                 FText::FromName(NewGoalBone.BoneName)));
 				return;
@@ -107,26 +148,20 @@ void FIKRigProcessor::Initialize(
 		}
 
 		// warn if goal is not connected to any solvers
-		const TArray<const FIKRigSolverBase*> AssetSolvers = InRigAsset->GetSolverArray();
-		TSet<FName> AllRequiredGoals;
-		for (const FIKRigSolverBase* AssetSolver : AssetSolvers)
-		{
-			AssetSolver->GetRequiredGoals(AllRequiredGoals);
-		}
-		if (!AllRequiredGoals.Contains(EffectorGoal->GoalName))
+		if (!AllRequiredGoals.Contains(Goal.Name))
 		{
 			Log.LogWarning(FText::Format(
 				LOCTEXT("DisconnectedGoal", "IK Rig, '{0}' has a Goal, '{1}' that is not connected to any solvers. It will have no effect."),
-				FText::FromString(InRigAsset->GetName()), FText::FromName(EffectorGoal->GoalName)));
+				FText::FromString(InRigAsset->GetName()), FText::FromName(Goal.Name)));
 		}
 		
-		GoalBones.Add(EffectorGoal->GoalName, NewGoalBone);
+		GoalBones.Add(Goal.Name, NewGoalBone);
 	}
 
 	// create copies of all the solvers in the IK rig
-	const TArray<FInstancedStruct>& AssetSolvers = InRigAsset->GetSolverStructs();
-	Solvers.Reset(AssetSolvers.Num());
-	for (const FInstancedStruct& AssetSolverStruct : AssetSolvers)
+	const TArray<FInstancedStruct>& AssetSolverStructs = InRigAsset->GetSolverStructs();
+	Solvers.Reset(AssetSolverStructs.Num());
+	for (const FInstancedStruct& AssetSolverStruct : AssetSolverStructs)
 	{
 		if (!AssetSolverStruct.IsValid())
 		{
@@ -146,9 +181,10 @@ void FIKRigProcessor::Initialize(
 		Solvers.Emplace(MakeShareable(NewSolver));
 
 		// remove excluded goals from the solver
-		for (const FName& ExcludedGoal : ExcludedGoals)
+		// this must be done BEFORE initializing the solver
+		for (const FName ExcludedGoalName : ExcludedGoals)
 		{
-			NewSolver->OnGoalRemoved(ExcludedGoal);
+			NewSolver->OnGoalRemoved(ExcludedGoalName);
 		}
 
 		// initialize it and store in the processor
@@ -164,14 +200,16 @@ void FIKRigProcessor::Initialize(
 		{
 			Log.LogWarning(FText::Format(
 				LOCTEXT("InvalidRetargetChain", "Invalid Retarget Chain: '{0}'. End bone is not a child of the start bone in Skeletal Mesh, '{1}'."),
-				FText::FromString(Chain.ChainName.ToString()), FText::FromString(SkeletalMesh->GetName())));
+				FText::FromString(Chain.ChainName.ToString()), FText::FromString(InSkeletalMesh->GetName())));
 		}
 	}
 
 	Log.LogInfo(FText::Format(
 				LOCTEXT("SuccessfulInit", "IK Rig, '{0}' ready to run on {1}."),
-				FText::FromString(InRigAsset->GetName()), FText::FromString(SkeletalMesh->GetName())));
-	
+				FText::FromString(InRigAsset->GetName()), FText::FromString(InSkeletalMesh->GetName())));
+
+	GoalContainer.bRigNeedsInitialized = false;
+	bTriedToInitialize = false;
 	bInitialized = true;
 }
 
@@ -275,7 +313,10 @@ bool FIKRigProcessor::IsIKRigCompatibleWithSkeleton(
 
 void FIKRigProcessor::SetInputPoseGlobal(const TArray<FTransform>& InGlobalBoneTransforms) 
 {
-	check(bInitialized);
+	if (!ensure(bInitialized))
+	{
+		return;
+	}
 	check(InGlobalBoneTransforms.Num() == Skeleton.CurrentPoseGlobal.Num());
 	Skeleton.CurrentPoseGlobal = InGlobalBoneTransforms;
 	Skeleton.UpdateAllLocalTransformFromGlobal();
@@ -283,26 +324,49 @@ void FIKRigProcessor::SetInputPoseGlobal(const TArray<FTransform>& InGlobalBoneT
 
 void FIKRigProcessor::SetInputPoseToRefPose()
 {
-	check(bInitialized);
+	if (!ensure(bInitialized))
+	{
+		return;
+	}
+	
 	Skeleton.CurrentPoseGlobal = Skeleton.RefPoseGlobal;
 	Skeleton.UpdateAllLocalTransformFromGlobal();
 }
 
+void FIKRigProcessor::ApplyGoalsFromOtherContainer(const FIKRigGoalContainer& InGoalContainer)
+{
+	for (const FIKRigGoal& InGoal : InGoalContainer.GetGoalArray())
+	{
+		SetIKGoal(InGoal);
+	}
+}
+
 void FIKRigProcessor::SetIKGoal(const FIKRigGoal& InGoal)
 {
-	check(bInitialized);
+	if (!ensure(bInitialized))
+	{
+		return;
+	}
+
 	GoalContainer.SetIKGoal(InGoal);
 }
 
 void FIKRigProcessor::SetIKGoal(const UIKRigEffectorGoal* InGoal)
 {
-	check(bInitialized);
+	if (!ensure(bInitialized))
+	{
+		return;
+	}
+	
 	GoalContainer.SetIKGoal(InGoal);
 }
 
 void FIKRigProcessor::Solve(const FTransform& ComponentToWorld)
 {
-	check(bInitialized);
+	if (!bInitialized)
+	{
+		return;
+	}
 	
 	// convert goals into component space and blend towards input pose by alpha
 	ResolveFinalGoalTransforms(ComponentToWorld);
@@ -354,7 +418,15 @@ void FIKRigProcessor::CopyAllSettingsFromAsset(const UIKRigDefinition* SourceAss
 	const TArray<UIKRigEffectorGoal*>& AssetGoals =  SourceAsset->GetGoalArray();
 	for (const UIKRigEffectorGoal* AssetGoal : AssetGoals)
 	{
-		SetIKGoal(AssetGoal);
+		if (FIKRigGoal* Goal = GoalContainer.FindGoalByName(AssetGoal->GoalName))
+		{
+			Goal->RotationAlpha = AssetGoal->RotationAlpha;
+			Goal->PositionAlpha = AssetGoal->PositionAlpha;
+		}
+		else
+		{
+			SetIKGoal(AssetGoal);
+		}
 	}
 
 	// copy solver settings
diff --git a/Engine/Plugins/Animation/IKRig/Source/IKRig/Public/ActorComponents/IKRigComponent.h b/Engine/Plugins/Animation/IKRig/Source/IKRig/Public/ActorComponents/IKRigComponent.h
--- a/Engine/Plugins/Animation/IKRig/Source/IKRig/Public/ActorComponents/IKRigComponent.h
+++ b/Engine/Plugins/Animation/IKRig/Source/IKRig/Public/ActorComponents/IKRigComponent.h
@@ -29,12 +29,14 @@ public:
 		GoalContainer.SetIKGoal(
 			FIKRigGoal(
 				GoalName,
+				NAME_None,
 				Position,
 				Rotation,
 				PositionAlpha,
 				RotationAlpha,
 				EIKRigGoalSpace::Component,
-				EIKRigGoalSpace::Component));
+				EIKRigGoalSpace::Component,
+				true));
 	};
 
 	/** Set an IK Rig Goal transform (assumed in Component Space of Skeletal Mesh) with separate alpha values. */
@@ -50,12 +52,14 @@ public:
 		GoalContainer.SetIKGoal(
 			FIKRigGoal(
 				GoalName,
+				NAME_None,
 				Position,
 				Rotation,
 				PositionAlpha,
 				RotationAlpha,
 				EIKRigGoalSpace::Component,
-				EIKRigGoalSpace::Component));
+				EIKRigGoalSpace::Component,
+				true));
 	};
 	
 	/** Apply a IKRigGoal and store it on this rig. Goal transform assumed in Component Space of Skeletal Mesh. */
diff --git a/Engine/Plugins/Animation/IKRig/Source/IKRig/Public/AnimNodes/AnimNode_RetargetPoseFromMesh.h b/Engine/Plugins/Animation/IKRig/Source/IKRig/Public/AnimNodes/AnimNode_RetargetPoseFromMesh.h
--- a/Engine/Plugins/Animation/IKRig/Source/IKRig/Public/AnimNodes/AnimNode_RetargetPoseFromMesh.h
+++ b/Engine/Plugins/Animation/IKRig/Source/IKRig/Public/AnimNodes/AnimNode_RetargetPoseFromMesh.h
@@ -46,10 +46,6 @@ struct IKRIG_API FAnimNode_RetargetPoseFromMesh : public FAnimNode_Base
 	UPROPERTY(EditAnywhere, BlueprintReadWrite, Category = Settings, meta=(PinHiddenByDefault))
 	FRetargetProfile CustomRetargetProfile;
 	
-	// Copy curves from the Source Mesh Component. This will copy any curves the source/target Skeleton have in common.
-	UPROPERTY(EditAnywhere, BlueprintReadWrite, Category = Settings, meta = (NeverAsPin))
-	bool bCopyCurves = true;
-	
 	// Max LOD that this node is allowed to run.
 	// For example if you have LODThreshold at 2, it will run until LOD 2 (based on 0 index) when the component LOD becomes 3, it will stop update/evaluate
 	// A value of -1 forces the node to execute at all LOD levels.
@@ -106,12 +102,6 @@ private:
 
 	// mapping from required compact bone indices to target mesh bones expected by the retargeter
 	TArray< TPair<int32, int32> > CompactToTargetBoneIndexMap;
-
-	// cached curves, copied on the game thread in PreUpdate()
-	FBlendedHeapCurve SourceCurves;
-	
-	// map of curve names to values, passed to retargeter for IK planting (copied from source mesh)
-	TMap<FName, float> SpeedValuesFromCurves;
 	
 	// the delta time this tick
 	float DeltaTime;
diff --git a/Engine/Plugins/Animation/IKRig/Source/IKRig/Public/IKRig.h b/Engine/Plugins/Animation/IKRig/Source/IKRig/Public/IKRig.h
--- a/Engine/Plugins/Animation/IKRig/Source/IKRig/Public/IKRig.h
+++ b/Engine/Plugins/Animation/IKRig/Source/IKRig/Public/IKRig.h
@@ -5,6 +5,9 @@
 #include "CoreMinimal.h"
 #include "Modules/ModuleManager.h"
 
+// a UProperty metadata tag used to indicate that a property should trigger reinitialization when edited
+static const TCHAR* IKRigReinitOnEditMetaLabel = TEXT("ReinitializeOnEdit");
+
 class FIKRigModule : public IModuleInterface
 {
 public:
diff --git a/Engine/Plugins/Animation/IKRig/Source/IKRig/Public/IKRigObjectVersion.h b/Engine/Plugins/Animation/IKRig/Source/IKRig/Public/IKRigObjectVersion.h
--- a/Engine/Plugins/Animation/IKRig/Source/IKRig/Public/IKRigObjectVersion.h
+++ b/Engine/Plugins/Animation/IKRig/Source/IKRig/Public/IKRigObjectVersion.h
@@ -27,6 +27,9 @@ struct IKRIG_API FIKRigObjectVersion
 		// we switched "Use Attach Parent" from a bool to an enum on retarget pose from mesh node
 		UseAttachedParentDeprecated,
 
+		// modularized retargeting internals (converted phases to ops)
+		ModularRetargeterOps,
+
 		// -----<new versions can be added above this line>-------------------------------------------------
 		VersionPlusOne,
 		LatestVersion = VersionPlusOne - 1
