[IKRig] Major refactor part 3 of 3.

diff --git a/Engine/Plugins/Animation/IKRig/Source/IKRig/Public/Retargeter/IKRetargetOps.h b/Engine/Plugins/Animation/IKRig/Source/IKRig/Public/Retargeter/IKRetargetOps.h
--- a/Engine/Plugins/Animation/IKRig/Source/IKRig/Public/Retargeter/IKRetargetOps.h
+++ b/Engine/Plugins/Animation/IKRig/Source/IKRig/Public/Retargeter/IKRetargetOps.h
@@ -2,31 +2,35 @@
 
 #pragma once
 
-#include "Retargeter/IKRetargetProcessor.h"
-
 #include "IKRetargetOps.generated.h"
 
+struct FIKRetargetProcessor;
+struct FRetargetSkeleton;
+struct FTargetSkeleton;
+struct FIKRigLogger;
+class UIKRetargeter;
+struct FInstancedStruct;
+class USkeleton;
+struct FPoseContext;
+
 #if WITH_EDITOR
+class FPrimitiveDrawInterface;
 class FIKRetargetEditorController;
+struct FIKRetargetDebugDrawState;
 #endif
 
-USTRUCT(BlueprintType)
-struct FCurveRemapPair
-{
-	GENERATED_BODY()
-	
-	// The curve name on the SOURCE skeletal mesh to copy animation data from.
-	UPROPERTY(EditAnywhere, Category=Settings)
-	FName SourceCurve;
-
-	// The curve name on the TARGET skeletal mesh to receive animation data.
-	UPROPERTY(EditAnywhere, Category=Settings)
-	FName TargetCurve;
-};
-
 /** This is the base class for defining editable settings for your custom retargeting operation.
- * All user-configurable properties for your "op" should be stored in a subclass of this
+ * All user-configurable properties for your "op" should be stored in a subclass of this.
+ * These settings will automatically be:
+ * 1. Displayed in the details panel when the op is selected
+ * 2. Saved/loaded with the op in the retarget asset
+ * 3. Applied to the op at runtime as part of a profile
+ * 
  * NOTE: the derived type must be returned by the op's GetSettingsType() and GetSettings()
+ * 
+ * NOTE: UProperties that require reinitialization when modified must be marked meta=(ReinitializeOnEdit)
+ * When modified, in the editor, these properties will trigger a reinitialization at which point
+ * the runtime Op will get the latest values automatically.
  */
 USTRUCT(BlueprintType)
 struct IKRIG_API FIKRetargetOpSettingsBase
@@ -38,6 +42,9 @@ struct IKRIG_API FIKRetargetOpSettingsBase
 #if WITH_EDITOR
 	/** the instance of this op currently running in the editor viewport */
 	mutable FIKRetargetOpSettingsBase* EditorInstance;
+	/** a reference to the skeletons for bone selector widgets (can return these from GetSkeleton) */
+	mutable const USkeleton* SourceSkeletonAsset;
+	mutable const USkeleton* TargetSkeletonAsset;
 	/** allow settings to provide a skeleton for any given FBoneReference widget */
 	virtual USkeleton* GetSkeleton(const FName InPropertyName) { return nullptr; };
 #endif
@@ -45,7 +52,7 @@ struct IKRIG_API FIKRetargetOpSettingsBase
 
 /**
  * This is the base class for defining operations that live in the retargeter "op" stack.
- * These operations are executed in order by called the virtual Run() function on each one in order.
+ * These operations are executed in order by calling the virtual Run() function on each one in order.
  * The Run() function takes an input pose on the source skeletal mesh and affects the output pose on the target mesh.
  * NOTE: any user defined settings associated with a retarget op must be aggregated into a custom UStruct derived from FIKRetargetOpSettingsBase
  * This ensures that the settings are user editable in the details panel with full undo/redo support and serialization.
@@ -62,63 +69,164 @@ public:
 	/** (optional) override to cache internal data when initializing the processor
 	 * NOTE: you must set bIsInitialized to true to inform the retargeter that this op is ok to execute. */
 	virtual bool Initialize(
-		const FIKRetargetProcessor& Processor,
-		const FRetargetSkeleton& SourceSkeleton,
-		const FTargetSkeleton& TargetSkeleton,
+		const FIKRetargetProcessor& InProcessor,
+		const FRetargetSkeleton& InSourceSkeleton,
+		const FTargetSkeleton& InTargetSkeleton,
 		FIKRigLogger& Log) { bIsInitialized = true; return true; };
 
 	/** (optional) override to evaluate this operation and modify the output pose */
 	virtual void Run(
-		const FIKRetargetProcessor& Processor,
+		FIKRetargetProcessor& InProcessor,
+		const double InDeltaTime,
 		const TArray<FTransform>& InSourceGlobalPose,
 		TArray<FTransform>& OutTargetGlobalPose){};
+	
+	/** (optional) a second pass of initialization that ops can use after Op::Initialize() is called on all ops. */
+	virtual void PostInitialize(
+		const FIKRetargetProcessor& InProcessor,
+		const FRetargetSkeleton& InSourceSkeleton,
+		const FTargetSkeleton& InTargetSkeleton,
+		FIKRigLogger& InLog) {};
+	
+	/** get if this op is enabled */
+	bool IsEnabled() const { return bIsEnabled; };
+
+	/** turn this operation on/off (will be skipped during execution if disabled) */
+	void SetEnabled(const bool bEnabled){ bIsEnabled = bEnabled; };
+
+	/** return true if this op is initialized and ready to run */
+	bool IsInitialized() const { return bIsInitialized; };
 
 	/** (required) override and return a pointer to the settings struct used by this operation */ 
 	virtual FIKRetargetOpSettingsBase* GetSettings() { checkNoEntry(); return nullptr; };
 	
+	/** (optional) override to customize which settings are copied into the op at runtime
+	 * NOTE:
+	 * This is called while the Op is running AFTER Initialize() so it up to the Op author
+	 * to copy only those settings which are safe to be updated while the op is running.
+	 * The default implementation copies ALL settings all the time.
+	 * @param InSettings - pointer to the settings to copy, can be safely cast to the type returned by GetSettingsType() */ 
+	virtual void SetSettings(const FIKRetargetOpSettingsBase* InSettings);
+	
 	/** (required) override and return the type used to house the settings for this operation */
 	virtual const UScriptStruct* GetSettingsType() const { checkNoEntry(); return nullptr; };
 
-	/** (optional) override to specify curves you want to retarget from the source animation graph */
-	virtual void GetCurvesToRetarget(TArray<FCurveRemapPair>& OutCurvesToRetarget){};
+	/** (required) override and return the type of this op (the derived subclass) */
+	virtual const UScriptStruct* GetType() const { checkNoEntry(); return FIKRetargetOpBase::StaticStruct(); };
 
-	/** get if this op is enabled */
-	bool IsEnabled() const { return bIsEnabled; };
+	/** (optional) override to reset any internal state when animation playback is reset or stopped (ie, springs / dampers etc) */
+	virtual void OnPlaybackReset() {};
 
-	/** turn this operation on/off (will be skipped during execution if disabled) */
-	void SetEnabled(const bool bEnabled){ bIsEnabled = bEnabled; };
+	/** (optional) override to get any data from the source or target skeletal mesh component
+	 * NOTE: this is called during AnimGraph::PreUpdate() which runs on the main thread, use caution. */
+	virtual void AnimGraphPreUpdateMainThread(USkeletalMeshComponent& SourceMeshComponent, USkeletalMeshComponent& TargetMeshComponent) {};
 
-	/** return true if this op is initialized and ready to run */
-	bool IsInitialized() const { return bIsInitialized; };
+	/** (optional) override to get any data from the anim graph as it's evaluating.
+	 * NOTE: this is called during AnimGraph::Evaluate_AnyThread() BEFORE the retargeting Ops are executed. */
+	virtual void AnimGraphEvaluateAnyThread(FPoseContext& Output) {};
+
+	/** (optional) override and add any bones that your op modifies to the output TSet of bone indices
+	 * Add indices of any bone that this op modifies. Any bone not registered here will be FK parented by other operations*/
+	virtual void CollectRetargetedBones(TSet<int32>& OutRetargetedBones) const {};
+	
+	/** provide a scripting object to edit your custom solver type in the editor via BP/Python
+	* NOTE: this returns an instance of whatever type is returned by GetControllerType(). This can be custom per op. */
+	UIKRetargetOpControllerBase* GetController(UObject* Outer);
+
+	/** (optional, but recommended) provide a custom controller type (deriving from UIKRetargetOpControllerBase) as an API for editing the op
+	 * NOTE: this type will automatically be lazy instantiated when/if the user calls GetController() on the op. */
+	virtual const UClass* GetControllerType() const;
 
 #if WITH_EDITOR
-	/** override to automate initial setup after being added to the stack */
-	virtual void OnAddedToStack(const UIKRetargeter* Asset) {};
+	/** get the nice name to display in the viewport (defaults to "DisplayName" UStruct metadata) */
+	FText GetNiceName() const;
+	
+	/** (optional) override to automate initial setup after being added to the stack */
+	virtual void RunInitialSetup(const UIKRetargeter* InRetargetAsset) {};
 
-	/** override to give your operation a nice name to display in the UI */
-	virtual FText GetNiceName() const { return FText::FromString(TEXT("Default Op Name")); };
+	/** (optional) override to display a warning message in the op stack */
+	virtual FText GetWarningMessage() const;
 
-	/** override to display a warning message in the op stack */
-	virtual FText GetWarningMessage() const { return FText::GetEmpty(); };
+	/** (optional) override to draw debug info into the editor viewport when the Op is selected */
+	virtual void DebugDraw(
+		FPrimitiveDrawInterface* InPDI,
+		const FTransform& InComponentTransform,
+		const double InComponentScale,
+		const FIKRetargetDebugDrawState& InEditorState) const {};
 
+	/** (optional) override to support resetting all settings for a given chain */
+	virtual void SetChainSettingsToDefault(const FName& InChainName) {};
+
+	/** (optional) override to support "Reset to Default" UI for chains */
+	virtual bool AreChainSettingsAtDefault(const FName& InChainName) { return true; };
 #endif
 
+	/** wholesale copy all settings from the input settings into this op.
+	 * NOTE: InSettings is assumed to be castable to the type returned by GetSettingsType()
+	 * NOTE: be careful calling this on an initialized op as it may invalidate it's runtime state (use SetSettings instead) */
+	void CopySettingsRaw(const FIKRetargetOpSettingsBase* InSettings, const TArray<FName>& InPropertiesToIgnore);
+
 private:
 	
 	/** when false, execution of this op is skipped */
 	bool bIsEnabled = true;
 
 protected:
+
+	/** a convenience function to copy all properties from one struct to another while ignoring some */
+	static void CopyStructProperties(
+		const UStruct* InStructType,
+		const void* InSrcStruct,
+		void* InOutDestStruct,
+		const TArray<FName>& InPropertiesToIgnore);
+
+	// helper function for ops to instantiate their own controller
+	UIKRetargetOpControllerBase* CreateControllerIfNeeded(UObject* Outer);
+
+	// the  controller used to edit this solver by script/blueprint (lazy instantiated when needed)
+	TStrongObjectPtr<UIKRetargetOpControllerBase> Controller = nullptr;
 	
 	/** set to true when the op is ready to run */
 	bool bIsInitialized = false;
 };
 
+/**
+ * This is the base class for defining a custom controller for a given retarget op type.
+ * Controllers provide an API for editing ops from blueprint or python.
+ *
+ * To use a controller:
+ * 1. Get a controller by calling UIKRetargeterController::GetOpController(int InOpIndex)
+ * 2. Cast the returned UIKRetargetOpControllerBase* to the type of controller used by the op you want to modify.
+ * 3. Call the public getter/setters to modify the op as desired
+ *
+ * NOTE: these controllers are necessary because the op UStructs do not support UFunctions
+*/ 
+UCLASS(BlueprintType)
+class IKRIG_API UIKRetargetOpControllerBase : public UObject
+{
+	GENERATED_BODY()
+	
+public:
+
+	/* Turn the op on/off in the retargeting stack.
+	 * @param bIsEnabled if true, the op will execute, false otherwise. */
+	UFUNCTION(BlueprintCallable, Category = RetargetOp)
+	void SetEnabled(bool bIsEnabled) { OpToControl->SetEnabled(bIsEnabled); };
+
+	/* Get if the op is on or off in the retargeting stack.
+	 * @param bIsEnabled returns true if the op is enabled, false otherwise. */
+	UFUNCTION(BlueprintCallable, Category = RetargetOp)
+	bool GetEnabled() { return OpToControl->IsEnabled(); };
+	
+	// the op this controller controls
+	FIKRetargetOpBase* OpToControl = nullptr;
+};
+
 
 //
-// BEGIN DEPRECATED UOBJECT-BASED RETARGET OP
+// BEGIN DEPRECATED OP BASE
 //
-
+PRAGMA_DISABLE_DEPRECATION_WARNINGS
 UCLASS(Deprecated)
 class UE_DEPRECATED(5.6, "This type has been replaced by FRetargetOpBase. URetargetOpBase-based ops no longer work, please refactor into the new FRetargetOpBase struct.")
 IKRIG_API UDEPRECATED_RetargetOpBase : public UObject
@@ -150,7 +258,6 @@ IKRIG_API UDEPRECATED_RetargetOpBase : public UObject
 	virtual void OnAddedToStack(const UIKRetargeter* Asset) {};
 	virtual FText GetNiceName() const { return FText::FromString(TEXT("Default Op Name")); };
 	virtual FText WarningMessage() const { return FText::GetEmpty(); };
-	mutable URetargetOpBase* EditorInstance;
 #endif
 };
 
@@ -160,10 +267,12 @@ IKRIG_API UDEPRECATED_RetargetOpStack : public UObject
 {
 	GENERATED_BODY()
 	
-public:
-
-	PRAGMA_DISABLE_DEPRECATION_WARNINGS
+	public:
+	
 	UPROPERTY()
 	TArray<TObjectPtr<UDEPRECATED_RetargetOpBase>> RetargetOps_DEPRECATED;
-	PRAGMA_ENABLE_DEPRECATION_WARNINGS
-};
\ No newline at end of file
+};
+PRAGMA_ENABLE_DEPRECATION_WARNINGS
+//
+// END DEPRECATED OP BASE
+//
\ No newline at end of file
diff --git a/Engine/Plugins/Animation/IKRig/Source/IKRig/Public/Retargeter/IKRetargetProcessor.h b/Engine/Plugins/Animation/IKRig/Source/IKRig/Public/Retargeter/IKRetargetProcessor.h
--- a/Engine/Plugins/Animation/IKRig/Source/IKRig/Public/Retargeter/IKRetargetProcessor.h
+++ b/Engine/Plugins/Animation/IKRig/Source/IKRig/Public/Retargeter/IKRetargetProcessor.h
@@ -4,6 +4,7 @@
 
 #include "CoreMinimal.h"
 #include "IKRetargeter.h"
+#include "IKRetargetOps.h"
 #include "IKRetargetSettings.h"
 #include "IKRigLogger.h"
 #include "Kismet/KismetMathLibrary.h"
@@ -20,11 +21,21 @@ struct FIKRigProcessor;
 struct FReferenceSkeleton;
 struct FBoneChain;
 struct FIKRetargetPose;
-class UObject;
 class UIKRetargeter;
 class USkeletalMesh;
 class URetargetOpBase;
-
+struct FIKRetargetPelvisMotionOp;
+struct FResolvedBoneChain;
+struct FIKRetargetDebugDrawState;
+
+/**
+ * A retarget skeleton contains:
+ * 1. hierarchy data (bones, parents)
+ * 2. a retarget pose
+ * 3. functions for converting Local/Global poses
+ * A retarget skeleton is created for both the source and target skeleton.
+ * The target skeleton uses a specialized skeleton subclass for managing a bone mask of retargeted bones
+ */
 struct IKRIG_API FRetargetSkeleton
 {
 	TArray<FName> BoneNames;				// list of all bone names in ref skeleton order
@@ -34,7 +45,7 @@ struct IKRIG_API FRetargetSkeleton
 	
 	FName RetargetPoseName;					// the name of the retarget pose this was initialized with
 	int32 RetargetPoseVersion;				// the version of the retarget pose this was initialized with (transient)
-	float RetargetPoseScale;				// the scale of the retarget pose this was initialized with (transient)
+	double RetargetPoseScale;				// the scale of the retarget pose this was initialized with (transient)
 	
 	const USkeletalMesh* SkeletalMesh;		// the skeletal mesh this was initialized with
 	TArray<FName> ChainThatContainsBone;	// record which chain is actually controlling each bone
@@ -45,7 +56,9 @@ struct IKRIG_API FRetargetSkeleton
 		const FName InRetargetPoseName,
 		const FIKRetargetPose* RetargetPose,
 		const FName RetargetRootBone,
-		const float ScaleFactor);
+		const double ScaleFactor);
+
+	void FillPerBoneChainNames(const TArray<FResolvedBoneChain>& InBoneChains);
 
 	void Reset();
 
@@ -53,7 +66,7 @@ struct IKRIG_API FRetargetSkeleton
 		const FName InRetargetPoseName,
 		const FIKRetargetPose* InRetargetPose,
 		const FName RetargetRootBone,
-		const float ScaleFactor);
+		const double ScaleFactor);
 
 	int32 FindBoneIndexByName(const FName InName) const;
 
@@ -98,13 +111,14 @@ private:
 	mutable TArray<int32> CachedEndOfBranchIndices;
 };
 
+/** A retarget skeleton for the target skeleton.
+ * Contains the output pose buffer for the target skeleton
+ * Provides a per-bone mask specifying which bones are retargeted.
+ */
 struct FTargetSkeleton : public FRetargetSkeleton
 {
 	TArray<FTransform> OutputGlobalPose;
 	TArray<FTransform> InputLocalPose;
-	// true for bones that are in a target chain that is ALSO mapped to a source chain
-	// ie, bones that are actually posed based on a mapped source chain
-	TArray<bool> IsBoneRetargeted;
 
 	void Initialize(
 		const USkeletalMesh* InSkeletalMesh,
@@ -115,372 +129,144 @@ struct FTargetSkeleton : public FRetargetSkeleton
 
 	void Reset();
 
-	void SetBoneIsRetargeted(const int32 BoneIndex, const bool IsRetargeted);
-
-	void UpdateGlobalTransformsAllNonRetargetedBones(const TArray<FTransform>& LocalPose, TArray<FTransform>& InOutGlobalPose);
-};
-
-/** resolving an FBoneChain to an actual skeleton, used to validate compatibility and get all chain indices */
-struct IKRIG_API FResolvedBoneChain
-{
-	FResolvedBoneChain(const FBoneChain& BoneChain, const FRetargetSkeleton& Skeleton, TArray<int32> &OutBoneIndices);
-
-	/* Does the START bone exist in the skeleton? */
-	bool bFoundStartBone = false;
-	/* Does the END bone exist in the skeleton? */
-	bool bFoundEndBone = false;
-	/* Is the END bone equals or a child of the START bone? */
-	bool bEndIsStartOrChildOfStart  = false;
-
-	bool IsValid() const
-	{
-		return bFoundStartBone && bFoundEndBone && bEndIsStartOrChildOfStart;
-	}
-};
-
-struct FRootSource
-{
-	FName BoneName;
-	int32 BoneIndex;
-	FQuat InitialRotation;
-	float InitialHeightInverse;
-	FVector InitialPosition;
-	FVector CurrentPosition;
-	FVector CurrentPositionNormalized;
-	FQuat CurrentRotation;
-};
-
-struct FRootTarget
-{
-	FName BoneName;
-	int32 BoneIndex;
-	FVector InitialPosition;
-	FQuat InitialRotation;
-	float InitialHeight;
-
-	FVector RootTranslationDelta;
-	FQuat RootRotationDelta;
-};
-
-struct FRootRetargeter
-{	
-	FRootSource Source;
-	FRootTarget Target;
-	FTargetRootSettings Settings;
-
-	void Reset();
-	
-	bool InitializeSource(
-		const FName SourceRootBoneName,
-		const FRetargetSkeleton& SourceSkeleton,
-		FIKRigLogger& Log);
-	
-	bool InitializeTarget(
-		const FName TargetRootBoneName,
-		const FTargetSkeleton& TargetSkeleton,
-		FIKRigLogger& Log);
-
-	void EncodePose(const TArray<FTransform> &SourceGlobalPose);
-	
-	void DecodePose(TArray<FTransform> &OutTargetGlobalPose);
+	void UpdateGlobalTransformsWithMask(
+		const TArray<bool>& InBoneMask,
+		const TArray<FTransform>& InLocalPose,
+		TArray<FTransform>& InOutGlobalPose) const;
 
-	FVector GetGlobalScaleVector() const
-	{
-		return GlobalScaleFactor * FVector(Settings.ScaleHorizontal, Settings.ScaleHorizontal, Settings.ScaleVertical);
-	}
+	void SetRetargetedBones(const TSet<int32>& InRetargetedBones);
 
-private:
-	FVector GlobalScaleFactor;
-};
+	bool GetIsBoneRetargeted(int32 InBoneIndex) const;
 
-struct FPoleVectorMatcher
-{
-	EAxis::Type SourcePoleAxis;
-	EAxis::Type TargetPoleAxis;
-	float TargetToSourceAngularOffsetAtRefPose;
-	TArray<int32> AllChildrenWithinChain;
+	const TArray<bool>& GetRetargetedBonesMask() const;
 
-	bool Initialize(
-		const TArray<int32>& SourceIndices,
-		const TArray<int32>& TargetIndices,
-		const TArray<FTransform> &SourceGlobalPose,
-		const TArray<FTransform> &TargetGlobalPose,
-		const FRetargetSkeleton& TargetSkeleton);
-
-	void MatchPoleVector(
-		const FTargetChainFKSettings& Settings,
-		const TArray<int32>& SourceIndices,
-		const TArray<int32>& TargetIndices,
-		const TArray<FTransform> &SourceGlobalPose,
-		TArray<FTransform> &OutTargetGlobalPose,
-		FRetargetSkeleton& TargetSkeleton);
-	
 private:
-
-	EAxis::Type CalculateBestPoleAxisForChain(
-		const TArray<int32>& BoneIndices,
-		const TArray<FTransform>& GlobalPose);
 	
-	static FVector CalculatePoleVector(
-		const EAxis::Type& PoleAxis,
-		const TArray<int32>& BoneIndices,
-		const TArray<FTransform>& GlobalPose);
-
-	static EAxis::Type GetMostDifferentAxis(
-		const FTransform& Transform,
-		const FVector& InNormal);
-
-	static FVector GetChainAxisNormalized(
-		const TArray<int32>& BoneIndices,
-		const TArray<FTransform>& GlobalPose);
+	/** a boolean mask with size NumBones, with value of "true" for any bone that is retargeted
+	 * ie, bones that are actually posed based on a mapped source chain
+	 * NOTE: this mask is only available in Op::PostInitialize()... it is built AFTER Op::Initialize() */
+	bool bIsMaskInitialized = false;
+	TArray<bool> IsBoneRetargeted;
 };
 
-struct FChainFK
+/** A "resolved" bone chain serves a few purposes:
+ * 1. takes the Start/End bone names from the IK Rig and converts them into a list of bone indices on the skeletal mesh
+ * 2. stores Local/Global ref pose of the chain
+ * 3. provides facilities to generate local/global poses of the chain
+ * 4. parameterizes the chain allowing you to GetTransformAtChainParam()
+ * 
+ * The intention with this data structure is to provide read-only chain data for ops to work with as they need.
+ */
+struct IKRIG_API FResolvedBoneChain
 {
-	TArray<FTransform> InitialGlobalTransforms;
-	TArray<FTransform> InitialLocalTransforms;
-	TArray<FTransform> CurrentGlobalTransforms;
+	FResolvedBoneChain(
+		const FBoneChain& InBoneChain,
+		const FRetargetSkeleton& InSkeleton,
+		FIKRigLogger& InLog);
 
-	TArray<float> Params;
-	TArray<int32> BoneIndices;
+	bool IsValid() const;
 
-	int32 ChainParentBoneIndex;
-	FTransform ChainParentInitialGlobalTransform;
-
-	bool Initialize(
-		const FRetargetSkeleton& Skeleton,
-		const TArray<int32>& InBoneIndices,
-		const TArray<FTransform> &InitialGlobalPose,
-		FIKRigLogger& Log);
+	void UpdatePoseFromSkeleton(const FRetargetSkeleton& InSkeleton);
+	
+	void GetWarnings( const FRetargetSkeleton& Skeleton, FIKRigLogger& Log) const;
 
 	FTransform GetTransformAtChainParam(
 		const TArray<FTransform>& Transforms,
-		const float& Param) const;
+		const double& Param) const;
 	
-	float GetStretchAtParam(
+	double GetStretchAtParam(
 		const TArray<FTransform>& InitialTransforms,
 		const TArray<FTransform>& CurrentTransforms,
-		const double& Param);
+		const double& Param) const;
 	
 	static double GetChainLength(const TArray<FTransform>& Transforms);
 
-private:
-	
-	bool CalculateBoneParameters(FIKRigLogger& Log);
-
-protected:
-
 	static void FillTransformsWithLocalSpaceOfChain(
 		const FRetargetSkeleton& Skeleton,
 		const TArray<FTransform>& InGlobalPose,
 		const TArray<int32>& BoneIndices,
 		TArray<FTransform>& OutLocalTransforms);
 
-	void PutCurrentTransformsInRefPose(
-		const TArray<int32>& BoneIndices,
+	static void FillTransformsWithGlobalRefPoseOfChain(
 		const FRetargetSkeleton& Skeleton,
-		const TArray<FTransform>& InCurrentGlobalPose);
-};
-
-struct FChainEncoderFK : public FChainFK
-{
-	TArray<FTransform> CurrentLocalTransforms;
-	FTransform ChainParentCurrentGlobalTransform;
-	
-	void EncodePose(
-		const FRetargetSkeleton& SourceSkeleton,
-		const TArray<int32>& SourceBoneIndices,
-		const TArray<FTransform> &InSourceGlobalPose);
-
-	void TransformCurrentChainTransforms(const FTransform& NewParentTransform);
-};
-
-struct FChainDecoderFK : public FChainFK
-{
-	void InitializeIntermediateParentIndices(
-		const int32 RetargetRootBoneIndex,
-		const int32 ChainRootBoneIndex,
-		const FTargetSkeleton& TargetSkeleton);
-	
-	void DecodePose(
-		const FRootRetargeter& RootRetargeter,
-		const FTargetChainFKSettings& Settings,
-		const TArray<int32>& TargetBoneIndices,
-		FChainEncoderFK& SourceChain,
-		const FTargetSkeleton& TargetSkeleton,
-		TArray<FTransform> &InOutGlobalPose);
-
-	void MatchPoleVector(
-		const FTargetChainFKSettings& Settings,
-		const TArray<int32>& TargetBoneIndices,
-		FChainEncoderFK& SourceChain,
-		const FTargetSkeleton& TargetSkeleton,
-		TArray<FTransform> &InOutGlobalPose);
-	
-private:
+		const TArray<FTransform>& InGlobalPose,
+		const TArray<int32>& BoneIndices,
+		TArray<FTransform>& OutGlobalTransforms);
 
-	void MatchChain(
-		const bool bScaleSourceChain,
-		const FChainEncoderFK& SourceChain,
-		const TArray<int32>& TargetBoneIndices,
-		const FVector& TargetChainOrigin);
+	TArray<FTransform> GetChainTransformsFromPose(const TArray<FTransform>& InPose) const;
 
-	static FVector GetPointOnSplineFromIndexAndAlpha(const TArray<FVector>& InSplinePoints, const int32 InPointIndex, const double InSegmentAlpha);
-	
-	static void GetPointOnSplineDistanceFromPoint(
-		const TArray<FVector>& InSplinePoints,
-		const int32 InOriginPointIndex,
-		const double InOriginPointAlpha,
-		const double InTargetDistanceFromOrigin,
-		int32& OutPointIndex,
-		double& OutPointAlpha);
+	FName ChainName;
+	FName StartBone;
+	FName EndBone;
+	FName IKGoalName;
 	
-	void UpdateIntermediateParents(
-		const FTargetSkeleton& TargetSkeleton,
-		TArray<FTransform> &InOutGlobalPose);
+	TArray<FTransform> RefPoseGlobalTransforms;
+	TArray<FTransform> RefPoseLocalTransforms;
+	TArray<FTransform> CurrentLocalTransforms;
 
-	TArray<int32> IntermediateParentIndices;
-};
+	TArray<float> Params;
+	TArray<int32> BoneIndices;
+	float InitialChainLength;
 
-struct FDecodedIKChain
-{
-	FVector EndEffectorPosition = FVector::ZeroVector;
-	FQuat EndEffectorRotation = FQuat::Identity;
-	FVector PoleVectorPosition = FVector::ZeroVector;
-};
+	int32 ChainParentBoneIndex;
+	FTransform ChainParentInitialGlobalTransform;
 
-struct FSourceChainIK
-{
-	int32 StartBoneIndex = INDEX_NONE;
-	int32 EndBoneIndex = INDEX_NONE;
+private:
 	
-	FVector InitialEndPosition = FVector::ZeroVector;
-	FQuat InitialEndRotation = FQuat::Identity;
-	float InvInitialLength = 1.0f;
-
-	// results after encoding...
-	FVector PreviousEndPosition = FVector::ZeroVector;
-	FVector CurrentEndPosition = FVector::ZeroVector;
-	FVector CurrentEndDirectionNormalized = FVector::ZeroVector;
-	FQuat CurrentEndRotation = FQuat::Identity;
-	float CurrentHeightFromGroundNormalized = 0.0f;
-};
-
-struct FTargetChainIK
-{
-	int32 BoneIndexA = INDEX_NONE;
-	int32 BoneIndexC = INDEX_NONE;
+	bool bFoundStartBone = false;
+	bool bFoundEndBone = false;
+	bool bEndIsStartOrChildOfStart  = false;
 	
-	float InitialLength = 1.0f;
-	FVector InitialEndPosition = FVector::ZeroVector;
-	FQuat InitialEndRotation = FQuat::Identity;
-	FVector PrevEndPosition = FVector::ZeroVector;
-};
-
-struct FChainDebugData
-{
-	FTransform InputTransformStart = FTransform::Identity;
-	FTransform InputTransformEnd = FTransform::Identity;
-	FTransform OutputTransformEnd = FTransform::Identity;
+	void CalculateBoneParameters(FIKRigLogger& Log);
 };
 
-struct FChainRetargeterIK
+/** a container for ALL the fully resolved bone chains on both
+ * the source and target skeletons, as well as the mapping between them */
+struct IKRIG_API FRetargeterBoneChains
 {
-	FSourceChainIK Source;
-	FTargetChainIK Target;
+	/** Load all chains from the IK Rigs and resolved them on the current skeletal meshes*/
+	bool Initialize(
+		const UIKRetargeter* InRetargetAsset,
+		const FRetargetSkeleton& InSourceSkeleton,
+		const FRetargetSkeleton& InTargetSkeleton,
+		FIKRigLogger& InLog);
 	
-	FDecodedIKChain Results;
+	/** Get read only access to the resolved bone chains on either the source or the target */
+	const TArray<FResolvedBoneChain>& GetResolvedBoneChains(ERetargetSourceOrTarget SourceOrTarget) const;
 
-	bool ResetThisTick;
-	FVectorSpringState PlantingSpringState;
+	/** Get the index of a bone chain by a given name. Returns INDEX_None if not found. */
+	int32 GetIndexOfBoneChainByName(const FName InChainName, const ERetargetSourceOrTarget SourceOrTarget) const;
 
-	#if WITH_EDITOR
-	FChainDebugData DebugData;
-	#endif
+	/** Given the name of a target chain, returns the source chain that it's mapped to. Returns nullptr if no mapping exists */
+	const FResolvedBoneChain* GetSourceChainMappedToTargetChain(const FName InTargetChainName) const;
 
-	bool InitializeSource(
-		const TArray<int32>& BoneIndices,
-		const TArray<FTransform> &SourceInitialGlobalPose,
-		FIKRigLogger& Log);
-	
-	bool InitializeTarget(
-		const TArray<int32>& BoneIndices,
-		const TArray<FTransform> &TargetInitialGlobalPose,
-		FIKRigLogger& Log);
+	/** Get a fully resolved bone chain by name */
+	const FResolvedBoneChain* GetResolvedBoneChainByName(const FName InChainName, const ERetargetSourceOrTarget SourceOrTarget) const;
 
-	void EncodePose(const TArray<FTransform> &SourceInputGlobalPose);
-	
-	void DecodePose(
-		const FTargetChainIKSettings& Settings,
-		const FTargetChainSpeedPlantSettings& SpeedPlantSettings,
-		const FRootRetargeter& RootRetargeter,
-		const TMap<FName, float>& SpeedValuesFromCurves,
-		const float DeltaTime,
-		const TArray<FTransform>& InGlobalPose);
+	/** Get a fully resolved bone chain by inddex */
+	const FResolvedBoneChain* GetResolvedBoneChainByIndex(const int32 InChainIndex, const ERetargetSourceOrTarget SourceOrTarget) const;
 
-	void SaveDebugInfo(const TArray<FTransform>& InGlobalPose);
-};
+	/** Get read only access to the mapping from each target bone chain to a source chain (or potentially None) */
+	const TArray<int32>& GetChainMappings() const;
 
-struct FRetargetChainPair
-{
-	TArray<int32> SourceBoneIndices;
-	TArray<int32> TargetBoneIndices;
-	
-	FName SourceBoneChainName;
-	FName TargetBoneChainName;
+	/** Update all transforms stored in the chains. This is needed whenever the retarget pose is modified */
+	void UpdatePoseFromSkeleton(const FRetargetSkeleton& InSkeleton, const ERetargetSourceOrTarget SourceOrTarget);
 
-	virtual ~FRetargetChainPair() = default;
-	
-	virtual bool Initialize(
-		const FBoneChain& SourceBoneChain,
-		const FBoneChain& TargetBoneChain,
-		const FRetargetSkeleton& SourceSkeleton,
-		const FTargetSkeleton& TargetSkeleton,
-		FIKRigLogger& Log);
+	/** Reset everything (between initializations) */
+	void Reset();
 
 private:
-
-	bool ValidateBoneChainWithSkeletalMesh(
-		const bool IsSource,
-		const FBoneChain& BoneChain,
-		const FRetargetSkeleton& RetargetSkeleton,
-		FIKRigLogger& Log);
-};
-
-struct FRetargetChainPairFK : FRetargetChainPair
-{
-	FTargetChainFKSettings Settings;
-	FChainEncoderFK FKEncoder;
-	FChainDecoderFK FKDecoder;
-	FPoleVectorMatcher PoleVectorMatcher;
 	
-	virtual bool Initialize(
-        const FBoneChain& SourceBoneChain,
-        const FBoneChain& TargetBoneChain,
-        const FRetargetSkeleton& SourceSkeleton,
-        const FTargetSkeleton& TargetSkeleton,
-        FIKRigLogger& Log) override;
-};
+	/** All resolved chains on the source mesh */
+	TArray<FResolvedBoneChain> SourceBoneChains;
 
-struct FRetargetChainPairIK : FRetargetChainPair
-{
-	FTargetChainIKSettings Settings;
-	FTargetChainSpeedPlantSettings SpeedPlantSettings;
-	FChainRetargeterIK IKChainRetargeter;
-	FName IKGoalName;
-	FName PoleVectorGoalName;
-
-	virtual bool Initialize(
-        const FBoneChain& SourceBoneChain,
-        const FBoneChain& TargetBoneChain,
-        const FRetargetSkeleton& SourceSkeleton,
-        const FTargetSkeleton& TargetSkeleton,
-        FIKRigLogger& Log) override;
-};
+	/** All resolved chains on the target mesh */
+	TArray<FResolvedBoneChain> TargetBoneChains;
 
-struct FRetargetDebugData
-{
-	FTransform StrideWarpingFrame;
+	/** This array is same length as TargetBoneChains and stores index into the mapped source chain for each target chain */
+	TArray<int32> ChainMapping;
+
+	bool bIsInitialized = false;
 };
 
 enum class ERetargetBoneSpace : uint8
@@ -489,7 +275,13 @@ enum class ERetargetBoneSpace : uint8
 	Local
 };
 
-// the system to do a straight copy of bone transform with matching names (does not do any retargeting!)
+enum class ERetargetOpsToSearch : uint8
+{
+	ProcessorOps,
+	AssetOps
+};
+
+/** the system to do a straight copy of bone transform with matching names (does not do any retargeting!) */
 struct FBasePoseCopier
 {
 	void CopyPoseFromSource(
@@ -523,8 +315,9 @@ private:
 /** The runtime processor that converts an input pose from a source skeleton into an output pose on a target skeleton.
  * To use:
  * 1. Initialize a processor with a Source/Target skeletal mesh and a UIKRetargeter asset.
- * 2. Call RunRetargeter and pass in a source pose as an array of global-space transforms
- * 3. RunRetargeter returns an array of global space transforms for the target skeleton.
+ * 2. Call ScaleSourcePose() and pass in the global space source pose (see func comments why this is necessary)
+ * 3. Call RunRetargeter and pass in a source pose as an array of global-space transforms
+ * 3. RunRetargeter() returns an array of global space transforms for the target skeleton.
  */
 struct IKRIG_API FIKRetargetProcessor
 {
@@ -536,6 +329,7 @@ public:
 	* @param SourceSkeleton - the skeletal mesh to poses FROM
 	* @param TargetSkeleton - the skeletal mesh to poses TO
 	* @param InRetargeterAsset - the source asset to use for retargeting settings
+	* @param InRetargetProfile - the collection of settings used to initialize with
 	* @param bSuppressWarnings - if true, will not output warnings during initialization
 	* @warning - Initialization does a lot of validation and can fail for many reasons. Check bIsLoadedAndValid afterwards.
 	*/
@@ -543,22 +337,40 @@ public:
 		const USkeletalMesh *SourceSkeleton,
 		const USkeletalMesh *TargetSkeleton,
 		const UIKRetargeter* InRetargeterAsset,
-		const FRetargetProfile& Settings,
+		const FRetargetProfile& InRetargetProfile,
 		const bool bSuppressWarnings=false);
 
+	/**
+	 * This scales the input source pose according to the source scale factor
+	 * This must be called outside of the retargeter itself because otherwise the pose may be scaled multiple times.
+	 * This is because we don't copy the source pose and we don't control when it's updated.
+	 * By explicitly scaling outside of RunRetargeter(), outside systems can scale it whenever they update it.
+	 * @param InOutSourceGlobalPose: an array of global space transforms for each bone in the source input pose */ 
+	void ScaleSourcePose(TArray<FTransform>& InOutSourceGlobalPose) const;
+
 	/**
 	* Run the retarget to generate a new pose.
 	* @param InSourceGlobalPose -  is the source mesh input pose in Component/Global space
-	* @param SpeedValuesFromCurves - the speed of each curve used by speed planting (blended in anim graph)
-	* @param DeltaTime -  time since last tick in seconds (used by speed planting)
-	* @param Profile -  the retarget profile to use for this update
+	* @param InProfile -  the retarget profile to use for this update
+	* @param InDeltaTime -  time since last tick in seconds (passed to ops)
 	* @return The retargeted Component/Global space pose for the target skeleton
 	*/
 	TArray<FTransform>& RunRetargeter(
 		TArray<FTransform>& InSourceGlobalPose,
-		const TMap<FName, float>& SpeedValuesFromCurves,
-		const float DeltaTime,
-		const FRetargetProfile& Profile);
+		const FRetargetProfile& InProfile,
+		const float InDeltaTime);
+
+	/** Get whether this processor is ready to call RunRetargeter() and generate new poses. */
+	bool IsInitialized() const;
+
+	/** Set that this processor needs to be reinitialized. */
+	void SetNeedsInitialized();
+
+	/** Get whether this processor was initialized with these skeletal meshes and retarget asset*/
+	bool WasInitializedWithTheseAssets(
+		const USkeletalMesh* InSourceMesh,
+		const USkeletalMesh* InTargetMesh,
+		const UIKRetargeter* InRetargetAsset) const;
 
 	/** Does a partial reinitialization (at runtime) whenever the retarget pose is swapped to a different or if the
 	 * pose has been modified. Does nothing if the pose has not changed. */
@@ -567,44 +379,82 @@ public:
 	/** Get read-only access to either source or target skeleton. */
 	const FRetargetSkeleton& GetSkeleton(ERetargetSourceOrTarget SourceOrTarget) const;
 
-	/** Get name of the root bone of the given skeleton. */
-	FName GetRetargetRoot(ERetargetSourceOrTarget SourceOrTarget) const;
-	
-	/** Get whether this processor is ready to call RunRetargeter() and generate new poses. */
-	bool IsInitialized() const { return bIsInitialized; };
+	/** Get read/write access to the target skeleton.
+	 * NOTE: this skeleton contains the output pose of the target skeleton */
+	FTargetSkeleton& GetTargetSkeleton();
+	const FTargetSkeleton& GetTargetSkeleton() const;
 
-	/** Get whether this processor was initialized with these skeletal meshes and retarget asset*/
-	bool WasInitializedWithTheseAssets(
-		const USkeletalMesh* InSourceMesh,
-		const USkeletalMesh* InTargetMesh,
-		const UIKRetargeter* InRetargetAsset) const;
+	/** Get read-only access to all the fully resolved bone chains for both source and target skeletons */
+	const FRetargeterBoneChains& GetBoneChains() const { return AllBoneChains; };
 
-	/** Get the currently running IK Rig processor for the target */
-	const FIKRigProcessor& GetTargetIKRigProcessor() const { return IKRigProcessor; };
-	
-	/** Get read only access to the core IK retarget chains */
-	const TArray<FRetargetChainPairIK>& GetIKChainPairs() const { return ChainPairsIK; };
-	
-	/** Get read only access to the core FK retarget chains */
-	const TArray<FRetargetChainPairFK>& GetFKChainPairs() const { return ChainPairsFK; };
-	
-	/** Get read only access to the core root retargeter */
-	const FRootRetargeter& GetRootRetargeter() const { return RootRetargeter; };
+	/** Get read-write access so that ops can modify the IK Rig goals
+	 * NOTE: pointers to goals in the container are not stable between updates, do not store them.
+	 * NOTE: any op that sets a goal position must also specify the space of the goal (they are all additive by default) */
+	FIKRigGoalContainer& GetIKRigGoalContainer() { return GoalContainer; };
+	const FIKRigGoalContainer& GetIKRigGoalContainer() const { return GoalContainer; };
 	
 	/** Get read only access to the retarget ops currently running in processor */
-	const TArray<TSharedPtr<FIKRetargetOpBase>>& GetRetargetOps() const {return OpStack; };
+	const TArray<FInstancedStruct>& GetRetargetOps() const {return OpStack; };
 
-	/** Get read only access to the current global settings */
-	const FRetargetGlobalSettings& GetGlobalSettings() const {return GlobalSettings; };
-	
-	/** Reset the IK planting state. */
-	void ResetPlanting();
+	/** Get list of ops of a given type */
+	TArray<const FIKRetargetOpBase*> GetRetargetOpsByType(const UScriptStruct* OpType) const;
 
-	/** Set that this processor needs to be reinitialized. */
-	void SetNeedsInitialized();
+	/** Get the first op in the stack of the given type */
+	template <typename T>
+	const T* GetFirstRetargetOpOfType(ERetargetOpsToSearch SourceOfOps = ERetargetOpsToSearch::ProcessorOps) const
+	{
+		// get the op stack to search (either processor or asset)
+		const TArray<FInstancedStruct>* OpsToSearch;
+		if (SourceOfOps == ERetargetOpsToSearch::AssetOps)
+		{
+			if (!RetargeterAsset)
+			{
+				return nullptr;
+			}
+			OpsToSearch = &RetargeterAsset->GetRetargetOps();
+		}
+		else
+		{
+			OpsToSearch = &OpStack;
+		}
+		
+		// search the op stack for first op of the given type
+		for (const FInstancedStruct& OpStruct : *OpsToSearch)
+		{
+			if (OpStruct.GetScriptStruct()->IsChildOf(T::StaticStruct()))
+			{
+				return OpStruct.GetPtr<T>();
+			}
+		}
+
+		return nullptr;
+	}
+
+	/** Returns true if the IK has not been forced off */
+	bool IsIKForcedOff() const { return bIKForcedOff; };
+
+	/** Get read only access to the retarget asset */
+	const UIKRetargeter* GetRetargetAsset() const { return RetargeterAsset; };
 	
-	/** Apply the settings stored in the IKRig asset. */
-	void CopyIKRigSettingsFromAsset();
+	/** Calls OnPlaybackReset() for all ops in the stack */
+	void OnPlaybackReset();
+
+	/** Calls AnimGraphPreUpdateMainThread() for all ops in the stack */
+	void OnAnimGraphPreUpdateMainThread(USkeletalMeshComponent& SourceMeshComponent, USkeletalMeshComponent& TargetMeshComponent);
+
+	/** Calls AnimGraphEvaluateAnyThread() for all ops in the stack */
+	void OnAnimGraphEvaluateAnyThread(FPoseContext& Output);
+
+	/** Get name of the root bone of the given skeleton. */
+	FName GetPelvisBone(ERetargetSourceOrTarget SourceOrTarget) const;
+
+	/** Determine if we should be directly copying in the source pose into the target
+	 * @param OutCopyBasePoseRoot: will contain the name of the bone to copy from if one was specified, otherwise Name_None
+	 * @return true if a CopyBasePoseOp is present and set to enable copying*/
+	bool ShouldCopyBasePose(FName& OutCopyBasePoseRoot) const;
+
+	/** Get the scale factor for the source pose (comes from presence of SourceScaleOp)*/
+	double GetSourceScaleFactor() const;
 
 	/** The logging system */
 	FIKRigLogger Log;
@@ -616,40 +466,49 @@ private:
 	FOnRetargeterInitialized RetargeterInitialized;
 public:
 	/** Returns true if the bone is part of a retarget chain or root bone, false otherwise. */
-	bool IsBoneRetargeted(const FName BoneName, const ERetargetSourceOrTarget SourceOrTarget) const;
+	bool IsBoneInAMappedChain(const FName BoneName, const ERetargetSourceOrTarget SourceOrTarget) const;
 	/** Returns index of the bone with the given name in either Source or Target skeleton. */
 	int32 GetBoneIndexFromName(const FName BoneName, const ERetargetSourceOrTarget SourceOrTarget) const;
+	/** Get access to all the fully resolved bone chain data */
+	const FRetargeterBoneChains& GetBoneChains();
 	/** Returns name of the chain associated with this bone. Returns NAME_None if bone is not in a chain. */
 	FName GetChainNameForBone(const FName BoneName, const ERetargetSourceOrTarget SourceOrTarget) const;
 	/** Get a transform at a given param in a chain */
 	FTransform GetGlobalRetargetPoseAtParam(const FName InChainName, const float Param, const ERetargetSourceOrTarget SourceOrTarget) const;
 	/** Get a transform of a bone in the retarget pose */
 	FTransform GetRetargetPoseBoneTransform(const FName InBoneName, const ERetargetSourceOrTarget SourceOrTarget, ERetargetBoneSpace BoneSpace) const;
-	/** Get access to the internal chain on the given skeleton */
-	const FChainFK* GetChain(const FName InChainName, const ERetargetSourceOrTarget SourceOrTarget) const;
+	
 	/** Get the param of the bone in it's retarget chain. Ranges from 0 to NumBonesInChain. */
 	float GetParamOfBoneInChain(const FName InBoneName, const ERetargetSourceOrTarget SourceOrTarget) const;
 	/** Get the bone in the chain at the given param. */
 	FName GetClosestBoneToParam(const FName InChainName, const float InParam, const ERetargetSourceOrTarget SourceOrTarget) const;
 	/** Get the chain mapped to this one. */
-	FName GetMappedChainName(const FName InChainName, const ERetargetSourceOrTarget SourceOrTarget);
-	/** store data for debug drawing */
-	FRetargetDebugData DebugData;
+	FName GetChainMappedToChain(const FName InChainName, const ERetargetSourceOrTarget InSourceOrTarget) const;
 	/** Attach a delegate to be notified whenever this processor is initialized. */
 	FOnRetargeterInitialized& OnRetargeterInitialized(){ return RetargeterInitialized; };
+	/** Run all debug drawing on all ops in the stack */
+	void DebugDrawAllOps(
+		FPrimitiveDrawInterface* InPDI,
+		const FTransform& InComponentTransform,
+		const FIKRetargetDebugDrawState& EditorState) const;
+	/** Run debug drawing on the specified op */
+	void DebugDrawOpAtIndex(
+		const int32 InOpIndex,
+		FPrimitiveDrawInterface* InPDI,
+		const FTransform& InComponentTransform,
+		const FIKRetargetDebugDrawState& EditorState) const;
+	
 #endif
 
 private:
 
+	/** If true, all IK operations are skipped
+	 * NOTE: this is used both for debugging and LOD'ing IK off */
+	bool bIKForcedOff = false;
+	
 	/** Only true once Initialize() has successfully completed.*/
 	bool bIsInitialized = false;
 	int32 AssetVersionInitializedWith = -2;
-	/** true when roots are able to be retargeted */
-	bool bRootsInitialized = false;
-	/** true when at least one pair of bone chains is able to be retargeted */
-	bool bAtLeastOneValidBoneChainPair = false;
-	/** true when IK Rig has been initialized and is ready to run*/
-	bool bIKRigInitialized = false;
 
 	/** The source asset this processor was initialized with. */
 	const UIKRetargeter* RetargeterAsset = nullptr;
@@ -660,70 +519,38 @@ private:
 	/** The internal data structures used to represent the TARGET skeleton / pose during retargeter.*/
 	FTargetSkeleton TargetSkeleton;
 
-	/** The IK Rig processor for running IK on the target */
-	FIKRigProcessor IKRigProcessor;
-
-	/** The Source/Target pairs of Bone Chains retargeted using the FK algorithm */
-	TArray<FRetargetChainPairFK> ChainPairsFK;
+	/** the named transforms that solvers use as end effectors */
+	FIKRigGoalContainer GoalContainer;
 
-	/** The Source/Target pairs of Bone Chains retargeted using the IK Rig */
-	TArray<FRetargetChainPairIK> ChainPairsIK;
-
-	/** The Source/Target pair of Root Bones retargeted with scaled translation */
-	FRootRetargeter RootRetargeter;
+	/** Storage for all bone chain data used by Ops and editor tools */
+	FRetargeterBoneChains AllBoneChains;
 
 	/** The system to do a straight copy of source-to-target bone transforms based on name */
 	FBasePoseCopier BasePoseCopier;
 
-	/** The currently used global settings (driven either by source asset or a profile) */
-	FRetargetGlobalSettings GlobalSettings;
-
-	/** The collection of operations to run in the final phase of retargeting */
-	TArray<TSharedPtr<FIKRetargetOpBase>> OpStack;
+	/** The collection of operations to run to transfer animation from source to target. */
+	TArray<FInstancedStruct> OpStack;
 
 	/** Apply the settings stored in a retarget profile. Called inside RunRetargeter(). */
-	void ApplySettingsFromProfile(const FRetargetProfile& Profile);
-	
-	/** Update chain settings at runtime (for use with a profile)
-	 * Will queue an IK Rig reinitialization if the "Enable IK" state is modified. */
-	void UpdateChainSettingsAtRuntime(const FName ChainName, const FTargetChainSettings& NewChainSettings);
-	
-	/** Initializes the FRootRetargeter */
-	bool InitializeRoots();
-
-	/** Initializes the all the chain pairs */
-	bool InitializeBoneChainPairs(const TMap<FName, FTargetChainSettings>& ChainSettings);
-
-	/** Initializes the IK Rig that evaluates the IK solve for the target IK chains */
-	bool InitializeIKRig(const USkeletalMesh* InSkeletalMesh);
-
-	// Initialize the retarget op stack
-	bool InitializeOpStack(const TArray<FInstancedStruct>& OpsFromAsset);
+	void ApplyProfileAtRuntime(const FRetargetProfile& Profile);
 
 	/** Internal retarget phase that does simple bone-to-bone copying from source*/
 	void GenerateBasePoses(TArray<FTransform>& InSourceGlobalPose);
-	
-	/** Internal retarget phase for the root. */
-	void RunRootRetarget(const TArray<FTransform>& InGlobalTransforms, TArray<FTransform>& OutGlobalTransforms);
-
-	/** Internal retarget phase for the FK chains. */
-	void RunFKRetarget(const TArray<FTransform>& InGlobalTransforms, TArray<FTransform>& OutGlobalTransforms);
 
-	/** Internal retarget phase for the IK chains. */
-	void RunIKRetarget(
-		const TArray<FTransform>& InSourceGlobalPose,
-		TArray<FTransform>& OutTargetGlobalPose,
-		const TMap<FName, float>& SpeedValuesFromCurves,
-		const float DeltaTime);
+	/** Setup the goal container for IK Rigs to use */
+	void ResetIKRigGoalContainer();
 
-	/** Internal retarget phase for the pole matching feature of FK chains. */
-	void RunPoleVectorMatching(const TArray<FTransform>& InGlobalTransforms, TArray<FTransform>& OutGlobalTransforms);
+	/** Initialize the retarget op stack */
+	bool InitializeOpStack(const TArray<FInstancedStruct>& OpsFromAsset, const FRetargetProfile& InRetargetProfile);
 
-	/** Runs in the after the base IK retarget to apply stride warping to IK goals. */
-	void RunStrideWarping(const TArray<FTransform>& InTargetGlobalPose);
+	/** Collect retargeted bones and call PostInitialize on each Op */
+	void PostInitializeOpStack();
 
 	/** Run all post process operations on the retargeted result. */
-	void RunRetargetOps(const TArray<FTransform>& InSourceGlobalPose, TArray<FTransform>& OutTargetGlobalPose);
+	void RunRetargetOps(
+		const double InDeltaTime,
+		const TArray<FTransform>& InSourceGlobalPose,
+		TArray<FTransform>& OutTargetGlobalPose);
 };
 
 //
@@ -733,7 +560,7 @@ private:
 // The UObject based processor is here for backwards compatibility.
 // If you have a system built off this processor, please convert it to the FIKRetargetProcessor as this will be removed in a future update.
 //
-
+PRAGMA_DISABLE_DEPRECATION_WARNINGS
 UCLASS()
 class UE_DEPRECATED(5.6, "UIKRetargetProcessor has been replaced by FIKRetargetProcessor. UIKRetargetProcessor is now a thin wrapper around FIKRetargetProcessor and will be removed in a future version.")
 IKRIG_API UIKRetargetProcessor : public UObject
@@ -758,7 +585,7 @@ public:
 		const float DeltaTime,
 		const FRetargetProfile& Profile)
 	{
-		return Processor.RunRetargeter(InSourceGlobalPose, SpeedValuesFromCurves, DeltaTime, Profile);
+		return Processor.RunRetargeter(InSourceGlobalPose, Profile, DeltaTime);
 	}
 
 	void UpdateRetargetPoseAtRuntime(
@@ -774,7 +601,7 @@ public:
 	}
 	FName GetRetargetRoot(ERetargetSourceOrTarget SourceOrTarget) const
 	{
-		return Processor.GetRetargetRoot(SourceOrTarget);
+		return Processor.GetPelvisBone(SourceOrTarget);
 	}
 	bool IsInitialized() const
 	{
@@ -787,49 +614,29 @@ public:
 	{
 		return Processor.WasInitializedWithTheseAssets(InSourceMesh, InTargetMesh, InRetargetAsset);
 	}
-	const FIKRigProcessor& GetTargetIKRigProcessor() const
-	{
-		return Processor.GetTargetIKRigProcessor();
-	}
-	const TArray<FRetargetChainPairIK>& GetIKChainPairs() const
-	{
-		return Processor.GetIKChainPairs();
-	}
-	const TArray<FRetargetChainPairFK>& GetFKChainPairs() const
-	{
-		return Processor.GetFKChainPairs();
-	}
-	const FRootRetargeter& GetRootRetargeter() const
-	{
-		return Processor.GetRootRetargeter();
-	};
 	const TArray<TObjectPtr<URetargetOpBase>>& GetRetargetOps() const
 	{
 		return DummyStack;
 	};
 	const FRetargetGlobalSettings& GetGlobalSettings() const
 	{
-		return Processor.GetGlobalSettings();
+		return DummyGlobalSettings;
 	};
 	void ResetPlanting()
 	{
-		Processor.ResetPlanting();
+		Processor.OnPlaybackReset();
 	};
 	void SetNeedsInitialized()
 	{
 		Processor.SetNeedsInitialized();
 	};
-	void CopyIKRigSettingsFromAsset()
-	{
-		Processor.CopyIKRigSettingsFromAsset();
-	};
 
 	FIKRigLogger Log;
 	
 #if WITH_EDITOR
 	bool IsBoneRetargeted(const FName BoneName, const ERetargetSourceOrTarget SourceOrTarget) const
 	{
-		return Processor.IsBoneRetargeted(BoneName, SourceOrTarget);
+		return Processor.IsBoneInAMappedChain(BoneName, SourceOrTarget);
 	}
 	int32 GetBoneIndexFromName(const FName BoneName, const ERetargetSourceOrTarget SourceOrTarget) const
 	{
@@ -847,10 +654,6 @@ public:
 	{
 		return Processor.GetRetargetPoseBoneTransform(InBoneName, SourceOrTarget, BoneSpace);
 	}
-	const FChainFK* GetChain(const FName InChainName, const ERetargetSourceOrTarget SourceOrTarget) const
-	{
-		return Processor.GetChain(InChainName, SourceOrTarget);
-	}
 	float GetParamOfBoneInChain(const FName InBoneName, const ERetargetSourceOrTarget SourceOrTarget) const
 	{
 		return Processor.GetParamOfBoneInChain(InBoneName, SourceOrTarget);
@@ -861,12 +664,17 @@ public:
 	}
 	FName GetMappedChainName(const FName InChainName, const ERetargetSourceOrTarget SourceOrTarget)
 	{
-		return Processor.GetMappedChainName(InChainName, SourceOrTarget);
+		return Processor.GetChainMappedToChain(InChainName, SourceOrTarget);
 	}
 #endif
 
 
 private:
+	FRetargetGlobalSettings DummyGlobalSettings;
 	TArray<TObjectPtr<URetargetOpBase>> DummyStack;
 	FIKRetargetProcessor Processor;
-};
\ No newline at end of file
+};
+PRAGMA_ENABLE_DEPRECATION_WARNINGS
+//
+// END DEPRECATED UOBJECT-BASED PROCESSOR
+//
\ No newline at end of file
diff --git a/Engine/Plugins/Animation/IKRig/Source/IKRig/Public/Retargeter/IKRetargetProfile.h b/Engine/Plugins/Animation/IKRig/Source/IKRig/Public/Retargeter/IKRetargetProfile.h
--- a/Engine/Plugins/Animation/IKRig/Source/IKRig/Public/Retargeter/IKRetargetProfile.h
+++ b/Engine/Plugins/Animation/IKRig/Source/IKRig/Public/Retargeter/IKRetargetProfile.h
@@ -3,90 +3,166 @@
 #pragma once
 
 #include "CoreMinimal.h"
-#include "Retargeter/IKRetargetSettings.h"
+#include "IKRetargetOps.h"
+#include "Retargeter/IKRetargetDeprecated.h"
+#include "StructUtils/InstancedStruct.h"
 
 #include "IKRetargetProfile.generated.h"
 
 class UIKRetargeter;
 class URetargetChainSettings;
+struct FIKRetargetOpSettingsBase;
+
+UENUM()
+enum class ECopyOpSettingsContext : uint8
+{
+	PreInitialize,	// copies ALL settings (used in editor during setup)
+	Runtime			// copy only settings that don't require reinitialization
+};
 
 USTRUCT(BlueprintType)
-struct FRetargetProfile
+struct FRetargetOpProfile
+{
+	GENERATED_BODY()
+
+	FRetargetOpProfile() = default;
+
+	FRetargetOpProfile(
+		int32 InOpIndex,
+		const UScriptStruct* InSettingsType,
+		const FIKRetargetOpSettingsBase* InOpSettings);
+
+	void CopyFromOtherOpProfile(const FRetargetOpProfile& OtherOpProfile);
+
+	/** Copies the settings from the profile into the given Op instanced struct.
+	 * @param InOutOpStruct - the op to apply these settings to
+	 * @param InApplyContext - in 'runtime' mode, will let op apply the settings itself to avoid changes that require reinitialization
+	 * @return true if the supplied op was compatible with the settings, false otherwise */
+	bool CopySettingsToOp(FInstancedStruct& InOutOpStruct, ECopyOpSettingsContext InApplyContext) const;
+	
+	/** The index of the op these settings should be applied to.
+	 * If left at -1, will apply to all ops that use the type of settings stored in 'SettingsToApply' */
+	UPROPERTY(EditAnywhere, BlueprintReadWrite, Category = "IK Retarget Op Profile")
+	int32 OpToApplySettingsTo = INDEX_NONE;
+
+	/** The settings to override for a specific op in the retargeter's op stack.
+	 * NOTE: all settings stored here will be applied to the op. However, at runtime each op independently
+	 * determines which settings are safe to be applied without requiring reinitialization. */
+	UPROPERTY(EditAnywhere, BlueprintReadWrite, Category = "IK Retarget Op Profile", meta = (ExcludeBaseStruct, BaseStruct = "/Script/IKRig.IKRetargetOpSettingsBase"))
+	FInstancedStruct SettingsToApply;
+
+	// TODO store sparse set of property overrides
+	// Is it possible we could do this with FInstancedPropertyBag?
+	// Perhaps we should store a bool mask with the same shape as the settings struct
+	// (ie with one bool per-property indicating if the profile should override it)
+	// Regardless, it will take some major UI work to expose this in a useful way.
+	// UPROPERTY(EditDefaultsOnly, Category = "IK Retarget Op Profile")
+	// FInstancedPropertyBag TestPropertyBag;
+};
+
+USTRUCT(BlueprintType)
+struct IKRIG_API FRetargetProfile
 {
 	GENERATED_BODY()
 	
 public:
 
+	PRAGMA_DISABLE_DEPRECATION_WARNINGS
+	// implicit copy assignment operator raises deprecation warnings (remove this once the deprecated properties are gone)
+	FRetargetProfile() = default;
+	FRetargetProfile(const FRetargetProfile&) = default;
+	FRetargetProfile(FRetargetProfile&&) = default;
+	FRetargetProfile& operator=(const FRetargetProfile&) = default;
+	FRetargetProfile& operator=(FRetargetProfile&&) = default;
+	PRAGMA_ENABLE_DEPRECATION_WARNINGS
+	
+	/** A polymorphic list of override settings to apply to retargeting operations in the stack */
+	UPROPERTY(EditAnywhere, BlueprintReadWrite, Category = "Op Settings")
+	TArray<FRetargetOpProfile> RetargetOpProfiles;
+
 	// If true, the TARGET Retarget Pose specified in this profile will be applied to the Retargeter (when plugged into the Retargeter).
-	UPROPERTY(EditAnywhere, BlueprintReadWrite, Category=RetargetPoses, meta = (DisplayName = "Override Target Retarget Pose"))
+	UPROPERTY(EditAnywhere, BlueprintReadWrite, Category="Retarget Poses", meta = (DisplayName = "Override Target Retarget Pose"))
 	bool bApplyTargetRetargetPose = false;
 	
 	// Override the TARGET Retarget Pose to use when this profile is active.
 	// The pose must be present in the Retarget Asset and is not applied unless bApplyTargetRetargetPose is true.
-	UPROPERTY(EditAnywhere, BlueprintReadWrite, Category=RetargetPoses, meta=(EditCondition="bApplyTargetRetargetPose"))
+	UPROPERTY(EditAnywhere, BlueprintReadWrite, Category="Retarget Poses", meta=(EditCondition="bApplyTargetRetargetPose"))
 	FName TargetRetargetPoseName;
 
 	// If true, the Source Retarget Pose specified in this profile will be applied to the Retargeter (when plugged into the Retargeter).
-	UPROPERTY(EditAnywhere, BlueprintReadWrite, Category=RetargetPoses, meta = (DisplayName = "Override Source Retarget Pose"))
+	UPROPERTY(EditAnywhere, BlueprintReadWrite, Category="Retarget Poses", meta = (DisplayName = "Override Source Retarget Pose"))
 	bool bApplySourceRetargetPose = false;
 
 	// Override the SOURCE Retarget Pose to use when this profile is active.
 	// The pose must be present in the Retarget Asset and is not applied unless bApplySourceRetargetPose is true.
-	UPROPERTY(EditAnywhere, BlueprintReadWrite, Category=RetargetPoses, meta=(EditCondition="bApplySourceRetargetPose"))
+	UPROPERTY(EditAnywhere, BlueprintReadWrite, Category="Retarget Poses", meta=(EditCondition="bApplySourceRetargetPose"))
 	FName SourceRetargetPoseName;
 
-	// If true, the Chain Settings stored in this profile will be applied to the Retargeter (when plugged into the Retargeter).
-	UPROPERTY(EditAnywhere, BlueprintReadWrite, Category=ChainSettings, meta = (DisplayName = "Override Chain Settings"))
-	bool bApplyChainSettings = false;
-	
-	// A (potentially sparse) set of setting overrides for the target chains (only applied when bApplyChainSettings is true).
-	UPROPERTY(EditAnywhere, BlueprintReadWrite, Category=ChainSettings, meta=(EditCondition="bApplyChainSettings"))
-	TMap<FName, FTargetChainSettings> ChainSettings;
+	/** Globally forces all IK solving off */
+	UPROPERTY(EditAnywhere, BlueprintReadWrite, Category = "IK Retarget Settings", DisplayName="Force All IK Off")
+	bool bForceAllIKOff = false;
 
-	// If true, the root settings stored in this profile will be applied to the Retargeter (when plugged into the Retargeter).
-	UPROPERTY(EditAnywhere, BlueprintReadWrite, Category=RootSettings, meta = (DisplayName = "Override Root Settings"))
-	bool bApplyRootSettings = false;
+	void FillProfileWithAssetSettings(const UIKRetargeter* InAsset);
 
-	// Retarget settings to control behavior of the retarget root motion (not applied unless bApplyRootSettings is true)
-	UPROPERTY(EditAnywhere, BlueprintReadWrite, Category=RootSettings, meta=(EditCondition="bApplyRootSettings"))
-	FTargetRootSettings RootSettings;
+	void MergeWithOtherProfile(const FRetargetProfile& OtherProfile);
 
-	// If true, the global settings stored in this profile will be applied to the Retargeter (when plugged into the Retargeter).
-	UPROPERTY(EditAnywhere, BlueprintReadWrite, Category=GlobalSettings, meta = (DisplayName = "Override Global Settings"))
-	bool bApplyGlobalSettings = false;
-
-	// Retarget settings to control global behavior, like Stride Warping (not applied unless bApplyGlobalSettings is true)
-	UPROPERTY(EditAnywhere, BlueprintReadWrite, Category=GlobalSettings, meta=(EditCondition="bApplyGlobalSettings"))
-	FRetargetGlobalSettings GlobalSettings;
+	FRetargetOpProfile* FindMatchingOpProfile(const FRetargetOpProfile& OtherOpProfile);
 
-	void MergeWithOtherProfile(const FRetargetProfile& OtherProfile)
+	bool ApplyOpProfilesToOpStruct(FInstancedStruct& InOutOpStruct, const int32 InIndexOfOpInStack, const ECopyOpSettingsContext InCopyContext) const;
+	
+	// search all the op profiles and get an array of all the settings of a particular type
+	template <typename T>
+	void GetOpSettingsByTypeInProfile(TArray<T*>& OutMatchingOpSettings)
 	{
-		if (OtherProfile.bApplyTargetRetargetPose)
-		{
-			TargetRetargetPoseName = OtherProfile.TargetRetargetPoseName;
-		}
-
-		if (OtherProfile.bApplySourceRetargetPose)
-		{
-			SourceRetargetPoseName = OtherProfile.SourceRetargetPoseName;
-		}
-
-		if (OtherProfile.bApplyChainSettings)
+		const UScriptStruct* SettingsTypeToMatch = T::StaticStruct();
+		
+		for (FRetargetOpProfile& OpProfile : RetargetOpProfiles)
 		{
-			for (const TPair<FName,FTargetChainSettings>& Pair : OtherProfile.ChainSettings)
+			FInstancedStruct& OpSettingsStruct = OpProfile.SettingsToApply;
+			if (OpSettingsStruct.GetScriptStruct() == SettingsTypeToMatch)
 			{
-				ChainSettings.Add(Pair.Key, Pair.Value);
+				OutMatchingOpSettings.Add(OpSettingsStruct.GetMutablePtr<T>());
 			}
 		}
+	}
 
-		if (OtherProfile.bApplyRootSettings)
-		{
-			RootSettings = OtherProfile.RootSettings;
-		}
+	//
+	// BEGIN DEPRECATED PROFILE SETTINGS
+	// NOTE: these are still stored and applied but should be removed and replaced with Op profiles
+	// NOTE: I cannot use the _DEPRECATED suffix on these properties because they live in a USTRUCT
+	// and UBT will complain that deprecated properties cannot be blueprint exposed without a getter but USTRUCT does not support UFUNCTIONS/getters
+	PRAGMA_DISABLE_DEPRECATION_WARNINGS
+	// (Deprecated) If true, the Chain Settings stored in this profile will be applied to the Retargeter (when plugged into the Retargeter).
+	UE_DEPRECATED(5.6, "Modifying chain settings must go through a retarget op profile.")
+	UPROPERTY(EditAnywhere, BlueprintReadWrite, Category="Deprecated", meta=(DisplayName="Apply Chain Settings (Deprecated)"))
+	bool bApplyChainSettings = false;
+	
+	// (Deprecated) A (potentially sparse) set of setting overrides for the target chains (only applied when bApplyChainSettings is true).
+	UE_DEPRECATED(5.6, "Modifying chain settings must go through a retarget op profile.")
+	UPROPERTY(EditAnywhere, BlueprintReadWrite, Category="Deprecated", meta=(DisplayName="Chain Settings (Deprecated)"))
+	TMap<FName, FTargetChainSettings> ChainSettings;
 
-		if (OtherProfile.bApplyGlobalSettings)
-		{
-			GlobalSettings = OtherProfile.GlobalSettings;
-		}
-	}
+	// (Deprecated) If true, the root settings stored in this profile will be applied to the Retargeter (when plugged into the Retargeter).
+	UE_DEPRECATED(5.6, "Modifying root settings must go through a retarget op profile.")
+	UPROPERTY(EditAnywhere, BlueprintReadWrite, Category="Deprecated", meta=(DisplayName="Apply Root Settings (Deprecated)"))
+	bool bApplyRootSettings = false;
+
+	// (Deprecated) Retarget settings to control behavior of the retarget root motion (not applied unless bApplyRootSettings is true)
+	UE_DEPRECATED(5.6, "Modifying root settings must go through a retarget op profile.")
+	UPROPERTY(EditAnywhere, BlueprintReadWrite, Category="Deprecated", meta=(DisplayName="Root Settings (Deprecated)"))
+	FTargetRootSettings RootSettings;
+
+	// (Deprecated) If true, the global settings stored in this profile will be applied to the Retargeter (when plugged into the Retargeter).
+	UE_DEPRECATED(5.6, "Modifying global settings must go through a retarget op profile.")
+	UPROPERTY(EditAnywhere, BlueprintReadWrite, Category="Deprecated", meta=(DisplayName="Apply Global Settings (Deprecated)"))
+	bool bApplyGlobalSettings = false;
+
+	// (Deprecated) Retarget settings to control global behavior, like Stride Warping (not applied unless bApplyGlobalSettings is true)
+	UE_DEPRECATED(5.6, "Modifying global settings must go through a retarget op profile.")
+	UPROPERTY(EditAnywhere, BlueprintReadWrite, Category="Deprecated", meta=(DisplayName="Global Settings (Deprecated)"))
+	FRetargetGlobalSettings GlobalSettings;
+	PRAGMA_ENABLE_DEPRECATION_WARNINGS
+	//
+	// END DEPRECATED API
+	//
 };
diff --git a/Engine/Plugins/Animation/IKRig/Source/IKRig/Public/Retargeter/IKRetargetSettings.h b/Engine/Plugins/Animation/IKRig/Source/IKRig/Public/Retargeter/IKRetargetSettings.h
--- a/Engine/Plugins/Animation/IKRig/Source/IKRig/Public/Retargeter/IKRetargetSettings.h
+++ b/Engine/Plugins/Animation/IKRig/Source/IKRig/Public/Retargeter/IKRetargetSettings.h
@@ -4,270 +4,16 @@
 
 #include "IKRetargetSettings.generated.h"
 
-class URetargetRootSettings;
-class URetargetChainSettings;
+class UPrimitiveComponent;
 
+// which skeleton are we referring to?
 UENUM(BlueprintType)
-enum class ERetargetTranslationMode : uint8
+enum class ERetargetSourceOrTarget : uint8
 {
-	None							UMETA(DisplayName = "None"),
-	GloballyScaled					UMETA(DisplayName = "Globally Scaled"),
-	Absolute						UMETA(DisplayName = "Absolute"),
-	StretchBoneLengthUniformly		UMETA(DisplayName = "Stretch Bone Length Uniformly"),
-	StretchBoneLengthNonUniformly	UMETA(DisplayName = "Stretch Bone Length Non-Uniformly"),
-};
-
-UENUM(BlueprintType)
-enum class ERetargetRotationMode : uint8
-{
-	Interpolated		UMETA(DisplayName = "Interpolated"),
-	OneToOne			UMETA(DisplayName = "One to One"),
-	OneToOneReversed	UMETA(DisplayName = "One to One Reversed"),
-	MatchChain			UMETA(DisplayName = "Match Chain"),
-	MatchScaledChain	UMETA(DisplayName = "Match Scaled Chain"),
-	None				UMETA(DisplayName = "None"),
-};
-
-USTRUCT(BlueprintType)
-struct IKRIG_API FTargetChainSpeedPlantSettings
-{
-	GENERATED_BODY()
-
-	/** The name of the curve on the source animation that contains the speed of the end effector bone.*/
-	UPROPERTY(EditAnywhere, BlueprintReadWrite, Category = "Plant IK by Speed")
-	bool EnableSpeedPlanting = false;
-	
-	/** The name of the curve on the source animation that contains the speed of the end effector bone.*/
-	UPROPERTY(EditAnywhere, BlueprintReadWrite, BlueprintReadWrite, Category = "Plant IK by Speed", meta = (ClampMin = "0.0", ClampMax = "100.0", UIMin = "0.0", UIMax = "100.0"))
-	FName SpeedCurveName;
-
-	/** Range 0 to 100. Default 15. The maximum speed a source bone can be moving while being considered 'planted'.
-	*  The target IK goal will not be allowed to move whenever the source bone speed drops below this threshold speed. */
-	UPROPERTY(EditAnywhere, BlueprintReadWrite, Category = "Plant IK by Speed", meta = (ClampMin = "0.0", ClampMax = "100.0", UIMin = "0.0", UIMax = "100.0"))
-	float SpeedThreshold = 15.0f;
-
-	// How stiff the spring model is that smoothly pulls the IK position after unplanting (more stiffness means more oscillation around the target value)
-	UPROPERTY(EditAnywhere, BlueprintReadWrite, Category = "Plant IK by Speed", meta = (ClampMin = "0.0", UIMin = "0.0"))
-	float UnplantStiffness = 250.0f;
-
-	// How much damping to apply to the spring (0 means no damping, 1 means critically damped which means no oscillation)
-	UPROPERTY(EditAnywhere, BlueprintReadWrite, Category = "Plant IK by Speed", meta = (ClampMin = "0.0", ClampMax = "1.0", UIMin = "0.0", UIMax = "1.0"))
-	float UnplantCriticalDamping = 1.0f;
-
-	bool operator==(const FTargetChainSpeedPlantSettings& Other) const;
-};
-
-USTRUCT(BlueprintType)
-struct IKRIG_API FTargetChainFKSettings
-{
-	GENERATED_BODY()
-
-	/** Whether to copy the shape of the chain from the source skeleton using the Rotation and Translation modes. Default is true.
-	* NOTE: All FK operations run before the IK pass to copy the shape of the FK chain from the source skeleton. */
-	UPROPERTY(EditAnywhere, BlueprintReadWrite, Category = "FK Adjustments")
-	bool EnableFK = true;
-	
-	/** Determines how rotation is copied from the source chain to the target chain. Default is Interpolated.
-	* Interpolated: Source and target chains are normalized by their length, then each target bone rotation is generated by finding the rotation at the same normalized distance on the source chain and interpolating between the neighboring bones.
-	* One to One: Each target bone rotation is copied from the equivalent bone in the source chain, based on the order in the chain, starting at the root of the chain. If the target chain has more bones than the source, the extra bones will remain at their reference pose.
-	* One to One Reversed: Same as One-to-One, but starting from the tip of the chain.
-	* Match Chain: Uses a Spline-IK type method to orient the target joint chain to exactly match the shape of the source chain.
-	* Match Scaled Chain: Same as Match Chain, but scales the source chain to the same length as the target chain before running the match procedure.
-	* None: The rotation of each target bone in the chain is left at the reference pose. */
-	UPROPERTY(EditAnywhere, BlueprintReadWrite, Category = "FK Adjustments")
-	ERetargetRotationMode RotationMode = ERetargetRotationMode::Interpolated;
-
-	/** Range +/- infinity. Default 1. Scales the amount of rotation that is applied.
-	*  If Rotation Mode is None this parameter has no effect.
-	*  Otherwise, this parameter blends the rotation of each bone in the chain from the base retarget pose (0) to the retargeted pose (1).*/
-	UPROPERTY(EditAnywhere, BlueprintReadWrite, Category = "FK Adjustments", meta = (UIMin = "0.0", UIMax = "1.0"))
-	float RotationAlpha = 1.0f;
-
-	/** Determines how translation is copied from the source chain to the target chain. Default is None.
-	* None: Translation of target bones are left unmodified from the retarget pose.
-	* Globally Scaled: Translation of target bone is set to the source bone offset multiplied by the global scale of the skeleton (determined by the relative height difference between retarget root bones).
-	* Absolute: Translation of target bone is set to the absolute position of the source bone.
-	* Stretch Bone Length Uniformly: squashes or stretches the local translation of each bone so that the length of the target chain matches the relative length of the source chain. Accounts for differences in initial chain lengths.
-	* Stretch Bone Length Non-Uniformly: same as "Stretch Bone Length Uniformly" but treats each bone independently to allow for chains that stretch non-uniformly along their length (ie the base may stretch more than the tip or vice versa).   
-	 */
-	UPROPERTY(EditAnywhere, BlueprintReadWrite, Category = "FK Adjustments")
-	ERetargetTranslationMode TranslationMode = ERetargetTranslationMode::None;
-
-	/** Range +/- infinity. Default 1. Scales the amount of translation that is applied. Exact behavior depends on the Translation Mode.
-	*  In None Mode, this parameter has no effect.
-	*  In Globally Scaled and Absolute modes, the translation offset is scaled by this parameter.*/
-	UPROPERTY(EditAnywhere, BlueprintReadWrite, Category = "FK Adjustments", meta = (UIMin = "0.0", UIMax = "1.0"))
-	float TranslationAlpha = 1.0f;
-
-	/** Range 0 to 1. Default 1. Matches the twist angle of this chain (along the Pole direction) to the source chain.
-	*  At 0, the chain's pole vector direction will be left alon
-	*  At 1, the root bone of the chain will be twist-rotated in the pole direction to match the orientation of the source chain.*/
-	UPROPERTY(EditAnywhere, BlueprintReadWrite, Category = "FK Adjustments", meta = (UIMin = "0.0", UIMax = "1.0", ClampMin = "0.0", ClampMax = "1.0"))
-	float PoleVectorMatching = 0.0f;
-
-	/** Default is False. When true, the original offset between the source/target pole vectors will be maintained when using Pole Vector Matching. */
-	UPROPERTY(EditAnywhere, BlueprintReadWrite, Category = "FK Adjustments")
-	bool PoleVectorMaintainOffset = false;
-
-	/** Range +/- 180. Default 0. An angular offset, in degrees, for the pole direction of the chain. */
-	UPROPERTY(EditAnywhere, BlueprintReadWrite, Category = "FK Adjustments", meta = (UIMin = "-180.0", UIMax = "180.0", ClampMin = "-180.0", ClampMax = "180.0"))
-	float PoleVectorOffset = 0.0f;
-
-	bool operator==(const FTargetChainFKSettings& Other) const;
-};
-
-USTRUCT(BlueprintType)
-struct IKRIG_API FTargetChainIKSettings
-{
-	GENERATED_BODY()
-
-	/** Whether to modify the location of the IK goal on this chain. Default is true.
-	 * NOTE: This only has an effect if the chain has an IK Goal assigned to it in the Target IK Rig asset.
-	 * NOTE: If off, and this chain has an IK Goal, the IK will still be evaluated, but the Goal is set to the input bone location (from the FK pass).*/
-	UPROPERTY(EditAnywhere, BlueprintReadWrite, Category = "IK Adjustments")
-	bool EnableIK = true;
-
-	/** Range 0 to 1. Default 0. Blends IK goal transform from retargeted transform (0) to source bone transform (1).
-	*  At 0 the goal is placed at the retargeted location and rotation.
-	*  At 1 the goal is placed at the location and rotation of the source chain's end bone. */
-	UPROPERTY(EditAnywhere, BlueprintReadWrite, Category = "IK Adjustments", meta = (ClampMin = "0.0", ClampMax = "1.0", UIMin = "0.0", UIMax = "1.0"))
-	float BlendToSource = 0.0f;
-
-	/** Range 0 to 1. Default 1. Blends the translational component of BlendToSource on/off.
-	*  At 0 the goal is placed at the retargeted location.
-	*  At 1 the goal is placed at the location of the source chain's end bone. */
-	UPROPERTY(EditAnywhere, BlueprintReadWrite, Category = "IK Adjustments", meta = (ClampMin = "0.0", ClampMax = "1.0", UIMin = "0.0", UIMax = "1.0"))
-	float BlendToSourceTranslation = 1.0f;
-
-	/** Range 0 to 1. Default 0. Blends the rotational component of BlendToSource on/off.
-	*  At 0 the goal is oriented to the retargeted rotation.
-	*  At 1 the goal is oriented to the source chain's end bone rotation. */
-	UPROPERTY(EditAnywhere, BlueprintReadWrite, Category = "IK Adjustments", meta = (ClampMin = "0.0", ClampMax = "1.0", UIMin = "0.0", UIMax = "1.0"))
-	float BlendToSourceRotation = 0.0f;
-	
-	/** Range 0 to 1. Default 1. Weight each axis separately when using Blend To Source.
-	*  At 0 the goal is placed at the retargeted location.
-	*  At 1 the goal is placed at the location of the source chain's end bone. */
-	UPROPERTY(EditAnywhere, BlueprintReadWrite, Category = "IK Adjustments", meta = (ClampMin = "0.0", ClampMax = "1.0", UIMin = "0.0", UIMax = "1.0"))
-	FVector BlendToSourceWeights = FVector::OneVector;
-
-	/** Default 0, 0, 0. Apply a static global-space offset to IK goal position. */
-	UPROPERTY(EditAnywhere, BlueprintReadWrite, Category = "IK Adjustments")
-	FVector StaticOffset = FVector::ZeroVector;
-
-	/** Default 0, 0, 0. Apply a static local-space offset to IK goal position. */
-	UPROPERTY(EditAnywhere, BlueprintReadWrite, Category = "IK Adjustments")
-	FVector StaticLocalOffset = FVector::ZeroVector;
-
-	/** Default 0, 0, 0. Apply a static local-space offset to IK goal rotation. */
-	UPROPERTY(EditAnywhere, BlueprintReadWrite, Category = "IK Adjustments")
-	FRotator StaticRotationOffset = FRotator::ZeroRotator;
-
-	/** Range 0 to infinity. Default 1. Scales the vertical component of the IK goal's position. */
-	UPROPERTY(EditAnywhere, BlueprintReadWrite, Category = "IK Adjustments", meta = (ClampMin = "0.0", UIMin = "0.0", UIMax = "5.0"))
-	float ScaleVertical = 1.0f;
-	
-	/** Range 0 to 5. Default 1. Brings IK goal closer (0) or further (1+) from origin of chain.
-	*  At 0 the effector is placed at the origin of the chain (ie Shoulder, Hip etc).
-	*  At 1 the effector is left at the end of the chain (ie Wrist, Foot etc)
-	*  Values in-between 0-1 will slide the effector along the vector from the start to the end of the chain.
-	*  Values greater than 1 will stretch the chain beyond the retargeted length. */
-	UPROPERTY(EditAnywhere, BlueprintReadWrite, Category = "IK Adjustments", meta = (ClampMin = "0.0", ClampMax = "5.0", UIMin = "0.1", UIMax = "2.0"))
-	float Extension = 1.0f;
-
-	/** Is this IK goal affected by the stride warping (in Global Settings). Typically this is true for all feet, but not for hands.  */
-	UPROPERTY(EditAnywhere, BlueprintReadWrite, Category = Warping, meta = (DisplayName = "Affected by IK Warping"))
-	bool bAffectedByIKWarping = true;
-
-	bool operator==(const FTargetChainIKSettings& Other) const;
-};
-
-USTRUCT(BlueprintType)
-struct IKRIG_API FTargetChainSettings
-{
-	GENERATED_BODY()
-
-public:
-
-	/** All settings for FK retargeting of this chain.
-	 *FK retargeting runs before the IK pass. */
-	UPROPERTY(EditAnywhere, BlueprintReadWrite, Category = "FK")
-	FTargetChainFKSettings FK;
-
-	/** All settings controlling the IK Goal assigned to this chain.
-	 *The IK pass runs AFTER the FK pass and can be used to fix contacts. */
-	UPROPERTY(EditAnywhere, BlueprintReadWrite, Category = "IK")
-	FTargetChainIKSettings IK;
-
-	/** All settings associated with planting IK goals based on the speed of the source.
-	 *Speed planting will pin the IK goal to the location determined by the IK settings above. */
-	UPROPERTY(EditAnywhere, BlueprintReadWrite, Category = "Speed Planting")
-	FTargetChainSpeedPlantSettings SpeedPlanting;
-	
-	bool operator==(const FTargetChainSettings& Other) const;
-};
-
-USTRUCT(BlueprintType)
-struct IKRIG_API FTargetRootSettings
-{
-	GENERATED_BODY()
-
-public:
-
-	/** Range 0 to 1. Default 1. Blends the amount of retargeted root rotation to apply.
-	*  At 0 the root is left at the rotation from the retarget pose.
-	*  At 1 the root is rotated fully to match the source root rotation. */
-	UPROPERTY(EditAnywhere, BlueprintReadWrite, Category = "Root Retarget Settings", meta = (ClampMin = "0.0", ClampMax = "1.0", UIMin = "0.0", UIMax = "1.0"))
-	float RotationAlpha = 1.0f;
-	
-	/** Range 0 to 1. Default 1. Blends the amount of retargeted root translation to apply.
-	*  At 0 the root is left at the position from the retarget pose.
-	*  At 1 the root will follow the source motion according to the behavior defined in the subsequent settings. */
-	UPROPERTY(EditAnywhere, BlueprintReadWrite, Category = "Root Retarget Settings", meta = (ClampMin = "0.0", ClampMax = "1.0", UIMin = "0.0", UIMax = "1.0"))
-	float TranslationAlpha = 1.0f;
-	
-	/** Range 0 to 1. Default 0. Blends the retarget root's translation to the exact source location.
-	*  At 0 the root is placed at the retargeted location.
-	*  At 1 the root is placed at the location of the source's retarget root bone.*/
-	UPROPERTY(EditAnywhere, BlueprintReadWrite, Category = "Root Retarget Settings", meta = (ClampMin = "0.0", ClampMax = "1.0", UIMin = "0.0", UIMax = "1.0"))
-	float BlendToSource = 0.0f;
-
-	/** Per-axis weights for the Blend to Source. */
-	UPROPERTY(EditAnywhere, BlueprintReadWrite, Category = "Root Retarget Settings", meta = (ClampMin = "0.0", ClampMax = "1.0", UIMin = "0.0", UIMax = "1.0"))
-	FVector BlendToSourceWeights = FVector::OneVector;
-
-	/** Default 1. Scales the motion of the root position in the horizontal plane (X,Y). */
-	UPROPERTY(EditAnywhere, BlueprintReadWrite, Category = "Root Retarget Settings", meta = (UIMin = "0.0", UIMax = "3.0"))
-	float ScaleHorizontal = 1.0f;
-
-	/** Default 1. Scales the motion of the root position in the vertical direction (Z). */
-	UPROPERTY(EditAnywhere, BlueprintReadWrite, Category = "Root Retarget Settings", meta = (UIMin = "0.0", UIMax = "3.0"))
-	float ScaleVertical = 1.0f;
-
-	/** Applies a static component-space translation offset to the retarget root.*/
-	UPROPERTY(EditAnywhere, BlueprintReadWrite, Category = "Root Retarget Settings")
-	FVector TranslationOffset = FVector::ZeroVector;
-
-	/** Applies a static local-space rotation offset to the retarget root.*/
-	UPROPERTY(EditAnywhere, BlueprintReadWrite, Category = "Root Retarget Settings", meta = (ClampMin = "-180.0", ClampMax = "180.0", UIMin = "-180.0", UIMax = "180.0"))
-	FRotator RotationOffset = FRotator::ZeroRotator;
-
-	/** Range 0 to 1. Default 1. Control whether modifications made to the root will affect the horizontal component of IK positions.
-	*  At 0 the IK positions are independent of the root modifications.
-	*  At 1 the IK positions are calculated relative to the modified root location.*/
-	UPROPERTY(EditAnywhere, BlueprintReadWrite, Category = "Root Retarget Settings", meta = (ClampMin = "0.0", ClampMax = "1.0", UIMin = "0.0", UIMax = "1.0", DisplayName = "Affect IK Horizontal"))
-	float AffectIKHorizontal = 1.0f;
-	
-	/** Range 0 to 1. Default 0. Control whether modifications made to the root will affect the vertical component of IK positions.
-	*  At 0 the IK positions are independent of the root modifications.
-	*  At 1 the IK positions are calculated relative to the modified root location.*/
-	UPROPERTY(EditAnywhere, BlueprintReadWrite, Category = "Root Retarget Settings", meta = (ClampMin = "0.0", ClampMax = "1.0", UIMin = "0.0", UIMax = "1.0", DisplayName = "Affect IK Vertical"))
-	float AffectIKVertical = 0.0f;
-
-	FVector GetAffectIKWeightVector() const
-	{
-		return FVector(AffectIKHorizontal, AffectIKHorizontal, AffectIKVertical);
-	}
+	// the SOURCE skeleton (to copy FROM)
+	Source,
+	// the TARGET skeleton (to copy TO)
+	Target,
 };
 
 UENUM()
@@ -281,118 +27,32 @@ enum class EBasicAxis
 	NegZ
 };
 
-UENUM()
-enum class EWarpingDirectionSource
+#if WITH_EDITOR
+enum class ERetargetSelectionType : uint8
 {
-	Goals,
-	Chain,
-	RootBone
+	BONE,
+	CHAIN,
+	MESH,
+	ROOT,
+	NONE
 };
 
-USTRUCT(BlueprintType)
-struct IKRIG_API FRetargetGlobalSettings
+struct FIKRetargetDebugDrawState
 {
-	GENERATED_BODY()
-	
-	/** When false, the motion of the Retarget Root bone is not copied from the source. Useful for debugging issues with the root settings.
-	* Note: the retargeting order is: Root > FK > IK > Post
-	 */
-	UPROPERTY(EditAnywhere, BlueprintReadWrite, Category = RetargetPhases)
-	bool bEnableRoot = true;
-	
-	/** When false, limbs are not copied via FK. Useful for debugging limb issues suspected to be caused by FK chain settings.
-	* Note: the retargeting order is: Root > FK > IK > Post
-	 */
-	UPROPERTY(EditAnywhere, BlueprintReadWrite, Category = RetargetPhases)
-	bool bEnableFK = true;
-	
-	/** When false, IK is not applied as part of retargeter. Useful for debugging limb issues suspected to be caused by IK.
-	* Note: the retargeting order is: Root > FK > IK > Post
-	 */
-	UPROPERTY(EditAnywhere, BlueprintReadWrite, Category = RetargetPhases)
-	bool bEnableIK = true;
-
-	/** When false, Post operations are not applied as part of retargeter. Useful for debugging issues suspected to be caused by the post phase.
-	 * Note: the retargeting order is: Root > FK > IK > Post
-	 */
-	UPROPERTY(EditAnywhere, BlueprintReadWrite, Category = RetargetPhases)
-	bool bEnablePost = true;
-
-	/** When true, will copy all the source bone transforms with matching names to use as a base pose. This can be useful for partial retargeting.
-	 * NOTE: no retargeting is applied to the bone transforms, they are assumed to be fully compatible between source/target skeletons.
-	 */
-	UPROPERTY(EditAnywhere, BlueprintReadWrite, Category = RetargetPhases)
-	bool bCopyBasePose = false;
+	// state of all things you can select in the retarget editor viewport
+	bool bIsRootSelected = false;
+	UPrimitiveComponent* SelectedMesh = nullptr;
+	TArray<FName> SelectedChains;
+	TArray<FName> SelectedGoals;
+	TMap<ERetargetSourceOrTarget, TArray<FName>> SelectedBoneNames;
+	ERetargetSelectionType LastSelectedItem = ERetargetSelectionType::NONE;
 
-	/** Filters the bones to copy when using "Copy Base Pose". If specified, will only copy all children of the specified bone (inclusive).*/
-	UPROPERTY(EditAnywhere, BlueprintReadWrite, Category = RetargetPhases, meta = (EditCondition = "bCopyBasePose"))
-	FName CopyBasePoseRoot = NAME_None;
-
-	/** Range 0.01 to +inf. Default 1. Scales the incoming source pose. Affects entire skeleton and all IK goals.*/
-	UPROPERTY(EditAnywhere, BlueprintReadWrite, Category = Scale, meta = (ClampMin = "0.01", UIMin = "0.01", UIMax = "10.0"))
-	float SourceScaleFactor = 1.0f;
-
-	/** Enable IK Warping.
-	 * These options allow for global modifications to all IK Goals that have "Affected by IK Warping" turned on (the default).
-	 * "Affected by IK Warping" can be found in the IK settings of a retarget chain. */
-	UPROPERTY(EditAnywhere, BlueprintReadWrite, Category = Warping)
-	bool bWarping = false;
-
-	/** Defines the source used to determine the forward direction as the character animates in world space. Default is "Goals".
-	 * This method uses various positions on the character to create a "best fit" global rotation that approximates the facing direction of the character over time.
-	 * This global rotation is used to define the forward and sideways directions used when warping goals along those axes.
-	 * The options are:
-	 * Goals: uses the positions of the IK goals to approximate the facing direction. This is best used on characters with a vertical spine, like bipeds.
-	 * Chain: uses the positions of the bones in a retarget chain to approximate the facing direction. This is best when used with the spine chain for characters with a horizontal spine, like quadrupeds.
-	 * Root Bone: uses the rotation of the root bone of the skeleton. This is most robust, but character must have correct root motion with yaw rotation in movement direction.
-	 */
-	UPROPERTY(EditAnywhere, BlueprintReadWrite, Category = Warping)
-	EWarpingDirectionSource DirectionSource = EWarpingDirectionSource::Goals;
-	
-	/** The world space axis that represents the forward facing direction for your character. By default, in Unreal, this is +Y.  */
-	UPROPERTY(EditAnywhere, BlueprintReadWrite, Category = Warping)
-	EBasicAxis ForwardDirection = EBasicAxis::Y;
-
-	/** When using the "Chain" option as a Direction Source, this defines the chain to use to determine the facing direction of the character.
-	 * Typically this would be the chain that contains the Spine bones. */
-	UPROPERTY(EditAnywhere, BlueprintReadWrite, Category = Warping, meta=(EditCondition="DirectionSource == EWarpingDirectionSource::Chain"))
-	FName DirectionChain;
-
-	/** Range 0 to Inf. Default 1. Warps IK goal positions in the forward direction. Useful for stride warping.
-	 * Values below 1 will create smaller, squashed strides. Values greater than 1 will create stretched, longer strides.*/
-	UPROPERTY(EditAnywhere, BlueprintReadWrite, Category = Warping, meta = (UIMin = "0.0", UIMax = "5.0", ClampMin = "0.0"))
-	float WarpForwards = 1.0f;
-
-	/** Range -+Inf. Default is 0. A static offset in world units to move the IK goals perpendicular to the forward direction.
-	 * Values less than zero will move the goals towards the center-line of the character. Values greater than zero push the goals outwards.*/
-	UPROPERTY(EditAnywhere, BlueprintReadWrite, Category = Warping, meta = (UIMin = "-10.0", UIMax = "100.0"))
-	float SidewaysOffset = 0.0f;
-
-	/** Range 0 to +Inf. Default is 1.0f.
-	 * Values below 1 pull all the goals towards the average of all the goals (towards each other).
-	 * Values greater than 1 push the goals apart.*/
-	UPROPERTY(EditAnywhere, BlueprintReadWrite, Category = Warping, meta = (UIMin = "0.0", UIMax = "2.0", ClampMin = "0.0"))
-	float WarpSplay = 1.0f;
+	// common viewport drawing color palette for selected things
+	static FLinearColor Muted;
+	static FLinearColor SourceColor; 
+	static FLinearColor GoalColor;
+	static FLinearColor MainColor;
+	static FLinearColor NonSelected;
+};
+#endif
 
-	static FVector GetAxisVector(const EBasicAxis& Axis)
-	{
-		switch (Axis)
-		{
-		case EBasicAxis::X:
-			return FVector::XAxisVector;
-		case EBasicAxis::Y:
-			return FVector::YAxisVector;
-		case EBasicAxis::Z:
-			return FVector::ZAxisVector;
-		case EBasicAxis::NegX:
-			return -FVector::XAxisVector;
-		case EBasicAxis::NegY:
-			return -FVector::YAxisVector;
-		case EBasicAxis::NegZ:
-			return -FVector::ZAxisVector;
-		default:
-			checkNoEntry();
-			return FVector::ZeroVector;
-		}
-	}
-};
\ No newline at end of file
