[IKRig] Major refactor part 3 of 3.

diff --git a/Engine/Plugins/Animation/IKRig/Source/IKRig/Public/Retargeter/IKRetargeter.h b/Engine/Plugins/Animation/IKRig/Source/IKRig/Public/Retargeter/IKRetargeter.h
--- a/Engine/Plugins/Animation/IKRig/Source/IKRig/Public/Retargeter/IKRetargeter.h
+++ b/Engine/Plugins/Animation/IKRig/Source/IKRig/Public/Retargeter/IKRetargeter.h
@@ -11,154 +11,8 @@
 #if WITH_EDITOR
 class FIKRetargetEditorController;
 #endif
-struct FIKRetargetPose;
-class URetargetChainSettings;
-class URetargetRootSettings;
-class UDEPRECATED_RetargetOpStack;
-
-struct UE_DEPRECATED(5.1, "Use URetargetChainSettings instead.") FRetargetChainMap;
-USTRUCT()
-struct IKRIG_API FRetargetChainMap
-{
-	GENERATED_BODY()
-
-	FRetargetChainMap() = default;
-	FRetargetChainMap(const FName& TargetChain) : TargetChain(TargetChain){}
-	
-	UPROPERTY(EditAnywhere, Category = Offsets)
-	FName SourceChain = NAME_None;
-	
-	UPROPERTY(EditAnywhere, Category = Offsets)
-	FName TargetChain = NAME_None;
-};
-
-UCLASS(BlueprintType)
-class IKRIG_API URetargetChainSettings : public UObject
-{
-	GENERATED_BODY()
-
-public:
-	
-	URetargetChainSettings() = default;
-	
-	URetargetChainSettings(const FName& TargetChain) : TargetChain(TargetChain){}
-	
-	// UObject 
-	virtual void Serialize(FArchive& Ar) override;
-	// END UObject 
-	
-	// The chain on the Source IK Rig asset to copy animation FROM. 
-	UPROPERTY(BlueprintReadOnly, VisibleAnywhere, Category = "Chain Mapping")
-	FName SourceChain;
-
-	// The chain on the Target IK Rig asset to copy animation TO. 
-	UPROPERTY(BlueprintReadOnly, VisibleAnywhere, Category = "Chain Mapping")
-	FName TargetChain;
-
-	// The settings used to control the motion on this target chain. 
-	UPROPERTY(BlueprintReadWrite, EditAnywhere, Category = "Settings")
-	FTargetChainSettings Settings;
-
-#if WITH_EDITORONLY_DATA
-	// Deprecated properties from before FTargetChainSettings / profile refactor  (July 2022)
-	UPROPERTY()
-	bool CopyPoseUsingFK_DEPRECATED = true;
-	UPROPERTY()
-	ERetargetRotationMode RotationMode_DEPRECATED;
-	UPROPERTY()
-	float RotationAlpha_DEPRECATED = 1.0f;
-	UPROPERTY()
-	ERetargetTranslationMode TranslationMode_DEPRECATED;
-	UPROPERTY()
-	float TranslationAlpha_DEPRECATED = 1.0f;
-	UPROPERTY()
-	bool DriveIKGoal_DEPRECATED = true;
-	UPROPERTY()
-	float BlendToSource_DEPRECATED = 0.0f;
-	UPROPERTY()
-	FVector BlendToSourceWeights_DEPRECATED = FVector::OneVector;
-	UPROPERTY()
-	FVector StaticOffset_DEPRECATED;
-	UPROPERTY()
-	FVector StaticLocalOffset_DEPRECATED;
-	UPROPERTY()
-	FRotator StaticRotationOffset_DEPRECATED;
-	UPROPERTY()
-	float Extension_DEPRECATED = 1.0f;
-	UPROPERTY()
-	bool UseSpeedCurveToPlantIK_DEPRECATED = false;
-	UPROPERTY()
-	FName SpeedCurveName_DEPRECATED;
-	UPROPERTY()
-	float VelocityThreshold_DEPRECATED = 15.0f;
-	UPROPERTY()
-	float UnplantStiffness_DEPRECATED = 250.0f;
-	UPROPERTY()
-	float UnplantCriticalDamping_DEPRECATED = 1.0f;
-	// END deprecated properties 
-#endif
-
-	// pointer to editor for details customization
-	#if WITH_EDITOR
-	TWeakPtr<FIKRetargetEditorController> EditorController;
-	#endif
-};
-
-UCLASS(BlueprintType)
-class IKRIG_API URetargetRootSettings: public UObject
-{
-	GENERATED_BODY()
-	
-public:
-	
-	// UObject 
-	virtual void Serialize(FArchive& Ar) override;
-	// END UObject 
-
-	// The settings used to control the motion of the target root bone. 
-	UPROPERTY(EditAnywhere, Category = "Settings")
-	FTargetRootSettings Settings;
-
-	// pointer to editor for details customization
-	#if WITH_EDITOR
-	TWeakPtr<FIKRetargetEditorController> EditorController;
-	#endif
-
-private:
-#if WITH_EDITORONLY_DATA
-	// Deprecated properties from before FTargetRootSettings / profile refactor 
-	UPROPERTY()
-	bool RetargetRootTranslation_DEPRECATED = true;
-	UPROPERTY()
-	float GlobalScaleHorizontal_DEPRECATED = 1.0f;
-	UPROPERTY()
-	float GlobalScaleVertical_DEPRECATED = 1.0f;
-	UPROPERTY()
-	FVector BlendToSource_DEPRECATED = FVector::ZeroVector;
-	UPROPERTY()
-	FVector StaticOffset_DEPRECATED = FVector::ZeroVector;
-	UPROPERTY()
-	FRotator StaticRotationOffset_DEPRECATED = FRotator::ZeroRotator;
-	// END deprecated properties 
-#endif
-};
-
-UCLASS(BlueprintType)
-class IKRIG_API UIKRetargetGlobalSettings: public UObject
-{
-	GENERATED_BODY()
-	
-public:
 
-	// Global retargeter settings. 
-	UPROPERTY(EditAnywhere, Category = "Settings")
-	FRetargetGlobalSettings Settings;
-
-	// pointer to editor for details customization
-	#if WITH_EDITOR
-	TWeakPtr<FIKRetargetEditorController> EditorController;
-	#endif
-};
+class UDEPRECATED_RetargetOpStack;
 
 USTRUCT(BlueprintType)
 struct IKRIG_API FIKRetargetPose
@@ -198,14 +52,70 @@ private:
 	friend class UIKRetargeterController;
 };
 
-// which skeleton are we referring to?
-UENUM()
-enum class ERetargetSourceOrTarget : uint8
+USTRUCT()
+struct IKRIG_API FRetargetChainPair
 {
-	// the SOURCE skeleton (to copy FROM)
-	Source,
-	// the TARGET skeleton (to copy TO)
-	Target,
+	GENERATED_BODY()
+
+	FRetargetChainPair() = default;
+	FRetargetChainPair(const FName InTargetChainName, const FName InSourceChainName) :
+		TargetChainName(InTargetChainName), SourceChainName(InSourceChainName) {}
+
+	UPROPERTY()
+	FName TargetChainName = NAME_None;
+	
+	UPROPERTY()
+	FName SourceChainName = NAME_None;
+};
+
+USTRUCT()
+struct IKRIG_API FRetargetChainMapping
+{
+	GENERATED_BODY()
+
+	/** Get the name of the bone chain mapped to this chain.
+	 * @param InChainName the name of either a source or target chain
+	 * @param InSourceOrTarget the side that InChainName belongs to (either source or target)
+	 * @return true if the chain is present in the mapping */
+	bool HasChain(const FName& InChainName, ERetargetSourceOrTarget InSourceOrTarget) const;
+
+	/** Get the name of the bone chain mapped to this chain.
+	 * @param InChainName the name of either a source or target chain
+	 * @param InSourceOrTarget the side that InChainName belongs to (either source or target)
+	 * @return the name of the chain that InChainName is mapped to (possibly NAME_None) */
+	FName GetChainMappedTo(const FName InChainName, ERetargetSourceOrTarget InSourceOrTarget) const;
+
+	/** returns the chain pair belonging to the supplied chain or nullptr if not found
+	* @param InChainName the name of either a source or target chain
+	* @param InSourceOrTarget the side that InChainName belongs to (either source or target)*/
+	FRetargetChainPair* FindChainPair(const FName& InChainName, ERetargetSourceOrTarget InSourceOrTarget) const;
+	const FRetargetChainPair* FindChainPairConst(const FName& InChainName, ERetargetSourceOrTarget InSourceOrTarget) const;
+
+	/** Get read-write access to the mapping from target to source chains (by name) */
+	TArray<FRetargetChainPair>& GetChainPairs();
+	const TArray<FRetargetChainPair>& GetChainPairs() const;
+
+	/** Map target chain to a source chain.
+	 * @param InTargetChainName the name of a target chain
+	 * @param InSourceChainName the name of a source chain*/
+	void SetChainMapping(const FName InTargetChainName, const FName InSourceChainName);
+
+	/* synchronizes chain names with given source/target IK Rigs */
+	void CleanMapping(const UIKRigDefinition* SourceIKRig, const UIKRigDefinition* TargetIKRig);
+
+	/* sort mapping hierarchically (root to leaf based on Target IK Rig hierarchy) */
+	void SortMapping(const UIKRigDefinition* TargetIKRig);
+	
+	PRAGMA_DISABLE_DEPRECATION_WARNINGS
+	void LoadFromDeprecatedChainSettings(const TArray<URetargetChainSettings*>& InChainSettings);
+	PRAGMA_ENABLE_DEPRECATION_WARNINGS
+
+private:
+	
+	/** mapping of target to source bone chains by name
+	 * NOTE: this is an array instead of a TMap because it needs to be sorted by hierarchy*/
+	UPROPERTY()
+	TArray<FRetargetChainPair> ChainMap;
 };
 
 UCLASS(BlueprintType)
@@ -227,20 +137,45 @@ public:
 	USkeletalMesh* GetPreviewMesh(ERetargetSourceOrTarget SourceOrTarget) const;
 	#endif
 
-	// Get read-only access to the chain mapping 
-	const TArray<TObjectPtr<URetargetChainSettings>>& GetAllChainSettings() const { return ChainSettings; };
-	// Get read-only access to the chain map for a given chain (null if chain not in retargeter) 
-	const TObjectPtr<URetargetChainSettings> GetChainMapByName(const FName& TargetChainName) const;
-	// Get read-only access to the chain settings for a given chain (null if chain not in retargeter) 
-	const FTargetChainSettings* GetChainSettingsByName(const FName& TargetChainName) const;
-	// Get access to the root settings 
-	URetargetRootSettings* GetRootSettingsUObject() const { return RootSettings; };
-	// Get access to the global settings uobject
-	UIKRetargetGlobalSettings* GetGlobalSettingsUObject() const { return GlobalSettings; };
-	// Get access to the global settings itself 
-	const FRetargetGlobalSettings& GetGlobalSettings() const { return GlobalSettings->Settings; };
+	// Get read/write access to the chain mapping
+	FRetargetChainMapping& GetChainMapping() { return ChainMap; };
+	// Get read-only access to the chain mapping
+	const FRetargetChainMapping& GetChainMapping() const { return ChainMap; };
 	// Get access to the stack of retargeting operations
 	const TArray<FInstancedStruct>& GetRetargetOps() const { return RetargetOps; };
+	/** Get the first op in the stack of the given type */
+	template <typename T>
+	T* GetFirstRetargetOpOfType()
+	{
+		for (FInstancedStruct& OpStruct : RetargetOps)
+		{
+			if (OpStruct.GetScriptStruct()->IsChildOf(T::StaticStruct()))
+			{
+				return OpStruct.GetMutablePtr<T>();
+			}
+		}
+
+		return nullptr;
+	}
+	/** Get all ops in the stack of the given type */
+	template <typename T>
+	TArray<T*> GetAllRetargetOpsOfType()
+	{
+		TArray<T*> Results;
+    
+		for (FInstancedStruct& OpStruct : RetargetOps)
+		{
+			if (OpStruct.GetScriptStruct()->IsChildOf(T::StaticStruct()))
+			{
+				if (T* Instance = OpStruct.GetMutablePtr<T>())
+				{
+					Results.Add(Instance);
+				}
+			}
+		}
+    
+		return Results;
+	}
 
 	// Get read-only access to a retarget pose 
 	const FIKRetargetPose* GetCurrentRetargetPose(const ERetargetSourceOrTarget& SourceOrTarget) const;
@@ -251,8 +186,6 @@ public:
 	// Get name of default pose 
 	static const FName GetDefaultPoseName();
 	
-	// Fill the provided profile with settings from this asset
-	void FillProfileWithAssetSettings(FRetargetProfile& InOutProfile) const;
 	// Get the current retarget profile (may be null) 
 	const FRetargetProfile* GetCurrentProfile() const;
 	// Get the retarget profile by name (may be null) 
@@ -262,124 +195,173 @@ public:
 	int32 GetVersion() const { return Version; };
 	// do this after any edit that would require running instance to reinitialize
 	void IncrementVersion() const { ++Version; };
+	
+	// Returns true if the source IK Rig has been assigned
+	UFUNCTION(BlueprintPure, Category=RetargetProfile)
+	bool HasSourceIKRig() const { return SourceIKRigAsset != nullptr; }
 
-	// BLUEPRINT GETTERS 
+	// Returns true if the target IK Rig has been assigned
+	UFUNCTION(BlueprintPure, Category=RetargetProfile)
+	bool HasTargetIKRig() const { return TargetIKRigAsset != nullptr; }
+
+	// UObject
+	virtual void PostDuplicate(bool bDuplicateForPIE) override;
+	virtual void Serialize(FArchive& Ar) override;
+	virtual void PostLoad() override;
+	// END UObject
+	#if WITH_EDITORONLY_DATA
+    	static void DeclareConstructClasses(TArray<FTopLevelAssetPath>& OutConstructClasses, const UClass* SpecificSubclass);
+    #endif
+
+#if WITH_EDITOR
+	// Get name of Source IK Rig property 
+	static const FName GetSourceIKRigPropertyName();
+	// Get name of Target IK Rig property 
+	static const FName GetTargetIKRigPropertyName();
+	// Get name of Source Preview Mesh property 
+	static const FName GetSourcePreviewMeshPropertyName();
+	// Get name of Target Preview Mesh property 
+	static const FName GetTargetPreviewMeshPropertyName();
+#endif
 
-	// Returns the chain settings associated with a given Goal in an IK Retargeter Asset using the given profile name (optional) 
-	UFUNCTION(BlueprintCallable, BlueprintPure, Category=RetargetAsset)
+	//
+	// BEGIN DEPRECATED API
+	
+	// (5.6 pre ops refactor)
+	//
+	// These were used to get/set Root/Chain/Global settings from outside systems
+	// Because these settings are now scattered across the various ops that perform the equivalent functions,
+	// we cannot provide backwards compatibility.
+	// Instead you can use the new op controller API to get a controller for any retarget op and get/set it's settings.
+	PRAGMA_DISABLE_DEPRECATION_WARNINGS
+	// Get read-only access to the chain mapping
+	UE_DEPRECATED(5.6, "Chain settings are now accessed through an Op controller.")
+	const TArray<TObjectPtr<URetargetChainSettings>>& GetAllChainSettings() const { return ChainSettings_DEPRECATED; };
+	
+	// Get read-only access to the chain map for a given chain (null if chain not in retargeter)
+	UE_DEPRECATED(5.6, "Chain mappings are now accessed from GetChainMapping().")
+	const TObjectPtr<URetargetChainSettings> GetChainMapByName(const FName& TargetChainName) const;
+	
+	// Get read-only access to the chain settings for a given chain (null if chain not in retargeter)
+	UE_DEPRECATED(5.6, "Chain settings are now accessed through an Op controller.")
+	const FTargetChainSettings* GetChainSettingsByName(const FName& TargetChainName) const;
+	
+	// Get access to the root settings
+	UE_DEPRECATED(5.6, "Root settings are now accessed through a Pelvis Motion Op.")
+	URetargetRootSettings* GetRootSettingsUObject() const { return RootSettings_DEPRECATED; };
+	
+	// Get access to the global settings uobject
+	UE_DEPRECATED(5.6, "Global settings are now accessed through various Ops depending on the feature.")
+	UIKRetargetGlobalSettings* GetGlobalSettingsUObject() const { return GlobalSettings_DEPRECATED; };
+	
+	// Get access to the global settings itself
+	UE_DEPRECATED(5.6, "Global settings are now accessed through various Ops depending on the feature.")
+	const FRetargetGlobalSettings& GetGlobalSettings() const { return GlobalSettings_DEPRECATED->Settings; };
+	
+	// Returns the chain settings associated with a given Goal in an IK Retargeter Asset using the given profile name (optional)
+	UE_DEPRECATED(5.6, "Chain settings are now accessed through an Op.")
+	UFUNCTION(BlueprintCallable, BlueprintPure, Category=RetargetAsset, meta = (DeprecatedFunction, DeprecationMessage = "Use IK Chain Op controller to get chains with goals."))
 	static FTargetChainSettings GetChainUsingGoalFromRetargetAsset(
 		const UIKRetargeter* RetargetAsset,
 		const FName IKGoalName);
 	
-	// Returns the chain settings associated with a given target chain in an IK Retargeter Asset using the given profile name (optional) 
-	UFUNCTION(BlueprintCallable, BlueprintPure, Category=RetargetAsset)
+	// Returns the chain settings associated with a given target chain in an IK Retargeter Asset using the given profile name (optional)
+	UE_DEPRECATED(5.6, "Chain settings are now accessed through an Op.")
+	UFUNCTION(BlueprintCallable, BlueprintPure, Category=RetargetAsset, meta = (DeprecatedFunction, DeprecationMessage = "Use op profiles to get/set chain settings."))
 	static FTargetChainSettings GetChainSettingsFromRetargetAsset(
 		const UIKRetargeter* RetargetAsset,
 		const FName TargetChainName,
 		const FName OptionalProfileName);
-
-	// Returns the chain settings associated with a given target chain in the supplied Retarget Profile. 
-	UFUNCTION(BlueprintCallable, BlueprintPure, Category=RetargetProfile)
+	
+	// Returns the chain settings associated with a given target chain in the supplied Retarget Profile.
+	UE_DEPRECATED(5.6, "Chain settings are now accessed through an Op.")
+	UFUNCTION(BlueprintCallable, BlueprintPure, Category=RetargetProfile, meta = (DeprecatedFunction, DeprecationMessage = "Use op profiles to get/set chain settings."))
 	static FTargetChainSettings GetChainSettingsFromRetargetProfile(
 		UPARAM(ref) FRetargetProfile& RetargetProfile,
 		const FName TargetChainName);
-
-	// Returns the root settings in an IK Retargeter Asset using the given profile name (optional) 
-	UFUNCTION(BlueprintCallable, BlueprintPure, Category=RetargetAsset)
+	
+	// Returns the root settings in an IK Retargeter Asset using the given profile name (optional)
+	UE_DEPRECATED(5.6, "Root settings are now accessed through a Pelvis Motion Op.")
+	UFUNCTION(BlueprintCallable, BlueprintPure, Category=RetargetAsset, meta = (DeprecatedFunction, DeprecationMessage = "Use op profiles to get/set pelvis settings."))
 	static void GetRootSettingsFromRetargetAsset(
 		const UIKRetargeter* RetargetAsset,
 		const FName OptionalProfileName,
 		UPARAM(DisplayName = "ReturnValue") FTargetRootSettings& OutSettings);
-
-	// Returns the root settings in the supplied Retarget Profile. 
-	UFUNCTION(BlueprintCallable, BlueprintPure, Category=RetargetProfile)
+	
+	// Returns the root settings in the supplied Retarget Profile.
+	UE_DEPRECATED(5.6, "Root settings are now accessed through a Pelvis Motion Op.")
+	UFUNCTION(BlueprintCallable, BlueprintPure, Category=RetargetProfile, meta = (DeprecatedFunction, DeprecationMessage = "Use op profiles to get/set pelvis settings."))
 	static FTargetRootSettings GetRootSettingsFromRetargetProfile(UPARAM(ref) FRetargetProfile& RetargetProfile);
-
-	// Returns the global settings in an IK Retargeter Asset using the given profile name (optional) 
-	UFUNCTION(BlueprintCallable, BlueprintPure, Category=RetargetAsset)
+	
+	// Returns the global settings in an IK Retargeter Asset using the given profile name (optional)
+	UE_DEPRECATED(5.6, "Global settings are now accessed through various Op.")
+	UFUNCTION(BlueprintCallable, BlueprintPure, Category=RetargetAsset, meta = (DeprecatedFunction, DeprecationMessage = "Use op profiles to get/set settings that previously used global settings."))
 	static void GetGlobalSettingsFromRetargetAsset(
 		const UIKRetargeter* RetargetAsset,
 		const FName OptionalProfileName,
 		UPARAM(DisplayName = "ReturnValue") FRetargetGlobalSettings& OutSettings);
-
-	// Returns the global settings in the supplied Retarget Profile. 
-	UFUNCTION(BlueprintCallable, BlueprintPure, Category=RetargetProfile)
+	
+	// Returns the global settings in the supplied Retarget Profile.
+	UE_DEPRECATED(5.6, "Global settings are now accessed through various Op.")
+	UFUNCTION(BlueprintCallable, BlueprintPure, Category=RetargetProfile, meta = (DeprecatedFunction, DeprecationMessage = "Use op profiles to get/set settings that previously used global settings."))
 	static FRetargetGlobalSettings GetGlobalSettingsFromRetargetProfile(UPARAM(ref) FRetargetProfile& RetargetProfile);
-
-	// Returns true if the source IK Rig has been assigned
-	UFUNCTION(BlueprintPure, Category=RetargetProfile)
-	bool HasSourceIKRig() const { return SourceIKRigAsset != nullptr; }
-
-	// Returns true if the target IK Rig has been assigned
-	UFUNCTION(BlueprintPure, Category=RetargetProfile)
-	bool HasTargetIKRig() const { return TargetIKRigAsset != nullptr; }
 	
-	// BLUEPRINT SETTERS 
-
 	// Set the global settings in a retarget profile (will set bApplyGlobalSettings to true). 
-	UFUNCTION(BlueprintCallable, Category=RetargetProfile)
+	UFUNCTION(BlueprintCallable, Category=RetargetProfile, meta = (DeprecatedFunction, DeprecationMessage = "Use op profiles to get/set settings that previously used global settings."))
+	UE_DEPRECATED(5.6, "Global settings are now accessed through various Op.")
 	static void SetGlobalSettingsInRetargetProfile(
 		UPARAM(ref) FRetargetProfile& RetargetProfile,
 		const FRetargetGlobalSettings& GlobalSettings);
-
-	// Set the root settings in a retarget profile (will set bApplyRootSettings to true). 
-	UFUNCTION(BlueprintCallable, Category=RetargetProfile)
+	
+	// Set the root settings in a retarget profile (will set bApplyRootSettings to true).
+	UE_DEPRECATED(5.6, "Root settings are now accessed through a Pelvis Motion Op.")
+	UFUNCTION(BlueprintCallable, Category=RetargetProfile, meta = (DeprecatedFunction, DeprecationMessage = "Use op profiles to get/set settings that previously used root settings."))
 	static void SetRootSettingsInRetargetProfile(
 		UPARAM(ref) FRetargetProfile& RetargetProfile,
 		const FTargetRootSettings& RootSettings);
 	
-	// Set the chain settings in a retarget profile (will set bApplyChainSettings to true). 
-	UFUNCTION(BlueprintCallable, Category=RetargetProfile)
+	// Set the chain settings in a retarget profile (will set bApplyChainSettings to true).
+	UE_DEPRECATED(5.6, "Chain settings are now accessed through an Op controller.")
+	UFUNCTION(BlueprintCallable, Category=RetargetProfile, meta = (DeprecatedFunction, DeprecationMessage = "Use op profiles to get/set chain settings."))
 	static void SetChainSettingsInRetargetProfile(
 		UPARAM(ref) FRetargetProfile& RetargetProfile,
 		const FTargetChainSettings& ChainSettings,
 		const FName TargetChainName);
-
-	// Set the chain FK settings in a retarget profile (will set bApplyChainSettings to true). 
-	UFUNCTION(BlueprintCallable, Category=RetargetProfile)
+	
+	// Set the chain FK settings in a retarget profile (will set bApplyChainSettings to true).
+	UE_DEPRECATED(5.6, "FK Chain settings are now accessed through an FK Chain Op.")
+	UFUNCTION(BlueprintCallable, Category=RetargetProfile, meta = (DeprecatedFunction, DeprecationMessage = "Use FK Chain Op profiles to get/set FK chain settings."))
 	static void SetChainFKSettingsInRetargetProfile(
 		UPARAM(ref) FRetargetProfile& RetargetProfile,
 		const FTargetChainFKSettings& FKSettings,
 		const FName TargetChainName);
-
-	// Set the chain IK settings in a retarget profile (will set bApplyChainSettings to true). 
-	UFUNCTION(BlueprintCallable, Category=RetargetProfile)
+	
+	// Set the chain IK settings in a retarget profile (will set bApplyChainSettings to true).
+	UE_DEPRECATED(5.6, "IK Chain settings are now accessed through an IK Chain Op.")
+	UFUNCTION(BlueprintCallable, Category=RetargetProfile, meta = (DeprecatedFunction, DeprecationMessage = "Use IK Chain Op profiles to get/set IK chain settings."))
 	static void SetChainIKSettingsInRetargetProfile(
 		UPARAM(ref) FRetargetProfile& RetargetProfile,
 		const FTargetChainIKSettings& IKSettings,
 		const FName TargetChainName);
-
-	// Set the chain Speed Plant settings in a retarget profile (will set bApplyChainSettings to true). 
-	UFUNCTION(BlueprintCallable, Category=RetargetProfile)
+	
+	// Set the chain Speed Plant settings in a retarget profile (will set bApplyChainSettings to true).
+	UE_DEPRECATED(5.6, "Speed Plant settings are now accessed through a Speed Plant Op.")
+	UFUNCTION(BlueprintCallable, Category=RetargetProfile, meta = (DeprecatedFunction, DeprecationMessage = "Use Speed Plant Op profiles to get/set speed plant settings."))
 	static void SetChainSpeedPlantSettingsInRetargetProfile(
 		UPARAM(ref) FRetargetProfile& RetargetProfile,
 		const FTargetChainSpeedPlantSettings& SpeedPlantSettings,
 		const FName TargetChainName);
-
-	// UObject
-	virtual void PostDuplicate(bool bDuplicateForPIE) override;
-	virtual void Serialize(FArchive& Ar) override;
-	virtual void PostLoad() override;
-	// END UObject
-	#if WITH_EDITORONLY_DATA
-    	static void DeclareConstructClasses(TArray<FTopLevelAssetPath>& OutConstructClasses, const UClass* SpecificSubclass);
-    #endif
-
-#if WITH_EDITOR
-	// Get name of Source IK Rig property 
-	static const FName GetSourceIKRigPropertyName();
-	// Get name of Target IK Rig property 
-	static const FName GetTargetIKRigPropertyName();
-	// Get name of Source Preview Mesh property 
-	static const FName GetSourcePreviewMeshPropertyName();
-	// Get name of Target Preview Mesh property 
-	static const FName GetTargetPreviewMeshPropertyName();
-	// Get the names of the all the speed curves the retargeter will be looking for 
-	void GetSpeedCurveNames(TArray<FName>& OutSpeedCurveNames) const;
-#endif
+	PRAGMA_ENABLE_DEPRECATION_WARNINGS
+	//
+	// END DEPRECATED API
+	//
 
 private:
 
+	// upgrade old assets to new data formats
+	void PostLoadOldSettingsToNew();
+	void PostLoadOldOpsToNewStructOps();
+	void PostLoadConvertEverythingToOps();
 	void CleanAndInitialize();
 	
 	// incremented by any edits that require re-initialization
@@ -409,17 +391,7 @@ private:
 public:
 
 #if WITH_EDITORONLY_DATA
-	UPROPERTY()
-	bool bRetargetRoot_DEPRECATED = true;
-	UPROPERTY()
-	bool bRetargetFK_DEPRECATED = true;
-	UPROPERTY()
-	bool bRetargetIK_DEPRECATED = true;
-	UPROPERTY()
-	float TargetActorOffset_DEPRECATED = 0.0f;
-	UPROPERTY()
-	float TargetActorScale_DEPRECATED = 0.0f;
-
+	
 	// The offset applied to the target mesh in the editor viewport. 
 	UPROPERTY(EditAnywhere, Category = PreviewSettings)
 	FVector TargetMeshOffset;
@@ -441,38 +413,6 @@ public:
 	// Toggle debug drawing for retargeting in the viewport. 
 	UPROPERTY(EditAnywhere, Category = DebugSettings)
 	bool bDebugDraw = true;
-
-	// Draw lines on each bone chain. 
-	UPROPERTY(EditAnywhere, Category = DebugSettings)
-	bool bDrawChainLines = true;
-
-	// Draw spheres on single bone chains. 
-	UPROPERTY(EditAnywhere, Category = DebugSettings)
-	bool bDrawSingleBoneChains = false;
-	
-	// Draw final IK goal locations. 
-	UPROPERTY(EditAnywhere, Category = DebugSettings)
-	bool bDrawFinalGoals = true;
-
-	// Draw goal locations from source skeleton. 
-	UPROPERTY(EditAnywhere, Category = DebugSettings)
-	bool bDrawSourceLocations = true;
-
-	// Draw circle on the floor below the retarget root. 
-	UPROPERTY(EditAnywhere, Category = DebugSettings)
-	bool bDrawRootCircle = true;
-
-	// Draw coordinate frame used to define stride warping directions.  
-	UPROPERTY(EditAnywhere, Category = DebugSettings)
-	bool bDrawWarpingFrame = false;
-	
-	// The visual size of the IK goals in the viewport. 
-	UPROPERTY(EditAnywhere, Category = DebugSettings)
-	float ChainDrawSize = 5.0f;
-
-	// The thickness of lines on the IK goals in the viewport. 
-	UPROPERTY(EditAnywhere, Category = DebugSettings)
-	float ChainDrawThickness = 0.5f;
 	
 	// The visual size of the bones in the viewport (saved between sessions). This is set from the viewport Character>Bones menu
 	UPROPERTY()
@@ -490,34 +430,21 @@ private:
 	
 private:
 
-	// (OLD VERSION) Mapping of chains to copy animation between source and target rigs.
-	PRAGMA_DISABLE_DEPRECATION_WARNINGS
-	UPROPERTY()
-	TArray<FRetargetChainMap> ChainMapping_DEPRECATED;
-	PRAGMA_ENABLE_DEPRECATION_WARNINGS
-	
-	// Settings for how to map source chains to target chains.
-	UPROPERTY()
-	TArray<TObjectPtr<URetargetChainSettings>> ChainSettings;
-	
-	// the retarget root settings 
-	UPROPERTY()
-	TObjectPtr<URetargetRootSettings> RootSettings;
-
-	// the retarget global settings 
-	UPROPERTY()
-	TObjectPtr<UIKRetargetGlobalSettings> GlobalSettings;
-
-	/** polymorphic stack of retargeting operations, executed in serial fashion where output of prior operation is input to the next */
+	/** polymorphic stack of retargeting operations
+	 * executed in serial fashion where output of prior operation is input to the next */
 	UPROPERTY(meta = (ExcludeBaseStruct, BaseStruct = "/Script/IKRig.RetargetOpBase"))
 	TArray<FInstancedStruct> RetargetOps;
 
+	/** mapping of target to source bone chains by name */
+	UPROPERTY()
+	FRetargetChainMapping ChainMap;
+
 	// settings profiles stored in this asset 
 	UPROPERTY()
 	TMap<FName, FRetargetProfile> Profiles;
 	UPROPERTY()
 	FName CurrentProfile = NAME_None;
-	
+
 	// The set of retarget poses for the SOURCE skeleton.
 	UPROPERTY()
 	TMap<FName, FIKRetargetPose> SourceRetargetPoses;
@@ -532,19 +459,47 @@ private:
 	UPROPERTY()
 	FName CurrentTargetRetargetPose;
 
-	// (OLD VERSION) Before retarget poses were stored for target AND source.
-	PRAGMA_DISABLE_DEPRECATION_WARNINGS
-	UPROPERTY()
-	TMap<FName, FIKRetargetPose> RetargetPoses_DEPRECATED;
-	UPROPERTY()
-	FName CurrentRetargetPose_DEPRECATED;
-	PRAGMA_ENABLE_DEPRECATION_WARNINGS
-
-	// 5.6 refactor to remove UObjects from retargeter
+	//
+	// BEGIN DEPRECATED DATA
+	//
 	PRAGMA_DISABLE_DEPRECATION_WARNINGS
+	UPROPERTY(meta = (DeprecatedProperty))
+	bool bRetargetRoot_DEPRECATED = true;
+	UPROPERTY(meta = (DeprecatedProperty))
+	bool bRetargetFK_DEPRECATED = true;
+	UPROPERTY(meta = (DeprecatedProperty))
+	bool bRetargetIK_DEPRECATED = true;
+	UPROPERTY(meta = (DeprecatedProperty))
+	float TargetActorOffset_DEPRECATED = 0.0f;
+	UPROPERTY(meta = (DeprecatedProperty))
+	float TargetActorScale_DEPRECATED = 0.0f;
+	// (OLD VERSION) Before retarget poses were stored for target AND source.
+	UPROPERTY(meta = (DeprecatedProperty))
+	TMap<FName, FIKRetargetPose> RetargetPoses;
+	UPROPERTY(meta = (DeprecatedProperty))
+	FName CurrentRetargetPose;
+	//
+	// BEGIN deprecated data from 5.6 refactor
+	//
+	// Settings for how to map source chains to target chains.
+	UPROPERTY(meta = (DeprecatedProperty, DeprecationMessage = "Use FK Chain and IK Chain Op settings instead to affect chain settings."))
+	TArray<TObjectPtr<URetargetChainSettings>> ChainSettings_DEPRECATED;
+	// the retarget root settings 
+	UPROPERTY(meta = (DeprecatedProperty, DeprecationMessage = "Use Pelvis Op settings to affect the pelvis motion instead."))
+	TObjectPtr<URetargetRootSettings> RootSettings_DEPRECATED;
+	// the retarget global settings 
+	UPROPERTY(meta = (DeprecatedProperty, DeprecationMessage = "Toggled 'phases' is now done by enable/disabling ops. Warping now in Stride Warp Op."))
+	TObjectPtr<UIKRetargetGlobalSettings> GlobalSettings_DEPRECATED;
+	// the stack of UObject-based ops
 	UPROPERTY(meta = (DeprecatedProperty, DeprecationMessage = "The pre 5.6 stack of UObject based solvers. Use RetargetOps instead."))
 	TObjectPtr<UDEPRECATED_RetargetOpStack> OpStack_DEPRECATED;
+	//
+	// END deprecated data from 5.6 refactor
+	//
 	PRAGMA_ENABLE_DEPRECATION_WARNINGS
+	//
+	// END DEPRECATED DATA
+	//
 
 	friend class UIKRetargeterController;
 };
diff --git a/Engine/Plugins/Animation/IKRig/Source/IKRig/Public/Retargeter/RetargetOps/AlignPoleVectorOp.h b/Engine/Plugins/Animation/IKRig/Source/IKRig/Public/Retargeter/RetargetOps/AlignPoleVectorOp.h
new file mode 100644
--- /dev/null
+++ b/Engine/Plugins/Animation/IKRig/Source/IKRig/Public/Retargeter/RetargetOps/AlignPoleVectorOp.h
@@ -0,0 +1,155 @@
+﻿// Copyright Epic Games, Inc. All Rights Reserved.
+#pragma once
+
+#include "Retargeter/IKRetargetOps.h"
+
+#include "AlignPoleVectorOp.generated.h"
+
+#define LOCTEXT_NAMESPACE "AlignPoleVectorOp"
+
+struct FResolvedBoneChain;
+
+struct FRetargetPoleVectorSettings;
+
+struct FPoleVectorMatcher
+{
+	EAxis::Type SourcePoleAxis;
+	EAxis::Type TargetPoleAxis;
+	float TargetToSourceAngularOffsetAtRefPose;
+	TArray<int32> AllChildrenWithinChain;
+	const FRetargetPoleVectorSettings* Settings;
+	const FResolvedBoneChain* SourceBoneChain;
+	const FResolvedBoneChain* TargetBoneChain;
+
+	bool Initialize(
+		const FRetargetPoleVectorSettings* InSettings,
+		const FResolvedBoneChain& InSourceBoneChain,
+		const FResolvedBoneChain& InTargetBoneChain,
+		const FRetargetSkeleton& InSourceSkeleton,
+		const FRetargetSkeleton& InTargetSkeleton);
+
+	void MatchPoleVector(
+		const TArray<FTransform> &SourceGlobalPose,
+		const FRetargetSkeleton& TargetSkeleton,
+		TArray<FTransform> &OutTargetGlobalPose);
+	
+private:
+
+	EAxis::Type CalculateBestPoleAxisForChain(
+		const TArray<int32>& BoneIndices,
+		const TArray<FTransform>& GlobalPose);
+	
+	static FVector CalculatePoleVector(
+		const EAxis::Type& PoleAxis,
+		const TArray<int32>& BoneIndices,
+		const TArray<FTransform>& GlobalPose);
+
+	static EAxis::Type GetMostDifferentAxis(
+		const FTransform& Transform,
+		const FVector& InNormal);
+
+	static FVector GetChainAxisNormalized(
+		const TArray<int32>& BoneIndices,
+		const TArray<FTransform>& GlobalPose);
+};
+
+USTRUCT(BlueprintType)
+struct IKRIG_API FRetargetPoleVectorSettings
+{
+	GENERATED_BODY()
+
+	FRetargetPoleVectorSettings() = default;
+	FRetargetPoleVectorSettings(FName InTargetChainName) : TargetChainName(InTargetChainName) {}
+	
+	/** The name of the target chain to align pole vectors on. */
+	UPROPERTY(EditAnywhere, BlueprintReadWrite, Category = "Setup", meta=(ReinitializeOnEdit))
+	FName TargetChainName;
+
+	/** Range 0 to 1. Default 1. Matches the twist angle of this chain (along the Pole direction) to the source chain.
+	*  At 0, the chain's pole vector direction will be left alone
+	*  At 1, the root bone of the chain will be twist-rotated in the pole direction to match the orientation of the source chain.*/
+	UPROPERTY(EditAnywhere, BlueprintReadWrite, Category = "Align Pole Vector", meta = (UIMin = "0.0", UIMax = "1.0", ClampMin = "0.0", ClampMax = "1.0"))
+	double AlignAlpha = 1.0f;
+
+	/** Range +/- 180. Default 0. An angular offset, in degrees, for the pole direction of the chain. */
+	UPROPERTY(EditAnywhere, BlueprintReadWrite, Category = "Align Pole Vector", meta = (UIMin = "-180.0", UIMax = "180.0", ClampMin = "-180.0", ClampMax = "180.0"))
+	double StaticAngularOffset = 0.0f;
+
+	/** Default is False. When true, the original angular offset between the source/target pole vectors will be maintained when aligning pole vectors. */
+	UPROPERTY(EditAnywhere, BlueprintReadWrite, Category = "Align Pole Vector")
+	bool MaintainOffset = false;
+	
+	bool operator==(const FRetargetPoleVectorSettings& Other) const;
+};
+
+USTRUCT(BlueprintType, meta = (DisplayName = "Pole Vector Alignment Settings"))
+struct IKRIG_API FIKRetargetAlignPoleVectorOpSettings : public FIKRetargetOpSettingsBase
+{
+	GENERATED_BODY()
+	
+	UPROPERTY(EditAnywhere, BlueprintReadWrite, Category = "Pole Vector Op Settings")
+	TArray<FRetargetPoleVectorSettings> ChainsToAlign;
+
+	void MergePoleVectorSettings(const FRetargetPoleVectorSettings& InSettingsToMerge);
+};
+
+USTRUCT(BlueprintType, meta = (DisplayName = "Pole Vector Alignment"))
+struct IKRIG_API FIKRetargetAlignPoleVectorOp : public FIKRetargetOpBase
+{
+	GENERATED_BODY()
+	
+	virtual bool Initialize(
+		const FIKRetargetProcessor& InProcessor,
+		const FRetargetSkeleton& InSourceSkeleton,
+		const FTargetSkeleton& InTargetSkeleton,
+		FIKRigLogger& InLog) override;
+	
+	virtual void Run(
+		FIKRetargetProcessor& InProcessor,
+		const double InDeltaTime,
+		const TArray<FTransform>& InSourceGlobalPose,
+		TArray<FTransform>& OutTargetGlobalPose) override;
+
+	virtual FIKRetargetOpSettingsBase* GetSettings() override;
+
+	virtual void SetSettings(const FIKRetargetOpSettingsBase* InSettings) override;
+	
+	virtual const UScriptStruct* GetSettingsType() const override;
+
+	virtual const UScriptStruct* GetType() const override;
+
+	virtual const UClass* GetControllerType() const override;
+
+#if WITH_EDITOR
+	virtual FText GetWarningMessage() const override;
+
+	virtual void SetChainSettingsToDefault(const FName& InChainName) override;
+	virtual bool AreChainSettingsAtDefault(const FName& InChainName) override;
+#endif
+
+	UPROPERTY()
+	FIKRetargetAlignPoleVectorOpSettings Settings;
+
+private:
+	TArray<FPoleVectorMatcher> PoleVectorMatchers;
+};
+
+/* The blueprint/python API for editing an Align Pole Vector Op */
+UCLASS(BlueprintType)
+class IKRIG_API UIKRetargetAlignPoleVectorController : public UIKRetargetOpControllerBase
+{
+	GENERATED_BODY()
+	
+public:
+	/* Get the current op settings as a struct.
+	 * @return FIKRetargetAlignPoleVectorOpSettings struct with the current settings used by the op. */
+	UFUNCTION(BlueprintCallable, Category = Settings)
+	FIKRetargetAlignPoleVectorOpSettings GetSettings();
+
+	/* Set the op settings. Input is a custom struct type for this op.
+	 * @param InSettings a FIKRetargetAlignPoleVectorOpSettings struct containing all the settings to apply to this op */
+	UFUNCTION(BlueprintCallable, Category = Settings)
+	void SetSettings(FIKRetargetAlignPoleVectorOpSettings InSettings);
+};
+
+#undef LOCTEXT_NAMESPACE
\ No newline at end of file
diff --git a/Engine/Plugins/Animation/IKRig/Source/IKRig/Public/Retargeter/RetargetOps/CopyBasePoseOp.h b/Engine/Plugins/Animation/IKRig/Source/IKRig/Public/Retargeter/RetargetOps/CopyBasePoseOp.h
new file mode 100644
--- /dev/null
+++ b/Engine/Plugins/Animation/IKRig/Source/IKRig/Public/Retargeter/RetargetOps/CopyBasePoseOp.h
@@ -0,0 +1,78 @@
+﻿// Copyright Epic Games, Inc. All Rights Reserved.
+#pragma once
+
+#include "Retargeter/IKRetargetOps.h"
+
+#include "CopyBasePoseOp.generated.h"
+
+#define LOCTEXT_NAMESPACE "CopyBasePoseOp"
+
+USTRUCT(BlueprintType, meta = (DisplayName = "Scale Source Settings"))
+struct IKRIG_API FIKRetargetCopyBasePoseOpSettings : public FIKRetargetOpSettingsBase
+{
+	GENERATED_BODY()
+	
+	/** When true, will copy all the source bone transforms with matching names to use as a base pose. This can be useful for partial retargeting.
+	 * NOTE: no retargeting is applied to the bone transforms, they are assumed to be fully compatible between source/target skeletons.
+	 */
+	UPROPERTY(EditAnywhere, BlueprintReadWrite, Category = RetargetPhases)
+	bool bCopyBasePose = false;
+
+	/** Filters the bones to copy when using "Copy Base Pose". If specified, will only copy all children of the specified bone (inclusive).*/
+	UPROPERTY(EditAnywhere, BlueprintReadWrite, Category = RetargetPhases, meta = (EditCondition = "bCopyBasePose"))
+	FName CopyBasePoseRoot = NAME_None;
+};
+
+USTRUCT(BlueprintType, meta = (DisplayName = "Copy Base Pose"))
+struct IKRIG_API FIKRetargetCopyBasePoseOp : public FIKRetargetOpBase
+{
+	GENERATED_BODY()
+	
+	// NOTE: this op does not do anything in Initialize() or Run().
+	// It is a special case op that the retargeter reads from when it needs to copy the base pose
+	
+	virtual bool Initialize(
+		const FIKRetargetProcessor& InProcessor,
+		const FRetargetSkeleton& InSourceSkeleton,
+		const FTargetSkeleton& InTargetSkeleton,
+		FIKRigLogger& Log) override;
+	
+	virtual void Run(
+		FIKRetargetProcessor& InProcessor,
+		const double InDeltaTime,
+		const TArray<FTransform>& InSourceGlobalPose,
+		TArray<FTransform>& OutTargetGlobalPose) override {};
+
+	virtual FIKRetargetOpSettingsBase* GetSettings() override;
+
+	virtual void SetSettings(const FIKRetargetOpSettingsBase* InSettings) override;
+	
+	virtual const UScriptStruct* GetSettingsType() const override;
+	
+	virtual const UScriptStruct* GetType() const override;
+	
+	virtual const UClass* GetControllerType() const override;
+
+	UPROPERTY()
+	FIKRetargetCopyBasePoseOpSettings Settings;
+};
+
+/* The blueprint/python API for editing a Copy Base Pose Op */
+UCLASS(BlueprintType)
+class IKRIG_API UIKRetargetCopyBasePoseController : public UIKRetargetOpControllerBase
+{
+	GENERATED_BODY()
+	
+public:
+	/* Get the current op settings as a struct.
+	 * @return FIKRetargetCopyBasePoseOpSettings struct with the current settings used by the op. */
+	UFUNCTION(BlueprintCallable, Category = Settings)
+	FIKRetargetCopyBasePoseOpSettings GetSettings();
+
+	/* Set the op settings. Input is a custom struct type for this op.
+	 * @param InSettings a FIKRetargetCopyBasePoseOpSettings struct containing all the settings to apply to this op */
+	UFUNCTION(BlueprintCallable, Category = Settings)
+	void SetSettings(FIKRetargetCopyBasePoseOpSettings InSettings);
+};
+
+#undef LOCTEXT_NAMESPACE
\ No newline at end of file
diff --git a/Engine/Plugins/Animation/IKRig/Source/IKRig/Public/Retargeter/RetargetOps/CurveRemapOp.h b/Engine/Plugins/Animation/IKRig/Source/IKRig/Public/Retargeter/RetargetOps/CurveRemapOp.h
--- a/Engine/Plugins/Animation/IKRig/Source/IKRig/Public/Retargeter/RetargetOps/CurveRemapOp.h
+++ b/Engine/Plugins/Animation/IKRig/Source/IKRig/Public/Retargeter/RetargetOps/CurveRemapOp.h
@@ -2,30 +2,48 @@
 #pragma once
 
 #include "Retargeter/IKRetargetOps.h"
+#include "StructUtils/InstancedStruct.h"
+#include "Animation/AnimCurveTypes.h"
+
 #include "CurveRemapOp.generated.h"
 
 #define LOCTEXT_NAMESPACE "CurveRemapOp"
 
 USTRUCT(BlueprintType)
+struct FCurveRemapPair
+{
+	GENERATED_BODY()
+	
+	// The curve name on the SOURCE skeletal mesh to copy animation data from.
+	UPROPERTY(EditAnywhere, Category=Settings)
+	FName SourceCurve;
+
+	// The curve name on the TARGET skeletal mesh to receive animation data.
+	UPROPERTY(EditAnywhere, Category=Settings)
+	FName TargetCurve;
+};
+
+USTRUCT(BlueprintType, meta = (DisplayName = "Remap Curve Settings"))
 struct IKRIG_API FIKRetargetCurveRemapOpSettings : public FIKRetargetOpSettingsBase
 {
 	GENERATED_BODY()
 
+	// Whether to copy all curves over to the target animation instance
+	// NOTE: This setting also applies when exporting retargeted animations.
+	// True: all source curves are copied to the target animation instance/asset
+	// False: only remapped curves are copied on the target animation instance/asset
+	UPROPERTY(EditAnywhere, Category="OpSettings")
+	bool bCopyAllSourceCurves = true;
+	
 	// Add pairs of Source/Target curve names to remap. While retargeting, the animation from the source curves
 	// will be redirected to the curves on the target skeletal meshes. Can be used to drive, blendshapes or other downstream systems.
 	// NOTE: By default the IK Retargeter will automatically copy all equivalently named curves from the source to the
-	// target. Remapping with this op is only necessary when the target curve name(s) are different.
+	// target. Remapping is only necessary when the target curve name(s) are different.
 	UPROPERTY(EditAnywhere, Category="OpSettings")
 	TArray<FCurveRemapPair> CurvesToRemap;
-
-	// This setting only applies to all curves when exporting retargeted animations.
-	// True: all source curves are copied to the target animation sequence asset.
-	// False: only remapped curves are left on the target animation sequence asset.
-	UPROPERTY(EditAnywhere, Category="OpSettings")
-	bool bCopyAllSourceCurves = true;
 };
 
-USTRUCT(BlueprintType)
+USTRUCT(BlueprintType, meta = (DisplayName = "Remap Curves"))
 struct IKRIG_API FIKRetargetCurveRemapOp : public FIKRetargetOpBase
 {
 	GENERATED_BODY()
@@ -33,13 +51,20 @@ struct IKRIG_API FIKRetargetCurveRemapOp : public FIKRetargetOpBase
 	// NOTE: this op does not do anything in Initialize() or Run().
 	// It implements GetCurvesToRetarget() which the retargeting anim node calls and manages
 	// In the future we may remove this coupling and make the op itself do the work via callbacks
+	
 	virtual bool Initialize(
-	const FIKRetargetProcessor& Processor,
-		const FRetargetSkeleton& SourceSkeleton,
-		const FTargetSkeleton& TargetSkeleton,
-		FIKRigLogger& Log) override { bIsInitialized = true; return true; };
+		const FIKRetargetProcessor& InProcessor,
+		const FRetargetSkeleton& InSourceSkeleton,
+		const FTargetSkeleton& InTargetSkeleton,
+		FIKRigLogger& Log) override
+	{
+		bIsInitialized = true;
+		return true;
+	};
+	
 	virtual void Run(
-		const FIKRetargetProcessor& Processor,
+		FIKRetargetProcessor& InProcessor,
+		const double InDeltaTime,
 		const TArray<FTransform>& InSourceGlobalPose,
 		TArray<FTransform>& OutTargetGlobalPose) override {};
 
@@ -47,25 +72,51 @@ struct IKRIG_API FIKRetargetCurveRemapOp : public FIKRetargetOpBase
 	{
 		return &Settings;
 	};
+
+	virtual void SetSettings(const FIKRetargetOpSettingsBase* InSettings) override
+	{
+		Settings = *reinterpret_cast<const FIKRetargetCurveRemapOpSettings*>(InSettings);
+	};
 	
 	virtual const UScriptStruct* GetSettingsType() const override
 	{
 		return FIKRetargetCurveRemapOpSettings::StaticStruct();
 	}
-
-	virtual void GetCurvesToRetarget(TArray<FCurveRemapPair>& OutCurvesToRetarget) override
+	
+	virtual const UScriptStruct* GetType() const override
 	{
-		OutCurvesToRetarget.Append(Settings.CurvesToRemap);
-	};
+		return FIKRetargetCurveRemapOp::StaticStruct();
+	}
+
+	virtual void AnimGraphPreUpdateMainThread(USkeletalMeshComponent& SourceMeshComponent, USkeletalMeshComponent& TargetMeshComponent) override;
+
+	virtual void AnimGraphEvaluateAnyThread(FPoseContext& Output) override;
+
+	virtual const UClass* GetControllerType() const override;
 
 	UPROPERTY()
 	FIKRetargetCurveRemapOpSettings Settings;
 	
-#if WITH_EDITOR
-	virtual FText GetNiceName() const override { return FText(LOCTEXT("RemapOpName", "Remap Curves")); };
-	virtual FText GetWarningMessage() const override { return Message; };
-	FText Message;
-#endif
+	// cached curves, copied on the game thread in PreUpdate()
+	FBlendedHeapCurve SourceCurves;
+};
+
+/* The blueprint/python API for editing a Curve Remap Op */
+UCLASS(BlueprintType)
+class IKRIG_API UIKRetargetCurveRemapController : public UIKRetargetOpControllerBase
+{
+	GENERATED_BODY()
+	
+public:
+	/* Get the current op settings as a struct.
+	 * @return FIKRetargetCurveRemapOpSettings struct with the current settings used by the op. */
+	UFUNCTION(BlueprintCallable, Category = Settings)
+	FIKRetargetCurveRemapOpSettings GetSettings();
+
+	/* Set the op settings. Input is a custom struct type for this op.
+	 * @param InSettings a FIKRetargetCurveRemapOpSettings struct containing all the settings to apply to this op */
+	UFUNCTION(BlueprintCallable, Category = Settings)
+	void SetSettings(FIKRetargetCurveRemapOpSettings InSettings);
 };
 
 //
diff --git a/Engine/Plugins/Animation/IKRig/Source/IKRig/Public/Retargeter/RetargetOps/IKChainsOp.h b/Engine/Plugins/Animation/IKRig/Source/IKRig/Public/Retargeter/RetargetOps/IKChainsOp.h
new file mode 100644
--- /dev/null
+++ b/Engine/Plugins/Animation/IKRig/Source/IKRig/Public/Retargeter/RetargetOps/IKChainsOp.h
@@ -0,0 +1,263 @@
+﻿// Copyright Epic Games, Inc. All Rights Reserved.
+#pragma once
+
+
+#include "Retargeter/IKRetargetOps.h"
+#include "IKChainsOp.generated.h"
+
+#define LOCTEXT_NAMESPACE "IKChainOp"
+
+struct FResolvedBoneChain;
+struct FIKRetargetPelvisMotionOp;
+
+struct FSourceChainIK
+{
+	bool Initialize(const FResolvedBoneChain& InSourceBoneChain, FIKRigLogger& InLog);
+
+	const FResolvedBoneChain* SourceBoneChain;
+	int32 StartBoneIndex = INDEX_NONE;
+	int32 EndBoneIndex = INDEX_NONE;
+	
+	FVector InitialEndPosition = FVector::ZeroVector;
+	FQuat InitialEndRotation = FQuat::Identity;
+	double InvInitialLength = 1.0f;
+
+	// results after encoding...
+	FVector PreviousEndPosition = FVector::ZeroVector;
+	FVector CurrentEndPosition = FVector::ZeroVector;
+	FVector CurrentEndDirectionNormalized = FVector::ZeroVector;
+	FQuat CurrentEndRotation = FQuat::Identity;
+	double CurrentHeightFromGroundNormalized = 0.0f;
+};
+
+struct FTargetChainIK
+{
+	bool Initialize(const FResolvedBoneChain& InTargetBoneChain, FIKRigLogger& InLog);
+
+	const FResolvedBoneChain* TargetBoneChain;
+	int32 BoneIndexA = INDEX_NONE;
+	int32 BoneIndexC = INDEX_NONE;
+	
+	double InitialLength = 1.0f;
+	FVector InitialEndPosition = FVector::ZeroVector;
+	FQuat InitialEndRotation = FQuat::Identity;
+	FVector PrevEndPosition = FVector::ZeroVector;
+};
+
+struct FDecodedIKChain
+{
+	FVector EndEffectorPosition = FVector::ZeroVector;
+	FQuat EndEffectorRotation = FQuat::Identity;
+};
+
+struct FIKChainRetargeter
+{
+	bool Initialize(
+		const FResolvedBoneChain& InSourceBoneChain,
+		const FResolvedBoneChain& InTargetBoneChain,
+		const FRetargetIKChainSettings& InSettings,
+		FIKRigLogger& InLog);
+
+	void EncodePose(const TArray<FTransform> &SourceInputGlobalPose);
+	
+	void DecodePose(
+		const FIKRetargetPelvisMotionOp* PelvisMotionOp,
+		const TArray<FTransform>& InGlobalPose);
+	
+	const FDecodedIKChain& GetResults() const { return Results; };
+
+	const FSourceChainIK& GetSource() const {return Source; };
+
+	const FTargetChainIK& GetTarget() const {return Target; };
+	
+	const FResolvedBoneChain& GetTargetChain() const { return *Target.TargetBoneChain; };
+	
+	const FRetargetIKChainSettings* GetSettings() const { return Settings; };
+
+private:
+
+	FSourceChainIK Source;
+	FTargetChainIK Target;
+	FDecodedIKChain Results;
+	const FRetargetIKChainSettings* Settings;
+};
+
+USTRUCT(BlueprintType)
+struct IKRIG_API FRetargetIKChainSettings
+{
+	GENERATED_BODY()
+
+	FRetargetIKChainSettings() = default;
+	FRetargetIKChainSettings(const FName InTargetChainName) : TargetChainName(InTargetChainName) {}
+
+	/** The name of the TARGET chain to transfer animation onto. */
+	UPROPERTY(EditAnywhere, BlueprintReadWrite, Category = "Setup", meta=(ReinitializeOnEdit))
+	FName TargetChainName;
+	
+	/** Whether to modify the location of the IK goal on this chain. Default is true.
+	 * NOTE: This only has an effect if the chain has an IK Goal assigned to it in the Target IK Rig asset.
+	 * NOTE: If off, and this chain has an IK Goal, the IK will still be evaluated, but the Goal is set to the input bone location (from the FK pass).*/
+	UPROPERTY(EditAnywhere, BlueprintReadWrite, Category = "Setup", meta=(ReinitializeOnEdit))
+	bool EnableIK = true;
+
+	/** Range 0 to 1. Default 0. Blends IK goal transform from retargeted transform (0) to source bone transform (1).
+	*  At 0 the goal is placed at the retargeted location and rotation.
+	*  At 1 the goal is placed at the location and rotation of the source chain's end bone. */
+	UPROPERTY(EditAnywhere, BlueprintReadWrite, Category = "Blend to Source", meta = (ClampMin = "0.0", ClampMax = "1.0", UIMin = "0.0", UIMax = "1.0"))
+	double BlendToSource = 0.0f;
+
+	/** Range 0 to 1. Default 1. Blends the translational component of BlendToSource on/off.
+	*  At 0 the goal is placed at the retargeted location.
+	*  At 1 the goal is placed at the location of the source chain's end bone. */
+	UPROPERTY(EditAnywhere, BlueprintReadWrite, Category = "Blend to Source", meta = (ClampMin = "0.0", ClampMax = "1.0", UIMin = "0.0", UIMax = "1.0"))
+	double BlendToSourceTranslation = 1.0f;
+
+	/** Range 0 to 1. Default 0. Blends the rotational component of BlendToSource on/off.
+	*  At 0 the goal is oriented to the retargeted rotation.
+	*  At 1 the goal is oriented to the source chain's end bone rotation. */
+	UPROPERTY(EditAnywhere, BlueprintReadWrite, Category = "Blend to Source", meta = (ClampMin = "0.0", ClampMax = "1.0", UIMin = "0.0", UIMax = "1.0"))
+	double BlendToSourceRotation = 0.0f;
+	
+	/** Range 0 to 1. Default 1. Weight each axis separately when using Blend To Source.
+	*  At 0 the goal is placed at the retargeted location.
+	*  At 1 the goal is placed at the location of the source chain's end bone. */
+	UPROPERTY(EditAnywhere, BlueprintReadWrite, Category = "Blend to Source", meta = (ClampMin = "0.0", ClampMax = "1.0", UIMin = "0.0", UIMax = "1.0"))
+	FVector BlendToSourceWeights = FVector::OneVector;
+
+	/** Default 0, 0, 0. Apply a static global-space offset to IK goal position. */
+	UPROPERTY(EditAnywhere, BlueprintReadWrite, Category = "Static Offset")
+	FVector StaticOffset = FVector::ZeroVector;
+
+	/** Default 0, 0, 0. Apply a static local-space offset to IK goal position. */
+	UPROPERTY(EditAnywhere, BlueprintReadWrite, Category = "Static Offset")
+	FVector StaticLocalOffset = FVector::ZeroVector;
+
+	/** Default 0, 0, 0. Apply a static local-space offset to IK goal rotation. */
+	UPROPERTY(EditAnywhere, BlueprintReadWrite, Category = "Static Offset")
+	FRotator StaticRotationOffset = FRotator::ZeroRotator;
+
+	/** Range 0 to infinity. Default 1. Scales the vertical component of the IK goal's position. */
+	UPROPERTY(EditAnywhere, BlueprintReadWrite, Category = "Scale Translation", meta = (ClampMin = "0.0", UIMin = "0.0", UIMax = "5.0"))
+	double ScaleVertical = 1.0f;
+	
+	/** Range 0 to 5. Default 1. Brings IK goal closer (0) or further (1+) from origin of chain.
+	*  At 0 the effector is placed at the origin of the chain (ie Shoulder, Hip etc).
+	*  At 1 the effector is left at the end of the chain (ie Wrist, Foot etc)
+	*  Values in-between 0-1 will slide the effector along the vector from the start to the end of the chain.
+	*  Values greater than 1 will stretch the chain beyond the retargeted length. */
+	UPROPERTY(EditAnywhere, BlueprintReadWrite, Category = "Scale Length", meta = (ClampMin = "0.0", ClampMax = "5.0", UIMin = "0.1", UIMax = "2.0"))
+	double Extension = 1.0f;
+
+	bool operator==(const FRetargetIKChainSettings& Other) const;
+};
+
+USTRUCT(BlueprintType, meta = (DisplayName = "Retarget IK Goals Settings"))
+struct IKRIG_API FIKRetargetIKChainsOpSettings : public FIKRetargetOpSettingsBase
+{
+	GENERATED_BODY()
+	
+	UPROPERTY(EditAnywhere, BlueprintReadWrite, Category = Chains, meta=(ReinitializeOnEdit))
+	TArray<FRetargetIKChainSettings> ChainsToRetarget;
+	
+	// Draw final IK goal locations. 
+	UPROPERTY(EditAnywhere, Category = DebugDraw)
+	bool bDrawFinalGoals = true;
+
+	// Draw goal locations from source skeleton. 
+	UPROPERTY(EditAnywhere, Category = DebugDraw)
+	bool bDrawSourceLocations = true;
+
+	// Adjust size of goal debug drawing in viewport
+	UPROPERTY(EditAnywhere, Category = DebugDraw)
+	double GoalDrawSize = 5.0;
+	
+	// Adjust thickness of goal debug drawing in viewport
+	UPROPERTY(EditAnywhere, Category = DebugDraw)
+	double GoalDrawThickness = 1.0;
+};
+
+#if WITH_EDITOR
+struct FChainDebugData
+{
+	FName TargetChainName;
+	FTransform InputTransformStart = FTransform::Identity;
+	FTransform InputTransformEnd = FTransform::Identity;
+	FTransform OutputTransformEnd = FTransform::Identity;
+	FTransform SourceTransformEnd = FTransform::Identity;
+};
+#endif
+
+USTRUCT(BlueprintType, meta = (DisplayName = "Retarget IK Goals"))
+struct IKRIG_API FIKRetargetIKChainsOp : public FIKRetargetOpBase
+{
+	GENERATED_BODY()
+	
+	virtual bool Initialize(
+		const FIKRetargetProcessor& InProcessor,
+		const FRetargetSkeleton& InSourceSkeleton,
+		const FTargetSkeleton& InTargetSkeleton,
+		FIKRigLogger& InLog) override;
+	
+	virtual void Run(
+		FIKRetargetProcessor& InProcessor,
+		const double InDeltaTime,
+		const TArray<FTransform>& InSourceGlobalPose,
+		TArray<FTransform>& OutTargetGlobalPose) override;
+
+	virtual FIKRetargetOpSettingsBase* GetSettings() override;
+	
+	virtual void SetSettings(const FIKRetargetOpSettingsBase* InSettings) override;
+	
+	virtual const UScriptStruct* GetSettingsType() const override;
+
+	virtual const UScriptStruct* GetType() const override;
+
+	virtual const UClass* GetControllerType() const override;
+
+	UPROPERTY()
+	FIKRetargetIKChainsOpSettings Settings;
+
+#if WITH_EDITOR
+public:
+	virtual void RunInitialSetup(const UIKRetargeter* InRetargetAsset) override;
+	virtual void DebugDraw(
+		FPrimitiveDrawInterface* InPDI,
+		const FTransform& InComponentTransform,
+		const double InComponentScale,
+		const FIKRetargetDebugDrawState& InEditorState) const override;
+	void SaveDebugData(
+		const FIKRetargetProcessor& InProcessor,
+		const TArray<FTransform>& InSourceGlobalPose,
+		const TArray<FTransform>& OutTargetGlobalPose);
+	TArray<FChainDebugData> AllChainsDebugData;
+	FVector DebugRootModification;
+	static FCriticalSection DebugDataMutex;
+
+	virtual void SetChainSettingsToDefault(const FName& InChainName) override;
+	virtual bool AreChainSettingsAtDefault(const FName& InChainName) override;
+#endif
+
+private:
+	
+	TArray<FIKChainRetargeter> IKChainRetargeters;
+};
+
+/* The blueprint/python API for editing a IK Chains Op */
+UCLASS(BlueprintType)
+class IKRIG_API UIKRetargetIKChainsController : public UIKRetargetOpControllerBase
+{
+	GENERATED_BODY()
+	
+public:
+	/* Get the current op settings as a struct.
+	 * @return FIKRetargetIKChainsOpSettings struct with the current settings used by the op. */
+	UFUNCTION(BlueprintCallable, Category = Settings)
+	FIKRetargetIKChainsOpSettings GetSettings();
+
+	/* Set the op settings. Input is a custom struct type for this op.
+	 * @param InSettings a FIKRetargetIKChainsOpSettings struct containing all the settings to apply to this op */
+	UFUNCTION(BlueprintCallable, Category = Settings)
+	void SetSettings(FIKRetargetIKChainsOpSettings InSettings);
+};
+
+#undef LOCTEXT_NAMESPACE
diff --git a/Engine/Plugins/Animation/IKRig/Source/IKRig/Public/Retargeter/RetargetOps/PinBoneOp.h b/Engine/Plugins/Animation/IKRig/Source/IKRig/Public/Retargeter/RetargetOps/PinBoneOp.h
--- a/Engine/Plugins/Animation/IKRig/Source/IKRig/Public/Retargeter/RetargetOps/PinBoneOp.h
+++ b/Engine/Plugins/Animation/IKRig/Source/IKRig/Public/Retargeter/RetargetOps/PinBoneOp.h
@@ -2,7 +2,7 @@
 
 #pragma once
 
-#include "Interfaces/Interface_BoneReferenceSkeletonProvider.h"
+#include "Retargeter/IKRetargeter.h"
 #include "Retargeter/IKRetargetOps.h"
 
 #include "PinBoneOp.generated.h"
@@ -30,14 +30,14 @@ struct FPinBoneData
 	void CachePinData(ERetargetSourceOrTarget InSkeletonToCopyFrom, const FIKRetargetProcessor& Processor);
 
 	// The bone to copy FROM.
-	UPROPERTY(EditAnywhere, Category=Settings)
+	UPROPERTY(EditAnywhere, Category=Settings, meta=(ReinitializeOnEdit))
 	FBoneReference BoneToCopyFrom;
 	// The bone to copy TO. Will have its transform pinned to BoneToCopyFrom
-	UPROPERTY(EditAnywhere, Category=Settings)
+	UPROPERTY(EditAnywhere, Category=Settings, meta=(ReinitializeOnEdit))
 	FBoneReference BoneToCopyTo;
 
 	// the scale factor used to calculate all the cached offset (lazy evaluated when scale factor is modified)
-	float SourceScaleFactor;
+	double SourceScaleFactor;
 	// the local transform of the BoneToCopyTo in the ref pose
 	FTransform LocalRefPoseBoneToCopyTo;
 	// the local transform of the BoneToCopyTo in the ref pose
@@ -85,85 +85,90 @@ enum class EPinBoneRotationMode : uint8
 	MaintainOffsetFromBoneToCopyFrom
 };
 
-USTRUCT(BlueprintType)
+USTRUCT(BlueprintType, meta = (DisplayName = "Pin Bone Settings"))
 struct IKRIG_API FIKRetargetPinBoneOpSettings : public FIKRetargetOpSettingsBase
 {
 	GENERATED_BODY()
 	
-	// A list of bone-pairs to copy transforms between
-	UPROPERTY(EditAnywhere, Category=Settings)
+	/** A list of bone-pairs to copy transforms between */
+	UPROPERTY(EditAnywhere, Category=Setup, meta=(ReinitializeOnEdit))
 	TArray<FPinBoneData> BonesToPin;
 
-	// Which skeleton to copy animation from.
+	/** Which skeleton to copy animation from. */
 	UPROPERTY(EditAnywhere, Category=Settings)
 	ERetargetSourceOrTarget SkeletonToCopyFrom = ERetargetSourceOrTarget::Target;
 
-	// Copy the translation of the bone.
+	/** Copy the translation of the bone. */
 	UPROPERTY(EditAnywhere, Category=Settings)
 	bool bCopyTranslation = true;
 
-	// Copy the rotation of the bone.
+	/** Copy the rotation of the bone. */
 	UPROPERTY(EditAnywhere, Category=Settings)
 	bool bCopyRotation = true;
 	
-	// Copy the scale of the bone.
+	/** Copy the scale of the bone. */
 	UPROPERTY(EditAnywhere, Category=Settings)
 	bool bCopyScale = true;
 
-	// The method used to calculate the translation of the bone.
+	/** The method used to calculate the translation of the bone. */
 	UPROPERTY(EditAnywhere, Category=Settings)
 	EPinBoneTranslationMode TranslationMode = EPinBoneTranslationMode::CopyGlobalPosition;
 
-	// The method used to calculate the rotation of the bone.
+	/** The method used to calculate the rotation of the bone. */
 	UPROPERTY(EditAnywhere, Category=Settings)
 	EPinBoneRotationMode RotationMode = EPinBoneRotationMode::CopyGlobalRotation;
 
-	// A manual offset to apply in global space
+	/** A manual offset to apply in global space */
 	UPROPERTY(EditAnywhere, Category=Settings)
 	FTransform GlobalOffset;
 
-	// A manual offset to apply in local space
+	/** A manual offset to apply in local space */
 	UPROPERTY(EditAnywhere, Category=Settings)
 	FTransform LocalOffset;
 
 #if WITH_EDITOR
 	virtual USkeleton* GetSkeleton(const FName InPropertyName) override;
-	const USkeleton* SourceSkeletonAsset;
-	const USkeleton* TargetSkeletonAsset;
 #endif
 };
 
-USTRUCT(BlueprintType)
+USTRUCT(BlueprintType, meta = (DisplayName = "Pin Bones"))
 struct IKRIG_API FIKRetargetPinBoneOp : public FIKRetargetOpBase
 {
 	GENERATED_BODY()
 
 	virtual bool Initialize(
-	const FIKRetargetProcessor& Processor,
+		const FIKRetargetProcessor& Processor,
 		const FRetargetSkeleton& SourceSkeleton,
 		const FTargetSkeleton& TargetSkeleton,
 		FIKRigLogger& Log) override;
 	
 	virtual void Run(
-		const FIKRetargetProcessor& Processor,
+		FIKRetargetProcessor& Processor,
+		const double InDeltaTime,
 		const TArray<FTransform>& InSourceGlobalPose,
 		TArray<FTransform>& OutTargetGlobalPose) override;
 
 	virtual FIKRetargetOpSettingsBase* GetSettings() override;
+
+	virtual void SetSettings(const FIKRetargetOpSettingsBase* InSettings) override;
 	
 	virtual const UScriptStruct* GetSettingsType() const override;
+	
+	virtual const UScriptStruct* GetType() const override;
+
+	virtual const UClass* GetControllerType() const override;
 
 	UPROPERTY()
 	FIKRetargetPinBoneOpSettings Settings;
 	
 #if WITH_EDITOR
-	virtual FText GetNiceName() const override { return FText(LOCTEXT("OpName", "Pin Bone")); };
-	virtual FText GetWarningMessage() const override { return Message; };
-	FText Message;
+	virtual FText GetWarningMessage() const override;
 #endif
 
 private:
 
+	bool bFoundAllBonesToPin = false;
+
 	FTransform GetNewBoneTransform(
 		const FPinBoneData& PinData,
 		const FIKRetargetProcessor& Processor,
@@ -171,6 +176,24 @@ private:
 		TArray<FTransform>& OutTargetGlobalPose) const;
 };
 
+/* The blueprint/python API for editing a Pin Bone Op */
+UCLASS(BlueprintType)
+class IKRIG_API UIKRetargetPinBoneController : public UIKRetargetOpControllerBase
+{
+	GENERATED_BODY()
+	
+public:
+	/* Get the current op settings as a struct.
+	 * @return FIKRetargetPinBoneOpSettings struct with the current settings used by the op. */
+	UFUNCTION(BlueprintCallable, Category = Settings)
+	FIKRetargetPinBoneOpSettings GetSettings();
+
+	/* Set the op settings. Input is a custom struct type for this op.
+	 * @param InSettings a FIKRetargetPinBoneOpSettings struct containing all the settings to apply to this op */
+	UFUNCTION(BlueprintCallable, Category = Settings)
+	void SetSettings(FIKRetargetPinBoneOpSettings InSettings);
+};
+
 //
 // BEGIN DEPRECATED UOBJECT-based OP
 //
diff --git a/Engine/Plugins/Animation/IKRig/Source/IKRig/Public/Retargeter/RetargetOps/RootMotionGeneratorOp.h b/Engine/Plugins/Animation/IKRig/Source/IKRig/Public/Retargeter/RetargetOps/RootMotionGeneratorOp.h
--- a/Engine/Plugins/Animation/IKRig/Source/IKRig/Public/Retargeter/RetargetOps/RootMotionGeneratorOp.h
+++ b/Engine/Plugins/Animation/IKRig/Source/IKRig/Public/Retargeter/RetargetOps/RootMotionGeneratorOp.h
@@ -2,6 +2,7 @@
 #pragma once
 
 #include "Retargeter/IKRetargetOps.h"
+#include "Animation/BoneReference.h"
 
 #include "RootMotionGeneratorOp.generated.h"
 
@@ -23,22 +24,22 @@ enum class ERootMotionHeightSource : uint8
 	SnapToGround,
 };
 
-USTRUCT(BlueprintType)
+USTRUCT(BlueprintType, meta = (DisplayName = "Root Motion Settings"))
 struct IKRIG_API FIKRetargetRootMotionOpSettings : public FIKRetargetOpSettingsBase
 {
 	GENERATED_BODY()
 
 	// The root of the source skeleton.
-	UPROPERTY(EditAnywhere, Category="OpSettings")
-	FName SourceRootBone;
+	UPROPERTY(EditAnywhere, Category="OpSettings", meta=(ReinitializeOnEdit))
+	FBoneReference SourceRoot;
 
 	// The root of the target skeleton.
-	UPROPERTY(EditAnywhere, Category="OpSettings")
-	FName TargetRootBone;
+	UPROPERTY(EditAnywhere, Category="OpSettings", meta=(ReinitializeOnEdit))
+	FBoneReference TargetRoot;
 
 	// The pelvis of the target skeleton.
-	UPROPERTY(EditAnywhere, Category="OpSettings")
-	FName TargetPelvisBone;
+	UPROPERTY(EditAnywhere, Category="OpSettings", meta=(ReinitializeOnEdit))
+	FBoneReference TargetPelvis;
 
 	// Where to copy the root motion from.
 	// Copy From Source Root: copies the motion from the source root bone, but scales it according to relative height difference.
@@ -72,9 +73,36 @@ struct IKRIG_API FIKRetargetRootMotionOpSettings : public FIKRetargetOpSettingsB
 	// A manual offset to apply in global space to the root bone.
 	UPROPERTY(EditAnywhere, Category=Settings)
 	FTransform GlobalOffset;
+
+	/** FBoneReference needs a skeleton for the UI widget. */
+#if WITH_EDITOR
+	virtual USkeleton* GetSkeleton(const FName InPropertyName) override;
+#endif
+
+	// patch to load old bone names into new FBoneReference
+	void PostSerialize(const FArchive& Ar);
+#if WITH_EDITORONLY_DATA
+	// Changed to FBoneReference to support bone selector UI
+	UPROPERTY(meta=(DeprecatedProperty))
+	FName SourceRootBone_DEPRECATED;
+	UPROPERTY(meta=(DeprecatedProperty))
+	FName TargetRootBone_DEPRECATED;
+	UPROPERTY(meta=(DeprecatedProperty))
+	FName TargetPelvisBone_DEPRECATED;
+#endif
 };
 
-USTRUCT(BlueprintType)
+template<>
+struct TStructOpsTypeTraits<FIKRetargetRootMotionOpSettings> : public TStructOpsTypeTraitsBase2<FIKRetargetRootMotionOpSettings>
+{
+	enum
+	{
+		WithPostSerialize = true,
+	};
+};
+
+
+USTRUCT(BlueprintType, meta = (DisplayName = "Root Motion"))
 struct IKRIG_API FIKRetargetRootMotionOp : public FIKRetargetOpBase
 {
 	GENERATED_BODY()
@@ -86,26 +114,39 @@ struct IKRIG_API FIKRetargetRootMotionOp : public FIKRetargetOpBase
 		FIKRigLogger& Log) override;
 	
 	virtual void Run(
-		const FIKRetargetProcessor& Processor,
+		FIKRetargetProcessor& Processor,
+		const double InDeltaTime,
 		const TArray<FTransform>& InSourceGlobalPose,
 		TArray<FTransform>& OutTargetGlobalPose) override;
+
+	virtual void PostInitialize(
+		const FIKRetargetProcessor& Processor,
+		const FRetargetSkeleton& SourceSkeleton,
+		const FTargetSkeleton& TargetSkeleton,
+		FIKRigLogger& Log) override;
 	
 	virtual FIKRetargetOpSettingsBase* GetSettings() override;
+
+	virtual void SetSettings(const FIKRetargetOpSettingsBase* InSettings) override;
 	
 	virtual const UScriptStruct* GetSettingsType() const override;
+	
+	virtual const UScriptStruct* GetType() const override;
+
+	virtual const UClass* GetControllerType() const override;
 
 	UPROPERTY()
 	FIKRetargetRootMotionOpSettings Settings;
 
 #if WITH_EDITOR
-	virtual void OnAddedToStack(const UIKRetargeter* Asset) override;
-	virtual FText GetNiceName() const override { return FText(LOCTEXT("OpName", "Root Motion Generator")); };
-	virtual FText GetWarningMessage() const override { return Message; };
-	FText Message;
+	virtual void RunInitialSetup(const UIKRetargeter* InRetargetAsset) override;
+	virtual FText GetWarningMessage() const override;
 #endif
 
 private:
 
+	void Reset();
+
 	void GenerateRootMotionFromTargetPelvis(
 		FTransform& OutRootTransform,
 		const TArray<FTransform>& InSourceGlobalPose,
@@ -128,11 +169,30 @@ private:
 	TArray<int32> NonRetargetedChildrenOfRoot;
 };
 
+/* The blueprint/python API for editing a Root Motion Op */
+UCLASS(BlueprintType)
+class IKRIG_API UIKRetargetRootMotionController : public UIKRetargetOpControllerBase
+{
+	GENERATED_BODY()
+	
+public:
+	/* Get the current op settings as a struct.
+	 * @return FIKRetargetRootMotionOpSettings struct with the current settings used by the op. */
+	UFUNCTION(BlueprintCallable, Category = Settings)
+	FIKRetargetRootMotionOpSettings GetSettings();
+
+	/* Set the op settings. Input is a custom struct type for this op.
+	 * @param InSettings a FIKRetargetRootMotionOpSettings struct containing all the settings to apply to this op */
+	UFUNCTION(BlueprintCallable, Category = Settings)
+	void SetSettings(FIKRetargetRootMotionOpSettings InSettings);
+};
+
 //
 // BEGIN DEPRECATED UOBJECT-based OP
 //
 
 PRAGMA_DISABLE_DEPRECATION_WARNINGS
+
 UCLASS(Deprecated)
 class
 UE_DEPRECATED(5.6, "This type has been replaced with FIKRetargetCurveRemapOp.")
@@ -142,20 +202,7 @@ IKRIG_API UDEPRECATED_RootMotionGeneratorOp : public UDEPRECATED_RetargetOpBase
 
 public:
 	
-	virtual void ConvertToInstancedStruct(FInstancedStruct& OutInstancedStruct) override
-	{
-		OutInstancedStruct.InitializeAs(FIKRetargetRootMotionOp::StaticStruct());
-		FIKRetargetRootMotionOp& NewOp = OutInstancedStruct.GetMutable<FIKRetargetRootMotionOp>();
-		NewOp.Settings.SourceRootBone = SourceRootBone;
-		NewOp.Settings.TargetRootBone = TargetRootBone;
-		NewOp.Settings.TargetPelvisBone = TargetPelvisBone;
-		NewOp.Settings.RootMotionSource = RootMotionSource;
-		NewOp.Settings.RootHeightSource = RootHeightSource;
-		NewOp.Settings.bPropagateToNonRetargetedChildren = bPropagateToNonRetargetedChildren;
-		NewOp.Settings.bMaintainOffsetFromPelvis = bMaintainOffsetFromPelvis;
-		NewOp.Settings.bRotateWithPelvis = bRotateWithPelvis;
-		NewOp.Settings.GlobalOffset = GlobalOffset;
-	};
+	virtual void ConvertToInstancedStruct(FInstancedStruct& OutInstancedStruct) override;
 
 	UPROPERTY()
 	FName SourceRootBone;
diff --git a/Engine/Plugins/Animation/IKRig/Source/IKRig/Public/Retargeter/RetargetOps/RunIKRigOp.h b/Engine/Plugins/Animation/IKRig/Source/IKRig/Public/Retargeter/RetargetOps/RunIKRigOp.h
new file mode 100644
--- /dev/null
+++ b/Engine/Plugins/Animation/IKRig/Source/IKRig/Public/Retargeter/RetargetOps/RunIKRigOp.h
@@ -0,0 +1,117 @@
+﻿// Copyright Epic Games, Inc. All Rights Reserved.
+#pragma once
+
+#include "Retargeter/IKRetargetOps.h"
+#include "Rig/IKRigProcessor.h"
+
+#include "RunIKRigOp.generated.h"
+
+#define LOCTEXT_NAMESPACE "RunIKRigOp"
+
+USTRUCT(BlueprintType, meta = (DisplayName = "Solve IK Goal Settings"))
+struct IKRIG_API FIKRetargetRunIKRigOpSettings : public FIKRetargetOpSettingsBase
+{
+	GENERATED_BODY()
+
+	/** The IK Rig asset to run when this op is executed. */
+	UPROPERTY(EditAnywhere, Category = "IK Rig Asset", meta=(ReinitializeOnEdit))
+	TObjectPtr<const UIKRigDefinition> IKRigAsset = nullptr;
+
+	/** Goals in this list will be excluded from the rig */
+	UPROPERTY(EditAnywhere, Category = "IK Rig Asset")
+	TArray<FName> ExcludedGoals;
+	
+	// Draw IK goal locations. 
+	UPROPERTY(EditAnywhere, Category = DebugDraw)
+	bool bDrawGoals = true;
+
+	// Draw locations of the source bone (pre-solve)
+	UPROPERTY(EditAnywhere, Category = DebugDraw)
+	bool bDrawGoalBoneLocations = true;
+
+	// Adjust size of goal debug drawing in viewport
+	UPROPERTY(EditAnywhere, Category = DebugDraw)
+	double GoalDrawSize = 5.0;
+	
+	// Adjust thickness of goal debug drawing in viewport
+	UPROPERTY(EditAnywhere, Category = DebugDraw)
+	double GoalDrawThickness = 1.0;
+};
+
+#if WITH_EDITOR
+struct FRunIKRigOpGoalDebugData
+{
+	FName GoalName;
+	FTransform InitialGoal;
+	FTransform CurrentGoal;
+};
+#endif
+
+USTRUCT(BlueprintType, meta = (DisplayName = "Solve IK Goals"))
+struct IKRIG_API FIKRetargetRunIKRigOp : public FIKRetargetOpBase
+{
+	GENERATED_BODY()
+	
+	virtual bool Initialize(
+		const FIKRetargetProcessor& InProcessor,
+		const FRetargetSkeleton& InSourceSkeleton,
+		const FTargetSkeleton& InTargetSkeleton,
+		FIKRigLogger& InLog) override;
+	
+	virtual void Run(
+		FIKRetargetProcessor& InProcessor,
+		const double InDeltaTime,
+		const TArray<FTransform>& InSourceGlobalPose,
+		TArray<FTransform>& OutTargetGlobalPose) override;
+
+	virtual FIKRetargetOpSettingsBase* GetSettings() override;
+	
+	virtual void SetSettings(const FIKRetargetOpSettingsBase* InSettings) override;
+	
+	virtual const UScriptStruct* GetSettingsType() const override;
+
+	virtual const UScriptStruct* GetType() const override;
+
+	virtual const UClass* GetControllerType() const override;
+
+	UPROPERTY()
+	FIKRetargetRunIKRigOpSettings Settings;
+
+#if WITH_EDITOR
+	virtual void RunInitialSetup(const UIKRetargeter* InRetargetAsset) override;
+	virtual void DebugDraw(
+		FPrimitiveDrawInterface* InPDI,
+		const FTransform& InComponentTransform,
+		const double InComponentScale,
+		const FIKRetargetDebugDrawState& InEditorState) const override;
+	void SaveInitialGoalTransformsIntoDebugData(const FIKRetargetProcessor& InProcessor, const TArray<FTransform>& InTargetGlobalPose);
+	void SaveCurrentGoalTransformsIntoDebugData();
+	TArray<FRunIKRigOpGoalDebugData> GoalDebugData;
+	static FCriticalSection DebugDataMutex;
+#endif
+
+private:
+
+	/** The IK Rig processor for running IK on the target */
+	FIKRigProcessor IKRigProcessor;
+};
+
+/* The blueprint/python API for editing a Run IK Rig Op */
+UCLASS(BlueprintType)
+class IKRIG_API UIKRetargetRunIKRigController : public UIKRetargetOpControllerBase
+{
+	GENERATED_BODY()
+	
+public:
+	/* Get the current op settings as a struct.
+	 * @return FIKRetargetRunIKRigOpSettings struct with the current settings used by the op. */
+	UFUNCTION(BlueprintCallable, Category = Settings)
+	FIKRetargetRunIKRigOpSettings GetSettings();
+
+	/* Set the op settings. Input is a custom struct type for this op.
+	 * @param InSettings a FIKRetargetRunIKRigOpSettings struct containing all the settings to apply to this op */
+	UFUNCTION(BlueprintCallable, Category = Settings)
+	void SetSettings(FIKRetargetRunIKRigOpSettings InSettings);
+};
+
+#undef LOCTEXT_NAMESPACE
diff --git a/Engine/Plugins/Animation/IKRig/Source/IKRig/Public/Retargeter/RetargetOps/ScaleSourceOp.h b/Engine/Plugins/Animation/IKRig/Source/IKRig/Public/Retargeter/RetargetOps/ScaleSourceOp.h
new file mode 100644
--- /dev/null
+++ b/Engine/Plugins/Animation/IKRig/Source/IKRig/Public/Retargeter/RetargetOps/ScaleSourceOp.h
@@ -0,0 +1,72 @@
+﻿// Copyright Epic Games, Inc. All Rights Reserved.
+#pragma once
+
+#include "Retargeter/IKRetargetOps.h"
+
+#include "ScaleSourceOp.generated.h"
+
+#define LOCTEXT_NAMESPACE "ScaleSourceOp"
+
+USTRUCT(BlueprintType, meta = (DisplayName = "Scale Source Settings"))
+struct IKRIG_API FIKRetargetScaleSourceOpSettings : public FIKRetargetOpSettingsBase
+{
+	GENERATED_BODY()
+	
+	/** Range 0.01 to +inf. Default 1. Scales the incoming source pose. Affects entire skeleton and all IK goals.*/
+	UPROPERTY(EditAnywhere, BlueprintReadWrite, Category = Scale, meta = (ReinitializeOnEdit, ClampMin = "0.01", UIMin = "0.01", UIMax = "10.0"))
+	double SourceScaleFactor = 1.0f;
+};
+
+USTRUCT(BlueprintType, meta = (DisplayName = "Scale Source"))
+struct IKRIG_API FIKRetargetScaleSourceOp : public FIKRetargetOpBase
+{
+	GENERATED_BODY()
+	
+	// NOTE: this op does not do anything in Initialize() or Run().
+	// It is a special case op that the retargeter reads from when it needs to scale the source pose.
+	
+	virtual bool Initialize(
+		const FIKRetargetProcessor& InProcessor,
+		const FRetargetSkeleton& InSourceSkeleton,
+		const FTargetSkeleton& InTargetSkeleton,
+		FIKRigLogger& Log) override;
+	
+	virtual void Run(
+		FIKRetargetProcessor& InProcessor,
+		const double InDeltaTime,
+		const TArray<FTransform>& InSourceGlobalPose,
+		TArray<FTransform>& OutTargetGlobalPose) override {};
+
+	virtual FIKRetargetOpSettingsBase* GetSettings() override;
+
+	virtual void SetSettings(const FIKRetargetOpSettingsBase* InSettings) override;
+	
+	virtual const UScriptStruct* GetSettingsType() const override;
+	
+	virtual const UScriptStruct* GetType() const override;
+
+	virtual const UClass* GetControllerType() const override;
+
+	UPROPERTY()
+	FIKRetargetScaleSourceOpSettings Settings;
+};
+
+/* The blueprint/python API for editing a Scale Source Op */
+UCLASS(BlueprintType)
+class IKRIG_API UIKRetargetScaleSourceController : public UIKRetargetOpControllerBase
+{
+	GENERATED_BODY()
+	
+public:
+	/* Get the current op settings as a struct.
+	 * @return FIKRetargetScaleSourceOpSettings struct with the current settings used by the op. */
+	UFUNCTION(BlueprintCallable, Category = Settings)
+	FIKRetargetScaleSourceOpSettings GetSettings();
+
+	/* Set the op settings. Input is a custom struct type for this op.
+	 * @param InSettings a FIKRetargetScaleSourceOpSettings struct containing all the settings to apply to this op */
+	UFUNCTION(BlueprintCallable, Category = Settings)
+	void SetSettings(FIKRetargetScaleSourceOpSettings InSettings);
+};
+
+#undef LOCTEXT_NAMESPACE
\ No newline at end of file
diff --git a/Engine/Plugins/Animation/IKRig/Source/IKRig/Public/Retargeter/RetargetOps/SpeedPlantingOp.h b/Engine/Plugins/Animation/IKRig/Source/IKRig/Public/Retargeter/RetargetOps/SpeedPlantingOp.h
new file mode 100644
--- /dev/null
+++ b/Engine/Plugins/Animation/IKRig/Source/IKRig/Public/Retargeter/RetargetOps/SpeedPlantingOp.h
@@ -0,0 +1,137 @@
+﻿// Copyright Epic Games, Inc. All Rights Reserved.
+#pragma once
+
+#include "Retargeter/IKRetargetOps.h"
+#include "Kismet/KismetMathLibrary.h"
+
+#include "SpeedPlantingOp.generated.h"
+
+#define LOCTEXT_NAMESPACE "SpeedPlantingOp"
+
+USTRUCT(BlueprintType)
+struct IKRIG_API FRetargetSpeedPlantingSettings
+{
+	GENERATED_BODY()
+
+	/** The name of the target chain to plant the IK on.
+	 * NOTE: this chain must have an IK Goal assigned to it! */
+	UPROPERTY(EditAnywhere, BlueprintReadWrite, Category = "Speed Planting", meta=(ReinitializeOnEdit))
+	FName TargetChainName;
+
+	/** The name of the curve on the source animation that contains the speed of the end effector bone.*/
+	UPROPERTY(EditAnywhere, BlueprintReadWrite, BlueprintReadWrite, Category = "Plant IK by Speed", meta = (ClampMin = "0.0", ClampMax = "100.0", UIMin = "0.0", UIMax = "100.0"))
+	FName SpeedCurveName;
+};
+
+USTRUCT(BlueprintType, meta = (DisplayName = "Speed Plant Settings"))
+struct IKRIG_API FIKRetargetSpeedPlantingOpSettings : public FIKRetargetOpSettingsBase
+{
+	GENERATED_BODY()
+	
+	UPROPERTY(EditAnywhere, BlueprintReadWrite, Category = "Speed Planting", meta=(ReinitializeOnEdit))
+	TArray<FRetargetSpeedPlantingSettings> ChainsToSpeedPlant;
+
+	/** Range 0 to 100. Default 15. The maximum speed a source bone can be moving while being considered 'planted'.
+	*  The target IK goal will not be allowed to move whenever the source bone speed drops below this threshold speed. */
+	UPROPERTY(EditAnywhere, BlueprintReadWrite, Category = "Speed Planting", meta = (ClampMin = "0.0", ClampMax = "100.0", UIMin = "0.0", UIMax = "100.0"))
+	double SpeedThreshold = 15.0f;
+
+	// How stiff the spring model is that smoothly pulls the IK position after unplanting (more stiffness means more oscillation around the target value)
+	UPROPERTY(EditAnywhere, BlueprintReadWrite, Category = "Speed Planting", meta = (ClampMin = "0.0", UIMin = "0.0"))
+	double Stiffness = 250.0f;
+
+	// How much damping to apply to the spring (0 means no damping, 1 means critically damped which means no oscillation)
+	UPROPERTY(EditAnywhere, BlueprintReadWrite, Category = "Speed Planting", meta = (ClampMin = "0.0", ClampMax = "1.0", UIMin = "0.0", UIMax = "1.0"))
+	double CriticalDamping = 1.0f;
+};
+
+struct FSpeedPlantingGoalState
+{
+	FSpeedPlantingGoalState(
+		const FName InGoalName,
+		const FRetargetSpeedPlantingSettings* InSettings,
+		const FVector& InInitialGoalPosition) :
+		GoalName(InGoalName),
+		Settings(InSettings),
+		PrevGoalPosition(InInitialGoalPosition) {};
+	
+	FName GoalName;
+	const FRetargetSpeedPlantingSettings *Settings;
+	FVectorSpringState PositionSpring;
+	FVector PrevGoalPosition;
+	double CurrentSpeedValue = -1.0f;
+
+#if WITH_EDITOR
+	bool bFoundCurveInSourceComponent = false;
+	bool bFoundCurveInTargetComponent = false;
+#endif
+};
+
+USTRUCT(BlueprintType, meta = (DisplayName = "Speed Plant IK Goals"))
+struct IKRIG_API FIKRetargetSpeedPlantingOp : public FIKRetargetOpBase
+{
+	GENERATED_BODY()
+	
+	virtual bool Initialize(
+		const FIKRetargetProcessor& InProcessor,
+		const FRetargetSkeleton& InSourceSkeleton,
+		const FTargetSkeleton& InTargetSkeleton,
+		FIKRigLogger& InLog) override;
+	
+	virtual void Run(
+		FIKRetargetProcessor& InProcessor,
+		const double InDeltaTime,
+		const TArray<FTransform>& InSourceGlobalPose,
+		TArray<FTransform>& OutTargetGlobalPose) override;
+
+	virtual FIKRetargetOpSettingsBase* GetSettings() override;
+
+	virtual void SetSettings(const FIKRetargetOpSettingsBase* InSettings) override;
+	
+	virtual const UScriptStruct* GetSettingsType() const override;
+	
+	virtual const UScriptStruct* GetType() const override;
+
+	virtual void OnPlaybackReset() override;
+
+	virtual void AnimGraphPreUpdateMainThread(USkeletalMeshComponent& SourceMeshComponent, USkeletalMeshComponent& TargetMeshComponent) override;
+
+	virtual void AnimGraphEvaluateAnyThread(FPoseContext& Output) override;
+	
+	virtual const UClass* GetControllerType() const override;
+
+	/** return array of names of all the speed curves referenced by this op. */
+	TArray<FName> GetRequiredSpeedCurves() const;
+
+	UPROPERTY()
+	FIKRetargetSpeedPlantingOpSettings Settings;
+
+#if WITH_EDITOR
+public:
+	virtual FText GetWarningMessage() const override;
+#endif
+
+	bool ResetThisTick;
+	
+	TArray<FSpeedPlantingGoalState> GoalsToPlant;
+};
+
+/* The blueprint/python API for editing a Speed Planting Op */
+UCLASS(BlueprintType)
+class IKRIG_API UIKRetargetSpeedPlantingController : public UIKRetargetOpControllerBase
+{
+	GENERATED_BODY()
+	
+public:
+	/* Get the current op settings as a struct.
+	 * @return FIKRetargetSpeedPlantingOpSettings struct with the current settings used by the op. */
+	UFUNCTION(BlueprintCallable, Category = Settings)
+	FIKRetargetSpeedPlantingOpSettings GetSettings();
+
+	/* Set the op settings. Input is a custom struct type for this op.
+	 * @param InSettings a FIKRetargetSpeedPlantingOpSettings struct containing all the settings to apply to this op */
+	UFUNCTION(BlueprintCallable, Category = Settings)
+	void SetSettings(FIKRetargetSpeedPlantingOpSettings InSettings);
+};
+
+#undef LOCTEXT_NAMESPACE
diff --git a/Engine/Plugins/Animation/IKRig/Source/IKRig/Public/Retargeter/RetargetOps/StrideWarpingOp.h b/Engine/Plugins/Animation/IKRig/Source/IKRig/Public/Retargeter/RetargetOps/StrideWarpingOp.h
new file mode 100644
--- /dev/null
+++ b/Engine/Plugins/Animation/IKRig/Source/IKRig/Public/Retargeter/RetargetOps/StrideWarpingOp.h
@@ -0,0 +1,157 @@
+﻿// Copyright Epic Games, Inc. All Rights Reserved.
+#pragma once
+
+#include "Retargeter/IKRetargetOps.h"
+#include "Retargeter/IKRetargetSettings.h"
+#include "StrideWarpingOp.generated.h"
+
+#define LOCTEXT_NAMESPACE "StrideWarpingOp"
+
+struct FIKRigGoal;
+
+UENUM()
+enum class EWarpingDirectionSource
+{
+	Goals,
+	Chain,
+	RootBone
+};
+
+USTRUCT(BlueprintType, meta = (DisplayName = "Stride Warp Settings"))
+struct IKRIG_API FIKRetargetStrideWarpingOpSettings : public FIKRetargetOpSettingsBase
+{
+	GENERATED_BODY()
+	
+	/** The chains to apply stride warping to */
+	UPROPERTY(EditAnywhere, BlueprintReadWrite, Category = Setup, meta=(ReinitializeOnEdit))
+	TArray<FName> ChainsToStrideWarp;
+
+	/** Defines the source used to determine the forward direction as the character animates in world space. Default is "Goals".
+	 * This method uses various positions on the character to create a "best fit" global rotation that approximates the facing direction of the character over time.
+	 * This global rotation is used to define the forward and sideways directions used when warping goals along those axes.
+	 * The options are:
+	 * Goals: uses the positions of the IK goals to approximate the facing direction. This is best used on characters with a vertical spine, like bipeds.
+	 * Chain: uses the positions of the bones in a retarget chain to approximate the facing direction. This is best when used with the spine chain for characters with a horizontal spine, like quadrupeds.
+	 * Root Bone: uses the rotation of the root bone of the skeleton. This is most robust, but character must have correct root motion with yaw rotation in movement direction.
+	 */
+	UPROPERTY(EditAnywhere, BlueprintReadWrite, Category = Warping)
+	EWarpingDirectionSource DirectionSource = EWarpingDirectionSource::Goals;
+	
+	/** The world space axis that represents the forward facing direction for your character. By default, in Unreal, this is +Y.  */
+	UPROPERTY(EditAnywhere, BlueprintReadWrite, Category = Warping)
+	EBasicAxis ForwardDirection = EBasicAxis::Y;
+
+	/** When using the "Chain" option as a Direction Source, this defines the chain to use to determine the facing direction of the character.
+	 * Typically this would be the chain that contains the Spine bones. */
+	UPROPERTY(EditAnywhere, BlueprintReadWrite, Category = Warping, meta=(EditCondition="DirectionSource == EWarpingDirectionSource::Chain"))
+	FName DirectionChain;
+
+	/** Range 0 to Inf. Default 1. Warps IK goal positions in the forward direction. Useful for stride warping.
+	 * Values below 1 will create smaller, squashed strides. Values greater than 1 will create stretched, longer strides.*/
+	UPROPERTY(EditAnywhere, BlueprintReadWrite, Category = Warping, meta = (UIMin = "0.0", UIMax = "5.0", ClampMin = "0.0"))
+	double WarpForwards = 1.0f;
+
+	/** Range -+Inf. Default is 0. A static offset in world units to move the IK goals perpendicular to the forward direction.
+	 * Values less than zero will move the goals towards the center-line of the character. Values greater than zero push the goals outwards.*/
+	UPROPERTY(EditAnywhere, BlueprintReadWrite, Category = Warping, meta = (UIMin = "-10.0", UIMax = "100.0"))
+	double SidewaysOffset = 0.0f;
+
+	/** Range 0 to +Inf. Default is 1.0f.
+	 * Values below 1 pull all the goals towards the average of all the goals (towards each other).
+	 * Values greater than 1 push the goals apart.*/
+	UPROPERTY(EditAnywhere, BlueprintReadWrite, Category = Warping, meta = (UIMin = "0.0", UIMax = "2.0", ClampMin = "0.0"))
+	double WarpSplay = 1.0f;
+
+	/** Enable debug drawing for this op. */
+	UPROPERTY(EditAnywhere, BlueprintReadWrite, Category = Debug)
+	bool bEnableDebugDraw = true;
+
+	/** Adjust the size of the debug drawing */
+	UPROPERTY(EditAnywhere, BlueprintReadWrite, Category = Debug)
+	double DebugDrawSize = 5.0;
+	
+	/** Adjust the thickness of the debug drawing */
+	UPROPERTY(EditAnywhere, BlueprintReadWrite, Category = Debug)
+	double DebugDrawThickness = 5.0;
+
+	static FVector GetAxisVector(const EBasicAxis& Axis);
+
+	bool operator==(const FIKRetargetStrideWarpingOpSettings& Other) const;
+};
+
+struct FStrideWarpGoalData
+{
+	FStrideWarpGoalData(const FName InGoalName, const FTransform& InGlobalRefPoseOfGoalBone) :
+		IKRigGoalName(InGoalName),
+		GlobalRefPoseOfGoalBone(InGlobalRefPoseOfGoalBone) {};
+	
+	FName IKRigGoalName;
+	FTransform GlobalRefPoseOfGoalBone;
+};
+
+USTRUCT(BlueprintType, meta = (DisplayName = "Stride Warp IK Goals"))
+struct IKRIG_API FIKRetargetStrideWarpingOp : public FIKRetargetOpBase
+{
+	GENERATED_BODY()
+	
+	virtual bool Initialize(
+		const FIKRetargetProcessor& InProcessor,
+		const FRetargetSkeleton& InSourceSkeleton,
+		const FTargetSkeleton& InTargetSkeleton,
+		FIKRigLogger& InLog) override;
+	
+	virtual void Run(
+		FIKRetargetProcessor& InProcessor,
+		const double InDeltaTime,
+		const TArray<FTransform>& InSourceGlobalPose,
+		TArray<FTransform>& OutTargetGlobalPose) override;
+
+	virtual FIKRetargetOpSettingsBase* GetSettings() override;
+
+	virtual void SetSettings(const FIKRetargetOpSettingsBase* InSettings) override;
+	
+	virtual const UScriptStruct* GetSettingsType() const override;
+	
+	virtual const UScriptStruct* GetType() const override;
+
+	virtual const UClass* GetControllerType() const override;
+
+	UPROPERTY()
+	FIKRetargetStrideWarpingOpSettings Settings;
+
+#if WITH_EDITOR
+	virtual void DebugDraw(
+		FPrimitiveDrawInterface* InPDI,
+		const FTransform& InComponentTransform,
+		const double InComponentScale,
+		const FIKRetargetDebugDrawState& InEditorState) const override;
+	FTransform DebugStrideWarpingFrame;
+	static FCriticalSection DebugDataMutex;
+
+	virtual void SetChainSettingsToDefault(const FName& InChainName) override;
+	virtual bool AreChainSettingsAtDefault(const FName& InChainName) override;
+#endif
+
+	TArray<FStrideWarpGoalData> GoalsToWarp;
+};
+
+/* The blueprint/python API for editing a Stride Warping Op */
+UCLASS(BlueprintType)
+class IKRIG_API UIKRetargetStrideWarpingController : public UIKRetargetOpControllerBase
+{
+	GENERATED_BODY()
+	
+public:
+
+	/* Get the current op settings as a struct.
+	 * @return FIKRetargetStrideWarpingOpSettings struct with the current settings used by the op. */
+	UFUNCTION(BlueprintCallable, Category = Settings)
+	FIKRetargetStrideWarpingOpSettings GetSettings();
+
+	/* Set the solver settings. Input is a custom struct type for this solver.
+	 * @param InSettings a FIKRetargetStrideWarpingOpSettings struct containing all the settings to apply to this op */
+	UFUNCTION(BlueprintCallable, Category = Settings)
+	void SetSettings(FIKRetargetStrideWarpingOpSettings InSettings);
+};
+
+#undef LOCTEXT_NAMESPACE
diff --git a/Engine/Plugins/Animation/IKRig/Source/IKRig/Public/Rig/IKRigDataTypes.h b/Engine/Plugins/Animation/IKRig/Source/IKRig/Public/Rig/IKRigDataTypes.h
--- a/Engine/Plugins/Animation/IKRig/Source/IKRig/Public/Rig/IKRigDataTypes.h
+++ b/Engine/Plugins/Animation/IKRig/Source/IKRig/Public/Rig/IKRigDataTypes.h
@@ -33,6 +33,10 @@ struct IKRIG_API FIKRigGoal
 	/** Name of the IK Goal. Must correspond to the name of a Goal in the target IKRig asset. */
 	UPROPERTY(EditAnywhere, BlueprintReadWrite, Category = Goal)
 	FName Name;
+	
+	/** Name of the bone associate with this goal. */
+	UPROPERTY(EditAnywhere, BlueprintReadWrite, Category = Goal)
+	FName BoneName;
 
 	/** Set the source of the transform data for the Goal.
 	 * "Manual Input" uses the values provided by the blueprint node pin.
@@ -52,71 +56,61 @@ struct IKRIG_API FIKRigGoal
 	
 	/** Position of the IK goal in Component Space of target actor component. */
 	UPROPERTY(EditAnywhere, BlueprintReadWrite, Category = Position, meta=(EditCondition="Source == EIKRigGoalSource::Manual"))
-	FVector Position;
+	FVector Position = FVector::ZeroVector;
 
 	/** Rotation of the IK goal in Component Space of target actor component. */
 	UPROPERTY(EditAnywhere, BlueprintReadWrite, Category = Rotation, meta=(EditCondition="Source == EIKRigGoalSource::Manual"))
-	FRotator Rotation;
+	FRotator Rotation = FRotator::ZeroRotator;
 
 	/** Range 0-1, default is 1.0. Smoothly blends the Goal position from the input pose (0.0) to the Goal position (1.0). */
 	UPROPERTY(EditAnywhere, BlueprintReadWrite, Category = Position, meta = (ClampMin = "0.0", ClampMax = "1.0", UIMin = "0.0", UIMax = "1.0"))
-	float PositionAlpha;
+	float PositionAlpha = 1.0f;
 
 	/** Range 0-1, default is 1.0. Smoothly blends the Goal rotation from the input pose (0.0) to the Goal rotation (1.0). */
 	UPROPERTY(EditAnywhere, BlueprintReadWrite, Category = Rotation, meta = (ClampMin = "0.0", ClampMax = "1.0", UIMin = "0.0", UIMax = "1.0"))
-	float RotationAlpha;
+	float RotationAlpha = 1.0f;
 
 	/** The space that the goal position is in.
 	 *"Additive" treats the goal transform as an additive offset relative to the Bone at the effector.
 	 *"Component" treats the goal transform as being in the space of the Skeletal Mesh Actor Component.
 	 *"World" treats the goal transform as being in the space of the World. */
 	UPROPERTY(EditAnywhere, BlueprintReadWrite, Category = Position)
-	EIKRigGoalSpace PositionSpace;
+	EIKRigGoalSpace PositionSpace = EIKRigGoalSpace::Additive;
 	
 	/** The space that the goal rotation is in.
 	*"Additive" treats the goal transform as an additive offset relative to the Bone at the effector.
 	*"Component" treats the goal transform as being in the space of the Skeletal Mesh Actor Component.
 	*"World" treats the goal transform as being in the space of the World. */
 	UPROPERTY(EditAnywhere, BlueprintReadWrite, Category = Rotation)
-	EIKRigGoalSpace RotationSpace;
+	EIKRigGoalSpace RotationSpace = EIKRigGoalSpace::Additive;
+	
+	UPROPERTY(Transient)
+	FVector FinalBlendedPosition = FVector::ZeroVector;
 
 	UPROPERTY(Transient)
-	FVector FinalBlendedPosition;
+	FQuat FinalBlendedRotation = FQuat::Identity;
 
+	/** If true, the goal is active and can be used by solvers in the rig.
+	 * NOTE: Toggling this on or off at runtime will trigger a re-initialization. Prefer Alpha if possible. */
 	UPROPERTY(Transient)
-	FQuat FinalBlendedRotation;
-
-	FIKRigGoal()
-    : Name(NAME_None),
-	Position(ForceInitToZero),
-	Rotation(FRotator::ZeroRotator),
-	PositionAlpha(1.0f),
-    RotationAlpha(1.0f),
-	PositionSpace(EIKRigGoalSpace::Additive),
-	RotationSpace(EIKRigGoalSpace::Additive),
-	FinalBlendedPosition(ForceInitToZero),
-	FinalBlendedRotation(FQuat::Identity){}
+	bool bEnabled = true;
+
+	FIKRigGoal() = default;
 	
-	FIKRigGoal(const FName& GoalName)
-    : Name(GoalName),
-	Position(ForceInitToZero),
-	Rotation(FRotator::ZeroRotator),
-    PositionAlpha(1.0f),
-    RotationAlpha(1.0f),
-	PositionSpace(EIKRigGoalSpace::Additive),
-    RotationSpace(EIKRigGoalSpace::Additive),
-	FinalBlendedPosition(ForceInitToZero),
-    FinalBlendedRotation(FQuat::Identity){}
+	FIKRigGoal(const FName& InGoalName, const FName& InBoneName) : Name(InGoalName), BoneName(InBoneName) {}
 
 	FIKRigGoal(
         const FName& InName,
+        const FName& InBoneName,
         const FVector& InPosition,
         const FQuat& InRotation,
         const float InPositionAlpha,
         const float InRotationAlpha,
         EIKRigGoalSpace InPositionSpace,
-        EIKRigGoalSpace InRotationSpace)
+        EIKRigGoalSpace InRotationSpace,
+        bool bInEnabled)
         : Name(InName),
+		BoneName(InBoneName),
         Position(InPosition),
         Rotation(InRotation.Rotator()),
         PositionAlpha(InPositionAlpha),
@@ -124,10 +118,12 @@ struct IKRIG_API FIKRigGoal
 		PositionSpace(InPositionSpace),
 		RotationSpace(InRotationSpace),
 		FinalBlendedPosition(InPosition),
-		FinalBlendedRotation(InRotation){}
+		FinalBlendedRotation(InRotation),
+		bEnabled(bInEnabled) {}
 
 	FIKRigGoal(const UIKRigEffectorGoal* InGoal)
 		: Name(InGoal->GoalName),
+		BoneName(InGoal->BoneName),
 		Position(InGoal->CurrentTransform.GetTranslation()),
 		Rotation(InGoal->CurrentTransform.Rotator()),
 		PositionAlpha(InGoal->PositionAlpha),
@@ -135,7 +131,7 @@ struct IKRIG_API FIKRigGoal
 		PositionSpace(EIKRigGoalSpace::Component),
         RotationSpace(EIKRigGoalSpace::Component),
         FinalBlendedPosition(ForceInitToZero),
-        FinalBlendedRotation(FQuat::Identity){}
+        FinalBlendedRotation(FQuat::Identity) {}
 
 	FString ToString() const
 	{
@@ -163,15 +159,33 @@ public:
 
 	/** Get an IK goal with the given name. Returns nullptr if no Goal is found in the container with the name. */
 	const FIKRigGoal* FindGoalByName(const FName& GoalName) const;
+	
+	/** Get an IK goal with the given name. Returns nullptr if no Goal is found in the container with the name. */
+	FIKRigGoal* FindGoalByName(const FName& GoalName);
 
 	/** Clear out all Goals in container. */
-	void Empty() { Goals.Empty(); };
+	void Empty() { Goals.Empty(); bRigNeedsInitialized = true; };
+
+	/** Returns true if there are no goals in the container */
+	bool IsEmpty() const { return Goals.Num() <= 0; };
+
+	/** Returns true if the container has a goal in it that triggered a reinitialization */
+	bool NeedsInitialized() const {return bRigNeedsInitialized; };
 
 	/** Get read-only access to array of Goals. */
 	const TArray<FIKRigGoal>& GetGoalArray() const { return Goals; };
+	TArray<FIKRigGoal>& GetGoalArray() { return Goals; };
+
+	/** Fill this contain with all the goals in the input array */
+	const void FillWithGoalArray(const TArray<UIKRigEffectorGoal*>& InGoals);
 	
 private:
 
+	/** Set to true when:
+	 * 1. a new goal is added
+	 * 2. a goal has it's enabled flag toggled */
+	bool bRigNeedsInitialized = true;
+
 	FIKRigGoal* FindGoalWriteable(const FName& GoalName) const;
 	
 	/** Array of Goals. Cannot contain duplicates (name is key). */
diff --git a/Engine/Plugins/Animation/IKRig/Source/IKRig/Public/Rig/IKRigDefinition.h b/Engine/Plugins/Animation/IKRig/Source/IKRig/Public/Rig/IKRigDefinition.h
--- a/Engine/Plugins/Animation/IKRig/Source/IKRig/Public/Rig/IKRigDefinition.h
+++ b/Engine/Plugins/Animation/IKRig/Source/IKRig/Public/Rig/IKRigDefinition.h
@@ -216,6 +216,9 @@ public:
 	/** runtime, read-only access to read a bone chain, all modifications must go through UIKRigController */
 	const FBoneChain* GetRetargetChainByName(FName ChainName) const;
 
+	/** get a list of the retarget chains contained in this IK Rig */
+	TArray<FName> GetRetargetChainNames() const;
+
 	/** UObject */
 	virtual void PostDuplicate(bool bDuplicateForPIE) override;
 	virtual void Serialize(FArchive& Ar) override;
diff --git a/Engine/Plugins/Animation/IKRig/Source/IKRig/Public/Rig/IKRigProcessor.h b/Engine/Plugins/Animation/IKRig/Source/IKRig/Public/Rig/IKRigProcessor.h
--- a/Engine/Plugins/Animation/IKRig/Source/IKRig/Public/Rig/IKRigProcessor.h
+++ b/Engine/Plugins/Animation/IKRig/Source/IKRig/Public/Rig/IKRigProcessor.h
@@ -41,12 +41,12 @@ public:
 	 *  NOTE!! this function creates new UObjects and consequently MUST be called from the main thread!!
 	 *  @param InRigAsset - the IK Rig defining the collection of solvers to execute and all the rig settings
 	 *  @param SkeletalMesh - the skeletal mesh you want to solve the IK on
-	 *  @param ExcludedGoals - a list of goal names to exclude from this processor instance (support per-instance removal of goals)
+	 *  @param GoalContainer - an optional list of goals to initialize with, if empty it will use the goals from the IK Rig asset
 	 */
 	void Initialize(
 		const UIKRigDefinition* InRigAsset,
-		const USkeletalMesh* SkeletalMesh,
-		const TArray<FName>& ExcludedGoals = TArray<FName>());
+		const USkeletalMesh* InSkeletalMesh,
+		const FIKRigGoalContainer& InOptionalGoals);
 
 	//
 	// BEGIN UPDATE SEQUENCE FUNCTIONS
@@ -60,6 +60,9 @@ public:
 	/** Optionally can be called before Solve() to use the reference pose as start pose */
 	void SetInputPoseToRefPose();
 
+	/** Assign an entire container of goals to the rig. */
+	void ApplyGoalsFromOtherContainer(const FIKRigGoalContainer& InGoalContainer);
+
 	/** Set a named IK goal to go to a specific location, rotation and space, blended by separate position/rotation alpha (0-1)*/
 	void SetIKGoal(const FIKRigGoal& Goal);
 
@@ -129,8 +132,9 @@ private:
 	/** storage for hierarchy and bone transforms */
 	FIKRigSkeleton Skeleton;
 
-	/** solving disabled until this flag is true */
+	/** true if processor is currently initialized and ready to run */
 	bool bInitialized = false;
+	/** true if processor attempted an initialization but failed (prevents infinite reinit) */
 	bool bTriedToInitialize = false;
 };
 
@@ -152,7 +156,7 @@ IKRIG_API UIKRigProcessor : public UObject
 	
 	void Initialize(const UIKRigDefinition* InRigAsset, const USkeletalMesh* SkeletalMesh, const TArray<FName>& ExcludedGoals = TArray<FName>())
 	{
-		Processor.Initialize(InRigAsset, SkeletalMesh, ExcludedGoals);
+		Processor.Initialize(InRigAsset, SkeletalMesh, FIKRigGoalContainer());
 	};
 	void SetInputPoseGlobal(const TArray<FTransform>& InGlobalBoneTransforms)
 	{
diff --git a/Engine/Plugins/Animation/IKRig/Source/IKRigDeveloper/Private/AnimGraphNode_IKRig.cpp b/Engine/Plugins/Animation/IKRig/Source/IKRigDeveloper/Private/AnimGraphNode_IKRig.cpp
--- a/Engine/Plugins/Animation/IKRig/Source/IKRigDeveloper/Private/AnimGraphNode_IKRig.cpp
+++ b/Engine/Plugins/Animation/IKRig/Source/IKRigDeveloper/Private/AnimGraphNode_IKRig.cpp
@@ -354,7 +354,7 @@ void UAnimGraphNode_IKRig::PostEditChangeProperty(struct FPropertyChangedEvent&
 			const TArray<UIKRigEffectorGoal*>& AssetGoals = Node.RigDefinitionAsset->GetGoalArray();
 			for (const UIKRigEffectorGoal* AssetGoal: AssetGoals)
 			{
-				const int32 GoalIndex = Node.Goals.Emplace(AssetGoal->GoalName);
+				const int32 GoalIndex = Node.Goals.Emplace(AssetGoal->GoalName, AssetGoal->BoneName);
 				SetupGoal(AssetGoal, Node.Goals[GoalIndex]);
 			}
 
@@ -863,7 +863,7 @@ void UAnimGraphNode_IKRig::UpdateGoalsFromAsset()
 			}
 			else
 			{
-				Node.Goals.Emplace(AssetGoal->GoalName);
+				Node.Goals.Emplace(AssetGoal->GoalName, AssetGoal->BoneName);
 			}
 			
 			SetupGoal(AssetGoal, Node.Goals.Last());
