[IKRig] Major refactor part 3 of 3.

diff --git a/Engine/Plugins/Animation/IKRig/Source/IKRigEditor/Private/IKRigEditor.cpp b/Engine/Plugins/Animation/IKRig/Source/IKRigEditor/Private/IKRigEditor.cpp
--- a/Engine/Plugins/Animation/IKRig/Source/IKRigEditor/Private/IKRigEditor.cpp
+++ b/Engine/Plugins/Animation/IKRig/Source/IKRigEditor/Private/IKRigEditor.cpp
@@ -58,8 +58,8 @@ void FIKRigEditor::StartupModule()
 	// register detail customizations
 	FPropertyEditorModule& PropertyEditorModule = FModuleManager::LoadModuleChecked<FPropertyEditorModule>("PropertyEditor");
 	// custom instanced struct details
-	PropertyEditorModule.RegisterCustomClassLayout(UStructViewer::StaticClass()->GetFName(), FOnGetDetailCustomizationInstance::CreateStatic(&FStructViewerCustomization::MakeInstance));
-	ClassesToUnregisterOnShutdown.Add(UStructViewer::StaticClass()->GetFName());
+	PropertyEditorModule.RegisterCustomClassLayout(UIKRigStructViewer::StaticClass()->GetFName(), FOnGetDetailCustomizationInstance::CreateStatic(&FStructViewerCustomization::MakeInstance));
+	ClassesToUnregisterOnShutdown.Add(UIKRigStructViewer::StaticClass()->GetFName());
 	// custom IK rig bone details
 	PropertyEditorModule.RegisterCustomClassLayout(UIKRigBoneDetails::StaticClass()->GetFName(), FOnGetDetailCustomizationInstance::CreateStatic(&FIKRigGenericDetailCustomization::MakeInstance));
 	ClassesToUnregisterOnShutdown.Add(UIKRigBoneDetails::StaticClass()->GetFName());
@@ -69,15 +69,6 @@ void FIKRigEditor::StartupModule()
 	// custom retargeter bone details
 	PropertyEditorModule.RegisterCustomClassLayout(UIKRetargetBoneDetails::StaticClass()->GetFName(), FOnGetDetailCustomizationInstance::CreateStatic(&FIKRetargetBoneDetailCustomization::MakeInstance));
 	ClassesToUnregisterOnShutdown.Add(UIKRetargetBoneDetails::StaticClass()->GetFName());
-	// custom retargeter chain details
-	PropertyEditorModule.RegisterCustomClassLayout(URetargetChainSettings::StaticClass()->GetFName(), FOnGetDetailCustomizationInstance::CreateStatic(&FRetargetChainSettingsCustomization::MakeInstance));
-	ClassesToUnregisterOnShutdown.Add(URetargetChainSettings::StaticClass()->GetFName());
-	// custom retargeter root details
-	PropertyEditorModule.RegisterCustomClassLayout(URetargetRootSettings::StaticClass()->GetFName(), FOnGetDetailCustomizationInstance::CreateStatic(&FRetargetRootSettingsCustomization::MakeInstance));
-	ClassesToUnregisterOnShutdown.Add(URetargetRootSettings::StaticClass()->GetFName());
-	// custom retargeter global details
-	PropertyEditorModule.RegisterCustomClassLayout(UIKRetargetGlobalSettings::StaticClass()->GetFName(), FOnGetDetailCustomizationInstance::CreateStatic(&FRetargetGlobalSettingsCustomization::MakeInstance));
-	ClassesToUnregisterOnShutdown.Add(UIKRetargetGlobalSettings::StaticClass()->GetFName());
 
 	// register a thumbnail renderer for the assets
 	UThumbnailManager::Get().RegisterCustomRenderer(UIKRigDefinition::StaticClass(), UIKRigThumbnailRenderer::StaticClass());
diff --git a/Engine/Plugins/Animation/IKRig/Source/IKRigEditor/Private/RetargetEditor/IKRetargetApplicationMode.cpp b/Engine/Plugins/Animation/IKRig/Source/IKRigEditor/Private/RetargetEditor/IKRetargetApplicationMode.cpp
--- a/Engine/Plugins/Animation/IKRig/Source/IKRigEditor/Private/RetargetEditor/IKRetargetApplicationMode.cpp
+++ b/Engine/Plugins/Animation/IKRig/Source/IKRigEditor/Private/RetargetEditor/IKRetargetApplicationMode.cpp
@@ -46,7 +46,7 @@ FIKRetargetApplicationMode::FIKRetargetApplicationMode(
 	TabFactories.RegisterFactory(MakeShared<FIKRetargetOpStackTabSummoner>(IKRetargetEditor));
 
 	// create tab layout
-	TabLayout = FTabManager::NewLayout("Standalone_IKRetargetEditor_Layout_v1.019")
+	TabLayout = FTabManager::NewLayout("Standalone_IKRetargetEditor_Layout_v1.020")
 		->AddArea
 		(
 			FTabManager::NewPrimaryArea()
@@ -71,6 +71,7 @@ FIKRetargetApplicationMode::FIKRetargetApplicationMode(
 							FTabManager::NewStack()
 							->SetSizeCoefficient(0.6f)
 							->AddTab(FIKRetargetOpStackTabSummoner::TabID, ETabState::OpenedTab)
+							->AddTab(FIKRetargetHierarchyTabSummoner::TabID, ETabState::OpenedTab)
 						)
 						->Split
 						(
@@ -109,7 +110,6 @@ FIKRetargetApplicationMode::FIKRetargetApplicationMode(
 					FTabManager::NewStack()
 						->SetSizeCoefficient(0.7f)
 						->AddTab(FPersonaTabs::DetailsID, ETabState::OpenedTab)
-						->AddTab(FIKRetargetHierarchyTabSummoner::TabID, ETabState::OpenedTab)
 						->AddTab(FPersonaTabs::AdvancedPreviewSceneSettingsID, ETabState::OpenedTab)
 						->SetForegroundTab(FPersonaTabs::DetailsID)
 					)
diff --git a/Engine/Plugins/Animation/IKRig/Source/IKRigEditor/Private/RetargetEditor/IKRetargetBatchOperation.cpp b/Engine/Plugins/Animation/IKRig/Source/IKRigEditor/Private/RetargetEditor/IKRetargetBatchOperation.cpp
--- a/Engine/Plugins/Animation/IKRig/Source/IKRigEditor/Private/RetargetEditor/IKRetargetBatchOperation.cpp
+++ b/Engine/Plugins/Animation/IKRig/Source/IKRigEditor/Private/RetargetEditor/IKRetargetBatchOperation.cpp
@@ -25,6 +25,7 @@
 #include "Retargeter/RetargetOps/CurveRemapOp.h"
 #include "ObjectTools.h"
 #include "AssetRegistry/AssetRegistryModule.h"
+#include "Retargeter/RetargetOps/SpeedPlantingOp.h"
 
 #define LOCTEXT_NAMESPACE "RetargetBatchOperation"
 
@@ -334,7 +335,7 @@ void UIKRetargetBatchOperation::ConvertAnimation(
 	UObject* TransientOuter = Cast<UObject>(GetTransientPackage());
 	FIKRetargetProcessor Processor;
 	FRetargetProfile RetargetProfile;
-	Context.IKRetargetAsset->FillProfileWithAssetSettings(RetargetProfile);
+	RetargetProfile.FillProfileWithAssetSettings(Context.IKRetargetAsset);
 	Processor.Initialize(Context.SourceMesh, Context.TargetMesh, Context.IKRetargetAsset, RetargetProfile);
 	if (!Processor.IsInitialized())
 	{
@@ -358,10 +359,15 @@ void UIKRetargetBatchOperation::ConvertAnimation(
 
 	TArray<FTransform> SourceComponentPose;
 	SourceComponentPose.SetNum(NumSourceBones);
-
-	// get names of the curves the retargeter is looking for
+	
+	// get list of all speed curves referenced by the speed planting ops (if there are any)
 	TArray<FName> SpeedCurveNames;
-	Context.IKRetargetAsset->GetSpeedCurveNames(SpeedCurveNames);
+	const TArray<FIKRetargetSpeedPlantingOp*> SpeedPlantingOps = Context.IKRetargetAsset->GetAllRetargetOpsOfType<FIKRetargetSpeedPlantingOp>();
+	for (FIKRetargetSpeedPlantingOp* SpeedPlantingOp : SpeedPlantingOps)
+	{
+		SpeedCurveNames.Append(SpeedPlantingOp->GetRequiredSpeedCurves());
+	}
+	
 	// while retargeting, record a map of (sequence name, curve name) for any curve missing in a sequence
 	TMap<FString, TSet<FName>> MissingSpeedCurves;
 	
@@ -411,8 +417,8 @@ void UIKRetargetBatchOperation::ConvertAnimation(
 		EvaluationOptions.bExtractRootMotion = false;
 		EvaluationOptions.bIncorporateRootMotionIntoPose = true;
 
-		// reset the planting state
-		Processor.ResetPlanting();
+		// reset playback of ops
+		Processor.OnPlaybackReset();
 		
 		// retarget each frame's pose from source to target
 		for (int32 FrameIndex=0; FrameIndex<NumFrames; ++FrameIndex)
@@ -440,9 +446,6 @@ void UIKRetargetBatchOperation::ConvertAnimation(
 			{
 				Transform.SetScale3D(FVector::OneVector);
 			}
-
-			// update goals 
-			Processor.CopyIKRigSettingsFromAsset();
 			
 			// calculate the delta time
 			const float TimeAtCurrentFrame = SourceSequence->GetTimeAtFrame(FrameIndex);
@@ -470,10 +473,10 @@ void UIKRetargetBatchOperation::ConvertAnimation(
 
 			// get the settings profile
 			FRetargetProfile SettingsProfile;
-			Context.IKRetargetAsset->FillProfileWithAssetSettings(SettingsProfile);
+			SettingsProfile.FillProfileWithAssetSettings(Context.IKRetargetAsset);
 
 			// run the retargeter
-			const TArray<FTransform>& TargetComponentPose = Processor.RunRetargeter(SourceComponentPose, SpeedCurveValues, DeltaTime, SettingsProfile);
+			const TArray<FTransform>& TargetComponentPose = Processor.RunRetargeter(SourceComponentPose, SettingsProfile, DeltaTime);
 
 			// convert to a local-space pose
 			TArray<FTransform> TargetLocalPose = TargetComponentPose;
diff --git a/Engine/Plugins/Animation/IKRig/Source/IKRigEditor/Private/RetargetEditor/IKRetargetCommands.cpp b/Engine/Plugins/Animation/IKRig/Source/IKRigEditor/Private/RetargetEditor/IKRetargetCommands.cpp
--- a/Engine/Plugins/Animation/IKRig/Source/IKRigEditor/Private/RetargetEditor/IKRetargetCommands.cpp
+++ b/Engine/Plugins/Animation/IKRig/Source/IKRigEditor/Private/RetargetEditor/IKRetargetCommands.cpp
@@ -11,9 +11,9 @@ void FIKRetargetCommands::RegisterCommands()
 	UI_COMMAND(ShowRetargetPose, "Show Retarget Pose", "Display the retarget pose and retarget to the target skeleton.", EUserInterfaceActionType::Button, FInputChord());
 
 	UI_COMMAND(ShowAssetSettings, "Asset Settings", "Set the source or target IK Rigs, preview meshes. Adjust debug settings.", EUserInterfaceActionType::ToggleButton, FInputChord());
-	UI_COMMAND(ShowGlobalSettings, "Global Settings", "Access Stride Warping and retarget phases.", EUserInterfaceActionType::ToggleButton, FInputChord());
-	UI_COMMAND(ShowRootSettings, "Root Settings", "Adjust the motion of the Pelvis.", EUserInterfaceActionType::ToggleButton, FInputChord());
-	
+
+	UI_COMMAND(AddDefaultOps, "Add Default Ops", "Add Pelvis Motion, IK/FK Chains and Root Motion Ops. Runs auto-setup on each op.", EUserInterfaceActionType::ToggleButton, FInputChord());
+
 	UI_COMMAND(ResetAllBones, "Reset All", "Sets the selected bones to the mesh reference pose.", EUserInterfaceActionType::Button, FInputChord());
 	UI_COMMAND(ResetSelectedBones, "Reset Selected Bones", "Sets the selected bones to the mesh reference pose.", EUserInterfaceActionType::Button, FInputChord());
 	UI_COMMAND(ResetSelectedAndChildrenBones, "Reset Selected and Children Bones", "Sets the selected bones (and all children recursively) to the mesh reference pose.", EUserInterfaceActionType::Button, FInputChord());
diff --git a/Engine/Plugins/Animation/IKRig/Source/IKRigEditor/Private/RetargetEditor/IKRetargetDefaultMode.cpp b/Engine/Plugins/Animation/IKRig/Source/IKRigEditor/Private/RetargetEditor/IKRetargetDefaultMode.cpp
--- a/Engine/Plugins/Animation/IKRig/Source/IKRigEditor/Private/RetargetEditor/IKRetargetDefaultMode.cpp
+++ b/Engine/Plugins/Animation/IKRig/Source/IKRigEditor/Private/RetargetEditor/IKRetargetDefaultMode.cpp
@@ -4,13 +4,13 @@
 
 #include "EditorViewportClient.h"
 #include "AssetEditorModeManager.h"
-#include "EngineUtils.h"
-#include "IKRigDebugRendering.h"
 #include "IPersonaPreviewScene.h"
 #include "Animation/DebugSkelMeshComponent.h"
 #include "RetargetEditor/IKRetargetEditor.h"
 #include "RetargetEditor/IKRetargetHitProxies.h"
 #include "ReferenceSkeleton.h"
+#include "Retargeter/RetargetOps/FKChainsOp.h"
+#include "Retargeter/RetargetOps/PelvisMotionOp.h"
 
 
 #define LOCTEXT_NAMESPACE "IKRetargetDefaultMode"
@@ -64,184 +64,25 @@ void FIKRetargetDefaultMode::Render(const FSceneView* View, FViewport* Viewport,
 
 void FIKRetargetDefaultMode::RenderDebugProxies(FPrimitiveDrawInterface* PDI, const FIKRetargetEditorController* Controller) const
 {
+	// asset settings can disable debug drawing
 	const UIKRetargeter* Asset = Controller->AssetController->GetAsset();
 	if (!Asset->bDebugDraw)
 	{
 		return;
 	}
-	
+
+	// skip until processor has been initialized
 	const FIKRetargetProcessor* RetargetProcessor = Controller->GetRetargetProcessor();
 	if (!(RetargetProcessor && RetargetProcessor->IsInitialized()))
 	{
 		return;
 	}
 
+	// debug draw the currently selected op
+	int32 SelectedOpIndex = Controller->GetSelectedOpIndex();
 	UDebugSkelMeshComponent* TargetSkelMesh = Controller->GetSkeletalMeshComponent(ERetargetSourceOrTarget::Target);
-	const FTransform ComponentTransform = TargetSkelMesh->GetComponentTransform();
-	const float ComponentScale = static_cast<float>(ComponentTransform.GetScale3D().GetMax());
-
-	const TArray<FName>& SelectedChains = Controller->GetSelectedChains();
-
-	constexpr FLinearColor Muted = FLinearColor(0.5,0.5,0.5, 0.5);
-	const FLinearColor SourceColor = (FLinearColor::Gray * FLinearColor::Blue) * Muted;
-	const FLinearColor GoalColor = FLinearColor::Yellow;
-	const FLinearColor MainColor = FLinearColor::Green;
-	const FLinearColor NonSelected = FLinearColor::Gray * 0.3f;
-
-	// draw IK goals on each IK chain
-	if (Asset->bDrawFinalGoals || Asset->bDrawSourceLocations)
-	{
-		// get the root modification
-		const FRootRetargeter& RootRetargeter = RetargetProcessor->GetRootRetargeter();
-		const FVector RootModification = RootRetargeter.Target.RootTranslationDelta * RootRetargeter.Settings.GetAffectIKWeightVector();
-		
-		// spin through all IK chains
-		const TArray<FRetargetChainPairIK>& IKChainPairs = RetargetProcessor->GetIKChainPairs();
-		for (const FRetargetChainPairIK& IKChainPair : IKChainPairs)
-		{
-			const FChainDebugData& ChainDebugData = IKChainPair.IKChainRetargeter.DebugData;
-			FTransform FinalTransform = ChainDebugData.OutputTransformEnd * ComponentTransform;
-
-			const bool bIsSelected = SelectedChains.Contains(IKChainPair.TargetBoneChainName);
-
-			PDI->SetHitProxy(new HIKRetargetEditorChainProxy(IKChainPair.TargetBoneChainName));
-
-			if (Asset->bDrawFinalGoals)
-			{
-				IKRigDebugRendering::DrawWireCube(
-				PDI,
-				FinalTransform,
-				bIsSelected ? GoalColor : GoalColor * NonSelected,
-				Asset->ChainDrawSize,
-				Asset->ChainDrawThickness * ComponentScale);
-			}
-		
-			if (Asset->bDrawSourceLocations)
-			{
-				const FSourceChainIK& SourceChain = IKChainPair.IKChainRetargeter.Source;
-				FTransform SourceGoalTransform;
-				SourceGoalTransform.SetTranslation(SourceChain.CurrentEndPosition + RootModification);
-				SourceGoalTransform.SetRotation(SourceChain.CurrentEndRotation);
-				SourceGoalTransform *= ComponentTransform;
-
-				FLinearColor Color = bIsSelected ? SourceColor : SourceColor * NonSelected;
-
-				DrawWireSphere(
-					PDI,
-					SourceGoalTransform,
-					Color,
-					Asset->ChainDrawSize * 0.5f,
-					12,
-					SDPG_World,
-					0.0f,
-					0.001f,
-					false);
-
-				if (Asset->bDrawFinalGoals)
-				{
-					DrawDashedLine(
-						PDI,
-						SourceGoalTransform.GetLocation(),
-						FinalTransform.GetLocation(),
-						Color,
-						1.0f,
-						SDPG_Foreground);
-				}
-			}
-
-			// done drawing chain proxies
-			PDI->SetHitProxy(nullptr);
-		}
-	}
-	
-
-	// draw lines on each FK chain
-	if (Asset->bDrawChainLines || Asset->bDrawSingleBoneChains)
-	{
-		const TArray<FRetargetChainPairFK>& FKChainPairs = RetargetProcessor->GetFKChainPairs();
-		for (const FRetargetChainPairFK& FKChainPair : FKChainPairs)
-		{
-			const TArray<int32>& TargetChainBoneIndices = FKChainPair.FKDecoder.BoneIndices;
-			if (TargetChainBoneIndices.IsEmpty())
-			{
-				continue;
-			}
-		
-			const bool bIsSelected = SelectedChains.Contains(FKChainPair.TargetBoneChainName);
-			FLinearColor Color = bIsSelected ? MainColor : MainColor * NonSelected;
-		
-			// draw a line from start to end of chain, or in the case of a chain with only 1 bone in it, draw a sphere
-			PDI->SetHitProxy(new HIKRetargetEditorChainProxy(FKChainPair.TargetBoneChainName));
-			if (Asset->bDrawChainLines && TargetChainBoneIndices.Num() > 1)
-			{
-				FTransform StartTransform = TargetSkelMesh->GetBoneTransform(TargetChainBoneIndices[0], ComponentTransform);
-				FTransform EndTransform = TargetSkelMesh->GetBoneTransform(TargetChainBoneIndices.Last(), ComponentTransform);
-				PDI->DrawLine(
-				StartTransform.GetLocation(),
-				EndTransform.GetLocation(),
-				Color,
-				SDPG_Foreground,
-				Asset->ChainDrawThickness * ComponentScale);
-			}
-			else if (Asset->bDrawSingleBoneChains)
-			{
-				// single bone chain, just draw a sphere on the bone
-				FTransform BoneTransform = TargetSkelMesh->GetBoneTransform(TargetChainBoneIndices[0], ComponentTransform);
-				DrawWireSphere(
-					PDI,
-					BoneTransform,
-					Color,
-					Asset->ChainDrawSize,
-					12,
-					SDPG_World,
-					Asset->ChainDrawThickness * ComponentScale,
-					0.001f,
-					false);
-			}
-		
-			PDI->SetHitProxy(nullptr);
-		}
-	}
-
-	// draw stride warping frame
-	if (Asset->bDrawWarpingFrame)
-	{
-		FTransform WarpingFrame = RetargetProcessor->DebugData.StrideWarpingFrame * ComponentTransform;
-		DrawCoordinateSystem(
-			PDI,
-			WarpingFrame.GetLocation(),
-			WarpingFrame.GetRotation().Rotator(),
-			Asset->ChainDrawSize * ComponentScale,
-			SDPG_World,
-			Asset->ChainDrawThickness * ComponentScale);	
-	}
-
-	// root bone name
-	if (Asset->bDrawRootCircle)
-	{
-		const FName RootBoneName = Controller->AssetController->GetRetargetRootBone(ERetargetSourceOrTarget::Target);
-		const int32 RootBoneIndex = TargetSkelMesh->GetReferenceSkeleton().FindBoneIndex(RootBoneName);
-		if (RootBoneIndex != INDEX_NONE)
-		{
-			const FTransform RootTransform = TargetSkelMesh->GetBoneTransform(RootBoneIndex, ComponentTransform);
-			const FVector RootCircleLocation = RootTransform.GetLocation() * FVector(1,1,0);
-			const bool bIsSelected = EditorController.Pin()->GetRootSelected();
-			const FLinearColor RootColor = bIsSelected ? MainColor : MainColor * NonSelected;
-	
-			PDI->SetHitProxy(new HIKRetargetEditorRootProxy());
-			DrawCircle(
-				PDI,
-				RootCircleLocation,
-				FVector(1, 0, 0),
-				FVector(0, 1, 0),
-				RootColor,
-				Asset->ChainDrawSize * 10.f * ComponentTransform.GetScale3D().GetMax(),
-				30,
-				SDPG_World,
-				Asset->ChainDrawThickness * 2.0f * ComponentScale);
-			PDI->SetHitProxy(nullptr);
-		}	
-	}
+	const FIKRetargetDebugDrawState& EditorState = Controller->GetSelectionState();
+	RetargetProcessor->DebugDrawOpAtIndex(SelectedOpIndex, PDI, TargetSkelMesh->GetComponentTransform(), EditorState);
 }
 
 bool FIKRetargetDefaultMode::HandleClick(FEditorViewportClient* InViewportClient, HHitProxy* HitProxy, const FViewportClick& Click)
diff --git a/Engine/Plugins/Animation/IKRig/Source/IKRigEditor/Private/RetargetEditor/IKRetargetDetails.cpp b/Engine/Plugins/Animation/IKRig/Source/IKRigEditor/Private/RetargetEditor/IKRetargetDetails.cpp
--- a/Engine/Plugins/Animation/IKRig/Source/IKRigEditor/Private/RetargetEditor/IKRetargetDetails.cpp
+++ b/Engine/Plugins/Animation/IKRig/Source/IKRigEditor/Private/RetargetEditor/IKRetargetDetails.cpp
@@ -8,7 +8,6 @@
 #include "Animation/AnimInstance.h"
 #include "Animation/DebugSkelMeshComponent.h"
 #include "Framework/MultiBox/SToolBarButtonBlock.h"
-#include "RetargetEditor/IKRetargetEditor.h"
 #include "RetargetEditor/IKRetargetEditorController.h"
 #include "RetargetEditor/IKRetargeterController.h"
 #include "Retargeter/IKRetargeter.h"
@@ -16,10 +15,7 @@
 #include "ScopedTransaction.h"
 #include "AnimationRuntime.h"
 #include "SPositiveActionButton.h"
-#include "SSearchableComboBox.h"
-#include "RetargetEditor/SRetargetOpStack.h"
-#include "Retargeter/IKRetargetOps.h"
-#include "UObject/UnrealTypePrivate.h"
+#include "Editor.h"
 
 #include UE_INLINE_GENERATED_CPP_BY_NAME(IKRetargetDetails)
 
@@ -27,6 +23,7 @@
 #include "HAL/PlatformApplicationMisc.h"
 #endif
 
+
 #define LOCTEXT_NAMESPACE "IKRetargeterDetails"
 
 
@@ -1008,477 +1005,5 @@ void FIKRetargetBoneDetailCustomization::GetTransformUIData(
 	}
 }
 
-void FRetargetChainSettingsCustomization::CustomizeDetails(IDetailLayoutBuilder& DetailBuilder)
-{
-	TArray<TWeakObjectPtr<UObject>> ObjectsBeingCustomized = DetailBuilder.GetSelectedObjects();
-	ChainSettingsObjects.Reset();
-	for (const TWeakObjectPtr<UObject>& Object : ObjectsBeingCustomized)
-	{
-		if (URetargetChainSettings* ChainSettings = Cast<URetargetChainSettings>(Object.Get()))
-		{
-			ChainSettingsObjects.Add(ChainSettings);
-		}
-	}
-	
-	if (ChainSettingsObjects.IsEmpty())
-	{
-		return;
-	}
-
-	FIKRetargetEditorController* Controller = ChainSettingsObjects[0]->EditorController.Pin().Get();
-	if (!Controller)
-	{
-		return;
-	}
-
-	// hide the settings category and build our own
-	DetailBuilder.HideCategory("Settings");
-	DetailBuilder.HideCategory("Chain Mapping");
-
-	// determine the current setup state of the chain (is it ready for FK and/or IK?)
-	// FK: it can only run FK retarget if it's mapped to a source chain
-	// IK: can only run IK if FK is setup AND it's mapped to a goal that is connected to a solver (in the target IK rig)
-	auto CheckChainSetupForIKAndFK([this](
-		FIKRetargetEditorController* Controller,
-		const TObjectPtr<URetargetChainSettings> ChainSettings,
-		bool& OutFKSetup,
-		bool& OutIKSetup,
-		bool& OutIKConnected)
-	{
-		OutFKSetup = true;
-		OutIKSetup = true;
-		OutIKConnected = true;
-		
-		if (ChainSettings->SourceChain == NAME_None || ChainSettings->TargetChain == NAME_None)
-		{
-			OutFKSetup = false;
-			OutIKSetup = false;
-			OutIKConnected = false;
-			return;
-		}
-
-		const FName ChainGoal = Controller->AssetController->GetChainGoal(ChainSettings.Get());
-		if (ChainGoal == NAME_None)
-		{
-			OutIKSetup = false;
-		}
-
-		if (!Controller->AssetController->IsChainGoalConnectedToASolver(ChainGoal))
-		{
-			OutIKConnected = false;
-		}
-	});
-
-	// get the FK and IK status of the first selected chain
-	bool bFirstChainFKSetup;
-	bool bFirstChainIKSetup;
-	bool bFirstChainIKConnected;
-	CheckChainSetupForIKAndFK(Controller, ChainSettingsObjects[0].Get(), bFirstChainFKSetup, bFirstChainIKSetup, bFirstChainIKConnected);
-
-	// check to see if there is a mix of chains with valid FK or IK setups
-	bool bIsFKSetupMultipleValues = false;
-	bool bIsIKSetupMultipleValues = false;
-	bool bIsIKConnectedMultipleValues = false;
-	for (const TWeakObjectPtr<URetargetChainSettings> ChainSettings : ChainSettingsObjects)
-	{
-		bool bIsFKSetup;
-		bool bIsIKSetup;
-		bool bIsIKConnected;
-		CheckChainSetupForIKAndFK(Controller, ChainSettings.Get(), bIsFKSetup, bIsIKSetup, bIsIKConnected);
-		bIsFKSetupMultipleValues = bIsFKSetup != bFirstChainFKSetup ? true : bIsFKSetupMultipleValues;
-		bIsIKSetupMultipleValues = bIsIKSetup != bFirstChainIKSetup ? true : bIsIKSetupMultipleValues;
-		bIsIKConnectedMultipleValues = bIsIKConnected != bFirstChainIKConnected ? true : bIsIKConnectedMultipleValues;
-	}
-
-	// are we editing multiple chains?
-	const bool bEditingMultipleChains = ChainSettingsObjects.Num() > 1;
-	TObjectPtr<URetargetChainSettings> PrimaryChainSettingsObject = ChainSettingsObjects[0].Get();
-
-	// get source chain options for combobox
-	SourceChainOptions.Reset();
-	SourceChainOptions.Add(MakeShareable(new FString(TEXT("None"))));
-	if (const UIKRigDefinition* SourceIKRig = Controller->AssetController->GetIKRig(ERetargetSourceOrTarget::Source))
-	{
-		const TArray<FBoneChain>& Chains = SourceIKRig->GetRetargetChains();
-		for (const FBoneChain& BoneChain : Chains)
-		{
-			SourceChainOptions.Add(MakeShareable(new FString(BoneChain.ChainName.ToString())));
-		}
-	}
-	
-	// replace the "Settings" struct property
-	const FText ChainTitleName = bEditingMultipleChains ?  LOCTEXT("MultipleChains_Label", "Multiple Chain") : FText::FromName(PrimaryChainSettingsObject->TargetChain);
-	const FText ChainCategoryTitle = FText::Format( LOCTEXT("ChainSettingsHeaderLabel", "{0} Settings"), ChainTitleName);
-	IDetailCategoryBuilder& SettingsCategory = DetailBuilder.EditCategory("Chain Settings", ChainCategoryTitle);
-	
-	// add row to select source chain to map to
-	SettingsCategory.AddCustomRow(LOCTEXT("ChainStatus_Label", "Status"))
-	.NameContent()
-	[
-		SNew(SBox)
-		.Content()
-		[
-			SNew(STextBlock)
-			.Text(LOCTEXT( "SoureChain_Label", "Source Chain" ))
-			.Font(IDetailLayoutBuilder::GetDetailFont())
-		]
-	]
-	.ValueContent()
-	.MinDesiredWidth(166.0f)
-	[
-		SNew(SSearchableComboBox)
-		.ToolTipText(LOCTEXT("SourceChainOptionsToolTip", "Select source chain to map to this target chain."))
-		.OptionsSource(&SourceChainOptions)
-		.OnGenerateWidget_Lambda([](TSharedPtr<FString> InItem)
-		{
-			return SNew(STextBlock).Text(FText::FromString(*InItem.Get()));
-		})
-		.OnSelectionChanged_Lambda([this, Controller]( TSharedPtr<FString> InString, ESelectInfo::Type SelectInfo)
-		{
-			const FName SourceChainName = FName(*InString.Get());
-			for (TWeakObjectPtr<URetargetChainSettings> ChainMapSettings : ChainSettingsObjects)
-			{
-				Controller->AssetController->SetSourceChain(SourceChainName, ChainMapSettings->TargetChain);
-			}
-		})
-		[
-			SNew(STextBlock)
-			.Font(IPropertyTypeCustomizationUtils::GetRegularFont())
-			.Text_Lambda([PrimaryChainSettingsObject]()
-			{
-				return FText::FromName(PrimaryChainSettingsObject->SourceChain);
-			})
-		]
-	];
-
-	// FK settings
-	const FText FKNotMappedWarning = LOCTEXT( "FKDisabled_Label", "Disabled: No source chain specified." );
-	const bool bFKEditingEnabled = bFirstChainFKSetup && !bIsFKSetupMultipleValues;
-	AddSettingsSection(
-		DetailBuilder,
-		Controller,
-		SettingsCategory,
-		GET_MEMBER_NAME_STRING_CHECKED(FTargetChainSettings, FK),
-		FName("FK"),
-		LOCTEXT("FKGroup_Label", "FK"),
-		FTargetChainFKSettings::StaticStruct(),
-		GET_MEMBER_NAME_STRING_CHECKED(FTargetChainFKSettings, EnableFK),
-		bFKEditingEnabled,
-		FKNotMappedWarning);
-
-	// IK settings
-	const FText NoIKGoalWarning = LOCTEXT( "IKDisabled_Label", "Disabled: No IK Goal specified." );
-	const FText IKNotConnectedWarning = LOCTEXT( "IKDisconnected_Label", "Disabled: IK Goal is not connected to any solvers." );
-	const FText IKEnabledMultipleValuesWarning = LOCTEXT( "IKDisabledMultipleValues_Label", "Some selected chain(s) do not have an IK Goal." );
-	const bool bIKIsConnected = bFirstChainIKConnected && !bIsIKConnectedMultipleValues;
-	const bool bIKEditingEnabled = bFirstChainIKSetup && !bIsIKSetupMultipleValues;
-	FText IKWarningToUse = FText();
-	if (!bFKEditingEnabled)
-	{
-		IKWarningToUse = FKNotMappedWarning;
-	}
-	else if (!bIKEditingEnabled)
-	{
-		IKWarningToUse = NoIKGoalWarning;
-	}
-	else if (!bIKIsConnected)
-	{
-		IKWarningToUse = IKNotConnectedWarning;
-	}
-	
-	AddSettingsSection(
-		DetailBuilder,
-		Controller,
-		SettingsCategory,
-		GET_MEMBER_NAME_STRING_CHECKED(FTargetChainSettings, IK),
-		FName("IK"),
-		LOCTEXT("IKGroup_Label", "IK"),
-		FTargetChainIKSettings::StaticStruct(),
-		GET_MEMBER_NAME_STRING_CHECKED(FTargetChainIKSettings, EnableIK),
-		bIKEditingEnabled && bIKIsConnected,
-		IKWarningToUse);
-
-	// Plant settings
-	AddSettingsSection(
-		DetailBuilder,
-		Controller,
-		SettingsCategory,
-		GET_MEMBER_NAME_STRING_CHECKED(FTargetChainSettings, SpeedPlanting),
-		FName("Speed Planting"),
-		LOCTEXT("PlantingGroup_Label", "Speed Planting"),
-		FTargetChainSpeedPlantSettings::StaticStruct(),
-		GET_MEMBER_NAME_STRING_CHECKED(FTargetChainSpeedPlantSettings, EnableSpeedPlanting),
-		bIKEditingEnabled,
-		IKWarningToUse);
-}
-
-void FRetargetChainSettingsCustomization::AddSettingsSection(
-	const IDetailLayoutBuilder& DetailBuilder,
-	const FIKRetargetEditorController* Controller,
-	IDetailCategoryBuilder& SettingsCategory,
-	const FString& StructPropertyName,
-	const FName& GroupName,
-	const FText& LocalizedGroupName,
-	const UScriptStruct* SettingsClass,
-	const FString& EnabledPropertyName,
-	const bool& bIsSectionEnabled,
-	const FText& DisabledMessage) const
-{
-	// create a group of all the properties in a settings struct,
-	// with a single "Enable X" property in the header of the details row
-	
-	const FString SettingsPropertyPath = GET_MEMBER_NAME_STRING_CHECKED(URetargetChainSettings, Settings);
-	const FString SettingsStructPropertyPath = FString::Printf(TEXT("%s.%s"), *SettingsPropertyPath, *StructPropertyName);
-	const FString EnablePropertyPath = FString::Printf(TEXT("%s.%s"), *SettingsStructPropertyPath, *EnabledPropertyName);
-	const TSharedRef<IPropertyHandle> IsEnabledProperty = DetailBuilder.GetProperty(*EnablePropertyPath);
-
-	IDetailGroup& Group = SettingsCategory.AddGroup(GroupName, LocalizedGroupName, false, true);
-
-	if (bIsSectionEnabled)
-	{
-		Group.HeaderProperty(IsEnabledProperty).CustomWidget()
-			.NameContent()
-			.VAlign(VAlign_Center)
-			[
-				SNew( STextBlock )
-				.Text( FText::FromName(GroupName) )
-				.Font( FAppStyle::GetFontStyle( TEXT( "PropertyWindow.NormalFont" ) ) )
-				.ToolTipText( IsEnabledProperty->GetToolTipText() )
-			]
-			.ValueContent()
-			[
-				SNew(SCheckBox)
-				.ToolTipText(IsEnabledProperty->GetToolTipText())
-				.IsChecked_Lambda([IsEnabledProperty]() -> ECheckBoxState
-				{
-					bool IsChecked;
-					IsEnabledProperty.Get().GetValue(IsChecked);
-					return IsChecked ? ECheckBoxState::Checked : ECheckBoxState::Unchecked;
-				})
-				.OnCheckStateChanged_Lambda([this, IsEnabledProperty, Controller](ECheckBoxState State)
-				{
-					IsEnabledProperty->SetValue(State == ECheckBoxState::Checked);
-					// clear the output log and trigger a reinitialization
-					Controller->ReinitializeRetargeterNoUIRefresh();
-				})
-			];
-	}
-	else
-	{
-		Group.HeaderRow()
-		.NameContent()
-		[
-			SNew(STextBlock)
-			.Font(IDetailLayoutBuilder::GetDetailFont())
-			.Text(LocalizedGroupName)
-		]
-		.ValueContent()
-		[
-			SNew(STextBlock)
-			.Font(IDetailLayoutBuilder::GetDetailFont())
-			.Text(DisabledMessage)
-		];
-	}
-	
-	const TAttribute<bool> EditCondition = TAttribute<bool>::Create([IsEnabledProperty, bIsSectionEnabled]()
-	{
-		bool bIsEnabledPropertyTrue = true;
-		IsEnabledProperty->GetValue(bIsEnabledPropertyTrue);
-		return bIsSectionEnabled && bIsEnabledPropertyTrue;
-	});
-	
-	for (TFieldIterator<FProperty> It(SettingsClass); It; ++It)
-	{
-		if (It->GetName() == EnabledPropertyName)
-		{
-			continue;
-		}
-		const FString PropertyPath = FString::Printf(TEXT("%s.%s"), *SettingsStructPropertyPath, *It->GetName());
-		Group.AddPropertyRow(DetailBuilder.GetProperty(*PropertyPath)).EditCondition(EditCondition, nullptr);
-	}
-}
-
-void FRetargetRootSettingsCustomization::CustomizeDetails(IDetailLayoutBuilder& DetailBuilder)
-{
-	TArray<TWeakObjectPtr<UObject>> ObjectsBeingCustomized = DetailBuilder.GetSelectedObjects();
-	if (ObjectsBeingCustomized.Num() < 0)
-	{
-		return;
-	}
-	
-	RootSettingsObject = Cast<URetargetRootSettings>(ObjectsBeingCustomized[0].Get());
-	if (!RootSettingsObject.IsValid())
-	{
-		return;
-	}
-
-	Controller = RootSettingsObject->EditorController;
-	if (!Controller.IsValid())
-	{
-		return;
-	}
-
-	const FString StructPropertyPath = GET_MEMBER_NAME_STRING_CHECKED(URetargetRootSettings, Settings);
-
-	auto AddPropertyRowToGroup = [&StructPropertyPath, &DetailBuilder](IDetailGroup& Group, const FString& PropertyName)
-	{
-		const FString PropertyPath = FString::Printf(TEXT("%s.%s"), *StructPropertyPath, *PropertyName);
-		Group.AddPropertyRow(DetailBuilder.GetProperty(*PropertyPath));
-	};
-
-	DetailBuilder.HideCategory("Settings");
-
-	// replace the "Settings" struct property
-	IDetailCategoryBuilder& SettingsCategory = DetailBuilder.EditCategory(FName("Root Settings"));
-
-	// alpha group
-	IDetailGroup& AlphaGroup = SettingsCategory.AddGroup("Alpha", LOCTEXT("Alpha_Label", "Alpha"), false, true);
-	AddPropertyRowToGroup(AlphaGroup,GET_MEMBER_NAME_STRING_CHECKED(FTargetRootSettings, RotationAlpha));
-	AddPropertyRowToGroup(AlphaGroup,GET_MEMBER_NAME_STRING_CHECKED(FTargetRootSettings, TranslationAlpha));
-
-	// blend to source group
-	IDetailGroup& BlendToSourceGroup = SettingsCategory.AddGroup("Blend To Source", LOCTEXT("BlendToSource_Label", "Blend to Source"), false, true);
-	AddPropertyRowToGroup(BlendToSourceGroup,GET_MEMBER_NAME_STRING_CHECKED(FTargetRootSettings, BlendToSource));
-	AddPropertyRowToGroup(BlendToSourceGroup,GET_MEMBER_NAME_STRING_CHECKED(FTargetRootSettings, BlendToSourceWeights));
-
-	// scale group
-	IDetailGroup& ScaleGroup = SettingsCategory.AddGroup("Scale Translation", LOCTEXT("ScaleRoot_Label", "Scale Translation"), false, true);
-	AddPropertyRowToGroup(ScaleGroup,GET_MEMBER_NAME_STRING_CHECKED(FTargetRootSettings, ScaleHorizontal));
-	AddPropertyRowToGroup(ScaleGroup,GET_MEMBER_NAME_STRING_CHECKED(FTargetRootSettings, ScaleVertical));
-
-	// offset group
-	IDetailGroup& OffsetGroup = SettingsCategory.AddGroup("Offsets", LOCTEXT("OffsetRoot_Label", "Offsets"), false, true);
-	AddPropertyRowToGroup(OffsetGroup,GET_MEMBER_NAME_STRING_CHECKED(FTargetRootSettings, TranslationOffset));
-	AddPropertyRowToGroup(OffsetGroup,GET_MEMBER_NAME_STRING_CHECKED(FTargetRootSettings, RotationOffset));
-
-	// affect IK directionally
-	IDetailGroup& AffectIKGroup = SettingsCategory.AddGroup("Affect IK", LOCTEXT("AffectIK_Label", "Affect IK"), false, true);
-	AddPropertyRowToGroup(AffectIKGroup,GET_MEMBER_NAME_STRING_CHECKED(FTargetRootSettings, AffectIKHorizontal));
-	AddPropertyRowToGroup(AffectIKGroup,GET_MEMBER_NAME_STRING_CHECKED(FTargetRootSettings, AffectIKVertical));
-}
-
-void FRetargetGlobalSettingsCustomization::CustomizeDetails(IDetailLayoutBuilder& DetailBuilder)
-{
-	TArray<TWeakObjectPtr<UObject>> ObjectsBeingCustomized = DetailBuilder.GetSelectedObjects();
-	if (ObjectsBeingCustomized.Num() < 0)
-	{
-		return;
-	}
-	
-	GlobalSettingsObject = Cast<UIKRetargetGlobalSettings>(ObjectsBeingCustomized[0].Get());
-	if (!GlobalSettingsObject.IsValid())
-	{
-		return;
-	}
-
-	Controller = GlobalSettingsObject->EditorController;
-	if (!Controller.IsValid())
-	{
-		return;
-	}
-
-	const FString StructPropertyPath = GET_MEMBER_NAME_STRING_CHECKED(URetargetRootSettings, Settings);
-
-	auto GetPropertyHandle = [&DetailBuilder, &StructPropertyPath](const FString& PropertyName)
-	{
-		const FString PropertyPath = FString::Printf(TEXT("%s.%s"), *StructPropertyPath, *PropertyName);
-		return DetailBuilder.GetProperty(*PropertyPath);
-	};
-
-	DetailBuilder.HideCategory("Settings");
-
-	// replace the "Settings" struct property
-	IDetailCategoryBuilder& SettingsCategory = DetailBuilder.EditCategory(FName("Global Settings"));
-
-	// global scale group
-	IDetailGroup& SourceScaleGroup = SettingsCategory.AddGroup("GlobalScaleFactor", LOCTEXT("GlobalScale_Label", "Global Scale Factor"), false, true);
-	SourceScaleGroup.AddPropertyRow( GetPropertyHandle(GET_MEMBER_NAME_STRING_CHECKED(FRetargetGlobalSettings, SourceScaleFactor)));
-	
-	// phases group
-	IDetailGroup& AlphaGroup = SettingsCategory.AddGroup("Phases", LOCTEXT("Phases_Label", "Global Retarget Phases"), false, true);
-	AlphaGroup.AddPropertyRow(GetPropertyHandle(GET_MEMBER_NAME_STRING_CHECKED(FRetargetGlobalSettings, bEnableRoot)));
-	AlphaGroup.AddPropertyRow(GetPropertyHandle(GET_MEMBER_NAME_STRING_CHECKED(FRetargetGlobalSettings, bEnableFK)));
-	AlphaGroup.AddPropertyRow(GetPropertyHandle(GET_MEMBER_NAME_STRING_CHECKED(FRetargetGlobalSettings, bEnableIK)));
-	AlphaGroup.AddPropertyRow(GetPropertyHandle(GET_MEMBER_NAME_STRING_CHECKED(FRetargetGlobalSettings, bEnablePost)));
-
-	// copy pose group
-	IDetailGroup& CopyPoseGroup = SettingsCategory.AddGroup("CopyPose", LOCTEXT("CopyPose_Label", "Copy Pose from Source"), false, true);
-	CopyPoseGroup.AddPropertyRow(GetPropertyHandle(GET_MEMBER_NAME_STRING_CHECKED(FRetargetGlobalSettings, bCopyBasePose)));
-	CopyPoseGroup.AddPropertyRow(GetPropertyHandle(GET_MEMBER_NAME_STRING_CHECKED(FRetargetGlobalSettings, CopyBasePoseRoot)));
-	
-	// stride warping group
-	const FText WarpingTitleLabel = LOCTEXT("Warping_Label", "Stride Warping");
-	IDetailGroup& WarpingGroup = SettingsCategory.AddGroup("Warping", WarpingTitleLabel, false, true);
-
-	// header with "enable" checkbox that disables whole group
-	const TSharedRef<IPropertyHandle> EnableWarpingPropertyHandle = GetPropertyHandle(GET_MEMBER_NAME_STRING_CHECKED(FRetargetGlobalSettings, bWarping));
-	WarpingGroup.HeaderProperty(EnableWarpingPropertyHandle).DisplayName(WarpingTitleLabel);
-
-	// add all the warping parameters
-	const TAttribute<bool> WarpingEditCondition = TAttribute<bool>::Create([EnableWarpingPropertyHandle]()
-	{
-		bool bIsEnabledPropertyTrue = true;
-		EnableWarpingPropertyHandle->GetValue(bIsEnabledPropertyTrue);
-		return bIsEnabledPropertyTrue;
-	});
-	WarpingGroup.AddPropertyRow(GetPropertyHandle(GET_MEMBER_NAME_STRING_CHECKED(FRetargetGlobalSettings, ForwardDirection))).EditCondition(WarpingEditCondition,nullptr);
-	WarpingGroup.AddPropertyRow(GetPropertyHandle(GET_MEMBER_NAME_STRING_CHECKED(FRetargetGlobalSettings, DirectionSource))).EditCondition(WarpingEditCondition,nullptr);
-
-	// get target chain options for combobox
-	TargetChainOptions.Reset();
-	TargetChainOptions.Add(MakeShareable(new FString(TEXT("None"))));
-	if (const UIKRigDefinition* TargetIKRig = Controller.Pin().Get()->AssetController->GetIKRig(ERetargetSourceOrTarget::Target))
-	{
-		const TArray<FBoneChain>& Chains = TargetIKRig->GetRetargetChains();
-		for (const FBoneChain& BoneChain : Chains)
-		{
-			TargetChainOptions.Add(MakeShareable(new FString(BoneChain.ChainName.ToString())));
-		}
-	}
-	
-	// add row to target chain to use for body direction
-	const TSharedRef<IPropertyHandle> DirectionChainPropertyHandle = GetPropertyHandle(GET_MEMBER_NAME_STRING_CHECKED(FRetargetGlobalSettings, DirectionChain));
-	WarpingGroup.AddPropertyRow(DirectionChainPropertyHandle).CustomWidget()
-	.IsEnabled(GlobalSettingsObject->Settings.DirectionSource == EWarpingDirectionSource::Chain)
-	.NameContent()
-	[
-		SNew(SBox)
-		.Content()
-		[
-			SNew(STextBlock)
-			.Text(LOCTEXT( "DirectionChainName_Label", "Direction Chain" ))
-			.Font(IDetailLayoutBuilder::GetDetailFont())
-		]
-	]
-	.ValueContent()
-	.MinDesiredWidth(166.0f)
-	[
-		SNew(SSearchableComboBox)
-		.ToolTipText(LOCTEXT("DirectionChainOptionsToolTip", "Select target chain (usually the Spine) to help define the forward direction of the character. This is especially useful for creatures with a horizontal body, like quadrupeds."))
-		.OptionsSource(&TargetChainOptions)
-		.OnGenerateWidget_Lambda([](TSharedPtr<FString> InItem)
-		{
-			return SNew(STextBlock).Text(FText::FromString(*InItem.Get()));
-		})
-		.OnSelectionChanged_Lambda([this](TSharedPtr<FString> InString, ESelectInfo::Type SelectInfo)
-		{
-			const FName TargetChainName = FName(*InString.Get());
-			GlobalSettingsObject->Settings.DirectionChain = TargetChainName;
-		})
-		[
-			SNew(STextBlock)
-			.Font(IPropertyTypeCustomizationUtils::GetRegularFont())
-			.Text_Lambda([this]()
-			{
-				return FText::FromName(GlobalSettingsObject->Settings.DirectionChain);
-			})
-		]
-	];
-	
-	WarpingGroup.AddPropertyRow(GetPropertyHandle(GET_MEMBER_NAME_STRING_CHECKED(FRetargetGlobalSettings, WarpForwards))).EditCondition(WarpingEditCondition,nullptr);
-	WarpingGroup.AddPropertyRow(GetPropertyHandle(GET_MEMBER_NAME_STRING_CHECKED(FRetargetGlobalSettings, WarpSplay))).EditCondition(WarpingEditCondition,nullptr);
-	WarpingGroup.AddPropertyRow(GetPropertyHandle(GET_MEMBER_NAME_STRING_CHECKED(FRetargetGlobalSettings, SidewaysOffset))).EditCondition(WarpingEditCondition,nullptr);
-}
-
 #undef LOCTEXT_NAMESPACE
 
diff --git a/Engine/Plugins/Animation/IKRig/Source/IKRigEditor/Private/RetargetEditor/IKRetargetEditor.cpp b/Engine/Plugins/Animation/IKRig/Source/IKRigEditor/Private/RetargetEditor/IKRetargetEditor.cpp
--- a/Engine/Plugins/Animation/IKRig/Source/IKRigEditor/Private/RetargetEditor/IKRetargetEditor.cpp
+++ b/Engine/Plugins/Animation/IKRig/Source/IKRigEditor/Private/RetargetEditor/IKRetargetEditor.cpp
@@ -130,11 +130,8 @@ void FIKRetargetEditor::BindCommands()
 		FExecuteAction::CreateSP(EditorController, &FIKRetargetEditorController::SetRetargeterMode, ERetargeterOutputMode::EditRetargetPose),
 		FCanExecuteAction::CreateSP(EditorController, &FIKRetargetEditorController::IsCurrentMeshLoaded),
 		FIsActionChecked());
-
-
-	//
-	// Show various settings in details panel
-	//
+	
+	// show asset settings in details panel
 	ToolkitCommands->MapAction(
 		Commands.ShowAssetSettings,
 		FExecuteAction::CreateLambda([this]()
@@ -148,24 +145,15 @@ void FIKRetargetEditor::BindCommands()
 			const UIKRetargeter* Asset = EditorController->AssetController->GetAsset();
 			return EditorController->IsObjectInDetailsView(Asset);	
 		}));
+
+	// add default ops to the stack
 	ToolkitCommands->MapAction(
-		Commands.ShowGlobalSettings,
-		FExecuteAction::CreateSP(EditorController, &FIKRetargetEditorController::ShowGlobalSettings),
-		FCanExecuteAction(),
-		FIsActionChecked::CreateLambda([this]() ->bool
-		{
-			const UIKRetargetGlobalSettings* GlobalSettings = EditorController->AssetController->GetAsset()->GetGlobalSettingsUObject();
-			return EditorController->IsObjectInDetailsView(GlobalSettings);	
-		}));
-	ToolkitCommands->MapAction(
-		Commands.ShowRootSettings,
-		FExecuteAction::CreateSP(EditorController, &FIKRetargetEditorController::ShowRootSettings),
-		FCanExecuteAction(),
-		FIsActionChecked::CreateLambda([this]() ->bool
+		Commands.AddDefaultOps,
+		FExecuteAction::CreateLambda([this]()
 		{
-			const URetargetRootSettings* RootSettings = EditorController->AssetController->GetAsset()->GetRootSettingsUObject();
-			return EditorController->IsObjectInDetailsView(RootSettings);	
-		}));
+			EditorController->AssetController->AddDefaultOps();
+		}),
+		FCanExecuteAction(), FIsActionChecked());
 
 	//
 	// Edit pose commands
@@ -283,8 +271,8 @@ void FIKRetargetEditor::ExtendToolbar()
 void FIKRetargetEditor::FillToolbar(FToolBarBuilder& ToolbarBuilder)
 {
 	ToolbarBuilder.BeginSection("Retarget Modes");
-	{
-		ToolbarBuilder.AddToolBarButton(
+
+	ToolbarBuilder.AddToolBarButton(
 			FExecuteAction::CreateLambda([this]{ EditorController->SetRetargetModeToPreviousMode(); }),
 			NAME_None, 
 			TAttribute<FText>::Create(TAttribute<FText>::FGetter::CreateSP(EditorController, &FIKRetargetEditorController::GetRetargeterModeLabel)),
@@ -292,21 +280,16 @@ void FIKRetargetEditor::FillToolbar(FToolBarBuilder& ToolbarBuilder)
 			TAttribute<FSlateIcon>::Create(TAttribute<FSlateIcon>::FGetter::CreateSP(EditorController, &FIKRetargetEditorController::GetCurrentRetargetModeIcon))
 		);
 		
-		ToolbarBuilder.AddComboButton(
-			FUIAction(),
-			FOnGetContent::CreateSP(this, &FIKRetargetEditor::GenerateRetargetModesMenu),
-			LOCTEXT("RetargetMode_Label", "UI Modes"),
-			LOCTEXT("RetargetMode_ToolTip", "Choose which mode to display in the viewport."),
-			FSlateIcon(FAppStyle::GetAppStyleSetName(), "LevelEditor.Recompile"),
-			true);
-
-		ToolbarBuilder.AddSeparator();
-		const TSharedPtr<SHorizontalBox> RetargetPhaseButtons = GenerateRetargetPhaseButtons();
-		ToolbarBuilder.AddWidget(RetargetPhaseButtons.ToSharedRef());
-	}
-	ToolbarBuilder.EndSection();
+	ToolbarBuilder.AddComboButton(
+		FUIAction(),
+		FOnGetContent::CreateSP(this, &FIKRetargetEditor::GenerateRetargetModesMenu),
+		LOCTEXT("RetargetMode_Label", "UI Modes"),
+		LOCTEXT("RetargetMode_ToolTip", "Choose which mode to display in the viewport."),
+		FSlateIcon(FAppStyle::GetAppStyleSetName(), "LevelEditor.Recompile"),
+		true);
 
-	ToolbarBuilder.AddSeparator();
+	ToolbarBuilder.EndSection();
+	
 	ToolbarBuilder.AddWidget(SNew(SSpacer), NAME_None, true, HAlign_Right);
 
 	ToolbarBuilder.BeginSection("Show Settings");
@@ -317,20 +300,13 @@ void FIKRetargetEditor::FillToolbar(FToolBarBuilder& ToolbarBuilder)
 		TAttribute<FText>(),
 		TAttribute<FText>(),
 		FSlateIcon(FIKRetargetEditorStyle::Get().GetStyleSetName(),"IKRetarget.AssetSettings"));
-		
-		ToolbarBuilder.AddToolBarButton(
-		FIKRetargetCommands::Get().ShowGlobalSettings,
-		NAME_None,
-		TAttribute<FText>(),
-		TAttribute<FText>(),
-		FSlateIcon(FIKRetargetEditorStyle::Get().GetStyleSetName(),"IKRetarget.GlobalSettings"));
 
 		ToolbarBuilder.AddToolBarButton(
-		FIKRetargetCommands::Get().ShowRootSettings,
+		FIKRetargetCommands::Get().AddDefaultOps,
 		NAME_None,
 		TAttribute<FText>(),
 		TAttribute<FText>(),
-		FSlateIcon(FIKRetargetEditorStyle::Get().GetStyleSetName(),"IKRetarget.RootSettings"));
+		FSlateIcon(FIKRetargetEditorStyle::Get().GetStyleSetName(),"IKRetarget.PostSettings"));
 	}
 	ToolbarBuilder.EndSection();
 }
@@ -347,121 +323,6 @@ TSharedRef<SWidget> FIKRetargetEditor::GenerateRetargetModesMenu()
 	return MenuBuilder.MakeWidget();
 }
 
-TSharedPtr<SHorizontalBox> FIKRetargetEditor::GenerateRetargetPhaseButtons() const
-{
-	constexpr float PhaseButtonMargin = 2.f;
-	
-	TSharedPtr<SHorizontalBox> Box = SNew(SHorizontalBox)
-	+ SHorizontalBox::Slot()
-	.VAlign(VAlign_Center)
-	.HAlign(HAlign_Center)
-	.AutoWidth()
-	[
-		SNew(SHorizontalBox)
-		+ SHorizontalBox::Slot()
-		.FillWidth(1.f)
-		.HAlign(HAlign_Center)
-		[
-			SNew(STextBlock).Text(FText::FromString("Retarget Phases: "))
-		]
-	]
-	+ SHorizontalBox::Slot()
-	.AutoWidth()
-	[
-		SNew(SHorizontalBox)
-		+ SHorizontalBox::Slot()
-		.AutoWidth()
-		.VAlign(VAlign_Center)
-		.Padding(FMargin(PhaseButtonMargin))
-		[
-			SNew(SCheckBox)
-			.Style(FAppStyle::Get(), "ToggleButtonCheckbox")
-			.IsChecked_Lambda([this]()
-			{
-				const bool bIsOn =  EditorController->GetGlobalSettings().bEnableRoot;
-				return bIsOn ? ECheckBoxState::Checked : ECheckBoxState::Unchecked;
-			})
-			.OnCheckStateChanged_Lambda([this](ECheckBoxState InCheckBoxState)
-			{
-				FRetargetGlobalSettings& GlobalSettings = EditorController->GetGlobalSettings();
-				GlobalSettings.bEnableRoot = !GlobalSettings.bEnableRoot;
-			})
-			[
-				SNew(STextBlock).Text(FText::FromString("Root"))
-			]
-		]
-
-		+ SHorizontalBox::Slot()
-		.AutoWidth()
-		.VAlign(VAlign_Center)
-		.Padding(FMargin(PhaseButtonMargin))
-		[
-			SNew(SCheckBox)
-			.Style(FAppStyle::Get(), "ToggleButtonCheckbox")
-			.IsChecked_Lambda([this]()
-			{
-				const bool bIsOn =  EditorController->GetGlobalSettings().bEnableFK;
-				return bIsOn ? ECheckBoxState::Checked : ECheckBoxState::Unchecked;
-			})
-			.OnCheckStateChanged_Lambda([this](ECheckBoxState InCheckBoxState)
-			{
-				FRetargetGlobalSettings& GlobalSettings = EditorController->GetGlobalSettings();
-				GlobalSettings.bEnableFK = !GlobalSettings.bEnableFK;
-			})
-			[
-				SNew(STextBlock).Text(FText::FromString("FK"))
-			]
-		]
-		
-		+ SHorizontalBox::Slot()
-		.AutoWidth()
-		.VAlign(VAlign_Center)
-		.Padding(FMargin(PhaseButtonMargin))
-		[
-
-			SNew(SCheckBox)
-			.Style(FAppStyle::Get(), "ToggleButtonCheckbox")
-			.IsChecked_Lambda([this]()
-			{
-				const bool bIsOn =  EditorController->GetGlobalSettings().bEnableIK;
-				return bIsOn ? ECheckBoxState::Checked : ECheckBoxState::Unchecked;
-			})
-			.OnCheckStateChanged_Lambda([this](ECheckBoxState InCheckBoxState)
-			{
-				FRetargetGlobalSettings& GlobalSettings = EditorController->GetGlobalSettings();
-				GlobalSettings.bEnableIK = !GlobalSettings.bEnableIK;
-			})
-			[
-				SNew(STextBlock).Text(FText::FromString("IK"))
-			]
-		]
-
-		+ SHorizontalBox::Slot()
-		.AutoWidth()
-		.VAlign(VAlign_Center)
-		.Padding(FMargin(PhaseButtonMargin))
-		[
-			SNew(SCheckBox)
-			.Style(FAppStyle::Get(), "ToggleButtonCheckbox")
-			.IsChecked_Lambda([this]()
-			{
-				const bool bIsOn =  EditorController->GetGlobalSettings().bEnablePost;
-				return bIsOn ? ECheckBoxState::Checked : ECheckBoxState::Unchecked;
-			})
-			.OnCheckStateChanged_Lambda([this](ECheckBoxState InCheckBoxState)
-			{
-				FRetargetGlobalSettings& GlobalSettings = EditorController->GetGlobalSettings();
-				GlobalSettings.bEnablePost = !GlobalSettings.bEnablePost;
-			})
-			[
-				SNew(STextBlock).Text(FText::FromString("Post"))
-			]
-		]
-	];
-	
-	return Box;
-}
-
 FName FIKRetargetEditor::GetToolkitFName() const
 {
 	return FName("IKRetargetEditor");
@@ -497,7 +358,7 @@ void FIKRetargetEditor::Tick(float DeltaTime)
 	constexpr float MaxSkipTimeBeforeReset = 0.25f;
 	if (CurrentTime < PreviousTime || CurrentTime > PreviousTime + MaxSkipTimeBeforeReset)
 	{
-		EditorController->ResetIKPlantingState();
+		EditorController->OnPlaybackReset();
 	}
 	PreviousTime = CurrentTime;
 	
diff --git a/Engine/Plugins/Animation/IKRig/Source/IKRigEditor/Private/RetargetEditor/IKRetargetEditorController.cpp b/Engine/Plugins/Animation/IKRig/Source/IKRigEditor/Private/RetargetEditor/IKRetargetEditorController.cpp
--- a/Engine/Plugins/Animation/IKRig/Source/IKRigEditor/Private/RetargetEditor/IKRetargetEditorController.cpp
+++ b/Engine/Plugins/Animation/IKRig/Source/IKRigEditor/Private/RetargetEditor/IKRetargetEditorController.cpp
@@ -159,9 +159,9 @@ void FIKRetargetEditorController::Initialize(TSharedPtr<FIKRetargetEditor> InEdi
 
 	PlaybackManager = MakeUnique<FRetargetPlaybackManager>(SharedThis(this));
 
-	SelectedBoneNames.Add(ERetargetSourceOrTarget::Source);
-	SelectedBoneNames.Add(ERetargetSourceOrTarget::Target);
-	LastSelectedItem = ERetargetSelectionType::NONE;
+	Selection.SelectedBoneNames.Add(ERetargetSourceOrTarget::Source);
+	Selection.SelectedBoneNames.Add(ERetargetSourceOrTarget::Target);
+	Selection.LastSelectedItem = ERetargetSelectionType::NONE;
 
 	// clean the asset before editing
 	AssetController->CleanAsset();
@@ -337,7 +337,8 @@ void FIKRetargetEditorController::UpdateSkeletalMeshComponents() const
 
 		const bool bIsSource  = SourceOrTarget == ERetargetSourceOrTarget::Source;
 		const FVector Position = bIsSource ? Asset->SourceMeshOffset : Asset->TargetMeshOffset;
-		const float Scale = bIsSource ? Asset->GetGlobalSettings().SourceScaleFactor : Asset->TargetMeshScale;
+		const double SourceScaleFactor = GetRetargetProcessor()->GetSourceScaleFactor();
+		const double Scale = bIsSource ? SourceScaleFactor : Asset->TargetMeshScale;
 		constexpr bool bSweep = false;
 		constexpr FHitResult* OutSweepHitResult = nullptr;
 		constexpr ETeleportType Teleport = ETeleportType::TeleportPhysics;
@@ -483,7 +484,7 @@ bool FIKRetargetEditorController::IsBoneRetargeted(const FName& BoneName, ERetar
 	}
 
 	// return if it's a retargeted bone
-	return Processor->IsBoneRetargeted(BoneName, SourceOrTarget);
+	return Processor->IsBoneInAMappedChain(BoneName, SourceOrTarget);
 }
 
 FName FIKRetargetEditorController::GetChainNameFromBone(const FName& BoneName, ERetargetSourceOrTarget SourceOrTarget) const
@@ -610,11 +611,11 @@ FIKRetargetProcessor* FIKRetargetEditorController::GetRetargetProcessor() const
 	return nullptr;	
 }
 
-void FIKRetargetEditorController::ResetIKPlantingState() const
+void FIKRetargetEditorController::OnPlaybackReset() const
 {
 	if (FIKRetargetProcessor* Processor = GetRetargetProcessor())
 	{
-		Processor->ResetPlanting();
+		Processor->OnPlaybackReset();
 	}
 }
 
@@ -756,7 +757,7 @@ void FIKRetargetEditorController::ShowDetailsForOp(const int32 OpIndex) const
 	const FIKRetargetOpBase* OpToEdit = InstancedStruct->GetPtr<FIKRetargetOpBase>();
 	const UScriptStruct* TypeToEdit = OpToEdit->GetSettingsType();
 
-	UStructViewer* StructViewer = AssetController->GetStructViewer();
+	UIKRigStructViewer* StructViewer = AssetController->GetStructViewer();
 	StructViewer->SetStructToView(TypeToEdit, MemoryProvider, AssetController->GetAsset());
 	
 	constexpr bool bForceRefresh = true; // must forcibly refresh because struct viewer is a recycled UObject
@@ -805,7 +806,7 @@ void FIKRetargetEditorController::SetSourceOrTargetMode(ERetargetSourceOrTarget
 	}
 
 	// make sure details panel updates with selected bone on OTHER skeleton
-	if (LastSelectedItem == ERetargetSelectionType::BONE)
+	if (Selection.LastSelectedItem == ERetargetSelectionType::BONE)
 	{
 		const TArray<FName>& SelectedBones = GetSelectedBones();
 		EditBoneSelection(SelectedBones, ESelectionEdit::Replace);
@@ -827,7 +828,7 @@ void FIKRetargetEditorController::EditBoneSelection(
 		return;
 	}
 
-	LastSelectedItem = ERetargetSelectionType::BONE;
+	Selection.LastSelectedItem = ERetargetSelectionType::BONE;
 	
 	SetRootSelected(false);
 	
@@ -837,7 +838,7 @@ void FIKRetargetEditorController::EditBoneSelection(
 		{
 			for (const FName& BoneName : InBoneNames)
 			{
-				SelectedBoneNames[CurrentlyEditingSourceOrTarget].AddUnique(BoneName);
+				Selection.SelectedBoneNames[CurrentlyEditingSourceOrTarget].AddUnique(BoneName);
 			}
 			
 			break;
@@ -846,13 +847,13 @@ void FIKRetargetEditorController::EditBoneSelection(
 		{
 			for (const FName& BoneName : InBoneNames)
 			{
-				SelectedBoneNames[CurrentlyEditingSourceOrTarget].Remove(BoneName);
+				Selection.SelectedBoneNames[CurrentlyEditingSourceOrTarget].Remove(BoneName);
 			}
 			break;
 		}
 		case ESelectionEdit::Replace:
 		{
-			SelectedBoneNames[CurrentlyEditingSourceOrTarget] = InBoneNames;
+			Selection.SelectedBoneNames[CurrentlyEditingSourceOrTarget] = InBoneNames;
 			break;
 		}
 		default:
@@ -872,14 +873,14 @@ void FIKRetargetEditorController::EditBoneSelection(
 	}
 
 	// update details
-	if (SelectedBoneNames[CurrentlyEditingSourceOrTarget].IsEmpty())
+	if (Selection.SelectedBoneNames[CurrentlyEditingSourceOrTarget].IsEmpty())
 	{
 		SetDetailsObject(AssetController->GetAsset());
 	}
 	else
 	{
 		TArray<UObject*> SelectedBoneDetails;
-		for (const FName& SelectedBone : SelectedBoneNames[CurrentlyEditingSourceOrTarget])
+		for (const FName& SelectedBone : Selection.SelectedBoneNames[CurrentlyEditingSourceOrTarget])
 		{
 			TObjectPtr<UIKRetargetBoneDetails> BoneDetails = GetOrCreateBoneDetailsObject(SelectedBone);
 			SelectedBoneDetails.Add(BoneDetails);
@@ -896,7 +897,7 @@ void FIKRetargetEditorController::EditChainSelection(
 	// deselect others
 	SetRootSelected(false);
 
-	LastSelectedItem = ERetargetSelectionType::CHAIN;
+	Selection.LastSelectedItem = ERetargetSelectionType::CHAIN;
 	
 	// update selection set based on edit mode
 	switch (EditMode)
@@ -905,7 +906,7 @@ void FIKRetargetEditorController::EditChainSelection(
 		{
 			for (const FName& ChainName : InChainNames)
 			{
-				SelectedChains.AddUnique(ChainName);
+				Selection.SelectedChains.AddUnique(ChainName);
 			}
 			
 			break;
@@ -914,13 +915,13 @@ void FIKRetargetEditorController::EditChainSelection(
 		{
 			for (const FName& ChainName : InChainNames)
 			{
-				SelectedChains.Remove(ChainName);
+				Selection.SelectedChains.Remove(ChainName);
 			}
 			break;
 		}
 	case ESelectionEdit::Replace:
 		{
-			SelectedChains = InChainNames;
+			Selection.SelectedChains = InChainNames;
 			break;
 		}
 	default:
@@ -936,36 +937,23 @@ void FIKRetargetEditorController::EditChainSelection(
 		}
 	}
 
-	// get selected chain UObjects to show in details view
-	TArray<UObject*> SelectedChainSettings;
-	const TArray<TObjectPtr<URetargetChainSettings>>& AllChainSettings = AssetController->GetAsset()->GetAllChainSettings();
-	for (const TObjectPtr<URetargetChainSettings>& ChainSettings : AllChainSettings)
-	{
-		if (ChainSettings->EditorController.Pin().Get() != this)
-		{
-			ChainSettings->EditorController = SharedThis(this);	
-		}
-		
-		if (SelectedChains.Contains(ChainSettings.Get()->TargetChain))
-		{
-			SelectedChainSettings.Add(ChainSettings);
-		}
-	}
+	// TODO
+	// We could get the chain settings struct from the currently selected op + chain map and display that
+	// but this needs more discussion with UX team on how we want to proceed.
 
 	// no chains selected, then show asset settings in the details view
-	SelectedChainSettings.IsEmpty() ? SetDetailsObject(AssetController->GetAsset()) : SetDetailsObjects(SelectedChainSettings);
+	SetDetailsObject(AssetController->GetAsset());
 }
 
 void FIKRetargetEditorController::SetRootSelected(const bool bIsSelected)
 {
-	bIsRootSelected = bIsSelected;
+	Selection.bIsRootSelected = bIsSelected;
 	if (!bIsSelected)
 	{
 		return;
 	}
 
-	LastSelectedItem = ERetargetSelectionType::ROOT;
-	ShowRootSettings();
+	Selection.LastSelectedItem = ERetargetSelectionType::ROOT;
 }
 
 void FIKRetargetEditorController::CleanSelection(ERetargetSourceOrTarget SourceOrTarget)
@@ -973,13 +961,13 @@ void FIKRetargetEditorController::CleanSelection(ERetargetSourceOrTarget SourceO
 	USkeletalMesh* SkeletalMesh = GetSkeletalMesh(SourceOrTarget);
 	if (!SkeletalMesh)
 	{
-		SelectedBoneNames[SourceOrTarget].Empty();
+		Selection.SelectedBoneNames[SourceOrTarget].Empty();
 		return;
 	}
 
 	TArray<FName> CleanedSelection;
 	const FReferenceSkeleton& RefSkeleton = SkeletalMesh->GetRefSkeleton();
-	for (const FName& SelectedBone : SelectedBoneNames[SourceOrTarget])
+	for (const FName& SelectedBone : Selection.SelectedBoneNames[SourceOrTarget])
 	{
 		if (RefSkeleton.FindBoneIndex(SelectedBone))
 		{
@@ -987,7 +975,7 @@ void FIKRetargetEditorController::CleanSelection(ERetargetSourceOrTarget SourceO
 		}
 	}
 
-	SelectedBoneNames[SourceOrTarget] = CleanedSelection;
+	Selection.SelectedBoneNames[SourceOrTarget] = CleanedSelection;
 }
 
 void FIKRetargetEditorController::ClearSelection(const bool bKeepBoneSelection)
@@ -999,18 +987,18 @@ void FIKRetargetEditorController::ClearSelection(const bool bKeepBoneSelection)
 	if (ChainsView.IsValid())
 	{
 		ChainsView->ClearSelection();
-		SelectedChains.Empty();
+		Selection.SelectedChains.Empty();
 	}
 
 	// clear bone selection
 	if (!bKeepBoneSelection)
 	{
 		SetRootSelected(false);
-		SelectedBoneNames[ERetargetSourceOrTarget::Source].Reset();
-		SelectedBoneNames[ERetargetSourceOrTarget::Target].Reset();
+		Selection.SelectedBoneNames[ERetargetSourceOrTarget::Source].Reset();
+		Selection.SelectedBoneNames[ERetargetSourceOrTarget::Target].Reset();
 	}
 
-	LastSelectedItem = ERetargetSelectionType::NONE;
+	Selection.LastSelectedItem = ERetargetSelectionType::NONE;
 
 	RefreshDetailsView();
 }
@@ -1026,6 +1014,16 @@ FIKRetargetOpBase* FIKRetargetEditorController::GetSelectedOp() const
 	return AssetController->GetRetargetOpAtIndex(SelectedOpIndex);
 }
 
+int32 FIKRetargetEditorController::GetSelectedOpIndex() const
+{
+	if (!OpStackView.IsValid())
+	{
+		return INDEX_NONE;
+	}
+	
+	return OpStackView->GetSelectedItemIndex();
+}
+
 void FIKRetargetEditorController::SetRetargeterMode(ERetargeterOutputMode Mode)
 {
 	if (OutputMode == Mode)
@@ -1115,11 +1113,6 @@ bool FIKRetargetEditorController::IsEditingPose() const
 	return GetRetargeterMode() == ERetargeterOutputMode::EditRetargetPose;
 }
 
-FRetargetGlobalSettings& FIKRetargetEditorController::GetGlobalSettings() const
-{
-	return AssetController->GetAsset()->GetGlobalSettingsUObject()->Settings;
-}
-
 void FIKRetargetEditorController::HandleNewPose()
 {
 	SetRetargeterMode(ERetargeterOutputMode::EditRetargetPose);
@@ -1493,7 +1486,7 @@ void FIKRetargetEditorController::RenderSkeleton(FPrimitiveDrawInterface* PDI, E
 		}
 
 		// record selected bone indices
-		for (const FName& SelectedBoneName : SelectedBoneNames[InSourceOrTarget])
+		for (const FName& SelectedBoneName : Selection.SelectedBoneNames[InSourceOrTarget])
 		{
 			int32 SelectedBoneIndex = RefSkeleton.FindBoneIndex(SelectedBoneName);
 			SelectedBones.Add(SelectedBoneIndex);
@@ -1519,13 +1512,13 @@ void FIKRetargetEditorController::RenderSkeleton(FPrimitiveDrawInterface* PDI, E
 		const bool bIsSource = InSourceOrTarget == ERetargetSourceOrTarget::Source;
 		if (const FIKRetargetProcessor* Processor = GetRetargetProcessor())
 		{
-			const TArray<FRetargetChainPairFK>& FKChains = Processor->GetFKChainPairs();
-			for (const FRetargetChainPairFK& FKChain : FKChains)
+			const FRetargeterBoneChains& AllBoneChains = Processor->GetBoneChains();
+			const TArray<FResolvedBoneChain>& BoneChains = AllBoneChains.GetResolvedBoneChains(InSourceOrTarget);
+			for (const FResolvedBoneChain& BoneChain : BoneChains)
 			{
-				if (SelectedChainNames.Contains(FKChain.TargetBoneChainName))
+				if (SelectedChainNames.Contains(BoneChain.ChainName))
 				{
-					const TArray<int32>& BoneIndices = bIsSource ? FKChain.SourceBoneIndices : FKChain.TargetBoneIndices;
-					for (const int32& BoneIndex : BoneIndices)
+					for (const int32& BoneIndex : BoneChain.BoneIndices)
 					{
 						BoneColors[BoneIndex] = HighlightedColor;
 					}	
@@ -1559,7 +1552,7 @@ TArray<FName> FIKRetargetEditorController::GetSelectedBonesAndChildren() const
 
 	// get list of all children of selected bones
 	TArray<int32> AllChildrenIndices;
-	for (const FName& SelectedBone : SelectedBoneNames[CurrentlyEditingSourceOrTarget])
+	for (const FName& SelectedBone : Selection.SelectedBoneNames[CurrentlyEditingSourceOrTarget])
 	{
 		const int32 SelectedBoneIndex = RefSkeleton.FindBoneIndex(SelectedBone);
 		AllChildrenIndices.Add(SelectedBoneIndex);
@@ -1575,7 +1568,7 @@ TArray<FName> FIKRetargetEditorController::GetSelectedBonesAndChildren() const
 	}
 
 	// merge total list of all selected bones and their children
-	TArray<FName> BonesToReturn = SelectedBoneNames[CurrentlyEditingSourceOrTarget];
+	TArray<FName> BonesToReturn = Selection.SelectedBoneNames[CurrentlyEditingSourceOrTarget];
 	for (const int32 ChildIndex : AllChildrenIndices)
 	{
 		BonesToReturn.AddUnique(RefSkeleton.GetBoneName(ChildIndex));
@@ -1607,7 +1600,8 @@ void FIKRetargetEditorController::FixZeroHeightRetargetRoot(ERetargetSourceOrTar
 		return;
 	}
 
-	const float ScaleFactor = SourceOrTarget == ERetargetSourceOrTarget::Source ? GetGlobalSettings().SourceScaleFactor : 1.0f;
+	const double SourceScaleFactor = GetRetargetProcessor()->GetSourceScaleFactor();
+	const double ScaleFactor = SourceOrTarget == ERetargetSourceOrTarget::Source ? SourceScaleFactor : 1.0;
 	FRetargetSkeleton DummySkeleton;
 	DummySkeleton.Initialize(
 		SkeletalMesh,
@@ -1675,26 +1669,4 @@ void FIKRetargetEditorController::OnPoseSelected(TSharedPtr<FName> InPose, ESele
 	}
 }
 
-void FIKRetargetEditorController::ShowGlobalSettings()
-{
-	UIKRetargetGlobalSettings* GlobalSettings = AssetController->GetAsset()->GetGlobalSettingsUObject();
-	if (GlobalSettings->EditorController.Pin().Get() != this)
-	{
-		GlobalSettings->EditorController = SharedThis(this);
-	}
-	
-	SetDetailsObject(GlobalSettings);
-}
-
-void FIKRetargetEditorController::ShowRootSettings()
-{
-	URetargetRootSettings* RootSettings = AssetController->GetAsset()->GetRootSettingsUObject();
-	if (RootSettings->EditorController.Pin().Get() != this)
-	{
-		RootSettings->EditorController = SharedThis(this);	
-	}
-	
-	SetDetailsObject(RootSettings);
-}
-
 #undef LOCTEXT_NAMESPACE
diff --git a/Engine/Plugins/Animation/IKRig/Source/IKRigEditor/Private/RetargetEditor/IKRetargetHitProxies.cpp b/Engine/Plugins/Animation/IKRig/Source/IKRigEditor/Private/RetargetEditor/IKRetargetHitProxies.cpp
--- a/Engine/Plugins/Animation/IKRig/Source/IKRigEditor/Private/RetargetEditor/IKRetargetHitProxies.cpp
+++ b/Engine/Plugins/Animation/IKRig/Source/IKRigEditor/Private/RetargetEditor/IKRetargetHitProxies.cpp
@@ -3,5 +3,5 @@
 #include "RetargetEditor/IKRetargetHitProxies.h"
 
 IMPLEMENT_HIT_PROXY(HIKRetargetEditorBoneProxy, HHitProxy);
-IMPLEMENT_HIT_PROXY(HIKRetargetEditorChainProxy, HHitProxy);
-IMPLEMENT_HIT_PROXY(HIKRetargetEditorRootProxy, HHitProxy);
+
+
diff --git a/Engine/Plugins/Animation/IKRig/Source/IKRigEditor/Private/RetargetEditor/IKRetargeterController.cpp b/Engine/Plugins/Animation/IKRig/Source/IKRigEditor/Private/RetargetEditor/IKRetargeterController.cpp
--- a/Engine/Plugins/Animation/IKRig/Source/IKRigEditor/Private/RetargetEditor/IKRetargeterController.cpp
+++ b/Engine/Plugins/Animation/IKRig/Source/IKRigEditor/Private/RetargetEditor/IKRetargeterController.cpp
@@ -9,8 +9,14 @@
 #include "Retargeter/IKRetargeter.h"
 #include "RigEditor/IKRigController.h"
 #include "IKRigEditor.h"
+#include "Components/Widget.h"
 #include "RetargetEditor/IKRetargeterPoseGenerator.h"
 #include "Retargeter/IKRetargetOps.h"
+#include "Retargeter/RetargetOps/FKChainsOp.h"
+#include "Retargeter/RetargetOps/IKChainsOp.h"
+#include "Retargeter/RetargetOps/PelvisMotionOp.h"
+#include "Retargeter/RetargetOps/RootMotionGeneratorOp.h"
+#include "Retargeter/RetargetOps/RunIKRigOp.h"
 #include "RigEditor/IKRigDetailCustomizations.h"
 
 #include UE_INLINE_GENERATED_CPP_BY_NAME(IKRetargeterController)
@@ -37,7 +43,11 @@ UIKRetargeterController* UIKRetargeterController::GetController(const UIKRetarge
 
 UIKRetargeterController::UIKRetargeterController()
 {
-	StructViewer = CreateDefaultSubobject<UStructViewer>(TEXT("RetargetSettingsViewer"));
+	StructViewer = CreateDefaultSubobject<UIKRigStructViewer>(TEXT("RetargetSettingsViewer"));
+	StructViewer->OnStructPropertyEdited().AddLambda([this]()
+	{
+		BroadcastNeedsReinitialized();
+	});
 }
 
 void UIKRetargeterController::PostInitProperties()
@@ -59,7 +69,7 @@ TObjectPtr<UIKRetargeter>& UIKRetargeterController::GetAssetPtr()
 void UIKRetargeterController::CleanAsset() const
 {
 	FScopeLock Lock(&ControllerLock);
-	CleanChainMapping();
+	Asset->GetChainMapping().CleanMapping(GetIKRig(ERetargetSourceOrTarget::Source), GetIKRig(ERetargetSourceOrTarget::Target));
 	CleanPoseList(ERetargetSourceOrTarget::Source);
 	CleanPoseList(ERetargetSourceOrTarget::Target);
 }
@@ -87,7 +97,7 @@ void UIKRetargeterController::SetIKRig(const ERetargetSourceOrTarget SourceOrTar
 		SetAskedToFixRootHeightForMesh(GetPreviewMesh(SourceOrTarget), false);
 	}
 
-	CleanChainMapping();
+	Asset->GetChainMapping().CleanMapping(GetIKRig(ERetargetSourceOrTarget::Source), GetIKRig(ERetargetSourceOrTarget::Target));
 	
 	constexpr bool bForceRemap = false;
 	AutoMapChains(EAutoMapChainType::Exact, bForceRemap);
@@ -156,67 +166,7 @@ USkeletalMesh* UIKRetargeterController::GetPreviewMesh(const ERetargetSourceOrTa
 	return IKRig->GetPreviewMesh();
 }
 
-FTargetRootSettings UIKRetargeterController::GetRootSettings() const
-{
-	FScopeLock Lock(&ControllerLock);
-	return GetAsset()->GetRootSettingsUObject()->Settings;
-}
-
-void UIKRetargeterController::SetRootSettings(const FTargetRootSettings& RootSettings) const
-{
-	FScopeLock Lock(&ControllerLock);
-	FScopedTransaction Transaction(LOCTEXT("SetRootSettings_Transaction", "Set Root Settings"));
-	FScopedReinitializeIKRetargeter Reinitialize(this);
-	GetAsset()->Modify();
-	GetAsset()->GetRootSettingsUObject()->Settings = RootSettings;
-}
-
-FRetargetGlobalSettings UIKRetargeterController::GetGlobalSettings() const
-{
-	FScopeLock Lock(&ControllerLock);
-	return GetAsset()->GetGlobalSettingsUObject()->Settings;
-}
-
-void UIKRetargeterController::SetGlobalSettings(const FRetargetGlobalSettings& GlobalSettings) const
-{
-	FScopeLock Lock(&ControllerLock);
-	FScopedTransaction Transaction(LOCTEXT("SetGlobalSettings_Transaction", "Set Global Settings"));
-	FScopedReinitializeIKRetargeter Reinitialize(this);
-	GetAsset()->Modify();
-	GetAsset()->GetGlobalSettingsUObject()->Settings = GlobalSettings;
-}
-
-FTargetChainSettings UIKRetargeterController::GetRetargetChainSettings(const FName& TargetChainName) const
-{
-	FScopeLock Lock(&ControllerLock);
-	
-	const URetargetChainSettings* ChainSettings = GetChainSettings(TargetChainName);
-	if (!ChainSettings)
-	{
-		return FTargetChainSettings();
-	}
-
-	return ChainSettings->Settings;
-}
-
-bool UIKRetargeterController::SetRetargetChainSettings(const FName& TargetChainName, const FTargetChainSettings& Settings) const
-{
-	FScopeLock Lock(&ControllerLock);
-
-	FScopedTransaction Transaction(LOCTEXT("SetChainSettings_Transaction", "Set Chain Settings"));
-	FScopedReinitializeIKRetargeter Reinitialize(this);
-	
-	if (URetargetChainSettings* ChainSettings = GetChainSettings(TargetChainName))
-	{
-		ChainSettings->Modify();
-		ChainSettings->Settings = Settings;
-		return true;
-	}
-
-	return false;
-}
-
-int32 UIKRetargeterController::AddRetargetOp(const FString InIKRetargetOpType) const
+int32 UIKRetargeterController::AddRetargetOp(const FString InIKRetargetOpType, const bool bAllowDuplicates, const bool bRunInitialSetup) const
 {
 	UScriptStruct* OpType = FindObject<UScriptStruct>(nullptr, *InIKRetargetOpType);
 	if (!OpType)
@@ -224,8 +174,12 @@ int32 UIKRetargeterController::AddRetargetOp(const FString InIKRetargetOpType) c
 		UE_LOG(LogIKRigEditor, Warning, TEXT("Specified retarget op type was not found, %s."), *InIKRetargetOpType);
 		return INDEX_NONE;
 	}
+	
+	FScopedTransaction Transaction(LOCTEXT("AddRetargetOp_Label", "Add Retarget Op"));
+	FScopedReinitializeIKRetargeter Reinitialize(this);
+	Asset->Modify();
 
-	return AddRetargetOp(OpType);
+	return AddRetargetOp(OpType, bAllowDuplicates);
 }
 
 bool UIKRetargeterController::RemoveRetargetOp(const int32 OpIndex) const
@@ -256,6 +210,36 @@ bool UIKRetargeterController::RemoveAllOps() const
 	return true;
 }
 
+void UIKRetargeterController::AddDefaultOps(const bool bRunInitialSetup) const
+{
+	FScopedTransaction Transaction(LOCTEXT("AddDefaultOps_Label", "Add Default Ops"));
+	FScopedReinitializeIKRetargeter Reinitialize(this);
+	Asset->Modify();
+
+	// add set of default ops for basic retargeting
+	constexpr bool bAllowDuplicateOps = false;
+	AddRetargetOp(FIKRetargetPelvisMotionOp::StaticStruct(), bAllowDuplicateOps, bRunInitialSetup);
+	AddRetargetOp(FIKRetargetFKChainsOp::StaticStruct(), bAllowDuplicateOps, bRunInitialSetup);
+	AddRetargetOp(FIKRetargetIKChainsOp::StaticStruct(), bAllowDuplicateOps, bRunInitialSetup);
+	AddRetargetOp(FIKRetargetRunIKRigOp::StaticStruct(), bAllowDuplicateOps, bRunInitialSetup);
+	AddRetargetOp(FIKRetargetRootMotionOp::StaticStruct(), bAllowDuplicateOps, bRunInitialSetup);
+}
+
+void UIKRetargeterController::RunOpInitialSetup(const int32 InOpIndex) const
+{
+	FIKRetargetOpBase* Op = GetRetargetOpAtIndex(InOpIndex);
+	if (!Op)
+	{
+		UE_LOG(LogIKRigEditor, Warning, TEXT("Retarget Op initial setup skipped. Invalid index, %d."), InOpIndex);
+		return;
+	}
+	
+	FScopedTransaction Transaction(LOCTEXT("RunOpInitialSetup_Label", "Run Op Initial Setup"));
+	FScopedReinitializeIKRetargeter Reinitialize(this);
+	Asset->Modify();
+	Op->RunInitialSetup(GetAsset());
+}
+
 FIKRetargetOpBase* UIKRetargeterController::GetRetargetOpAtIndex(int32 Index) const
 {
 	check(Asset)
@@ -308,6 +292,12 @@ int32 UIKRetargeterController::GetNumRetargetOps() const
 
 bool UIKRetargeterController::MoveRetargetOpInStack(int32 OpToMoveIndex, int32 TargetIndex) const
 {
+	// adjust TargetIndex if necessary (since RemoveAt shifts elements)
+	if (OpToMoveIndex < TargetIndex)
+	{
+		TargetIndex--;  // compensate for the removed element shifting indices down
+	}
+	
 	if (!Asset->RetargetOps.IsValidIndex(OpToMoveIndex))
 	{
 		UE_LOG(LogIKRigEditor, Warning, TEXT("Retarget Op not moved. Invalid source index, %d."), OpToMoveIndex);
@@ -329,7 +319,16 @@ bool UIKRetargeterController::MoveRetargetOpInStack(int32 OpToMoveIndex, int32 T
 	FScopedTransaction Transaction(LOCTEXT("ReorderRetargetOps_Label", "Reorder Retarget Ops"));
 	FScopedReinitializeIKRetargeter Reinitialize(this);
 	Asset->Modify();
-	Asset->RetargetOps.Swap(OpToMoveIndex, TargetIndex);
+
+	// extract the element to move
+	FInstancedStruct MovedOp = MoveTemp(Asset->RetargetOps[OpToMoveIndex]);
+
+	// Remove the element without shrinking the array
+	Asset->RetargetOps.RemoveAt(OpToMoveIndex, 1, EAllowShrinking::No);
+
+	// Insert at the corrected TargetIndex
+	Asset->RetargetOps.Insert(MoveTemp(MovedOp), TargetIndex);
+
 	return true;
 }
 
@@ -359,7 +358,46 @@ bool UIKRetargeterController::GetRetargetOpEnabled(int32 RetargetOpIndex) const
 	return GetRetargetOpAtIndex(RetargetOpIndex)->IsEnabled();
 }
 
-int32 UIKRetargeterController::AddRetargetOp(UScriptStruct* InRetargetOpType) const
+UIKRetargetOpControllerBase* UIKRetargeterController::GetOpController(int32 OpIndex)
+{
+	if (FIKRetargetOpBase* OpAtIndex = GetRetargetOpAtIndex(OpIndex))
+	{
+		return OpAtIndex->GetController(this);
+	}
+	return nullptr;
+}
+
+void UIKRetargeterController::ResetChainSettingsInAllOps(const FName InTargetChainName) const
+{
+	FScopedTransaction Transaction(LOCTEXT("ResetChainSettings_Label", "Reset Settings for Chain"));
+	FScopedReinitializeIKRetargeter Reinitialize(this);
+	Asset->Modify();
+	for (int32 OpIndex=0; OpIndex<Asset->RetargetOps.Num(); ++OpIndex)
+	{
+		FIKRetargetOpBase& Op = Asset->RetargetOps[OpIndex].GetMutable<FIKRetargetOpBase>();
+		Op.SetChainSettingsToDefault(InTargetChainName);
+	}
+}
+
+const FRetargetChainMapping& UIKRetargeterController::GetChainMapping() const
+{
+	return Asset->GetChainMapping();
+}
+
+bool UIKRetargeterController::ChainHasNonDefaultSettings(const FName InTargetChainName) const
+{
+	for (int32 OpIndex=0; OpIndex<Asset->RetargetOps.Num(); ++OpIndex)
+	{
+		FIKRetargetOpBase& Op = Asset->RetargetOps[OpIndex].GetMutable<FIKRetargetOpBase>();
+		if (!Op.AreChainSettingsAtDefault(InTargetChainName))
+		{
+			return true;
+		}
+	}
+	return false;
+}
+
+int32 UIKRetargeterController::AddRetargetOp(UScriptStruct* InRetargetOpType, const bool bAllowDuplicates, const bool bRunInitialSetup) const
 {
 	check(Asset)
 
@@ -375,10 +413,29 @@ int32 UIKRetargeterController::AddRetargetOp(UScriptStruct* InRetargetOpType) co
 		return INDEX_NONE;
 	}
 
-	FScopedTransaction Transaction(LOCTEXT("AddRetargetOp_Label", "Add Retarget Op"));
-	FScopedReinitializeIKRetargeter Reinitialize(this);
-	Asset->Modify();
+	// optionally avoid duplicate ops of the same type
+	if (!bAllowDuplicates)
+	{
+		for (int32 OpIndex=0; OpIndex<Asset->RetargetOps.Num(); ++OpIndex)
+		{
+			const FInstancedStruct& Op = Asset->RetargetOps[OpIndex];
+			if (Op.GetScriptStruct() == InRetargetOpType)
+			{
+				return OpIndex;
+			}
+		}
+	}
+
+	// add the op
 	const int32 NewOpIndex = Asset->RetargetOps.Emplace(InRetargetOpType);
+
+	// optionally run the initial setup
+	if (bRunInitialSetup)
+	{
+		FIKRetargetOpBase* NewOp = GetRetargetOpAtIndex(NewOpIndex);
+		NewOp->RunInitialSetup(GetAsset());
+	}
+	
 	return NewOpIndex;
 }
 
@@ -419,66 +476,6 @@ void UIKRetargeterController::GetChainNames(const ERetargetSourceOrTarget Source
 	}
 }
 
-void UIKRetargeterController::CleanChainMapping() const
-{
-	if (IsValid(Asset->GetIKRig(ERetargetSourceOrTarget::Target)))
-	{
-		TArray<FName> TargetChainNames;
-		GetChainNames(ERetargetSourceOrTarget::Target, TargetChainNames);
-
-		// remove all target chains that are no longer in the target IK rig asset
-		TArray<FName> TargetChainsToRemove;
-		for (const URetargetChainSettings* ChainMap : Asset->ChainSettings)
-		{
-			if (!TargetChainNames.Contains(ChainMap->TargetChain))
-			{
-				TargetChainsToRemove.Add(ChainMap->TargetChain);
-			}
-		}
-		for (FName TargetChainToRemove : TargetChainsToRemove)
-		{
-			Asset->ChainSettings.RemoveAll([&TargetChainToRemove](const URetargetChainSettings* Element)
-			{
-				return Element->TargetChain == TargetChainToRemove;
-			});
-		}
-
-		// add a mapping for each chain that is in the target IK rig (if it doesn't have one already)
-		for (FName TargetChainName : TargetChainNames)
-		{
-			const bool HasChain = Asset->ChainSettings.ContainsByPredicate([&TargetChainName](const URetargetChainSettings* Element)
-			{
-				return Element->TargetChain == TargetChainName;
-			});
-		
-			if (!HasChain)
-			{
-				TObjectPtr<URetargetChainSettings> ChainMap = NewObject<URetargetChainSettings>(Asset, URetargetChainSettings::StaticClass(), NAME_None, RF_Transactional);
-				ChainMap->TargetChain = TargetChainName;
-				Asset->ChainSettings.Add(ChainMap);
-			}
-		}
-	}
-
-	if (IsValid(Asset->GetIKRig(ERetargetSourceOrTarget::Source)))
-	{
-		TArray<FName> SourceChainNames;
-		GetChainNames(ERetargetSourceOrTarget::Source,SourceChainNames);
-	
-		// reset any sources that are no longer present to "None"
-		for (URetargetChainSettings* ChainMap : Asset->ChainSettings)
-		{
-			if (!SourceChainNames.Contains(ChainMap->SourceChain))
-			{
-				ChainMap->SourceChain = NAME_None;
-			}
-		}
-	}
-
-	// enforce the chain order based on the StartBone index
-	SortChainMapping();
-}
-
 void UIKRetargeterController::CleanPoseList(const ERetargetSourceOrTarget SourceOrTarget) const
 {
 	const UIKRigDefinition* IKRig = GetIKRig(SourceOrTarget);
@@ -513,7 +510,7 @@ void UIKRetargeterController::CleanPoseList(const ERetargetSourceOrTarget Source
 	}
 }
 
-UStructViewer* UIKRetargeterController::GetStructViewer()
+UIKRigStructViewer* UIKRetargeterController::GetStructViewer()
 {
 	StructViewer->Reset();
 	return StructViewer;
@@ -524,24 +521,30 @@ void UIKRetargeterController::AutoMapChains(const EAutoMapChainType AutoMapType,
 	FScopeLock Lock(&ControllerLock);
 	FScopedTransaction Transaction(LOCTEXT("AutoMapRetargetChains", "Auto-Map Retarget Chains"));
 	FScopedReinitializeIKRetargeter Reinitialize(this);
-	
-	CleanChainMapping();
+	Asset->Modify();
+
+	// get the chain mapping from the asset
+	FRetargetChainMapping& ChainMapping = Asset->GetChainMapping();
+
+	// clean the chain map
+	const UIKRigDefinition* SourceIKRig = GetIKRig(ERetargetSourceOrTarget::Source);
+	const UIKRigDefinition* TargetIKRig = GetIKRig(ERetargetSourceOrTarget::Target);
+	ChainMapping.CleanMapping(SourceIKRig, TargetIKRig);
 
 	// get names of all the potential chains we could map to on the source
 	TArray<FName> SourceChainNames;
 	GetChainNames(ERetargetSourceOrTarget::Source, SourceChainNames);
-	
-	// iterate over all the chain mappings and find matching source chain
-	for (URetargetChainSettings* ChainMap : Asset->ChainSettings)
+
+	// iterate over all the chain pairs and find matching source chain
+	TArray<FRetargetChainPair>& ChainPairs = Asset->GetChainMapping().GetChainPairs();
+	for (FRetargetChainPair& ChainPair : ChainPairs)
 	{
-		const bool bIsMappedAlready = ChainMap->SourceChain != NAME_None;
+		const bool bIsMappedAlready = ChainPair.SourceChainName != NAME_None;
 		if (bIsMappedAlready && !bForceRemap)
 		{
 			continue; // already set by user
 		}
 
-		ChainMap->Modify();
-
 		// find a source chain to map to
 		int32 SourceChainIndexToMapTo = -1;
 
@@ -550,7 +553,7 @@ void UIKRetargeterController::AutoMapChains(const EAutoMapChainType AutoMapType,
 		case EAutoMapChainType::Fuzzy:
 			{
 				// auto-map chains using a fuzzy string comparison
-				FString TargetNameLowerCase = ChainMap->TargetChain.ToString().ToLower();
+				FString TargetNameLowerCase = ChainPair.TargetChainName.ToString().ToLower();
 				float HighestScore = 0.2f;
 				for (int32 ChainIndex=0; ChainIndex<SourceChainNames.Num(); ++ChainIndex)
 				{
@@ -569,12 +572,12 @@ void UIKRetargeterController::AutoMapChains(const EAutoMapChainType AutoMapType,
 		case EAutoMapChainType::Exact:
 			{
 				// if no exact match is found, then set to None
-				ChainMap->SourceChain = NAME_None;
+				ChainPair.SourceChainName = NAME_None;
 				
 				// auto-map chains with EXACT same name
 				for (int32 ChainIndex=0; ChainIndex<SourceChainNames.Num(); ++ChainIndex)
 				{
-					if (SourceChainNames[ChainIndex] == ChainMap->TargetChain)
+					if (SourceChainNames[ChainIndex] == ChainPair.TargetChainName)
 					{
 						SourceChainIndexToMapTo = ChainIndex;
 						break;
@@ -584,7 +587,7 @@ void UIKRetargeterController::AutoMapChains(const EAutoMapChainType AutoMapType,
 			}
 		case EAutoMapChainType::Clear:
 			{
-				ChainMap->SourceChain = NAME_None;
+				ChainPair.SourceChainName = NAME_None;
 				break;
 			}
 		default:
@@ -594,12 +597,12 @@ void UIKRetargeterController::AutoMapChains(const EAutoMapChainType AutoMapType,
 		// apply source if any decent matches were found
 		if (SourceChainNames.IsValidIndex(SourceChainIndexToMapTo))
 		{
-			ChainMap->SourceChain = SourceChainNames[SourceChainIndexToMapTo];
+			ChainPair.SourceChainName = SourceChainNames[SourceChainIndexToMapTo];
 		}
 	}
 
-	// sort them
-	SortChainMapping();
+	// sort asset mapping again
+	ChainMapping.SortMapping(TargetIKRig);
 }
 
 void UIKRetargeterController::HandleRetargetChainAdded(UIKRigDefinition* IKRig) const
@@ -611,25 +614,33 @@ void UIKRetargeterController::HandleRetargetChainAdded(UIKRigDefinition* IKRig)
 		return;
 	}
 
-	// add the new chain to the mapping data
-	CleanChainMapping();
 	FScopedReinitializeIKRetargeter Reinitialize(this);
+	
+	// clean the chain map (this will add the new chain automatically)
+	const UIKRigDefinition* SourceIKRig = GetIKRig(ERetargetSourceOrTarget::Source);
+	const UIKRigDefinition* TargetIKRig = GetIKRig(ERetargetSourceOrTarget::Target);
+	Asset->GetChainMapping().CleanMapping(SourceIKRig, TargetIKRig);
 }
 
 void UIKRetargeterController::HandleRetargetChainRenamed(UIKRigDefinition* IKRig, FName OldChainName, FName NewChainName) const
 {
 	const bool bIsSourceRig = IKRig == Asset->GetIKRig(ERetargetSourceOrTarget::Source);
 	const bool bIsTargetRig = IKRig == Asset->GetIKRig(ERetargetSourceOrTarget::Target);
-	check(bIsSourceRig || bIsTargetRig)
-	for (URetargetChainSettings* ChainMap : Asset->ChainSettings)
+	if (!ensure(bIsSourceRig || bIsTargetRig))
+	{
+		return;
+	}
+	
+	FRetargetChainMapping& ChainMapping = Asset->GetChainMapping();
+	for (FRetargetChainPair& ChainPair : ChainMapping.GetChainPairs())
 	{
-		FName& ChainNameToUpdate = bIsSourceRig ? ChainMap->SourceChain : ChainMap->TargetChain;
+		FName& ChainNameToUpdate = bIsSourceRig ? ChainPair.SourceChainName : ChainPair.TargetChainName;
 		if (ChainNameToUpdate == OldChainName)
 		{
 			FScopedTransaction Transaction(LOCTEXT("RetargetChainRenamed_Label", "Retarget Chain Renamed"));
-			ChainMap->Modify();
-			ChainNameToUpdate = NewChainName;
 			FScopedReinitializeIKRetargeter Reinitialize(this);
+			Asset->Modify();
+			ChainNameToUpdate = NewChainName;
 			return;
 		}
 	}
@@ -637,22 +648,28 @@ void UIKRetargeterController::HandleRetargetChainRenamed(UIKRigDefinition* IKRig
 
 void UIKRetargeterController::HandleRetargetChainRemoved(UIKRigDefinition* IKRig, const FName& InChainRemoved) const
 {
-	FScopedTransaction Transaction(LOCTEXT("RetargetChainRemoved_Label", "Retarget Chain Removed"));
-	Asset->Modify();
-	
 	const bool bIsSourceRig = IKRig == Asset->GetIKRig(ERetargetSourceOrTarget::Source);
 	const bool bIsTargetRig = IKRig == Asset->GetIKRig(ERetargetSourceOrTarget::Target);
-	check(bIsSourceRig || bIsTargetRig)
+	if (!ensure(bIsSourceRig || bIsTargetRig))
+	{
+		return;
+	}
+	
+	FScopedTransaction Transaction(LOCTEXT("RetargetChainRemoved_Label", "Retarget Chain Removed"));
+	FScopedReinitializeIKRetargeter Reinitialize(this);
+	Asset->Modify();
+
+	FRetargetChainMapping& ChainMapping = Asset->GetChainMapping();
+	TArray<FRetargetChainPair>& ChainPairs = ChainMapping.GetChainPairs();
 	
 	// set source chain name to NONE if it has been deleted 
 	if (bIsSourceRig)
 	{
-		for (URetargetChainSettings* ChainMap : Asset->ChainSettings)
+		for (FRetargetChainPair& ChainPair : ChainPairs)
 		{
-			if (ChainMap->SourceChain == InChainRemoved)
+			if (ChainPair.SourceChainName == InChainRemoved)
 			{
-				ChainMap->SourceChain = NAME_None;
-				FScopedReinitializeIKRetargeter Reinitialize(this);
+				ChainPair.SourceChainName = NAME_None;
 				return;
 			}
 		}
@@ -660,23 +677,22 @@ void UIKRetargeterController::HandleRetargetChainRemoved(UIKRigDefinition* IKRig
 	}
 	
 	// remove target mapping if the target chain has been removed
-	const int32 ChainIndex = Asset->ChainSettings.IndexOfByPredicate([&InChainRemoved](const URetargetChainSettings* ChainMap)
+	const int32 ChainIndex = ChainPairs.IndexOfByPredicate([&InChainRemoved](const FRetargetChainPair& ChainPair)
 	{
-		return ChainMap->TargetChain == InChainRemoved;
+		return ChainPair.TargetChainName == InChainRemoved;
 	});
 	
 	if (ChainIndex != INDEX_NONE)
 	{
-		Asset->ChainSettings.RemoveAt(ChainIndex);
-		FScopedReinitializeIKRetargeter Reinitialize(this);
+		ChainPairs.RemoveAt(ChainIndex);
 	}
 }
 
 bool UIKRetargeterController::SetSourceChain(FName SourceChainName, FName TargetChainName) const
 {
 	FScopeLock Lock(&ControllerLock);
-	URetargetChainSettings* ChainSettings = GetChainSettings(TargetChainName);
-	if (!ChainSettings)
+	FRetargetChainMapping& ChainMap = Asset->GetChainMapping();
+	if (!ChainMap.HasChain(TargetChainName, ERetargetSourceOrTarget::Target))
 	{
 		UE_LOG(LogIKRigEditor, Warning, TEXT("Could not map target chain. Target chain not found, %s."), *TargetChainName.ToString());
 		return false;
@@ -684,43 +700,21 @@ bool UIKRetargeterController::SetSourceChain(FName SourceChainName, FName Target
 	
 	FScopedTransaction Transaction(LOCTEXT("SetRetargetChainSource", "Set Retarget Chain Source"));
 	FScopedReinitializeIKRetargeter Reinitialize(this);
-	ChainSettings->Modify();
-	ChainSettings->SourceChain = SourceChainName;
+	Asset->Modify();
+	ChainMap.SetChainMapping(TargetChainName, SourceChainName);
 	return true;
 }
 
 FName UIKRetargeterController::GetSourceChain(const FName& TargetChainName) const
 {
-	const URetargetChainSettings* ChainSettings = GetChainSettings(TargetChainName);
-	if (!ChainSettings)
+	FRetargetChainMapping& ChainMap = Asset->GetChainMapping();
+	if (!ChainMap.HasChain(TargetChainName, ERetargetSourceOrTarget::Target))
 	{
-		UE_LOG(LogIKRigEditor, Warning, TEXT("Could not map target chain. Target chain not found, %s."), *TargetChainName.ToString());
+		UE_LOG(LogIKRigEditor, Warning, TEXT("Could not get source chain. Target chain not found, %s."), *TargetChainName.ToString());
 		return NAME_None;
 	}
 	
-	return ChainSettings->SourceChain;
-}
-
-const TArray<URetargetChainSettings*>& UIKRetargeterController::GetAllChainSettings() const
-{
-	return Asset->ChainSettings;
-}
-
-FName UIKRetargeterController::GetChainGoal(const TObjectPtr<URetargetChainSettings> ChainSettings) const
-{
-	if (!ChainSettings)
-	{
-		return NAME_None;
-	}
-	
-	const UIKRigDefinition* TargetIKRig = GetIKRig(ERetargetSourceOrTarget::Target);
-	if (!TargetIKRig)
-	{
-		return NAME_None;
-	}
-
-	const UIKRigController* RigController = UIKRigController::GetController(TargetIKRig);
-	return RigController->GetRetargetChainGoal(ChainSettings->TargetChain);
+	return ChainMap.GetChainMappedTo(TargetChainName, ERetargetSourceOrTarget::Target);
 }
 
 bool UIKRetargeterController::IsChainGoalConnectedToASolver(const FName& GoalName) const
@@ -1048,62 +1042,92 @@ FName UIKRetargeterController::MakePoseNameUnique(const FString& PoseName, const
 	return FName(UniqueName);
 }
 
-URetargetChainSettings* UIKRetargeterController::GetChainSettings(const FName& TargetChainName) const
+//
+// BEGIN DEPRECATED API
+//
+PRAGMA_DISABLE_DEPRECATION_WARNINGS
+
+FTargetRootSettings UIKRetargeterController::GetRootSettings() const
 {
-	for (URetargetChainSettings* ChainMap : Asset->ChainSettings)
-	{
-		if (ChainMap->TargetChain == TargetChainName)
-		{
-			return ChainMap;
-		}
-	}
+	FScopeLock Lock(&ControllerLock);
+	return GetAsset()->GetRootSettingsUObject()->Settings;
+}
 
-	return nullptr;
+void UIKRetargeterController::SetRootSettings(const FTargetRootSettings& RootSettings) const
+{
+	FScopeLock Lock(&ControllerLock);
+	FScopedTransaction Transaction(LOCTEXT("SetRootSettings_Transaction", "Set Root Settings"));
+	FScopedReinitializeIKRetargeter Reinitialize(this);
+	GetAsset()->Modify();
+	GetAsset()->GetRootSettingsUObject()->Settings = RootSettings;
+}
+
+FRetargetGlobalSettings UIKRetargeterController::GetGlobalSettings() const
+{
+	FScopeLock Lock(&ControllerLock);
+	return GetAsset()->GetGlobalSettingsUObject()->Settings;
+}
+
+void UIKRetargeterController::SetGlobalSettings(const FRetargetGlobalSettings& GlobalSettings) const
+{
+	FScopeLock Lock(&ControllerLock);
+	FScopedTransaction Transaction(LOCTEXT("SetGlobalSettings_Transaction", "Set Global Settings"));
+	FScopedReinitializeIKRetargeter Reinitialize(this);
+	GetAsset()->Modify();
+	GetAsset()->GetGlobalSettingsUObject()->Settings = GlobalSettings;
 }
 
-void UIKRetargeterController::SortChainMapping() const
+FTargetChainSettings UIKRetargeterController::GetRetargetChainSettings(const FName& TargetChainName) const
 {
-	const UIKRigDefinition* TargetIKRig = Asset->GetIKRig(ERetargetSourceOrTarget::Target);
-	if (!IsValid(TargetIKRig))
+	FScopeLock Lock(&ControllerLock);
+	
+	const URetargetChainSettings* ChainSettings = GetChainSettings(TargetChainName);
+	if (!ChainSettings)
 	{
-		return;
+		return FTargetChainSettings();
 	}
+
+	return ChainSettings->Settings;
+}
+
+bool UIKRetargeterController::SetRetargetChainSettings(const FName& TargetChainName, const FTargetChainSettings& Settings) const
+{
+	FScopeLock Lock(&ControllerLock);
+
+	FScopedTransaction Transaction(LOCTEXT("SetChainSettings_Transaction", "Set Chain Settings"));
+	FScopedReinitializeIKRetargeter Reinitialize(this);
 	
-	Asset->ChainSettings.Sort([TargetIKRig](const URetargetChainSettings& A, const URetargetChainSettings& B)
+	if (URetargetChainSettings* ChainSettings = GetChainSettings(TargetChainName))
 	{
-		const TArray<FBoneChain>& BoneChains = TargetIKRig->GetRetargetChains();
-		const FIKRigSkeleton& TargetSkeleton = TargetIKRig->GetSkeleton();
+		ChainSettings->Modify();
+		ChainSettings->Settings = Settings;
+		return true;
+	}
 
-		// look for chains
-		const int32 IndexA = BoneChains.IndexOfByPredicate([&A](const FBoneChain& Chain)
-		{
-			return A.TargetChain == Chain.ChainName;
-		});
+	return false;
+}
 
-		const int32 IndexB = BoneChains.IndexOfByPredicate([&B](const FBoneChain& Chain)
-		{
-			return B.TargetChain == Chain.ChainName;
-		});
+const TArray<URetargetChainSettings*>& UIKRetargeterController::GetAllChainSettings() const
+{
+	return Asset->ChainSettings_DEPRECATED;
+}
 
-		// compare their StartBone Index 
-		if (IndexA > INDEX_NONE && IndexB > INDEX_NONE)
+URetargetChainSettings* UIKRetargeterController::GetChainSettings(const FName& TargetChainName) const
+{
+	for (URetargetChainSettings* ChainMap : Asset->ChainSettings_DEPRECATED)
+	{
+		if (ChainMap->TargetChain == TargetChainName)
 		{
-			const int32 StartBoneIndexA = TargetSkeleton.GetBoneIndexFromName(BoneChains[IndexA].StartBone.BoneName);
-			const int32 StartBoneIndexB = TargetSkeleton.GetBoneIndexFromName(BoneChains[IndexB].StartBone.BoneName);
-
-			if (StartBoneIndexA == StartBoneIndexB)
-			{
-				// fallback to sorting alphabetically
-				return BoneChains[IndexA].ChainName.LexicalLess(BoneChains[IndexB].ChainName);
-			}
-				
-			return StartBoneIndexA < StartBoneIndexB;	
+			return ChainMap;
 		}
+	}
 
-		// sort them according to the target ik rig if previously failed 
-		return IndexA < IndexB;
-	});
+	return nullptr;
 }
+PRAGMA_ENABLE_DEPRECATION_WARNINGS
+//
+// END DEPRECATED API
+//
 
 #undef LOCTEXT_NAMESPACE
 
diff --git a/Engine/Plugins/Animation/IKRig/Source/IKRigEditor/Private/RetargetEditor/IKRetargeterPoseGenerator.cpp b/Engine/Plugins/Animation/IKRig/Source/IKRigEditor/Private/RetargetEditor/IKRetargeterPoseGenerator.cpp
--- a/Engine/Plugins/Animation/IKRig/Source/IKRigEditor/Private/RetargetEditor/IKRetargeterPoseGenerator.cpp
+++ b/Engine/Plugins/Animation/IKRig/Source/IKRigEditor/Private/RetargetEditor/IKRetargeterPoseGenerator.cpp
@@ -62,7 +62,7 @@ void FRetargetAutoPoseGenerator::AlignBones(
 	if (!bSuppressWarnings)
 	{
 		// warn user if they tried to align the retarget root (doesn't work)
-		if (BonesToAlign.Contains(Processor->GetRetargetRoot(SourceOrTarget)))
+		if (BonesToAlign.Contains(Processor->GetPelvisBone(SourceOrTarget)))
 		{
 			Processor->Log.LogWarning(LOCTEXT("AutoAlignRootNotSupported", "Skipped aligning the retarget root. This is not supported."));
 		}
@@ -70,7 +70,7 @@ void FRetargetAutoPoseGenerator::AlignBones(
 		// warn user if they tried to align a non-retargeted bone (doesn't work)
 		for (const FName Bone : BonesToAlign)
 		{
-			if (!Processor->IsBoneRetargeted(Bone, SourceOrTarget))
+			if (!Processor->IsBoneInAMappedChain(Bone, SourceOrTarget))
 			{
 				Processor->Log.LogWarning(FText::Format(LOCTEXT("SkipAutoAlignNonRetargeted", "Skipped aligning {0}. The bone is not retargeted and therefore has no 'matching' equivalent on other skeleton."), FText::FromName(Bone)));
 			}
@@ -90,7 +90,7 @@ void FRetargetAutoPoseGenerator::AlignBone(
 	}
 	
 	// can't auto-align if the bone is not retargeted  (warning is elsewhere)
-	if (!Processor->IsBoneRetargeted(BoneToAlign, SourceOrTarget))
+	if (!Processor->IsBoneInAMappedChain(BoneToAlign, SourceOrTarget))
 	{
 		return;
 	}
@@ -100,7 +100,7 @@ void FRetargetAutoPoseGenerator::AlignBone(
 	Processor->UpdateRetargetPoseAtRuntime(CurrentPoseName, SourceOrTarget);
 
 	// use a custom procedure to auto-align retarget root
-	if (BoneToAlign == Processor->GetRetargetRoot(SourceOrTarget))
+	if (BoneToAlign == Processor->GetPelvisBone(SourceOrTarget))
 	{
 		AlignRetargetRootBone(SourceOrTarget);
 		return;
@@ -112,7 +112,7 @@ void FRetargetAutoPoseGenerator::AlignBone(
 	
 	// find chains on both sides
 	const FName ChainToAlign = Processor->GetChainNameForBone(BoneToAlign, SideToAlign);
-	const FName ChainToMatch = Processor->GetMappedChainName(ChainToAlign, SideToAlign);
+	const FName ChainToMatch = Processor->GetChainMappedToChain(ChainToAlign, SideToAlign);
 	checkf(ChainToMatch != NAME_None != 0, TEXT("Bone should never be retargeted and not in a mapped chain."));
 	
 	// get the normalized chain parameter to match
@@ -266,7 +266,7 @@ void FRetargetAutoPoseGenerator::AlignRetargetRootBone(const ERetargetSourceOrTa
 	const FQuat WorldRotationOffset = FQuat::FindBetweenNormals(DirectionToAlign, DirectionToMatch);
 
 	// convert to local space offset and store in retarget pose
-	const FName BoneToAlign = Processor->GetRetargetRoot(SourceOrTarget);
+	const FName BoneToAlign = Processor->GetPelvisBone(SourceOrTarget);
 	ApplyWorldRotationToRetargetPose(WorldRotationOffset, BoneToAlign, SourceOrTarget);
 }
 
@@ -292,7 +292,7 @@ void FRetargetAutoPoseGenerator::SnapToGround(const FName BoneToPutOnGround, con
 		// find the lowest bone in the current retarget pose
 		for (int32 BoneIndex=0; BoneIndex<Skeleton.BoneNames.Num(); ++BoneIndex)
 		{
-			if (!Processor->IsBoneRetargeted(Skeleton.BoneNames[BoneIndex], SourceOrTarget))
+			if (!Processor->IsBoneInAMappedChain(Skeleton.BoneNames[BoneIndex], SourceOrTarget))
 			{
 				continue;
 			}
@@ -333,10 +333,10 @@ bool FRetargetAutoPoseGenerator::GetSpineChain(const ERetargetSourceOrTarget Sou
 	OutChainName = NAME_None;
 	const FString ChainNameToFind = "spine";
 	float HighestScore = -1.f;
-	const TArray<FRetargetChainPairFK>& ChainPairs = Processor->GetFKChainPairs();
-	for (const FRetargetChainPairFK& ChainPair : ChainPairs)
+	const TArray<FResolvedBoneChain>& BoneChains = Processor->GetBoneChains().GetResolvedBoneChains(SourceOrTarget);
+	for (const FResolvedBoneChain& BoneChain : BoneChains)
 	{
-		FName ChainNameToTest = SourceOrTarget == ERetargetSourceOrTarget::Source ? ChainPair.SourceBoneChainName : ChainPair.TargetBoneChainName;
+		FName ChainNameToTest = BoneChain.ChainName;
 		FString ChainNameStr = ChainNameToTest.ToString().ToLower();
 		float WorstCase = static_cast<float>(ChainNameToFind.Len() + ChainNameStr.Len());
 		WorstCase = WorstCase < 1.0f ? 1.0f : WorstCase;
@@ -584,7 +584,7 @@ bool FRetargetAutoPoseGenerator::IsARetargetLeaf(
 	GetAllChildrenOfBone(BoneName, SourceOrTarget, AllChildren);
 	for (const FName Child : AllChildren)
 	{
-		if (Processor->IsBoneRetargeted(Child, SourceOrTarget))
+		if (Processor->IsBoneInAMappedChain(Child, SourceOrTarget))
 		{
 			return false;
 		}
@@ -611,7 +611,7 @@ bool FRetargetAutoPoseGenerator::CheckReadyToAlignBones() const
 	UIKRetargeter* RetargeterAsset = Controller->GetAssetPtr();
 	constexpr bool bSuppressWarnings = true;
 	FRetargetProfile RetargetProfile;
-	RetargeterAsset->FillProfileWithAssetSettings(RetargetProfile);
+	RetargetProfile.FillProfileWithAssetSettings(RetargeterAsset);
 	Processor->Initialize(SourceSkeletalMesh, TargetSkeletalMesh, RetargeterAsset, RetargetProfile, bSuppressWarnings);
 	
 	// can't auto align until processor is initialized because we need a fully resolved chain mapping on source/target
