[IKRig] Major refactor part 3 of 3.

diff --git a/Engine/Plugins/Animation/IKRig/Source/IKRigEditor/Private/RetargetEditor/SIKRetargetChainMapList.cpp b/Engine/Plugins/Animation/IKRig/Source/IKRigEditor/Private/RetargetEditor/SIKRetargetChainMapList.cpp
--- a/Engine/Plugins/Animation/IKRig/Source/IKRigEditor/Private/RetargetEditor/SIKRetargetChainMapList.cpp
+++ b/Engine/Plugins/Animation/IKRig/Source/IKRigEditor/Private/RetargetEditor/SIKRetargetChainMapList.cpp
@@ -73,7 +73,7 @@ TSharedRef< SWidget > SIKRetargetChainMapRow::GenerateWidgetForColumn(const FNam
 			.Padding(3.0f, 1.0f)
 			[
 				SNew(STextBlock)
-				.Text(FText::FromName(ChainMapElement.Pin()->ChainMap->TargetChain))
+				.Text(FText::FromName(ChainMapElement.Pin()->TargetChainName))
 				.Font(FAppStyle::GetFontStyle(TEXT("BoldFont")))
 			];
 		return NewWidget;
@@ -151,15 +151,20 @@ TSharedRef< SWidget > SIKRetargetChainMapRow::GenerateWidgetForColumn(const FNam
 
 FReply SIKRetargetChainMapRow::OnResetToDefaultClicked()
 {
-	ChainMapList.Pin()->ResetChainSettings(ChainMapElement.Pin()->ChainMap.Get());
+	ChainMapList.Pin()->ResetChainSettings(ChainMapElement.Pin()->TargetChainName);
 	return FReply::Handled();
 }
 
 EVisibility SIKRetargetChainMapRow::GetResetToDefaultVisibility() const
 {
-	const FTargetChainSettings& DefaultSettings = FTargetChainSettings();
-	const FTargetChainSettings& Settings = ChainMapElement.Pin()->ChainMap->Settings;
-	return (Settings == DefaultSettings) ? EVisibility::Hidden : EVisibility::Visible;
+	UIKRetargeterController* RetargeterController = ChainMapList.Pin()->GetRetargetController();
+	if (!RetargeterController)
+	{
+		return EVisibility::Hidden; 
+	}
+
+	bool bIsChainNonDefault = RetargeterController->ChainHasNonDefaultSettings(ChainMapElement.Pin()->TargetChainName);
+	return bIsChainNonDefault ? EVisibility::Visible : EVisibility::Hidden;
 }
 
 void SIKRetargetChainMapRow::OnSourceChainComboSelectionChanged(TSharedPtr<FString> InName, ESelectInfo::Type SelectInfo)
@@ -171,7 +176,7 @@ void SIKRetargetChainMapRow::OnSourceChainComboSelectionChanged(TSharedPtr<FStri
 	}
 
 	const FName SourceChainName = FName(*InName.Get());
-	RetargeterController->SetSourceChain(SourceChainName, ChainMapElement.Pin()->ChainMap->TargetChain);
+	RetargeterController->SetSourceChain(SourceChainName, ChainMapElement.Pin()->TargetChainName);
 }
 
 FText SIKRetargetChainMapRow::GetSourceChainName() const
@@ -182,7 +187,7 @@ FText SIKRetargetChainMapRow::GetSourceChainName() const
 		return FText::FromName(NAME_None); 
 	}
 
-	return FText::FromName(ChainMapElement.Pin()->ChainMap->SourceChain);
+	return FText::FromName(RetargeterController->GetSourceChain(ChainMapElement.Pin()->TargetChainName));
 }
 
 FText SIKRetargetChainMapRow::GetTargetIKGoalName() const
@@ -200,7 +205,7 @@ FText SIKRetargetChainMapRow::GetTargetIKGoalName() const
 	}
 
 	const UIKRigController* RigController = UIKRigController::GetController(IKRig);
-	const FBoneChain* Chain = RigController->GetRetargetChainByName(ChainMapElement.Pin()->ChainMap->TargetChain);
+	const FBoneChain* Chain = RigController->GetRetargetChainByName(ChainMapElement.Pin()->TargetChainName);
 	if (!Chain)
 	{
 		return FText(); 
@@ -370,11 +375,9 @@ void SIKRetargetChainMapList::ClearSelection() const
 	ListView->ClearSelection();
 }
 
-void SIKRetargetChainMapList::ResetChainSettings(URetargetChainSettings* ChainMap) const
+void SIKRetargetChainMapList::ResetChainSettings(const FName InTargetChainName) const
 {
-	FScopedTransaction Transaction(LOCTEXT("ResetChainSettings", "Reset Retarget Chain Settings"));
-	ChainMap->Modify();
-	ChainMap->Settings = FTargetChainSettings();
+	EditorController.Pin()->AssetController->ResetChainSettingsInAllOps(InTargetChainName);
 	EditorController.Pin()->RefreshDetailsView();
 }
 
@@ -464,41 +467,44 @@ void SIKRetargetChainMapList::RefreshView()
 
 	// refresh items
 	ListViewItems.Reset();
-	const TArray<URetargetChainSettings*>& ChainMappings = AssetController->GetAllChainSettings();
-	for (const TObjectPtr<URetargetChainSettings> ChainMap : ChainMappings)
+	const FRetargetChainMapping& ChainMapping = AssetController->GetChainMapping();
+	const TArray<FRetargetChainPair>& ChainMap = ChainMapping.GetChainPairs();
+	for (const FRetargetChainPair& ChainPair : ChainMap)
 	{
+		const FName TargetChainName = ChainPair.TargetChainName;
+		const FName SourceChainName = ChainPair.SourceChainName;
 		// apply text filter to items
 		if (!(TextFilter->GetFilterText().IsEmpty() ||
-			FilterString(ChainMap->SourceChain.ToString()) ||
-			FilterString(ChainMap->TargetChain.ToString())))
+			FilterString(SourceChainName.ToString()) ||
+			FilterString(TargetChainName.ToString())))
 		{
 			continue;
 		}
 		
 		// apply "only IK" filter
-		if (ChainFilterOptions.bHideChainsWithoutIK && !DoesChainHaveIK(ChainMap->TargetChain))
+		if (ChainFilterOptions.bHideChainsWithoutIK && !DoesChainHaveIK(TargetChainName))
 		{
 			continue;
 		}
 
 		// apply "hide mapped chains" filter
-		if (ChainFilterOptions.bHideMappedChains && ChainMap->SourceChain != LiteralNone)
+		if (ChainFilterOptions.bHideMappedChains && SourceChainName != LiteralNone)
 		{
 			continue;
 		}
 		
 		// apply "hide un-mapped chains" filter
-		if (ChainFilterOptions.bHideUnmappedChains && ChainMap->SourceChain == LiteralNone)
+		if (ChainFilterOptions.bHideUnmappedChains && SourceChainName == LiteralNone)
 		{
 			continue;
 		}
 		
 		// create an item for this chain
-		TSharedPtr<FRetargetChainMapElement> ChainItem = FRetargetChainMapElement::Make(ChainMap);
+		TSharedPtr<FRetargetChainMapElement> ChainItem = FRetargetChainMapElement::Make(TargetChainName);
 		ListViewItems.Add(ChainItem);
 
 		// select/deselect it
-		const bool bIsSelected = SelectedChains.Contains(ChainMap->TargetChain);
+		const bool bIsSelected = SelectedChains.Contains(TargetChainName);
 		ListView->SetItemSelection(ChainItem, bIsSelected, ESelectInfo::Direct);
 	}
 	
@@ -614,7 +620,7 @@ void SIKRetargetChainMapList::OnItemClicked(TSharedPtr<FRetargetChainMapElement>
 	TArray<TSharedPtr<FRetargetChainMapElement>> SelectedItems = ListView.Get()->GetSelectedItems();
 	for (const TSharedPtr<FRetargetChainMapElement>& Item : SelectedItems)
 	{
-		SelectedChains.Add(Item.Get()->ChainMap->TargetChain);
+		SelectedChains.Add(Item.Get()->TargetChainName);
 	}
 	
 	// replace the chain selection
@@ -741,7 +747,6 @@ void SIKRetargetChainMapList::AutoMapChains(const EAutoMapChainType AutoMapType,
 
 	Controller->ClearOutputLog();
 	RetargeterController->AutoMapChains(AutoMapType, bForceRemap);
-	Controller->HandleRetargeterNeedsInitialized();
 }
 
 #undef LOCTEXT_NAMESPACE
diff --git a/Engine/Plugins/Animation/IKRig/Source/IKRigEditor/Private/RetargetEditor/SRetargetAnimAssetsWindow.cpp b/Engine/Plugins/Animation/IKRig/Source/IKRigEditor/Private/RetargetEditor/SRetargetAnimAssetsWindow.cpp
--- a/Engine/Plugins/Animation/IKRig/Source/IKRigEditor/Private/RetargetEditor/SRetargetAnimAssetsWindow.cpp
+++ b/Engine/Plugins/Animation/IKRig/Source/IKRigEditor/Private/RetargetEditor/SRetargetAnimAssetsWindow.cpp
@@ -28,6 +28,7 @@
 #include "RetargetEditor/IKRetargetEditorController.h"
 #include "Retargeter/RetargetOps/PinBoneOp.h"
 #include "Retargeter/RetargetOps/RootMotionGeneratorOp.h"
+#include "Retargeter/RetargetOps/RunIKRigOp.h"
 #include "RigEditor/IKRigAutoCharacterizer.h"
 #include "RigEditor/IKRigController.h"
 #include "Settings/SkeletalMeshEditorSettings.h"
@@ -685,49 +686,52 @@ void FProceduralRetargetAssets::AutoGenerateIKRetargetAsset()
 			TargetCharacterizationResults.AutoRetargetDefinition.BonesToExcludeFromAutoPose,
 			ERetargetSourceOrTarget::Target);
 	}
-	
-	// disable IK pass because auto-IK is not yet robust enough, still useful for manual editing though
-	FRetargetGlobalSettings GlobalSettings = RetargetController->GetGlobalSettings();
-	GlobalSettings.bEnableIK = false;
-	RetargetController->SetGlobalSettings(GlobalSettings);
 
 	// clear the op stack and regenerate it
 	RetargetController->RemoveAllOps();
+
+	// add all the default ops for retargeting
+	constexpr bool bRunInitialSetup = true;
+	RetargetController->AddDefaultOps(bRunInitialSetup);
+
+	// disable IK pass because auto-IK is not yet robust enough, still useful for manual editing though
+	FIKRetargetRunIKRigOp* RunIKOp = RetargetController->GetFirstRetargetOpOfType<FIKRetargetRunIKRigOp>();
+	RunIKOp->SetEnabled(false);
 	
-	// add a root motion generator (but only if the skeleton has a root bone separate from the pelvis/retarget root)
+	// setup the root motion generator (but only if the skeleton has a root bone separate from the pelvis/retarget root)
+	FIKRetargetRootMotionOp* RootOp = RetargetController->GetFirstRetargetOpOfType<FIKRetargetRootMotionOp>();
 	const FName TargetSkeletonRoot = TargetIKRig->GetSkeleton().BoneNames[0];
 	const FName TargetPelvis = TargetCharacterizationResults.AutoRetargetDefinition.RetargetDefinition.RootBone;
 	if (TargetSkeletonRoot != TargetPelvis)
 	{
-		// add a root motion op to the post retarget op stack
-		const int32 RootOpIndex = RetargetController->AddRetargetOp(FIKRetargetRootMotionOp::StaticStruct());
-
 		// if the source does not have a dedicated root bone, configure the root motion to be generated from the pelvis motion on target
 		const FName SourceSkeletonRoot = SourceIKRig->GetSkeleton().BoneNames[0];
 		const FName SourcePelvis = SourceCharacterizationResults.AutoRetargetDefinition.RetargetDefinition.RootBone;
 		if (SourceSkeletonRoot == SourcePelvis)
 		{
-			if (FIKRetargetRootMotionOp* RootOp = (FIKRetargetRootMotionOp*)(RetargetController->GetRetargetOpAtIndex(RootOpIndex)))
-			{
-				RootOp->Settings.RootMotionSource = ERootMotionSource::GenerateFromTargetPelvis;
-				RootOp->Settings.RootHeightSource = ERootMotionHeightSource::SnapToGround;
-				RootOp->Settings.bRotateWithPelvis = true;
-				RootOp->Settings.bMaintainOffsetFromPelvis = true;
-			}
+			RootOp->Settings.RootMotionSource = ERootMotionSource::GenerateFromTargetPelvis;
+			RootOp->Settings.RootHeightSource = ERootMotionHeightSource::SnapToGround;
+			RootOp->Settings.bRotateWithPelvis = true;
+			RootOp->Settings.bMaintainOffsetFromPelvis = true;
 		}
 	}
+	else
+	{
+		// the pelvis IS the root, so we let the Pelvis op handle "root" motion
+		RootOp->SetEnabled(false);
+	}
 	
 	// setup a "Pin Bone Op" to pin IK bones to the new target locations
-	const int32 PinBonesOpIndex = RetargetController->AddRetargetOp(FIKRetargetPinBoneOp::StaticStruct());
-	FIKRetargetPinBoneOp* PinBoneOp = (FIKRetargetPinBoneOp*)RetargetController->GetRetargetOpAtIndex(PinBonesOpIndex);
+	RetargetController->AddRetargetOp(FIKRetargetPinBoneOp::StaticStruct());
+	FIKRetargetPinBoneOp* PinBoneOp = RetargetController->GetFirstRetargetOpOfType<FIKRetargetPinBoneOp>();
 	PinBoneOp->Settings.TranslationMode = EPinBoneTranslationMode::CopyGlobalPosition;
 	PinBoneOp->Settings.RotationMode = EPinBoneRotationMode::CopyGlobalRotation;
 	const TArray<FBoneToPin>& AllBonesToPin = TargetCharacterizationResults.AutoRetargetDefinition.BonesToPin.GetBonesToPin();
 	for (const FBoneToPin& BonesToPin : AllBonesToPin)
 	{
 		FPinBoneData PinBoneData;
-		PinBoneData.BoneToCopyFrom = BonesToPin.BoneToPinTo;
-		PinBoneData.BoneToCopyTo = BonesToPin.BoneToPin;
+		PinBoneData.BoneToCopyFrom.BoneName = BonesToPin.BoneToPinTo;
+		PinBoneData.BoneToCopyTo.BoneName = BonesToPin.BoneToPin;
 		PinBoneOp->Settings.BonesToPin.Add(PinBoneData);
 	}
 }
diff --git a/Engine/Plugins/Animation/IKRig/Source/IKRigEditor/Private/RetargetEditor/SRetargetOpStack.cpp b/Engine/Plugins/Animation/IKRig/Source/IKRigEditor/Private/RetargetEditor/SRetargetOpStack.cpp
--- a/Engine/Plugins/Animation/IKRig/Source/IKRigEditor/Private/RetargetEditor/SRetargetOpStack.cpp
+++ b/Engine/Plugins/Animation/IKRig/Source/IKRigEditor/Private/RetargetEditor/SRetargetOpStack.cpp
@@ -146,11 +146,12 @@ bool SRetargetOpStackItem::GetWarningMessage(FText& Message) const
 		return false;
 	}
 
-	const TArray<TSharedPtr<FIKRetargetOpBase>>& RetargetOps = Processor->GetRetargetOps();
+	const TArray<FInstancedStruct>& RetargetOps = Processor->GetRetargetOps();
 	const int32 OpIndex = StackElement.Pin()->IndexInStack;
 	if (RetargetOps.IsValidIndex(OpIndex))
 	{
-		Message = RetargetOps[OpIndex]->GetWarningMessage();
+		const FIKRetargetOpBase* Op = RetargetOps[OpIndex].GetPtr<FIKRetargetOpBase>();
+		Message = Op->GetWarningMessage();
 		return true;
 	}
 	
@@ -266,18 +267,26 @@ TSharedRef<SWidget> SRetargetOpStack::CreateAddNewMenuWidget()
 
 	// add menu option to create each retarget op type
 	const UStruct* BaseStruct = FIKRetargetOpBase::StaticStruct();
-	for(TObjectIterator<UStruct> StructIt; StructIt; ++StructIt)
+	for (TObjectIterator<UStruct> StructIt; StructIt; ++StructIt)
 	{
 		UStruct* Struct = *StructIt;
-		if(!Struct->IsChildOf(BaseStruct) || Struct == BaseStruct)
+		if (!Struct->IsChildOf(BaseStruct) || Struct == BaseStruct)
 		{
 			continue;
 		}
 
 		if (UScriptStruct* ScriptStruct = Cast<UScriptStruct>(Struct))
 		{
-			FString TypeName = ScriptStruct->GetName();
-			FUIAction Action = FUIAction( FExecuteAction::CreateSP(this, &SRetargetOpStack::AddNewRetargetOp, ScriptStruct));
+			// try to get the "nice name" from metadata
+			FString TypeName = ScriptStruct->GetMetaData(TEXT("DisplayName"));
+		
+			// if no "DisplayName" metadata is found, fall back to the struct name
+			if (TypeName.IsEmpty())
+			{
+				TypeName = ScriptStruct->GetName();
+			}
+
+			FUIAction Action = FUIAction(FExecuteAction::CreateSP(this, &SRetargetOpStack::AddNewRetargetOp, ScriptStruct));
 			MenuBuilder.AddMenuEntry(FText::FromString(TypeName), FText::GetEmpty(), FSlateIcon(), Action);
 		}
 	}
@@ -302,7 +311,7 @@ bool SRetargetOpStack::IsAddOpEnabled() const
 	return false;
 }
 
-void SRetargetOpStack::AddNewRetargetOp(UScriptStruct* ScriptStruct) const
+void SRetargetOpStack::AddNewRetargetOp(UScriptStruct* ScriptStruct)
 {
 	if (!EditorController.IsValid())
 	{
@@ -315,7 +324,18 @@ void SRetargetOpStack::AddNewRetargetOp(UScriptStruct* ScriptStruct) const
 		return;
 	}
 	
-	AssetController->AddRetargetOp(ScriptStruct);
+	int32 NewOpIndex = AssetController->AddRetargetOp(ScriptStruct);
+	
+	// ops have a chance at initial setup when first added to a stack
+	FIKRetargetProcessor* Processor = EditorController.Pin()->GetRetargetProcessor();
+	FIKRetargetOpBase* NewOp = AssetController->GetRetargetOpAtIndex(NewOpIndex);
+	if (NewOp && Processor)
+	{
+		NewOp->RunInitialSetup(AssetController->GetAsset());
+		Processor->SetNeedsInitialized();
+	}
+
+	RefreshStackView();
 }
 
 void SRetargetOpStack::DeleteRetargetOp(TSharedPtr<FRetargetOpStackElement> OpToDelete)
@@ -443,7 +463,7 @@ TOptional<EItemDropZone> SRetargetOpStack::OnCanAcceptDrop(
 	const TSharedPtr<FRetargetOpStackDragDropOp> DragDropOp = DragDropEvent.GetOperationAs<FRetargetOpStackDragDropOp>();
 	if (DragDropOp.IsValid())
 	{
-		ReturnedDropZone = EItemDropZone::BelowItem;	
+		ReturnedDropZone = DropZone;
 	}
 	
 	return ReturnedDropZone;
@@ -468,7 +488,8 @@ FReply SRetargetOpStack::OnAcceptDrop(
 
 	const FRetargetOpStackElement& DraggedElement = *DragDropOp.Get()->Element.Pin().Get();
 	const UIKRetargeterController* AssetController = Controller->AssetController;
-	const bool bWasReparented = AssetController->MoveRetargetOpInStack(DraggedElement.IndexInStack, TargetItem.Get()->IndexInStack);
+	int32 TargetIndexOffset = DropZone == EItemDropZone::BelowItem ? 1 : -1;
+	const bool bWasReparented = AssetController->MoveRetargetOpInStack(DraggedElement.IndexInStack, TargetItem.Get()->IndexInStack + TargetIndexOffset);
 	if (bWasReparented)
 	{
 		RefreshStackView();
diff --git a/Engine/Plugins/Animation/IKRig/Source/IKRigEditor/Private/RigEditor/IKRigController.cpp b/Engine/Plugins/Animation/IKRig/Source/IKRigEditor/Private/RigEditor/IKRigController.cpp
--- a/Engine/Plugins/Animation/IKRig/Source/IKRigEditor/Private/RigEditor/IKRigController.cpp
+++ b/Engine/Plugins/Animation/IKRig/Source/IKRigEditor/Private/RigEditor/IKRigController.cpp
@@ -22,7 +22,7 @@ UIKRigController::UIKRigController()
 	AutoCharacterizer = MakeUnique<FAutoCharacterizer>();
 	AutoFBIKCreator = MakeUnique<FAutoFBIKCreator>();
 
-	StructViewer = CreateDefaultSubobject<UStructViewer>(TEXT("SettingsViewer"));
+	StructViewer = CreateDefaultSubobject<UIKRigStructViewer>(TEXT("SettingsViewer"));
 }
 
 UIKRigController* UIKRigController::GetController(const UIKRigDefinition* InIKRigDefinition)
@@ -605,7 +605,7 @@ bool UIKRigController::ValidateChain(
 	}
 }
 
-UStructViewer* UIKRigController::GetStructViewer()
+UIKRigStructViewer* UIKRigController::GetStructViewer()
 {
 	StructViewer->Reset();
 	return StructViewer;
diff --git a/Engine/Plugins/Animation/IKRig/Source/IKRigEditor/Private/RigEditor/IKRigDetailCustomizations.cpp b/Engine/Plugins/Animation/IKRig/Source/IKRigEditor/Private/RigEditor/IKRigDetailCustomizations.cpp
--- a/Engine/Plugins/Animation/IKRig/Source/IKRigEditor/Private/RigEditor/IKRigDetailCustomizations.cpp
+++ b/Engine/Plugins/Animation/IKRig/Source/IKRigEditor/Private/RigEditor/IKRigDetailCustomizations.cpp
@@ -8,6 +8,7 @@
 #include "DetailWidgetRow.h"
 #include "IDetailGroup.h"
 #include "IDetailPropertyRow.h"
+#include "IKRig.h"
 #include "Widgets/SWidget.h"
 #include "Widgets/Input/SSegmentedControl.h"
 #include "Framework/MultiBox/MultiBoxBuilder.h"
@@ -20,7 +21,39 @@ namespace IKRigDetailCustomizationsConstants
 	static const float ItemWidth = 125.0f;
 }
 
-USkeleton* UStructViewer::GetSkeleton(bool& bInvalidSkeletonIsError, const IPropertyHandle* PropertyHandle)
+void UIKRigStructViewer::TriggerReinitIfNeeded(const FPropertyChangedEvent& InEvent)
+{
+	// never reinitialize while interacting with a slider
+	if (InEvent.ChangeType == EPropertyChangeType::Interactive || !InEvent.Property)
+	{
+		return;
+	}
+
+	// always reinitialize if property is marked for it
+	if (InEvent.Property->HasMetaData(IKRigReinitOnEditMetaLabel))
+	{
+		OnStructPropertyEdited().Broadcast();
+		return;
+	}
+
+	// if it is a name property, we want to reinitialize to validate it
+	if (CastField<FNameProperty>(InEvent.Property) != nullptr)
+	{
+		OnStructPropertyEdited().Broadcast();
+		return;
+	}
+
+	// always reinitialize if the property is a bone name
+	// NOTE: this is needed because FBoneReference's child BoneName property
+	// cannot be marked with the meta-tag (IKRig plugin doesn't own it and it's only relevant here)
+	if (InEvent.Property->GetFName() == FName(TEXT("BoneName")))
+	{
+		OnStructPropertyEdited().Broadcast();
+		return;
+	}
+}
+
+USkeleton* UIKRigStructViewer::GetSkeleton(bool& bInvalidSkeletonIsError, const IPropertyHandle* PropertyHandle)
 {
 	if (!StructType)
 	{
@@ -51,7 +84,7 @@ void FStructViewerCustomization::CustomizeDetails(IDetailLayoutBuilder& DetailBu
 {
 	TArray<TWeakObjectPtr<UObject>> ObjectsBeingCustomized;
 	DetailBuilder.GetObjectsBeingCustomized(ObjectsBeingCustomized);
-	UStructViewer* StructViewer = Cast<UStructViewer>(ObjectsBeingCustomized[0]);
+	UIKRigStructViewer* StructViewer = Cast<UIKRigStructViewer>(ObjectsBeingCustomized[0]);
 	if (!ensure(StructViewer))
 	{
 		return;
@@ -74,12 +107,26 @@ void FStructViewerCustomization::CustomizeDetails(IDetailLayoutBuilder& DetailBu
 	{
 		Row->GetPropertyHandle()->SetOnPropertyValuePreChange(FSimpleDelegate::CreateLambda([ObjectToTransact]()
 		{
-			GEditor->BeginTransaction(LOCTEXT("StructViewer", "Edited property."));
+			GEditor->BeginTransaction(LOCTEXT("IKRigStructViewer", "Edited property."));
 			ObjectToTransact->Modify();
 		}));
-		Row->GetPropertyHandle()->SetOnPropertyValueChanged(FSimpleDelegate::CreateLambda([ObjectToTransact]()
+		
+		Row->GetPropertyHandle()->SetOnPropertyValueChangedWithData(TDelegate<void(const FPropertyChangedEvent&)>::CreateLambda([StructViewer](const FPropertyChangedEvent& InEvent)
+		{
+			GEditor->EndTransaction();
+			StructViewer->TriggerReinitIfNeeded(InEvent);
+		}));
+
+		Row->GetPropertyHandle()->SetOnChildPropertyValuePreChange(FSimpleDelegate::CreateLambda([ObjectToTransact]()
+		{
+			GEditor->BeginTransaction(LOCTEXT("IKRigStructViewerChild", "Edited child property."));
+			ObjectToTransact->Modify();
+		}));
+
+		Row->GetPropertyHandle()->SetOnChildPropertyValueChangedWithData(TDelegate<void(const FPropertyChangedEvent&)>::CreateLambda([StructViewer](const FPropertyChangedEvent& InEvent)
 		{
 			GEditor->EndTransaction();
+			StructViewer->TriggerReinitIfNeeded(InEvent);
 		}));
 	}
 }
diff --git a/Engine/Plugins/Animation/IKRig/Source/IKRigEditor/Private/RigEditor/IKRigEditorController.cpp b/Engine/Plugins/Animation/IKRig/Source/IKRigEditor/Private/RigEditor/IKRigEditorController.cpp
--- a/Engine/Plugins/Animation/IKRig/Source/IKRigEditor/Private/RigEditor/IKRigEditorController.cpp
+++ b/Engine/Plugins/Animation/IKRig/Source/IKRigEditor/Private/RigEditor/IKRigEditorController.cpp
@@ -973,7 +973,7 @@ void FIKRigEditorController::ShowDetailsForBoneSettings(const FName& BoneName, i
 	FInstancedStruct* InstancedStruct = AssetController->GetSolverStructAtIndex(SolverIndex);
 	const UScriptStruct* TypeToEdit = InstancedStruct->GetPtr<FIKRigSolverBase>()->GetBoneSettingsType();
 	
-	UStructViewer* StructViewer = AssetController->GetStructViewer();
+	UIKRigStructViewer* StructViewer = AssetController->GetStructViewer();
 	StructViewer->SetStructToView(TypeToEdit, MemoryProvider, AssetController->GetAsset());
 	
 	ShowStructViewerInDetailsView(StructViewer);
@@ -1002,7 +1002,7 @@ void FIKRigEditorController::ShowDetailsForGoalSettings(const FName GoalName, co
 	FInstancedStruct* InstancedStruct = AssetController->GetSolverStructAtIndex(SolverIndex);
 	const UScriptStruct* TypeToEdit = InstancedStruct->GetPtr<FIKRigSolverBase>()->GetGoalSettingsType();
 	
-	UStructViewer* StructViewer = AssetController->GetStructViewer();
+	UIKRigStructViewer* StructViewer = AssetController->GetStructViewer();
 	StructViewer->SetStructToView(TypeToEdit, MemoryProvider, AssetController->GetAsset());
 	
 	ShowStructViewerInDetailsView(StructViewer);
@@ -1026,7 +1026,7 @@ void FIKRigEditorController::ShowDetailsForSolver(const int32 SolverIndex) const
 	FInstancedStruct* InstancedStruct = AssetController->GetSolverStructAtIndex(SolverIndex);
 	const UScriptStruct* TypeToEdit = InstancedStruct->GetPtr<FIKRigSolverBase>()->GetSolverSettingsType();
 	
-	UStructViewer* StructViewer = AssetController->GetStructViewer();
+	UIKRigStructViewer* StructViewer = AssetController->GetStructViewer();
 	StructViewer->SetStructToView(TypeToEdit, MemoryProvider, AssetController->GetAsset());
 	
 	ShowStructViewerInDetailsView(StructViewer);
@@ -1071,7 +1071,7 @@ void FIKRigEditorController::ShowDetailsForElements(const TArray<TSharedPtr<FIKR
 	}
 }
 
-void FIKRigEditorController::ShowStructViewerInDetailsView(UStructViewer* StructViewer) const
+void FIKRigEditorController::ShowStructViewerInDetailsView(UIKRigStructViewer* StructViewer) const
 {
 	constexpr bool bForceRefresh = true; // must forcibly refresh because struct viewer is a recycled UObject
 	DetailsView->SetObject(StructViewer, bForceRefresh);
diff --git a/Engine/Plugins/Animation/IKRig/Source/IKRigEditor/Public/RetargetEditor/IKRetargetCommands.h b/Engine/Plugins/Animation/IKRig/Source/IKRigEditor/Public/RetargetEditor/IKRetargetCommands.h
--- a/Engine/Plugins/Animation/IKRig/Source/IKRigEditor/Public/RetargetEditor/IKRetargetCommands.h
+++ b/Engine/Plugins/Animation/IKRig/Source/IKRigEditor/Public/RetargetEditor/IKRetargetCommands.h
@@ -27,10 +27,9 @@ public:
 
 	// open asset settings
 	TSharedPtr< FUICommandInfo > ShowAssetSettings;
-	// open global settings
-	TSharedPtr< FUICommandInfo > ShowGlobalSettings;
-	// open root settings
-	TSharedPtr< FUICommandInfo > ShowRootSettings;
+	
+	// add default ops to the stack
+	TSharedPtr< FUICommandInfo > AddDefaultOps;
 
 	// enable Root retarget pass
 	TSharedPtr< FUICommandInfo > EnableRoot;
diff --git a/Engine/Plugins/Animation/IKRig/Source/IKRigEditor/Public/RetargetEditor/IKRetargetDetails.h b/Engine/Plugins/Animation/IKRig/Source/IKRigEditor/Public/RetargetEditor/IKRetargetDetails.h
--- a/Engine/Plugins/Animation/IKRig/Source/IKRigEditor/Public/RetargetEditor/IKRetargetDetails.h
+++ b/Engine/Plugins/Animation/IKRig/Source/IKRigEditor/Public/RetargetEditor/IKRetargetDetails.h
@@ -166,78 +166,3 @@ private:
 
 	TArray<TObjectPtr<UIKRetargetBoneDetails>> Bones;
 };
-
-/** ------------------------------------- BEGIN CHAIN DETAILS CUSTOMIZATION -------------*/
-
-
-class FRetargetChainSettingsCustomization : public IDetailCustomization
-{
-public:
-
-	static TSharedRef<IDetailCustomization> MakeInstance()
-	{
-		return MakeShareable(new FRetargetChainSettingsCustomization);
-	}
-
-	/** IDetailCustomization interface */
-	virtual void CustomizeDetails(IDetailLayoutBuilder& DetailBuilder) override;
-
-private:
-
-	void AddSettingsSection(
-		const IDetailLayoutBuilder& DetailBuilder,
-		const FIKRetargetEditorController* Controller,
-		IDetailCategoryBuilder& SettingsCategory,
-		const FString& StructPropertyName,
-		const FName& GroupName,
-		const FText& LocalizedGroupName,
-		const UScriptStruct* SettingsClass,
-		const FString& EnabledPropertyName,
-		const bool& bIsSectionEnabled,
-		const FText& DisabledMessage) const;
-	
-	TArray<TWeakObjectPtr<URetargetChainSettings>> ChainSettingsObjects;
-	TArray<TSharedPtr<FString>> SourceChainOptions;
-};
-
-/** ------------------------------------- BEGIN ROOT DETAILS CUSTOMIZATION -------------*/
-
-class FRetargetRootSettingsCustomization : public IDetailCustomization
-{
-public:
-
-	static TSharedRef<IDetailCustomization> MakeInstance()
-	{
-		return MakeShareable(new FRetargetRootSettingsCustomization);
-	}
-
-	/** IDetailCustomization interface */
-	virtual void CustomizeDetails(IDetailLayoutBuilder& DetailBuilder) override;
-
-private:
-
-	TWeakObjectPtr<URetargetRootSettings> RootSettingsObject;
-	TWeakPtr<FIKRetargetEditorController> Controller;
-};
-
-/** ------------------------------------- BEGIN GLOBAL DETAILS CUSTOMIZATION -------------*/
-
-class FRetargetGlobalSettingsCustomization : public IDetailCustomization
-{
-public:
-
-	static TSharedRef<IDetailCustomization> MakeInstance()
-	{
-		return MakeShareable(new FRetargetGlobalSettingsCustomization);
-	}
-
-	/** IDetailCustomization interface */
-	virtual void CustomizeDetails(IDetailLayoutBuilder& DetailBuilder) override;
-
-private:
-
-	TWeakObjectPtr<UIKRetargetGlobalSettings> GlobalSettingsObject;
-	TWeakPtr<FIKRetargetEditorController> Controller;
-
-	TArray<TSharedPtr<FString>> TargetChainOptions;
-};
diff --git a/Engine/Plugins/Animation/IKRig/Source/IKRigEditor/Public/RetargetEditor/IKRetargetEditor.h b/Engine/Plugins/Animation/IKRig/Source/IKRigEditor/Public/RetargetEditor/IKRetargetEditor.h
--- a/Engine/Plugins/Animation/IKRig/Source/IKRigEditor/Public/RetargetEditor/IKRetargetEditor.h
+++ b/Engine/Plugins/Animation/IKRig/Source/IKRigEditor/Public/RetargetEditor/IKRetargetEditor.h
@@ -80,7 +80,6 @@ private:
 	void ExtendToolbar();
 	void FillToolbar(FToolBarBuilder& ToolbarBuilder);
 	TSharedRef<SWidget> GenerateRetargetModesMenu();
-	TSharedPtr<SHorizontalBox> GenerateRetargetPhaseButtons() const;
 	/** END toolbar */
 	
 	/** preview scene setup */
diff --git a/Engine/Plugins/Animation/IKRig/Source/IKRigEditor/Public/RetargetEditor/IKRetargetEditorController.h b/Engine/Plugins/Animation/IKRig/Source/IKRigEditor/Public/RetargetEditor/IKRetargetEditorController.h
--- a/Engine/Plugins/Animation/IKRig/Source/IKRigEditor/Public/RetargetEditor/IKRetargetEditorController.h
+++ b/Engine/Plugins/Animation/IKRig/Source/IKRigEditor/Public/RetargetEditor/IKRetargetEditorController.h
@@ -41,15 +41,6 @@ enum class ESelectionEdit : uint8
 	Replace	// replace selection entirely
 };
 
-enum class ERetargetSelectionType : uint8
-{
-	BONE,
-	CHAIN,
-	MESH,
-	ROOT,
-	NONE
-};
-
 struct FBoundIKRig
 {
 	FBoundIKRig(UIKRigDefinition* InIKRig, const FIKRetargetEditorController& InController);
@@ -168,11 +159,6 @@ public:
 	bool IsCurrentMeshLoaded() const;
 	bool IsEditingPose() const;
 
-	// display settings in the details panel
-	void ShowGlobalSettings();
-	void ShowRootSettings();
-	FRetargetGlobalSettings& GetGlobalSettings() const;
-
 	// clear the output log
 	void ClearOutputLog() const;
 
@@ -204,7 +190,7 @@ public:
 	// get the retargeter that is running in the viewport (which is a duplicate of the source asset)
 	FIKRetargetProcessor* GetRetargetProcessor() const;
 	// Reset the planting state of the IK (when scrubbing or animation loops over)
-	void ResetIKPlantingState() const;
+	void OnPlaybackReset() const;
 
 	// Set viewport / editor tool mode
 	void SetRetargeterMode(ERetargeterOutputMode Mode);
@@ -228,7 +214,8 @@ public:
 		const TArray<FName>& InBoneNames,
 		ESelectionEdit EditMode,
 		const bool bFromHierarchyView = false);
-	const TArray<FName>& GetSelectedBones() const {return SelectedBoneNames[CurrentlyEditingSourceOrTarget]; };
+	const TArray<FName>& GetSelectedBones() const {return Selection.SelectedBoneNames[CurrentlyEditingSourceOrTarget]; };
+	const FIKRetargetDebugDrawState& GetSelectionState() const { return Selection; };
 	// END bone selection
 
 	// SELECTION - CHAINS (viewport or chains view)
@@ -236,18 +223,19 @@ public:
 		const TArray<FName>& InChainNames,
 		ESelectionEdit EditMode,
 		const bool bFromChainsView);
-	const TArray<FName>& GetSelectedChains() const {return SelectedChains; };
+	const TArray<FName>& GetSelectedChains() const {return Selection.SelectedChains; };
 	// END chain selection
 
 	void SetRootSelected(const bool bIsSelected);
-	bool GetRootSelected() const { return bIsRootSelected; };
+	bool GetRootSelected() const { return Selection.bIsRootSelected; };
 
 	void CleanSelection(ERetargetSourceOrTarget SourceOrTarget);
 	void ClearSelection(const bool bKeepBoneSelection=false);
-	ERetargetSelectionType GetLastSelectedItemType() const { return LastSelectedItem; };
+	ERetargetSelectionType GetLastSelectedItemType() const { return Selection.LastSelectedItem; };
 
 	// op stack selection
 	FIKRetargetOpBase* GetSelectedOp() const;
+	int32 GetSelectedOpIndex() const;
 	int32 LastSelectedOpIndex = 0;
 
 	// to frame selection when pressing "f" in viewport
@@ -354,11 +342,7 @@ private:
 	ERetargetSourceOrTarget CurrentlyEditingSourceOrTarget = ERetargetSourceOrTarget::Target;
 
 	// current selection set
-	bool bIsRootSelected = false;
-	UPrimitiveComponent* SelectedMesh = nullptr;
-	TArray<FName> SelectedChains;
-	TMap<ERetargetSourceOrTarget, TArray<FName>> SelectedBoneNames;
-	ERetargetSelectionType LastSelectedItem = ERetargetSelectionType::NONE;
+	FIKRetargetDebugDrawState Selection;
 	UPROPERTY()
 	TMap<FName,TObjectPtr<UIKRetargetBoneDetails>> AllBoneDetails;
 
diff --git a/Engine/Plugins/Animation/IKRig/Source/IKRigEditor/Public/RetargetEditor/IKRetargetHitProxies.h b/Engine/Plugins/Animation/IKRig/Source/IKRigEditor/Public/RetargetEditor/IKRetargetHitProxies.h
--- a/Engine/Plugins/Animation/IKRig/Source/IKRigEditor/Public/RetargetEditor/IKRetargetHitProxies.h
+++ b/Engine/Plugins/Animation/IKRig/Source/IKRigEditor/Public/RetargetEditor/IKRetargetHitProxies.h
@@ -29,47 +29,6 @@ struct HIKRetargetEditorBoneProxy : public HHitProxy
 		return EMouseCursor::Crosshairs;
 	}
 
-	virtual bool AlwaysAllowsTranslucentPrimitives() const override
-	{
-		return true;
-	}
-};
-
-// select chains/goals to edit chain settings
-struct HIKRetargetEditorChainProxy : public HHitProxy
-{
-	DECLARE_HIT_PROXY();
-
-	FName TargetChainName;
-	
-	HIKRetargetEditorChainProxy(const FName& InTargetChainName)
-		: HHitProxy(HPP_World)
-		, TargetChainName(InTargetChainName) {}
-
-	virtual EMouseCursor::Type GetMouseCursor()
-	{
-		return EMouseCursor::Crosshairs;
-	}
-
-	virtual bool AlwaysAllowsTranslucentPrimitives() const override
-	{
-		return true;
-	}
-};
-
-// select root control to edit root settings
-struct HIKRetargetEditorRootProxy : public HHitProxy
-{
-	DECLARE_HIT_PROXY();
-	
-	HIKRetargetEditorRootProxy()
-		: HHitProxy(HPP_World){}
-
-	virtual EMouseCursor::Type GetMouseCursor()
-	{
-		return EMouseCursor::Crosshairs;
-	}
-
 	virtual bool AlwaysAllowsTranslucentPrimitives() const override
 	{
 		return true;
diff --git a/Engine/Plugins/Animation/IKRig/Source/IKRigEditor/Public/RetargetEditor/IKRetargeterController.h b/Engine/Plugins/Animation/IKRig/Source/IKRigEditor/Public/RetargetEditor/IKRetargeterController.h
--- a/Engine/Plugins/Animation/IKRig/Source/IKRigEditor/Public/RetargetEditor/IKRetargeterController.h
+++ b/Engine/Plugins/Animation/IKRig/Source/IKRigEditor/Public/RetargetEditor/IKRetargeterController.h
@@ -21,7 +21,7 @@ class URetargetChainSettings;
 class UIKRigDefinition;
 class UIKRetargeter;
 class USkeletalMesh;
-class UStructViewer;
+class UIKRigStructViewer;
 
 UENUM()
 enum class EAutoMapChainType : uint8
@@ -62,96 +62,121 @@ private:
 public:
 	
 	/** Use this to get the controller for the given retargeter asset */
-	UFUNCTION(BlueprintCallable, Category=IKRetargeter, meta = (BlueprintThreadSafe, DisplayName = "Get IK Retarget Controller"))
+	UFUNCTION(BlueprintCallable, Category="IK Retarget", meta = (BlueprintThreadSafe, DisplayName = "Get IK Retarget Controller"))
 	static UIKRetargeterController* GetController(const UIKRetargeter* InRetargeterAsset);
 	
 	// Set the IK Rig to use as the source or target (to copy animation FROM/TO) 
-	UFUNCTION(BlueprintCallable, BlueprintPure=false, Category=IKRetargeter, meta = (BlueprintThreadSafe, DisplayName = "Set IK Rig"))
+	UFUNCTION(BlueprintCallable, BlueprintPure=false, Category="IK Retarget", meta = (BlueprintThreadSafe, DisplayName = "Set IK Rig"))
 	void SetIKRig(const ERetargetSourceOrTarget SourceOrTarget, UIKRigDefinition* IKRig) const;
 	
 	// Get either source or target IK Rig 
-	UFUNCTION(BlueprintCallable, Category=IKRetargeter, meta = (BlueprintThreadSafe, DisplayName = "Get IK Rig"))
+	UFUNCTION(BlueprintCallable, Category="IK Retarget", meta = (BlueprintThreadSafe, DisplayName = "Get IK Rig"))
 	const UIKRigDefinition* GetIKRig(const ERetargetSourceOrTarget SourceOrTarget) const;
 
 	// Set the preview skeletal mesh for either source or target
-	UFUNCTION(BlueprintCallable, BlueprintPure=false, Category=IKRetargeter, meta = (BlueprintThreadSafe))
+	UFUNCTION(BlueprintCallable, BlueprintPure=false, Category="IK Retarget", meta = (BlueprintThreadSafe))
 	void SetPreviewMesh(const ERetargetSourceOrTarget SourceOrTarget, USkeletalMesh* InPreviewMesh) const;
 	
 	// Get the preview skeletal mesh
-	UFUNCTION(BlueprintCallable, Category=IKRetargeter, meta = (BlueprintThreadSafe))
+	UFUNCTION(BlueprintCallable, Category="IK Retarget", meta = (BlueprintThreadSafe))
 	USkeletalMesh* GetPreviewMesh(const ERetargetSourceOrTarget SourceOrTarget) const;
 
 	//
 	// GET/SET SETTINGS PUBLIC/SCRIPTING API
 	//
 	
-	// Get a copy of the retarget root settings for this asset.
-	UFUNCTION(BlueprintCallable, BlueprintPure=false, Category=IKRetargeter, meta = (BlueprintThreadSafe))
-	FTargetRootSettings GetRootSettings() const;
+	// TODO
 	
-	// Set the retarget root settings for this asset.
-	UFUNCTION(BlueprintCallable, BlueprintPure=false, Category=IKRetargeter, meta = (BlueprintThreadSafe))
-	void SetRootSettings(const FTargetRootSettings& RootSettings) const;
-
-	// Get a copy of the global settings for this asset.
-	UFUNCTION(BlueprintCallable, BlueprintPure=false, Category=IKRetargeter, meta = (BlueprintThreadSafe))
-	FRetargetGlobalSettings GetGlobalSettings() const;
-
-	// Get a copy of the global settings for this asset.
-	UFUNCTION(BlueprintCallable, BlueprintPure=false, Category=IKRetargeter, meta = (BlueprintThreadSafe))
-	void SetGlobalSettings(const FRetargetGlobalSettings& GlobalSettings) const;
+	// Tell each op to reset any settings associated with this chain
+	UFUNCTION(BlueprintCallable, Category="IK Retarget")
+	void ResetChainSettingsInAllOps(const FName InTargetChainName) const;
 
-	// Get a copy of the settings for the target chain by name.
-	UFUNCTION(BlueprintCallable, BlueprintPure=false, Category=IKRetargeter, meta = (BlueprintThreadSafe))
-	FTargetChainSettings GetRetargetChainSettings(const FName& TargetChainName) const;
+	//
+	// GET/SET SETTINGS C++ API
+	//
 
-	// Set the settings for the target chain by name. Returns true if the chain settings were applied, false otherwise.
-	UFUNCTION(BlueprintCallable, BlueprintPure=false, Category=IKRetargeter, meta = (BlueprintThreadSafe))
-	bool SetRetargetChainSettings(const FName& TargetChainName, const FTargetChainSettings& Settings) const;
+	// get access to the mapping between source/target chains
+	const FRetargetChainMapping& GetChainMapping() const;
+	
+	// Ask if any op has settings for the given chain that are not at the default
+	bool ChainHasNonDefaultSettings(const FName InTargetChainName) const;
 
 	//
 	// RETARGET OPS PUBLIC/SCRIPTING API
 	//
 
-	// Add a new solver of the given type to the bottom of the stack. Returns the stack index.
-	// The type here is the full package path of the UStruct type. For example: /Script/IKRig.IKRetargetPinBoneOp
-	UFUNCTION(BlueprintCallable, BlueprintPure=false, Category=IKRig)
-	int32 AddRetargetOp(const FString InIKRetargetOpType) const;
+	/** Add a new retarget op of the given type to the bottom of the stack. Returns the stack index.
+	 * @param InIKRetargetOpType: the full package path of the UStruct type, ie /Script/IKRig.IKRetargetPinBoneOp
+	 * @param bAllowDuplicates: if false, will not add a new op if one of that type already exists (returns index of existing op instead)
+	 * @param bRunInitialSetup: if true, will run the Op's setup routine after adding it to the stack */
+	UFUNCTION(BlueprintCallable, BlueprintPure=false, Category="IK Retarget")
+	int32 AddRetargetOp(
+		const FString InIKRetargetOpType, 
+		const bool bAllowDuplicates = true,
+		const bool bRunInitialSetup = true) const;
 
 	// Remove the retarget op at the given stack index. 
-	UFUNCTION(BlueprintCallable, BlueprintPure=false, Category=RetargetOps)
+	UFUNCTION(BlueprintCallable, BlueprintPure=false, Category="IK Retarget")
 	bool RemoveRetargetOp(const int32 OpIndex) const;
 
 	// Remove all ops in the stack.
-	UFUNCTION(BlueprintCallable, Category=RetargetOps)
+	UFUNCTION(BlueprintCallable, Category="IK Retarget")
 	bool RemoveAllOps() const;
 
+	/** Automatically add basic retargeting operations
+	 * Adds these ops in the following order: Pelvis Motion, FK Chains, IK Chains, IK Solve and Root Motion
+	 * If any of these ops are already present, they will not be re-added.
+	 * @param bRunInitialSetup: if true, will run the initial setup routine on each op after they are all added to the stack. */
+	UFUNCTION(BlueprintCallable, Category="IK Retarget", meta = (BlueprintThreadSafe))
+	void AddDefaultOps(const bool bRunInitialSetup=true) const;
+
+	/** Force the op to run the initial setup. This is normally run when an Op is added through the editor UI.
+	 * NOTE: ops may not have custom initial setup routines; the exact behavior is op-dependent */
+	UFUNCTION(BlueprintCallable, Category="IK Retarget", meta = (BlueprintThreadSafe))
+	void RunOpInitialSetup(const int32 InOpIndex) const;
+
 	// Get the number of Ops in the stack.
-	UFUNCTION(BlueprintCallable, Category=RetargetOps)
+	UFUNCTION(BlueprintCallable, Category="IK Retarget")
 	int32 GetNumRetargetOps() const;
 
 	// Move the retarget op at the given index to the target index. 
-	UFUNCTION(BlueprintCallable, BlueprintPure=false, Category=RetargetOps)
+	UFUNCTION(BlueprintCallable, BlueprintPure=false, Category="IK Retarget")
 	bool MoveRetargetOpInStack(int32 OpToMoveIndex, int32 TargetIndex) const;
 
 	// Set enabled/disabled status of the given retarget operation. 
-	UFUNCTION(BlueprintCallable, BlueprintPure=false, Category=RetargetOps)
+	UFUNCTION(BlueprintCallable, BlueprintPure=false, Category="IK Retarget")
 	bool SetRetargetOpEnabled(int32 RetargetOpIndex, bool bIsEnabled) const;
 
 	// Get enabled status of the given Op. 
-	UFUNCTION(BlueprintCallable, Category=RetargetOps)
+	UFUNCTION(BlueprintCallable, Category="IK Retarget")
 	bool GetRetargetOpEnabled(int32 RetargetOpIndex) const;
+	
+	/** Get a controller to get/set the settings for a given op in the stack
+	 * NOTE: the returned UIKRetargetOpControllerBase* must be cast to the type specific to your op.
+	 * @param InOpIndex: the index of the op in the stack to return a controller for */
+	UFUNCTION(BlueprintCallable, Category="IK Rig")
+	UIKRetargetOpControllerBase* GetOpController(int32 InOpIndex);
 
 	//
 	// GENERAL C++ ONLY API
 	//
 	
 	// add a new retarget op to the stack
-	int32 AddRetargetOp(UScriptStruct* InIKRigSolverType) const;
+	int32 AddRetargetOp(
+		UScriptStruct* InRetargetOpType,
+		const bool bAllowDuplicates = true,
+		const bool bRunInitialSetup = true) const;
 	
 	// Get access to the given retarget operation.
 	FIKRetargetOpBase* GetRetargetOpAtIndex(int32 Index) const;
 
+	// Get access to the first retarget operation of the given type
+	template <typename T>
+	T* GetFirstRetargetOpOfType() const
+	{
+		return Asset->GetFirstRetargetOpOfType<T>();
+	}
+
 	// Get access to the given retarget operation. 
 	int32 GetIndexOfRetargetOp(FIKRetargetOpBase* RetargetOp) const;
 
@@ -180,28 +205,21 @@ public:
 	//
 	
 	// Use fuzzy string search to find "best" Source chain to map to each Target chain
-	UFUNCTION(BlueprintCallable, BlueprintPure=false, Category=IKRetargeter, meta = (BlueprintThreadSafe))
+	UFUNCTION(BlueprintCallable, BlueprintPure=false, Category="IK Retarget", meta = (BlueprintThreadSafe))
 	void AutoMapChains(const EAutoMapChainType AutoMapType, const bool bForceRemap) const;
 
 	// Assign a source chain to the given target chain. Animation will be copied from the source to the target.
-	UFUNCTION(BlueprintCallable, BlueprintPure=false, Category=IKRetargeter, meta = (BlueprintThreadSafe))
+	UFUNCTION(BlueprintCallable, BlueprintPure=false, Category="IK Retarget", meta = (BlueprintThreadSafe))
 	bool SetSourceChain(FName SourceChainName, FName TargetChainName) const;
 
 	// Get the name of the source chain mapped to a given target chain (the chain animation is copied FROM).
-	UFUNCTION(BlueprintCallable, Category=IKRetargeter, meta = (BlueprintThreadSafe))
+	UFUNCTION(BlueprintCallable, Category="IK Retarget", meta = (BlueprintThreadSafe))
 	FName GetSourceChain(const FName& TargetChainName) const;
-	
-	// Get read-only access to the list of settings for each target chain
-	UFUNCTION(BlueprintCallable, Category=IKRetargeter, meta = (BlueprintThreadSafe))
-	const TArray<URetargetChainSettings*>& GetAllChainSettings() const;
 
 	//
 	// RETARGET CHAIN MAPPING C++ ONLY API
 	//
 	
-	// Get the name of the goal assigned to a chain 
-	FName GetChainGoal(const TObjectPtr<URetargetChainSettings> ChainSettings) const;
-	
 	// Get whether the given chain's IK goal is connected to a solver 
 	bool IsChainGoalConnectedToASolver(const FName& GoalName) const;
 	
@@ -215,20 +233,10 @@ public:
 	void HandleRetargetChainRemoved(UIKRigDefinition* IKRig, const FName& InChainRemoved) const;
 	
 private:
-
-	// Remove invalid chain mappings (no longer existing in currently referenced source/target IK Rig assets)
-	// NOTE: be sure to reinitialize any running processors after chain mappings are cleaned!
-	void CleanChainMapping() const;
 	
 	// Get names of all the bone chains.
 	void GetChainNames(const ERetargetSourceOrTarget SourceOrTarget, TArray<FName>& OutNames) const;
 
-	// Sort the Asset ChainMapping based on the StartBone of the target chains. 
-	void SortChainMapping() const;
-
-	// convenience to get chain settings UObject by name
-	URetargetChainSettings* GetChainSettings(const FName& TargetChainName) const;
-
 public:
 
 	//
@@ -236,75 +244,75 @@ public:
 	//
 	
 	// Add new retarget pose. Returns the name of the new retarget pose.
-	UFUNCTION(BlueprintCallable, BlueprintPure=false, Category=IKRetargeter, meta =(BlueprintThreadSafe))
+	UFUNCTION(BlueprintCallable, BlueprintPure=false, Category="IK Retarget", meta =(BlueprintThreadSafe))
 	FName CreateRetargetPose(const FName& NewPoseName, const ERetargetSourceOrTarget SourceOrTarget) const;
 
 	// Remove a retarget pose. Returns true if the pose was found and removed.
-	UFUNCTION(BlueprintCallable, BlueprintPure=false, Category=IKRetargeter, meta =(BlueprintThreadSafe))
+	UFUNCTION(BlueprintCallable, BlueprintPure=false, Category="IK Retarget", meta =(BlueprintThreadSafe))
 	bool RemoveRetargetPose(const FName& PoseToRemove, const ERetargetSourceOrTarget SourceOrTarget) const;
 
 	// Duplicate a retarget pose. Returns the name of the new, duplicate pose (or NAME_None if PoseToDuplicate is not found).
-	UFUNCTION(BlueprintCallable, BlueprintPure=false, Category=IKRetargeter, meta =(BlueprintThreadSafe))
+	UFUNCTION(BlueprintCallable, BlueprintPure=false, Category="IK Retarget", meta =(BlueprintThreadSafe))
 	FName DuplicateRetargetPose(const FName PoseToDuplicate, const FName NewName, const ERetargetSourceOrTarget SourceOrTarget) const;
 	
 	// Rename current retarget pose. Returns true if a pose was renamed.
-	UFUNCTION(BlueprintCallable, BlueprintPure=false, Category=IKRetargeter, meta =(BlueprintThreadSafe))
+	UFUNCTION(BlueprintCallable, BlueprintPure=false, Category="IK Retarget", meta =(BlueprintThreadSafe))
 	bool RenameRetargetPose(const FName OldPoseName, const FName NewPoseName, const ERetargetSourceOrTarget SourceOrTarget) const;
 
 	// Change which retarget pose is used by the retargeter at runtime
-	UFUNCTION(BlueprintCallable, BlueprintPure=false, Category=IKRetargeter, meta =(BlueprintThreadSafe))
+	UFUNCTION(BlueprintCallable, BlueprintPure=false, Category="IK Retarget", meta =(BlueprintThreadSafe))
 	bool SetCurrentRetargetPose(FName CurrentPose, const ERetargetSourceOrTarget SourceOrTarget) const;
 	
 	// Get the current retarget pose
-	UFUNCTION(BlueprintCallable, Category=IKRetargeter, meta =(BlueprintThreadSafe))
+	UFUNCTION(BlueprintCallable, Category="IK Retarget", meta =(BlueprintThreadSafe))
     FName GetCurrentRetargetPoseName(const ERetargetSourceOrTarget SourceOrTarget) const;
 	
 	// Get access to array of retarget poses
-	UFUNCTION(BlueprintCallable, Category=IKRetargeter, meta =(BlueprintThreadSafe))
+	UFUNCTION(BlueprintCallable, Category="IK Retarget", meta =(BlueprintThreadSafe))
 	TMap<FName, FIKRetargetPose>& GetRetargetPoses(const ERetargetSourceOrTarget SourceOrTarget) const;
 	
 	// Get the current retarget pose
-	UFUNCTION(BlueprintCallable, Category=IKRetargeter, meta =(BlueprintThreadSafe))
+	UFUNCTION(BlueprintCallable, Category="IK Retarget", meta =(BlueprintThreadSafe))
 	FIKRetargetPose& GetCurrentRetargetPose(const ERetargetSourceOrTarget SourceOrTarget) const;
 
 	// Reset a retarget pose for the specified bones.
 	// If BonesToReset is Empty, will removes all stored deltas, returning pose to reference pose
-	UFUNCTION(BlueprintCallable, BlueprintPure=false, Category=IKRetargeter, meta =(BlueprintThreadSafe))
+	UFUNCTION(BlueprintCallable, BlueprintPure=false, Category="IK Retarget", meta =(BlueprintThreadSafe))
 	void ResetRetargetPose(const FName& PoseToReset, const TArray<FName>& BonesToReset, const ERetargetSourceOrTarget SourceOrTarget) const;
 	
 	// Set a delta rotation for a given bone for the current retarget pose
-	UFUNCTION(BlueprintCallable, BlueprintPure=false, Category=IKRetargeter, meta =(BlueprintThreadSafe))
+	UFUNCTION(BlueprintCallable, BlueprintPure=false, Category="IK Retarget", meta =(BlueprintThreadSafe))
 	void SetRotationOffsetForRetargetPoseBone(
 		const FName& BoneName,
 		const FQuat& RotationOffset,
 		const ERetargetSourceOrTarget SkeletonMode) const;
 	
 	// Get a delta rotation for a given bone for the current retarget pose
-	UFUNCTION(BlueprintCallable, BlueprintPure=false, Category=IKRetargeter, meta =(BlueprintThreadSafe))
+	UFUNCTION(BlueprintCallable, BlueprintPure=false, Category="IK Retarget", meta =(BlueprintThreadSafe))
 	FQuat GetRotationOffsetForRetargetPoseBone(
 		const FName& BoneName,
 		const ERetargetSourceOrTarget SourceOrTarget) const;
 	
 	// Set the translation offset on the retarget root bone for the current retarget pose
-	UFUNCTION(BlueprintCallable, BlueprintPure=false, Category=IKRetargeter, meta =(BlueprintThreadSafe))
+	UFUNCTION(BlueprintCallable, BlueprintPure=false, Category="IK Retarget", meta =(BlueprintThreadSafe))
 	void SetRootOffsetInRetargetPose(
 		const FVector& TranslationOffset,
 		const ERetargetSourceOrTarget SourceOrTarget) const;
 
 	// Get the translation offset on the retarget root bone for the current retarget pose
-	UFUNCTION(BlueprintCallable, Category=IKRetargeter, meta = (BlueprintThreadSafe))
+	UFUNCTION(BlueprintCallable, Category="IK Retarget", meta = (BlueprintThreadSafe))
 	FVector GetRootOffsetInRetargetPose(const ERetargetSourceOrTarget SourceOrTarget) const;
 
 	// Automatically align all bones in mapped chains and store in the current retarget pose.
-	UFUNCTION(BlueprintCallable, Category=RetargetOps)
+	UFUNCTION(BlueprintCallable, Category="IK Retarget")
 	void AutoAlignAllBones(ERetargetSourceOrTarget SourceOrTarget, const ERetargetAutoAlignMethod Method = ERetargetAutoAlignMethod::ChainToChain) const;
 
 	// Automatically align an array of bones and store in the current retarget pose. Bones not in a mapped chain will be ignored.
-	UFUNCTION(BlueprintCallable, Category=RetargetOps)
+	UFUNCTION(BlueprintCallable, Category="IK Retarget")
 	void AutoAlignBones(const TArray<FName>& BonesToAlign, const ERetargetAutoAlignMethod Method, ERetargetSourceOrTarget SourceOrTarget) const;
 
 	// Moves the entire skeleton vertically until the specified bone is the same height off the ground as in the reference pose.
-	UFUNCTION(BlueprintCallable, Category=RetargetOps)
+	UFUNCTION(BlueprintCallable, Category="IK Retarget")
 	void SnapBoneToGround(FName ReferenceBone, ERetargetSourceOrTarget SourceOrTarget);
 
 	//
@@ -320,8 +328,57 @@ private:
 	
 	// Remove bones from retarget poses that are no longer in skeleton 
 	void CleanPoseList(const ERetargetSourceOrTarget SourceOrTarget) const;
+	
+	//
+	// BEGIN DEPRECATED PUBLIC API
+	//
+	PRAGMA_DISABLE_DEPRECATION_WARNINGS
+public:
+
+	// Get a copy of the retarget root settings for this asset.
+	UE_DEPRECATED(5.6, "Root settings are now accessed through the Pelvis Motion Op controller.")
+	UFUNCTION(BlueprintCallable, BlueprintPure=false, Category="IK Retarget", meta = (BlueprintThreadSafe))
+	FTargetRootSettings GetRootSettings() const;
+	
+	// Set the retarget root settings for this asset.
+	UE_DEPRECATED(5.6, "Root settings are now accessed through the Pelvis Motion Op controller.")
+	UFUNCTION(BlueprintCallable, BlueprintPure=false, Category="IK Retarget", meta = (BlueprintThreadSafe))
+	void SetRootSettings(const FTargetRootSettings& RootSettings) const;
+
+	// Get a copy of the global settings for this asset.
+	UE_DEPRECATED(5.6, "Global settings are now accessed through Op controllers that perform the same duties.")
+	UFUNCTION(BlueprintCallable, BlueprintPure=false, Category="IK Retarget", meta = (BlueprintThreadSafe))
+	FRetargetGlobalSettings GetGlobalSettings() const;
 
+	// Get a copy of the global settings for this asset.
+	UE_DEPRECATED(5.6, "Global settings are now accessed through Op controllers that perform the same duties.")
+	UFUNCTION(BlueprintCallable, BlueprintPure=false, Category="IK Retarget", meta = (BlueprintThreadSafe))
+	void SetGlobalSettings(const FRetargetGlobalSettings& GlobalSettings) const;
 
+	// Get a copy of the settings for the target chain by name.
+	UE_DEPRECATED(5.6, "Access to chain settings must go through an Op controller now.")
+	UFUNCTION(BlueprintCallable, BlueprintPure=false, Category="IK Retarget", meta = (BlueprintThreadSafe))
+	FTargetChainSettings GetRetargetChainSettings(const FName& TargetChainName) const;
+
+	// Set the settings for the target chain by name. Returns true if the chain settings were applied, false otherwise.
+	UE_DEPRECATED(5.6, "Access to chain settings must go through an Op controller now.")
+	UFUNCTION(BlueprintCallable, BlueprintPure=false, Category="IK Retarget", meta = (BlueprintThreadSafe))
+	bool SetRetargetChainSettings(const FName& TargetChainName, const FTargetChainSettings& Settings) const;
+
+	// Get read-only access to the list of settings for each target chain
+	UE_DEPRECATED(5.6, "Access to chain settings must go through an Op controller now.")
+	UFUNCTION(BlueprintCallable, Category="IK Retarget", meta = (BlueprintThreadSafe))
+	const TArray<URetargetChainSettings*>& GetAllChainSettings() const;
+
+	// convenience to get chain settings UObject by name
+	UE_DEPRECATED(5.6, "Access to chain settings must go through an Op controller now.")
+	URetargetChainSettings* GetChainSettings(const FName& TargetChainName) const;
+
+	PRAGMA_ENABLE_DEPRECATION_WARNINGS
+	//
+	// END DEPRECATED PUBLIC API
+	//
+	
 	//
 	// DELEGATE CALLBACKS FOR OUTSIDE SYSTEMS
 	//
@@ -363,7 +420,7 @@ private:
 	
 	// a dummy UObject to wrap UStructs for the details panel
 	UPROPERTY(transient)
-	TObjectPtr<UStructViewer> StructViewer = nullptr;
+	TObjectPtr<UIKRigStructViewer> StructViewer = nullptr;
 	
 public:
 
@@ -377,7 +434,7 @@ public:
 	FOnRetargeterNeedsInitialized& OnRetargeterNeedsInitialized(){ return RetargeterNeedsInitialized; };
 
 	// get the singleton struct viewer for editing structs belonging to the asset this controller controls
-	UStructViewer* GetStructViewer();
+	UIKRigStructViewer* GetStructViewer();
 	
 	friend class UIKRetargeter;
 	friend struct FScopedReinitializeIKRetargeter;
diff --git a/Engine/Plugins/Animation/IKRig/Source/IKRigEditor/Public/RetargetEditor/SIKRetargetChainMapList.h b/Engine/Plugins/Animation/IKRig/Source/IKRigEditor/Public/RetargetEditor/SIKRetargetChainMapList.h
--- a/Engine/Plugins/Animation/IKRig/Source/IKRigEditor/Public/RetargetEditor/SIKRetargetChainMapList.h
+++ b/Engine/Plugins/Animation/IKRig/Source/IKRigEditor/Public/RetargetEditor/SIKRetargetChainMapList.h
@@ -27,17 +27,17 @@ public:
         TSharedRef<FRetargetChainMapElement> InStackElement,
         TSharedPtr<SIKRetargetChainMapList> InChainList);
 
-	static TSharedRef<FRetargetChainMapElement> Make(TObjectPtr<URetargetChainSettings> InChainMap)
+	static TSharedRef<FRetargetChainMapElement> Make(FName InTargetChainName)
 	{
-		return MakeShareable(new FRetargetChainMapElement(InChainMap));
+		return MakeShareable(new FRetargetChainMapElement(InTargetChainName));
 	}
 
-	TWeakObjectPtr<URetargetChainSettings> ChainMap;
+	FName TargetChainName;
 
 private:
 	
 	/** Hidden constructor, always use Make above */
-	FRetargetChainMapElement(TObjectPtr<URetargetChainSettings> InChainMap) : ChainMap(InChainMap) {}
+	FRetargetChainMapElement(FName InTargetChainName) : TargetChainName(InTargetChainName) {}
 
 	/** Hidden constructor, always use Make above */
 	FRetargetChainMapElement() = default;
@@ -99,7 +99,7 @@ public:
 
 	void ClearSelection() const;
 
-	void ResetChainSettings(URetargetChainSettings* Settings) const;
+	void ResetChainSettings(const FName InTargetChainName) const;
 	
 private:
 	
diff --git a/Engine/Plugins/Animation/IKRig/Source/IKRigEditor/Public/RetargetEditor/SRetargetOpStack.h b/Engine/Plugins/Animation/IKRig/Source/IKRigEditor/Public/RetargetEditor/SRetargetOpStack.h
--- a/Engine/Plugins/Animation/IKRig/Source/IKRigEditor/Public/RetargetEditor/SRetargetOpStack.h
+++ b/Engine/Plugins/Animation/IKRig/Source/IKRigEditor/Public/RetargetEditor/SRetargetOpStack.h
@@ -97,7 +97,7 @@ private:
 	// END op stack menu
 	
 	// menu command callback for adding a new op
-	void AddNewRetargetOp(UScriptStruct* ScriptStruct) const;
+	void AddNewRetargetOp(UScriptStruct* ScriptStruct);
 	// delete op from stack
 	void DeleteRetargetOp(TSharedPtr<FRetargetOpStackElement> OpToDelete);
 	// when an op is selected on in the stack view
diff --git a/Engine/Plugins/Animation/IKRig/Source/IKRigEditor/Public/RigEditor/IKRigDetailCustomizations.h b/Engine/Plugins/Animation/IKRig/Source/IKRigEditor/Public/RigEditor/IKRigDetailCustomizations.h
--- a/Engine/Plugins/Animation/IKRig/Source/IKRigEditor/Public/RigEditor/IKRigDetailCustomizations.h
+++ b/Engine/Plugins/Animation/IKRig/Source/IKRigEditor/Public/RigEditor/IKRigDetailCustomizations.h
@@ -13,7 +13,7 @@
 
 // a thin wrapper around UStruct data to display in a details panel
 UCLASS(Blueprintable)
-class UStructViewer : public UObject, public IBoneReferenceSkeletonProvider
+class UIKRigStructViewer : public UObject, public IBoneReferenceSkeletonProvider
 {
 	GENERATED_BODY()
 
@@ -60,7 +60,16 @@ public:
 
 	FName GetStructName() const
 	{
-		return FName(*StructType->GetName());
+		// try to get the "nice name" from metadata
+		FString TypeName = StructType->GetMetaData(TEXT("DisplayName"));
+
+		// if no "DisplayName" metadata is found, fall back to the struct name
+		if (TypeName.IsEmpty())
+		{
+			TypeName = StructType->GetName();
+		}
+		
+		return FName(TypeName);
 	}
 
 	UObject* GetStructOwner() const
@@ -68,11 +77,16 @@ public:
 		return StructOwner.Get();
 	};
 
+	void TriggerReinitIfNeeded(const FPropertyChangedEvent& InEvent);
+
 	// IBoneReferenceSkeletonProvider overrides.
 	virtual USkeleton* GetSkeleton(bool& bInvalidSkeletonIsError, const IPropertyHandle* PropertyHandle) override;
 	// ~END IBoneReferenceSkeletonProvider overrides.
 
 private:
+
+	DECLARE_MULTICAST_DELEGATE(FOnStructPropertyEdited);
+	FOnStructPropertyEdited StructPropertyEditedDelegate;
 	
 	// the type of struct we are editing
 	const UScriptStruct* StructType;
@@ -84,6 +98,11 @@ private:
 
 	// a UObject that owns the struct for Undo/Redo transactions (ie, the asset we are editing)
 	TWeakObjectPtr<UObject> StructOwner = nullptr;
+
+public:
+	
+	// Attach a delegate to be notified whenever a property (or child property) in the currently displayed UStruct is edited
+	FOnStructPropertyEdited& OnStructPropertyEdited(){ return StructPropertyEditedDelegate; };
 };
 
 class FStructViewerCustomization : public IDetailCustomization
diff --git a/Engine/Plugins/Animation/IKRig/Source/IKRigEditor/Public/RigEditor/IKRigEditorController.h b/Engine/Plugins/Animation/IKRig/Source/IKRigEditor/Public/RigEditor/IKRigEditorController.h
--- a/Engine/Plugins/Animation/IKRig/Source/IKRigEditor/Public/RigEditor/IKRigEditorController.h
+++ b/Engine/Plugins/Animation/IKRig/Source/IKRigEditor/Public/RigEditor/IKRigEditorController.h
@@ -21,7 +21,7 @@ class FSolverStackElement;
 class FIKRigTreeElement;
 class UAnimInstance;
 class UDebugSkelMeshComponent;
-class UStructViewer;
+class UIKRigStructViewer;
 struct FIKRigProcessor;
 
 enum EIKRigSelectionType : int8
@@ -184,7 +184,7 @@ public:
 	// show selected items in details view 
 	void ShowDetailsForElements(const TArray<TSharedPtr<FIKRigTreeElement>>& InItems) const;
 	// show the currently configured StructViewer object in the details tab
-	void ShowStructViewerInDetailsView(UStructViewer* StructViewer) const;
+	void ShowStructViewerInDetailsView(UIKRigStructViewer* StructViewer) const;
 	// callback when detail is edited 
 	void OnFinishedChangingDetails(const FPropertyChangedEvent& PropertyChangedEvent) const;
 	// returns true if the supplied UObject is being shown in the details panel
