[IKRig] Missing copyright boilerplate in IKRetargetDeprecated.cpp

diff --git a/Engine/Plugins/Animation/IKRig/Source/IKRig/Private/Retargeter/IKRetargetDeprecated.cpp b/Engine/Plugins/Animation/IKRig/Source/IKRig/Private/Retargeter/IKRetargetDeprecated.cpp
--- a/Engine/Plugins/Animation/IKRig/Source/IKRig/Private/Retargeter/IKRetargetDeprecated.cpp
+++ b/Engine/Plugins/Animation/IKRig/Source/IKRig/Private/Retargeter/IKRetargetDeprecated.cpp
@@ -1,4 +1,6 @@
-﻿#include "Retargeter/IKRetargetDeprecated.h"
+﻿// Copyright Epic Games, Inc. All Rights Reserved.
+
+#include "Retargeter/IKRetargetDeprecated.h"
 
 #include "IKRigObjectVersion.h"
 
