[IKRig] Fix minor issue caught by static analysis in RunIKRigOp.cpp

diff --git a/Engine/Plugins/Animation/IKRig/Source/IKRig/Private/Retargeter/RetargetOps/RunIKRigOp.cpp b/Engine/Plugins/Animation/IKRig/Source/IKRig/Private/Retargeter/RetargetOps/RunIKRigOp.cpp
--- a/Engine/Plugins/Animation/IKRig/Source/IKRig/Private/Retargeter/RetargetOps/RunIKRigOp.cpp
+++ b/Engine/Plugins/Animation/IKRig/Source/IKRig/Private/Retargeter/RetargetOps/RunIKRigOp.cpp
@@ -142,7 +142,7 @@ void FIKRetargetRunIKRigOp::DebugDraw(
 	// spin through all the IK goals debug data
 	for (const FRunIKRigOpGoalDebugData& GoalData : GoalDebugData)
 	{
-		bool bIsSelected = bIsSelected = InEditorState.SelectedGoals.Contains(GoalData.GoalName);
+		bool bIsSelected = InEditorState.SelectedGoals.Contains(GoalData.GoalName);
 
 		FTransform Initial = GoalData.InitialGoal * InComponentTransform;
 		FTransform Current = GoalData.CurrentGoal * InComponentTransform;
