[IKRIg] Fixing Asan issue from storying IK Rig solvers as TSharedPtr (converted to InstancedStrucT).

diff --git a/Engine/Plugins/Animation/IKRig/Source/IKRig/Private/Rig/IKRigDefinition.cpp b/Engine/Plugins/Animation/IKRig/Source/IKRig/Private/Rig/IKRigDefinition.cpp
--- a/Engine/Plugins/Animation/IKRig/Source/IKRig/Private/Rig/IKRigDefinition.cpp
+++ b/Engine/Plugins/Animation/IKRig/Source/IKRig/Private/Rig/IKRigDefinition.cpp
@@ -360,23 +360,6 @@ void UIKRigDefinition::PostLoad()
 	});
 }
 
-const TArray<const FIKRigSolverBase*> UIKRigDefinition::GetSolverArray() const
-{
-	TArray<const FIKRigSolverBase*> Result;
-	for (const FInstancedStruct& SolverStruct : SolverStack)
-	{
-		if (!ensureMsgf(SolverStruct.IsValid(), TEXT("IK Rig asset contains unknown solver type. Possibly due to an unloaded plugin?")))
-		{
-			// this can happen if asset references deleted or unknown IK Solver type
-			// which should only happen during development (if at all)
-			continue;
-		}
-		const FIKRigSolverBase* Solver = SolverStruct.GetPtr<FIKRigSolverBase>();
-		Result.Add(Solver);
-	}
-	return MoveTemp(Result);
-}
-
 const TArray<FInstancedStruct>& UIKRigDefinition::GetSolverStructs() const
 {
 	return SolverStack;
diff --git a/Engine/Plugins/Animation/IKRig/Source/IKRig/Private/Rig/IKRigProcessor.cpp b/Engine/Plugins/Animation/IKRig/Source/IKRig/Private/Rig/IKRigProcessor.cpp
--- a/Engine/Plugins/Animation/IKRig/Source/IKRig/Private/Rig/IKRigProcessor.cpp
+++ b/Engine/Plugins/Animation/IKRig/Source/IKRig/Private/Rig/IKRigProcessor.cpp
@@ -67,11 +67,12 @@ void FIKRigProcessor::Initialize(
 	}
 
 	// get the set of ALL required goals so we can validate that each goal is used by at least 1 solver
-	const TArray<const FIKRigSolverBase*> AssetSolvers = InRigAsset->GetSolverArray();
+	const TArray<FInstancedStruct>& AssetSolverStructs = InRigAsset->GetSolverStructs();
 	TSet<FName> AllRequiredGoals;
-	for (const FIKRigSolverBase* AssetSolver : AssetSolvers)
+	for (const FInstancedStruct& AssetSolverStruct : AssetSolverStructs)
 	{
-		AssetSolver->GetRequiredGoals(AllRequiredGoals);
+		const FIKRigSolverBase& AssetSolver = AssetSolverStruct.Get<FIKRigSolverBase>();
+		AssetSolver.GetRequiredGoals(AllRequiredGoals);
 	}
 
 	// generate list of excluded goals
@@ -159,7 +160,6 @@ void FIKRigProcessor::Initialize(
 	}
 
 	// create copies of all the solvers in the IK rig
-	const TArray<FInstancedStruct>& AssetSolverStructs = InRigAsset->GetSolverStructs();
 	Solvers.Reset(AssetSolverStructs.Num());
 	for (const FInstancedStruct& AssetSolverStruct : AssetSolverStructs)
 	{
@@ -172,13 +172,9 @@ void FIKRigProcessor::Initialize(
 			continue;
 		}
 
-		// copy the solver memory
-		const UScriptStruct* ScriptStruct = AssetSolverStruct.GetScriptStruct();
-		FIKRigSolverBase* NewSolver = static_cast<FIKRigSolverBase*>(FMemory::Malloc(ScriptStruct->GetStructureSize()));
-		ScriptStruct->InitializeStruct(NewSolver);
-		ScriptStruct->CopyScriptStruct(NewSolver, AssetSolverStruct.GetMemory());
-		// store it in the runtime stack
-		Solvers.Emplace(MakeShareable(NewSolver));
+		// copy the solver memory into the processor
+		const int32 NewSolverIndex = Solvers.Emplace(AssetSolverStruct);
+		FIKRigSolverBase* NewSolver = Solvers[NewSolverIndex].GetMutablePtr<FIKRigSolverBase>();
 
 		// remove excluded goals from the solver
 		// this must be done BEFORE initializing the solver
@@ -221,10 +217,11 @@ bool FIKRigProcessor::IsIKRigCompatibleWithSkeleton(
 	// first we validate that all the required bones are in the input skeleton...
 	
 	TSet<FName> RequiredBones;
-	const TArray<const FIKRigSolverBase*> AssetSolvers = InRigAsset->GetSolverArray();
-	for (const FIKRigSolverBase* AssetSolver : AssetSolvers)
+	const TArray<FInstancedStruct>& AssetSolverStructs = InRigAsset->GetSolverStructs();
+	for (const FInstancedStruct& AssetSolverStruct : AssetSolverStructs)
 	{
-		AssetSolver->GetRequiredBones(RequiredBones);
+		const FIKRigSolverBase& AssetSolver = AssetSolverStruct.Get<FIKRigSolverBase>();
+		AssetSolver.GetRequiredBones(RequiredBones);
 	}
 
 	const TArray<UIKRigEffectorGoal*>& Goals = InRigAsset->GetGoalArray();
@@ -372,8 +369,9 @@ void FIKRigProcessor::Solve(const FTransform& ComponentToWorld)
 	ResolveFinalGoalTransforms(ComponentToWorld);
 
 	// run all the solvers
-	for (const TSharedPtr<FIKRigSolverBase>& Solver : Solvers)
+	for (FInstancedStruct& SolverStruct : Solvers)
 	{
+		FIKRigSolverBase* Solver = SolverStruct.GetMutablePtr<FIKRigSolverBase>();
 		if (Solver->IsEnabled())
 		{
 			Solver->Solve(Skeleton, GoalContainer);
@@ -403,7 +401,17 @@ void FIKRigProcessor::SetNeedsInitialized()
 {
 	bInitialized = false;
 	bTriedToInitialize = false;
-};
+}
+
+const FIKRigSolverBase* FIKRigProcessor::GetSolver(const int32 InSolverIndex) const
+{
+	if (Solvers.IsValidIndex(InSolverIndex))
+	{
+		return Solvers[InSolverIndex].GetPtr<FIKRigSolverBase>();
+	}
+	
+	return nullptr;
+}
 
 void FIKRigProcessor::CopyAllSettingsFromAsset(const UIKRigDefinition* SourceAsset)
 {
@@ -430,13 +438,13 @@ void FIKRigProcessor::CopyAllSettingsFromAsset(const UIKRigDefinition* SourceAss
 	}
 
 	// copy solver settings
-	const TArray<const FIKRigSolverBase*> AssetSolvers = SourceAsset->GetSolverArray();
-	check(Solvers.Num() == AssetSolvers.Num()); // if number of solvers has been changed, processor should have been reinitialized
+	const TArray<FInstancedStruct>& AssetSolverStructs = SourceAsset->GetSolverStructs();
+	check(Solvers.Num() == AssetSolverStructs.Num()); // if number of solvers has been changed, processor should have been reinitialized
 	for (int32 SolverIndex=0; SolverIndex<Solvers.Num(); ++SolverIndex)
 	{
-		const FIKRigSolverBase& AssetSolver = *AssetSolvers[SolverIndex];
-		FIKRigSolverBase* RunningSolver = Solvers[SolverIndex].Get();
-		RunningSolver->UpdateSettingsFromAsset(AssetSolver);
+		const FIKRigSolverBase& AssetSolver = AssetSolverStructs[SolverIndex].Get<FIKRigSolverBase>();
+		FIKRigSolverBase& RunningSolver = Solvers[SolverIndex].GetMutable<FIKRigSolverBase>();
+		RunningSolver.UpdateSettingsFromAsset(AssetSolver);
 	}
 }
 
diff --git a/Engine/Plugins/Animation/IKRig/Source/IKRig/Public/Rig/IKRigDefinition.h b/Engine/Plugins/Animation/IKRig/Source/IKRig/Public/Rig/IKRigDefinition.h
--- a/Engine/Plugins/Animation/IKRig/Source/IKRig/Public/Rig/IKRigDefinition.h
+++ b/Engine/Plugins/Animation/IKRig/Source/IKRig/Public/Rig/IKRigDefinition.h
@@ -201,9 +201,6 @@ public:
 	/** runtime, read-only access to array of Goals, all modifications must go through UIKRigController */
 	const TArray<UIKRigEffectorGoal*>& GetGoalArray() const { return Goals; };
 
-	/** runtime, read-only access to array of Solvers, all modifications must go through UIKRigController */
-	const TArray<const FIKRigSolverBase*> GetSolverArray() const;
-
 	/** runtime, read-only access to array of Solvers structs, all modifications must go through UIKRigController */
 	const TArray<FInstancedStruct>& GetSolverStructs() const;
 
diff --git a/Engine/Plugins/Animation/IKRig/Source/IKRig/Public/Rig/IKRigProcessor.h b/Engine/Plugins/Animation/IKRig/Source/IKRig/Public/Rig/IKRigProcessor.h
--- a/Engine/Plugins/Animation/IKRig/Source/IKRig/Public/Rig/IKRigProcessor.h
+++ b/Engine/Plugins/Animation/IKRig/Source/IKRig/Public/Rig/IKRigProcessor.h
@@ -105,8 +105,8 @@ public:
 	/** force the processor to reinitialize */
 	void SetNeedsInitialized();
 
-	/** Get read-only access to currently running solvers */
-	const TArray<TSharedPtr<FIKRigSolverBase>>& GetSolvers(){ return Solvers; };
+	/** Get read-only access to a currently running solver */
+	const FIKRigSolverBase* GetSolver(const int32 InSolverIndex) const;
 	
 	/** Used to propagate setting values from the source asset at runtime (settings that do not require re-initialization)
 	 * This allows the user to edit their IK Rig asset and see the effects in PIE. */
@@ -121,7 +121,7 @@ private:
 	void ResolveFinalGoalTransforms(const FTransform& WorldToComponent);
 
 	/** the stack of solvers to run in order */
-	TArray<TSharedPtr<FIKRigSolverBase>> Solvers;
+	TArray<FInstancedStruct> Solvers;
 
 	/** the named transforms that solvers use as end effectors */
 	FIKRigGoalContainer GoalContainer;
@@ -214,10 +214,6 @@ IKRIG_API UIKRigProcessor : public UObject
 	{
 		Processor.SetNeedsInitialized();
 	};
-	const TArray<TSharedPtr<FIKRigSolverBase>>& GetSolvers()
-	{
-		return Processor.GetSolvers();
-	};
 	void CopyAllSettingsFromAsset(const UIKRigDefinition* SourceAsset)
 	{
 		Processor.CopyAllSettingsFromAsset(SourceAsset);
diff --git a/Engine/Plugins/Animation/IKRig/Source/IKRigEditor/Private/RigEditor/SIKRigSolverStack.cpp b/Engine/Plugins/Animation/IKRig/Source/IKRigEditor/Private/RigEditor/SIKRigSolverStack.cpp
--- a/Engine/Plugins/Animation/IKRig/Source/IKRigEditor/Private/RigEditor/SIKRigSolverStack.cpp
+++ b/Engine/Plugins/Animation/IKRig/Source/IKRigEditor/Private/RigEditor/SIKRigSolverStack.cpp
@@ -204,10 +204,7 @@ const FIKRigSolverBase* SIKRigSolverStackItem::GetSolver(const bool bFromAsset)
 		// this is needed to report solver warnings dependent on being initialized
 		if (FIKRigProcessor* Processor = EditorController->GetIKRigProcessor())
 		{
-			if (Processor->GetSolvers().IsValidIndex(SolverIndex))
-			{
-				return Processor->GetSolvers()[SolverIndex].Get();
-			}
+			return Processor->GetSolver(SolverIndex);
 		}
 	}
 
