[IKRig] Fixed issue moving goals does not affect pose.

diff --git a/Engine/Plugins/Animation/IKRig/Source/IKRig/Private/Rig/IKRigProcessor.cpp b/Engine/Plugins/Animation/IKRig/Source/IKRig/Private/Rig/IKRigProcessor.cpp
--- a/Engine/Plugins/Animation/IKRig/Source/IKRig/Private/Rig/IKRigProcessor.cpp
+++ b/Engine/Plugins/Animation/IKRig/Source/IKRig/Private/Rig/IKRigProcessor.cpp
@@ -428,6 +428,8 @@ void FIKRigProcessor::CopyAllSettingsFromAsset(const UIKRigDefinition* SourceAss
 	{
 		if (FIKRigGoal* Goal = GoalContainer.FindGoalByName(AssetGoal->GoalName))
 		{
+			Goal->Position = AssetGoal->CurrentTransform.GetTranslation();
+			Goal->Rotation = AssetGoal->CurrentTransform.GetRotation().Rotator();
 			Goal->RotationAlpha = AssetGoal->RotationAlpha;
 			Goal->PositionAlpha = AssetGoal->PositionAlpha;
 		}
diff --git a/Engine/Plugins/Animation/IKRig/Source/IKRigEditor/Private/RigEditor/IKRigController.cpp b/Engine/Plugins/Animation/IKRig/Source/IKRigEditor/Private/RigEditor/IKRigController.cpp
--- a/Engine/Plugins/Animation/IKRig/Source/IKRigEditor/Private/RigEditor/IKRigController.cpp
+++ b/Engine/Plugins/Animation/IKRig/Source/IKRigEditor/Private/RigEditor/IKRigController.cpp
@@ -129,25 +129,21 @@ bool UIKRigController::CanRemoveBoneSetting(const FName BoneName, int32 SolverIn
 	const FIKRigSolverBase* Solver = GetSolverAtIndex(SolverIndex);
 	if (!Solver)
 	{
-		UE_LOG(LogIKRigEditor, Warning, TEXT("Solver does not exist at index: %d."), SolverIndex);
 		return false; // solver doesn't exist
 	}
 
 	if (!Solver->UsesCustomBoneSettings())
 	{
-		UE_LOG(LogIKRigEditor, Warning, TEXT("Solver does not support bone settings."));
 		return false; // solver doesn't use bone settings
 	}
 
 	if (Asset->Skeleton.GetBoneIndexFromName(BoneName) == INDEX_NONE)
 	{
-		UE_LOG(LogIKRigEditor, Warning, TEXT("Bone does not exist, %s."), *BoneName.ToString());
 		return false; // bone doesn't exist
 	}
 
 	if (!Solver->HasSettingsOnBone(BoneName))
 	{
-		UE_LOG(LogIKRigEditor, Warning, TEXT("Solver does not have settings on bone, %s."), *BoneName.ToString());
 		return false; // solver doesn't have any settings for this bone
 	}
 
