[IK Retarget] Added "Propagate to Children" option to Pin Bone op.

diff --git a/Engine/Plugins/Animation/IKRig/Source/IKRig/Private/Retargeter/IKRetargetProcessor.cpp b/Engine/Plugins/Animation/IKRig/Source/IKRig/Private/Retargeter/IKRetargetProcessor.cpp
--- a/Engine/Plugins/Animation/IKRig/Source/IKRig/Private/Retargeter/IKRetargetProcessor.cpp
+++ b/Engine/Plugins/Animation/IKRig/Source/IKRig/Private/Retargeter/IKRetargetProcessor.cpp
@@ -173,6 +173,23 @@ void FRetargetSkeleton::UpdateGlobalTransformsBelowBone(
 	}
 }
 
+void FRetargetSkeleton::SetGlobalTransformAndUpdateChildren(const int32 InBoneToSetIndex, const FTransform& InNewTransform, TArray<FTransform>& InOutGlobalPose) const
+{
+	check(BoneNames.Num() == InOutGlobalPose.Num());
+	check(InOutGlobalPose.IsValidIndex(InBoneToSetIndex));
+
+	const FTransform PrevTransform = InOutGlobalPose[InBoneToSetIndex];
+	InOutGlobalPose[InBoneToSetIndex] = InNewTransform;
+
+	TArray<int32> ChildBoneIndices;
+	GetChildrenIndicesRecursive(InBoneToSetIndex, ChildBoneIndices);
+	for (const int32 ChildBoneIndex : ChildBoneIndices)
+	{
+		const FTransform RelativeToPrev = InOutGlobalPose[ChildBoneIndex].GetRelativeTransform(PrevTransform);
+		InOutGlobalPose[ChildBoneIndex] = RelativeToPrev * InNewTransform;
+	}
+}
+
 void FRetargetSkeleton::UpdateLocalTransformsBelowBone(
 	const int32 StartBoneIndex,
 	TArray<FTransform>& OutLocalPose,
diff --git a/Engine/Plugins/Animation/IKRig/Source/IKRig/Private/Retargeter/RetargetOps/PinBoneOp.cpp b/Engine/Plugins/Animation/IKRig/Source/IKRig/Private/Retargeter/RetargetOps/PinBoneOp.cpp
--- a/Engine/Plugins/Animation/IKRig/Source/IKRig/Private/Retargeter/RetargetOps/PinBoneOp.cpp
+++ b/Engine/Plugins/Animation/IKRig/Source/IKRig/Private/Retargeter/RetargetOps/PinBoneOp.cpp
@@ -160,21 +160,34 @@ void FIKRetargetPinBoneOp::Run(
 		// calculate new transform for bone to pin
 		const FTransform NewTransform = GetNewBoneTransform(PinData, Processor, InSourceGlobalPose, OutTargetGlobalPose);
 		// apply static local and global offsets
-		const FTransform Result = Settings.LocalOffset * (NewTransform * Settings.GlobalOffset);
+		FTransform Result = Settings.LocalOffset * (NewTransform * Settings.GlobalOffset);
 
-		// store result in output pose
-		FTransform& BoneToCopyToTransform = OutTargetGlobalPose[PinData.BoneToCopyTo.BoneIndex];
-		if (Settings.bCopyTranslation)
+		// filter channels
+		const FTransform& BoneToCopyToTransform = OutTargetGlobalPose[PinData.BoneToCopyTo.BoneIndex];
+		if (!Settings.bCopyTranslation)
 		{
-			BoneToCopyToTransform.SetTranslation(Result.GetLocation());
+			Result.SetTranslation(BoneToCopyToTransform.GetLocation());
 		}
-		if (Settings.bCopyRotation)
+		if (!Settings.bCopyRotation)
 		{
-			BoneToCopyToTransform.SetRotation(Result.GetRotation());
+			Result.SetRotation(BoneToCopyToTransform.GetRotation());
 		}
-		if (Settings.bCopyScale)
+		if (!Settings.bCopyScale)
 		{
-			BoneToCopyToTransform.SetScale3D(Result.GetScale3D());
+			Result.SetScale3D(BoneToCopyToTransform.GetScale3D());
+		}
+
+		// apply to pose
+		if (Settings.bPropagateToChildren)
+		{
+			// assign result and update children
+			const FRetargetSkeleton& TargetSkeleton = Processor.GetSkeleton(ERetargetSourceOrTarget::Target);
+			TargetSkeleton.SetGlobalTransformAndUpdateChildren(PinData.BoneToCopyTo.BoneIndex, Result,OutTargetGlobalPose);
+		}
+		else
+		{
+			// assign result directly
+			OutTargetGlobalPose[PinData.BoneToCopyTo.BoneIndex] = Result;
 		}
 	}
 }
diff --git a/Engine/Plugins/Animation/IKRig/Source/IKRig/Public/Retargeter/IKRetargetProcessor.h b/Engine/Plugins/Animation/IKRig/Source/IKRig/Public/Retargeter/IKRetargetProcessor.h
--- a/Engine/Plugins/Animation/IKRig/Source/IKRig/Public/Retargeter/IKRetargetProcessor.h
+++ b/Engine/Plugins/Animation/IKRig/Source/IKRig/Public/Retargeter/IKRetargetProcessor.h
@@ -77,6 +77,11 @@ struct IKRIG_API FRetargetSkeleton
 		const TArray<FTransform>& InLocalPose,
 		TArray<FTransform>& OutGlobalPose) const;
 
+	void SetGlobalTransformAndUpdateChildren(
+		const int32 InBoneToSetIndex,
+		const FTransform& NewTransform,
+		TArray<FTransform>& InOutGlobalPose) const;
+
 	void UpdateLocalTransformsBelowBone(
 		const int32 StartBoneIndex,
 		TArray<FTransform>& OutLocalPose,
diff --git a/Engine/Plugins/Animation/IKRig/Source/IKRig/Public/Retargeter/RetargetOps/PinBoneOp.h b/Engine/Plugins/Animation/IKRig/Source/IKRig/Public/Retargeter/RetargetOps/PinBoneOp.h
--- a/Engine/Plugins/Animation/IKRig/Source/IKRig/Public/Retargeter/RetargetOps/PinBoneOp.h
+++ b/Engine/Plugins/Animation/IKRig/Source/IKRig/Public/Retargeter/RetargetOps/PinBoneOp.h
@@ -110,6 +110,10 @@ struct IKRIG_API FIKRetargetPinBoneOpSettings : public FIKRetargetOpSettingsBase
 	UPROPERTY(EditAnywhere, Category=Settings)
 	bool bCopyScale = true;
 
+	/** Update children bone transforms (recursively). */
+	UPROPERTY(EditAnywhere, Category=Settings)
+	bool bPropagateToChildren = false;
+
 	/** The method used to calculate the translation of the bone. */
 	UPROPERTY(EditAnywhere, Category=Settings)
 	EPinBoneTranslationMode TranslationMode = EPinBoneTranslationMode::CopyGlobalPosition;
