[IK Retarget] Fixed uninitialized property in FRetargetGlobalSettings

diff --git a/Engine/Plugins/Animation/IKRig/Source/IKRig/Public/Retargeter/IKRetargetDeprecated.h b/Engine/Plugins/Animation/IKRig/Source/IKRig/Public/Retargeter/IKRetargetDeprecated.h
--- a/Engine/Plugins/Animation/IKRig/Source/IKRig/Public/Retargeter/IKRetargetDeprecated.h
+++ b/Engine/Plugins/Animation/IKRig/Source/IKRig/Public/Retargeter/IKRetargetDeprecated.h
@@ -2,6 +2,7 @@
 
 #pragma once
 #include "IKRetargetSettings.h"
+#include "RetargetOps/StrideWarpingOp.h"
 
 #include "IKRetargetDeprecated.generated.h"
 
@@ -166,7 +167,7 @@ public:
 	UPROPERTY(EditAnywhere, BlueprintReadWrite, Category=DeprecatedSettings, meta=(DeprecatedProperty))
 	bool bWarping = false;
 	UPROPERTY(EditAnywhere, BlueprintReadWrite, Category=DeprecatedSettings, meta=(DeprecatedProperty))
-	EWarpingDirectionSource DirectionSource;
+	EWarpingDirectionSource DirectionSource = EWarpingDirectionSource::RootBone;
 	UPROPERTY(EditAnywhere, BlueprintReadWrite, Category=DeprecatedSettings, meta=(DeprecatedProperty))
 	EBasicAxis ForwardDirection = EBasicAxis::Y;
 	UPROPERTY(EditAnywhere, BlueprintReadWrite, Category=DeprecatedSettings, meta=(DeprecatedProperty))
