[IK Retargeter] Adds new UI to retarget op stack.

diff --git a/Engine/Plugins/Animation/IKRig/Source/IKRig/Private/AnimNodes/AnimNode_RetargetPoseFromMesh.cpp b/Engine/Plugins/Animation/IKRig/Source/IKRig/Private/AnimNodes/AnimNode_RetargetPoseFromMesh.cpp
--- a/Engine/Plugins/Animation/IKRig/Source/IKRig/Private/AnimNodes/AnimNode_RetargetPoseFromMesh.cpp
+++ b/Engine/Plugins/Animation/IKRig/Source/IKRig/Private/AnimNodes/AnimNode_RetargetPoseFromMesh.cpp
@@ -292,8 +292,11 @@ void FAnimNode_RetargetPoseFromMesh::PreUpdate(const UAnimInstance* InAnimInstan
 	if (Processor.IsInitialized() && RetargetFrom != ERetargetSourceMode::SourcePosePin)
 	{
 		GetPoseToRetargetFromSourceMesh(TargetMeshComponent);
-		
-		Processor.OnAnimGraphPreUpdateMainThread(*SourceMeshComponent.Get(), *TargetMeshComponent);
+
+		if (SourceMeshComponent.IsValid() && TargetMeshComponent)
+		{
+			Processor.OnAnimGraphPreUpdateMainThread(*SourceMeshComponent.Get(), *TargetMeshComponent);
+		}
 	}
 }
 
diff --git a/Engine/Plugins/Animation/IKRig/Source/IKRig/Private/Retargeter/IKRetargetOps.cpp b/Engine/Plugins/Animation/IKRig/Source/IKRig/Private/Retargeter/IKRetargetOps.cpp
--- a/Engine/Plugins/Animation/IKRig/Source/IKRig/Private/Retargeter/IKRetargetOps.cpp
+++ b/Engine/Plugins/Animation/IKRig/Source/IKRig/Private/Retargeter/IKRetargetOps.cpp
@@ -12,17 +12,17 @@ void FIKRetargetOpBase::SetSettings(const FIKRetargetOpSettingsBase* InSettings)
 }
 
 #if WITH_EDITOR
-FText FIKRetargetOpBase::GetNiceName() const
+FName FIKRetargetOpBase::GetDefaultName() const
 {
 	// calls virtual GetType() to get the derived type
 	const UScriptStruct* ScriptStruct = GetType();
 
 	// try to get the "nice name" from metadata
 	FString TypeName;
-	if (ScriptStruct->HasMetaData(TEXT("DisplayName")))
-	{
-		TypeName = ScriptStruct->GetMetaData(TEXT("DisplayName"));
-	}
+    if (ScriptStruct->HasMetaData(TEXT("DisplayName")))
+    {
+    	TypeName = ScriptStruct->GetMetaData(TEXT("DisplayName"));
+    }
 
 	// if no "DisplayName" metadata is found, fall back to the struct name
 	if (TypeName.IsEmpty())
@@ -30,7 +30,7 @@ FText FIKRetargetOpBase::GetNiceName() const
 		TypeName = ScriptStruct->GetName();
 	}
 	
-	return FText::FromString(TypeName);
+	return FName(TypeName);
 }
 
 FText FIKRetargetOpBase::GetWarningMessage() const
@@ -54,6 +54,27 @@ const UClass* FIKRetargetOpBase::GetControllerType() const
 	return UIKRetargetOpControllerBase::StaticClass();
 }
 
+FName FIKRetargetOpBase::GetName() const
+{
+	return Name;
+}
+
+void FIKRetargetOpBase::SetName(const FName InName)
+{
+	ensure(InName != NAME_None);
+	Name = InName;
+}
+
+void FIKRetargetOpBase::SetParentOpName(const FName InName)
+{
+	ParentOpName = InName;
+}
+
+FName FIKRetargetOpBase::GetParentOpName() const
+{
+	return ParentOpName;
+}
+
 void FIKRetargetOpBase::CopySettingsRaw(const FIKRetargetOpSettingsBase* InSettings, const TArray<FName>& InPropertiesToIgnore)
 {
 	// copy everything except any property that is filtered
diff --git a/Engine/Plugins/Animation/IKRig/Source/IKRig/Private/Retargeter/IKRetargetProcessor.cpp b/Engine/Plugins/Animation/IKRig/Source/IKRig/Private/Retargeter/IKRetargetProcessor.cpp
--- a/Engine/Plugins/Animation/IKRig/Source/IKRig/Private/Retargeter/IKRetargetProcessor.cpp
+++ b/Engine/Plugins/Animation/IKRig/Source/IKRig/Private/Retargeter/IKRetargetProcessor.cpp
@@ -1181,7 +1181,7 @@ bool FIKRetargetProcessor::InitializeOpStack(
 		// create a copy of the op
 		int32 NewOpIndex = OpStack.Emplace(AssetOpStruct);
 		// apply the profile settings to the op (may be different than what was loaded from the asset)
-		InRetargetProfile.ApplyOpProfilesToOpStruct(OpStack[NewOpIndex], NewOpIndex, ECopyOpSettingsContext::PreInitialize);
+		InRetargetProfile.ApplyOpProfilesToOpStruct(OpStack[NewOpIndex], ECopyOpSettingsContext::PreInitialize);
 		
 		// initialize it
 		FIKRetargetOpBase& NewOp = OpStack[NewOpIndex].GetMutable<FIKRetargetOpBase>();
@@ -1680,9 +1680,9 @@ void FIKRetargetProcessor::SetNeedsInitialized()
 void FIKRetargetProcessor::ApplyProfileAtRuntime(const FRetargetProfile& Profile)
 {
 	// apply op settings from the profile
-	for (int32 OpIndex=0; OpIndex < OpStack.Num(); ++OpIndex)
+	for (FInstancedStruct& OpStruct : OpStack)
 	{
-		Profile.ApplyOpProfilesToOpStruct(OpStack[OpIndex], OpIndex, ECopyOpSettingsContext::Runtime);
+		Profile.ApplyOpProfilesToOpStruct(OpStruct, ECopyOpSettingsContext::Runtime);
 	}
 
 	// apply flag to force IK off
diff --git a/Engine/Plugins/Animation/IKRig/Source/IKRig/Private/Retargeter/IKRetargetProfile.cpp b/Engine/Plugins/Animation/IKRig/Source/IKRig/Private/Retargeter/IKRetargetProfile.cpp
--- a/Engine/Plugins/Animation/IKRig/Source/IKRig/Private/Retargeter/IKRetargetProfile.cpp
+++ b/Engine/Plugins/Animation/IKRig/Source/IKRig/Private/Retargeter/IKRetargetProfile.cpp
@@ -15,11 +15,11 @@
 
 
 FRetargetOpProfile::FRetargetOpProfile(
-	int32 InOpIndex,
+	const FName InOpName,
 	const UScriptStruct* InSettingsType,
 	const FIKRetargetOpSettingsBase* InOpSettings)
 {
-	OpToApplySettingsTo = InOpIndex;
+	OpToApplySettingsTo = InOpName;
 	SettingsToApply.InitializeAs(InSettingsType, reinterpret_cast<const uint8*>(InOpSettings));
 }
 
@@ -37,7 +37,7 @@ bool FRetargetOpProfile::CopySettingsToOp(FInstancedStruct& InOutOpStruct, ECopy
 	{
 		return false; // settings were wrong type
 	}
-	
+
 	switch (InApplyContext)
 	{
 		case ECopyOpSettingsContext::Runtime:
@@ -73,7 +73,7 @@ void FRetargetProfile::FillProfileWithAssetSettings(const UIKRetargeter* InAsset
 	for (int32 OpIndex=0; OpIndex<AssetOps.Num(); ++OpIndex)
 	{
 		FIKRetargetOpBase* Op = AssetOps[OpIndex].GetMutablePtr<FIKRetargetOpBase>();
-		RetargetOpProfiles.Emplace(OpIndex, Op->GetSettingsType(), Op->GetSettings());
+		RetargetOpProfiles.Emplace(Op->GetName(), Op->GetSettingsType(), Op->GetSettings());
 	}
 
 	// now override any settings in the asset's current profile
@@ -274,26 +274,16 @@ FRetargetOpProfile* FRetargetProfile::FindMatchingOpProfile(const FRetargetOpPro
 
 bool FRetargetProfile::ApplyOpProfilesToOpStruct(
 	FInstancedStruct& InOutOpStruct,
-	const int32 InIndexOfOpInStack,
 	const ECopyOpSettingsContext InCopyContext) const
 {
 	bool bOpSettingsApplied = true;
-	for (const FRetargetOpProfile& RetargetOpProfile : RetargetOpProfiles)
+	const FName OpName = InOutOpStruct.GetPtr<FIKRetargetOpBase>()->GetName();
+	for (const FRetargetOpProfile& OpProfile : RetargetOpProfiles)
 	{
-		if (RetargetOpProfile.OpToApplySettingsTo == INDEX_NONE)
+		if (OpProfile.OpToApplySettingsTo == NAME_None || OpProfile.OpToApplySettingsTo == OpName)
 		{
-			// no index was specified, so we apply the settings to all compatible types
-			bOpSettingsApplied &= RetargetOpProfile.CopySettingsToOp(InOutOpStruct, InCopyContext);
-		}
-		else if (RetargetOpProfile.OpToApplySettingsTo == InIndexOfOpInStack)
-		{
-			bOpSettingsApplied &= RetargetOpProfile.CopySettingsToOp(InOutOpStruct, InCopyContext);
-			if (!bOpSettingsApplied)
-			{
-				UE_LOG(LogTemp, Warning, TEXT("Retargeter profile found settings for a different op type at the same index. Op settings ignored."));
-			}
+			bOpSettingsApplied &= OpProfile.CopySettingsToOp(InOutOpStruct, InCopyContext);
 		}
 	}
-
 	return bOpSettingsApplied;
 }
diff --git a/Engine/Plugins/Animation/IKRig/Source/IKRig/Private/Retargeter/IKRetargeter.cpp b/Engine/Plugins/Animation/IKRig/Source/IKRig/Private/Retargeter/IKRetargeter.cpp
--- a/Engine/Plugins/Animation/IKRig/Source/IKRig/Private/Retargeter/IKRetargeter.cpp
+++ b/Engine/Plugins/Animation/IKRig/Source/IKRig/Private/Retargeter/IKRetargeter.cpp
@@ -205,8 +205,6 @@ UIKRetargeter::UIKRetargeter(const FObjectInitializer& ObjectInitializer)
 	OpStack_DEPRECATED = CreateDefaultSubobject<UDEPRECATED_RetargetOpStack>(TEXT("PostSettings"));
 	OpStack_DEPRECATED->SetFlags(RF_Transactional);
 	PRAGMA_ENABLE_DEPRECATION_WARNINGS
-
-	CleanAndInitialize();
 }
 
 const UIKRigDefinition* UIKRetargeter::GetIKRig(ERetargetSourceOrTarget SourceOrTarget) const
@@ -286,15 +284,20 @@ void UIKRetargeter::PostLoad()
 	// very early versions of the asset may not have been set as standalone
 	SetFlags(RF_Standalone);
 
+#if WITH_EDITOR
 	PostLoadOldSettingsToNew();
 	
 	PostLoadOldOpsToNewStructOps();
 
 	PostLoadConvertEverythingToOps();
+#endif
+
+	CleanRetargetPoses();
 	
-	CleanAndInitialize();
+	CleanOpStack();
 }
 
+#if WITH_EDITOR
 void UIKRetargeter::PostLoadOldSettingsToNew()
 {
 	PRAGMA_DISABLE_DEPRECATION_WARNINGS
@@ -422,6 +425,8 @@ void UIKRetargeter::PostLoadConvertEverythingToOps()
 	
 	int32 OpIndex = 0;
 
+	FName RunIKOpName = FIKRetargetRunIKRigOp().GetDefaultName();
+
 	// Source Scale Op from old "global" settings
 	{
 		if (!FMath::IsNearlyEqual(GlobalSettings_DEPRECATED->Settings.SourceScaleFactor, 1.0f))
@@ -479,6 +484,7 @@ void UIKRetargeter::PostLoadConvertEverythingToOps()
 		++OpIndex;
 		RetargetOps.EmplaceAt(OpIndex,FIKRetargetIKChainsOp::StaticStruct());
 		FIKRetargetOpBase* Op = RetargetOps[OpIndex].GetMutablePtr<FIKRetargetOpBase>();
+		Op->SetParentOpName(RunIKOpName);
 		Op->SetEnabled(GlobalSettings_DEPRECATED->Settings.bEnableIK);
 		FIKRetargetIKChainsOpSettings* Settings = reinterpret_cast<FIKRetargetIKChainsOpSettings*>(Op->GetSettings());
 		for (URetargetChainSettings* Chain : ChainSettings_DEPRECATED)
@@ -509,6 +515,7 @@ void UIKRetargeter::PostLoadConvertEverythingToOps()
 		++OpIndex;
 		RetargetOps.EmplaceAt(OpIndex,FIKRetargetStrideWarpingOp::StaticStruct());
 		FIKRetargetOpBase* Op = RetargetOps[OpIndex].GetMutablePtr<FIKRetargetOpBase>();
+		Op->SetParentOpName(RunIKOpName);
 		Op->SetEnabled(GlobalSettings_DEPRECATED->Settings.bWarping);
 		FIKRetargetStrideWarpingOpSettings* Settings = reinterpret_cast<FIKRetargetStrideWarpingOpSettings*>(Op->GetSettings());
 		Settings->DirectionSource = GlobalSettings_DEPRECATED->Settings.DirectionSource;
@@ -537,6 +544,7 @@ void UIKRetargeter::PostLoadConvertEverythingToOps()
 		++OpIndex;
 		RetargetOps.EmplaceAt(OpIndex,FIKRetargetSpeedPlantingOp::StaticStruct());
 		FIKRetargetOpBase* Op = RetargetOps[OpIndex].GetMutablePtr<FIKRetargetOpBase>();
+		Op->SetParentOpName(RunIKOpName);
 		FIKRetargetSpeedPlantingOpSettings* Settings = reinterpret_cast<FIKRetargetSpeedPlantingOpSettings*>(Op->GetSettings());
 		for (URetargetChainSettings* Chain : ChainSettings_DEPRECATED)
 		{
@@ -593,8 +601,56 @@ void UIKRetargeter::PostLoadConvertEverythingToOps()
 	
 	PRAGMA_ENABLE_DEPRECATION_WARNINGS
 }
+#endif
 
-void UIKRetargeter::CleanAndInitialize()
+FName UIKRetargeter::GetCleanAndUniqueOpName(const FName& InOpName, const int32 InIndexOfOp)
+{
+	FName OutCleanedName = InOpName;
+
+	// if None, revert to default name for the op
+	if (InOpName == NAME_None)
+	{
+		#if WITH_EDITOR
+		OutCleanedName = RetargetOps[InIndexOfOp].GetPtr<FIKRetargetOpBase>()->GetDefaultName();
+		#else
+		OutCleanedName = FName("DefaultRetargetOpName");	
+		#endif
+	}
+
+	auto OpNameInUse = [this](const FName InOpNameToCheck, int32 InOpIndexToIgnore)
+	{
+		for (int32 OpIndex=0; OpIndex<RetargetOps.Num(); ++OpIndex)
+		{
+			if (InOpIndexToIgnore != INDEX_NONE && OpIndex == InOpIndexToIgnore)
+			{
+				continue;
+			}
+		
+			const FIKRetargetOpBase* Op = RetargetOps[OpIndex].GetPtr<FIKRetargetOpBase>();
+			if (Op->GetName() == InOpNameToCheck)
+			{
+				return true;
+			}
+		}
+		return false;
+	};
+
+	if (!OpNameInUse(OutCleanedName, InIndexOfOp))
+	{
+		return OutCleanedName;
+	}
+
+	// keep concatenating an incremented integer suffix until name is unique
+	int32 Number = OutCleanedName.GetNumber() + 1;
+	while(OpNameInUse(FName(OutCleanedName, Number), InIndexOfOp))
+	{
+		Number++;
+	}
+
+	return OutCleanedName;
+}
+
+void UIKRetargeter::CleanRetargetPoses()
 {
 	// use default pose as current pose unless set to something else
 	if (CurrentSourceRetargetPose == NAME_None)
@@ -625,6 +681,140 @@ void UIKRetargeter::CleanAndInitialize()
 	{
 		CurrentTargetRetargetPose = GetDefaultPoseName();
 	}
+}
+
+void UIKRetargeter::CleanOpStack()
+{
+	auto GetNamesOfTopLevelOps = [this]() -> TArray<FName>
+	{
+		TArray<FName> NamesOfTopLevelOps;
+		for (FInstancedStruct& OpStruct : RetargetOps)
+		{
+			const FIKRetargetOpBase& Op = OpStruct.Get<FIKRetargetOpBase>();
+			if (Op.GetParentOpName() == NAME_None)
+			{
+				NamesOfTopLevelOps.Add(Op.GetName());
+			}
+		}
+		return MoveTemp(NamesOfTopLevelOps);
+	};
+
+	auto GetChildOpNames = [this](const int32 InOpIndex) -> TArray<FName>
+	{
+		TArray<FName> ChildrenNames;
+		const FName InOpName = RetargetOps[InOpIndex].Get<FIKRetargetOpBase>().GetName();
+		for (FInstancedStruct& OpStruct : RetargetOps)
+		{
+			const FIKRetargetOpBase& Op = OpStruct.Get<FIKRetargetOpBase>();
+			if (Op.GetParentOpName() == InOpName)
+			{
+				ChildrenNames.Add(Op.GetName());
+			}
+		}
+		return MoveTemp(ChildrenNames);
+	};
+
+	auto GetIndexOfOpByName = [this](const FName InOpName) -> int32
+	{
+		for (int32 OpIndex=0; OpIndex<RetargetOps.Num(); ++OpIndex)
+		{
+			if (RetargetOps[OpIndex].Get<FIKRetargetOpBase>().GetName() == InOpName)
+			{
+				return OpIndex;
+			}
+		}
+		return INDEX_NONE;
+	};
+	
+	// enforce unique non-None names on all ops
+	for (int32 OpIndex=0; OpIndex<RetargetOps.Num(); ++OpIndex)
+	{
+		FIKRetargetOpBase* Op = RetargetOps[OpIndex].GetMutablePtr<FIKRetargetOpBase>();
+		FName OldOpName = Op->GetName();
+		FName CleanedOpName = GetCleanAndUniqueOpName(OldOpName, OpIndex);
+		Op->SetName(CleanedOpName);
+
+		// update any children pointing at the old name
+		if (OldOpName != NAME_None)
+		{
+			for (FInstancedStruct& OpStruct : RetargetOps)
+			{
+				FIKRetargetOpBase& OtherOp = OpStruct.GetMutable<FIKRetargetOpBase>();
+				if (OtherOp.GetParentOpName() == OldOpName)
+				{
+					OtherOp.SetParentOpName(CleanedOpName);
+				}
+			}
+		}	
+	}
+
+	// auto parent ops with missing or unset parent
+	for (FInstancedStruct& OpStruct : RetargetOps)
+	{
+		FIKRetargetOpBase& Op = OpStruct.GetMutable<FIKRetargetOpBase>();
+		const UScriptStruct* ParentType = Op.GetParentOpType();
+		if (ParentType == nullptr)
+		{
+			// op doesn't need a parent
+			continue;
+		}
+		
+		if (GetIndexOfOpByName(Op.GetParentOpName()) != INDEX_NONE)
+		{
+			// op already has a parent and it exists
+			continue;
+		}
+
+		// op needs a parent but doesn't have one
+		// find the first op of the correct type and parent it
+		for (FInstancedStruct& OtherOpStruct : RetargetOps)
+		{
+			if (OtherOpStruct.GetScriptStruct() == ParentType)
+			{
+				const FIKRetargetOpBase& OtherOp = OtherOpStruct.Get<FIKRetargetOpBase>();
+				Op.SetParentOpName(OtherOp.GetName());
+				break;
+			}
+		}
+	}
+
+	// enforce correct execution order of the ops according to the following constraints
+	// 1. all children must come BEFORE parent
+	// 2. compact
+	// 3. retain existing order of ops to the extent possible
+	// 4. no non-siblings between siblings
+
+	// get list of op names in the correct execution order
+	TArray<FName> TopLevelOpsNames = GetNamesOfTopLevelOps();
+	TArray<FName> CorrectedOpOrder;
+	for (const FName TopLevelOpName : TopLevelOpsNames)
+	{
+		const int32 OpIndex = GetIndexOfOpByName(TopLevelOpName);
+		CorrectedOpOrder.Append(GetChildOpNames(OpIndex));
+		CorrectedOpOrder.Add(TopLevelOpName);
+	}
+	
+	// re-order op stack accordingly
+	TArray<FInstancedStruct> TempOps = MoveTemp(RetargetOps);
+	for (const FName& NameOfNextOpToAdd : CorrectedOpOrder)
+	{
+		for (FInstancedStruct& TempOpStruct : TempOps)
+		{
+			if (!TempOpStruct.IsValid())
+			{
+				continue; // already moved
+			}
+			
+			FIKRetargetOpBase* TempOp = TempOpStruct.GetMutablePtr<FIKRetargetOpBase>();
+			if (TempOp->GetName() == NameOfNextOpToAdd)
+			{
+				RetargetOps.Add(MoveTemp(TempOpStruct));
+				break;
+			}
+		}
+	}
+	
+	ensureAlwaysMsgf(RetargetOps.Num() == TempOps.Num(), TEXT("Retarget ops were lost during cleaning."));
 };
 
 #if WITH_EDITORONLY_DATA
diff --git a/Engine/Plugins/Animation/IKRig/Source/IKRig/Private/Retargeter/RetargetOps/AlignPoleVectorOp.cpp b/Engine/Plugins/Animation/IKRig/Source/IKRig/Private/Retargeter/RetargetOps/AlignPoleVectorOp.cpp
--- a/Engine/Plugins/Animation/IKRig/Source/IKRig/Private/Retargeter/RetargetOps/AlignPoleVectorOp.cpp
+++ b/Engine/Plugins/Animation/IKRig/Source/IKRig/Private/Retargeter/RetargetOps/AlignPoleVectorOp.cpp
@@ -194,13 +194,10 @@ void FIKRetargetAlignPoleVectorOpSettings::MergePoleVectorSettings(const FRetarg
 	{
 		if (ChainToAlign.TargetChainName == InSettingsToMerge.TargetChainName)
 		{
-			// settings exist for that chain, so update them
 			ChainToAlign = InSettingsToMerge;
 			return;
 		}
 	}
-
-	// new settings, so add them
 	ChainsToAlign.Add(InSettingsToMerge);
 }
 
diff --git a/Engine/Plugins/Animation/IKRig/Source/IKRig/Private/Retargeter/RetargetOps/IKChainsOp.cpp b/Engine/Plugins/Animation/IKRig/Source/IKRig/Private/Retargeter/RetargetOps/IKChainsOp.cpp
--- a/Engine/Plugins/Animation/IKRig/Source/IKRig/Private/Retargeter/RetargetOps/IKChainsOp.cpp
+++ b/Engine/Plugins/Animation/IKRig/Source/IKRig/Private/Retargeter/RetargetOps/IKChainsOp.cpp
@@ -5,6 +5,7 @@
 #include "Retargeter/IKRetargetProcessor.h"
 #include "Retargeter/RetargetOps/FKChainsOp.h"
 #include "Retargeter/RetargetOps/PelvisMotionOp.h"
+#include "Retargeter/RetargetOps/RunIKRigOp.h"
 #include "Retargeter/RetargetOps/StrideWarpingOp.h"
 
 #if WITH_EDITOR
@@ -357,6 +358,11 @@ const UClass* FIKRetargetIKChainsOp::GetControllerType() const
 	return UIKRetargetIKChainsController::StaticClass();
 }
 
+const UScriptStruct* FIKRetargetIKChainsOp::GetParentOpType() const
+{
+	return FIKRetargetRunIKRigOp::StaticStruct();
+}
+
 #if WITH_EDITOR
 void FIKRetargetIKChainsOp::RunInitialSetup(const UIKRetargeter* InRetargetAsset)
 {
diff --git a/Engine/Plugins/Animation/IKRig/Source/IKRig/Private/Retargeter/RetargetOps/SpeedPlantingOp.cpp b/Engine/Plugins/Animation/IKRig/Source/IKRig/Private/Retargeter/RetargetOps/SpeedPlantingOp.cpp
--- a/Engine/Plugins/Animation/IKRig/Source/IKRig/Private/Retargeter/RetargetOps/SpeedPlantingOp.cpp
+++ b/Engine/Plugins/Animation/IKRig/Source/IKRig/Private/Retargeter/RetargetOps/SpeedPlantingOp.cpp
@@ -6,6 +6,7 @@
 #include "Animation/AnimInstance.h"
 #include "Components/SkeletalMeshComponent.h"
 #include "Animation/AnimNodeBase.h"
+#include "Retargeter/RetargetOps/RunIKRigOp.h"
 
 #define LOCTEXT_NAMESPACE "SpeedPlantingOp"
 
@@ -139,6 +140,11 @@ const UClass* FIKRetargetSpeedPlantingOp::GetControllerType() const
 	return UIKRetargetSpeedPlantingController::StaticClass();
 }
 
+const UScriptStruct* FIKRetargetSpeedPlantingOp::GetParentOpType() const
+{
+	return FIKRetargetRunIKRigOp::StaticStruct();
+}
+
 void FIKRetargetSpeedPlantingOp::OnPlaybackReset()
 {
 	for (FSpeedPlantingGoalState& GoalToPlant : GoalsToPlant)
diff --git a/Engine/Plugins/Animation/IKRig/Source/IKRig/Private/Retargeter/RetargetOps/StrideWarpingOp.cpp b/Engine/Plugins/Animation/IKRig/Source/IKRig/Private/Retargeter/RetargetOps/StrideWarpingOp.cpp
--- a/Engine/Plugins/Animation/IKRig/Source/IKRig/Private/Retargeter/RetargetOps/StrideWarpingOp.cpp
+++ b/Engine/Plugins/Animation/IKRig/Source/IKRig/Private/Retargeter/RetargetOps/StrideWarpingOp.cpp
@@ -3,6 +3,7 @@
 #include "Retargeter/RetargetOps/StrideWarpingOp.h"
 
 #include "Retargeter/IKRetargetProcessor.h"
+#include "Retargeter/RetargetOps/RunIKRigOp.h"
 #include "Rig/Solvers/PointsToRotation.h"
 
 #if WITH_EDITOR
@@ -314,6 +315,11 @@ const UClass* FIKRetargetStrideWarpingOp::GetControllerType() const
 	return UIKRetargetStrideWarpingController::StaticClass();
 }
 
+const UScriptStruct* FIKRetargetStrideWarpingOp::GetParentOpType() const
+{
+	return FIKRetargetRunIKRigOp::StaticStruct();
+}
+
 #if WITH_EDITOR
 
 FCriticalSection FIKRetargetStrideWarpingOp::DebugDataMutex;
diff --git a/Engine/Plugins/Animation/IKRig/Source/IKRig/Public/Retargeter/IKRetargetOps.h b/Engine/Plugins/Animation/IKRig/Source/IKRig/Public/Retargeter/IKRetargetOps.h
--- a/Engine/Plugins/Animation/IKRig/Source/IKRig/Public/Retargeter/IKRetargetOps.h
+++ b/Engine/Plugins/Animation/IKRig/Source/IKRig/Public/Retargeter/IKRetargetOps.h
@@ -137,9 +137,16 @@ public:
 	 * NOTE: this type will automatically be lazy instantiated when/if the user calls GetController() on the op. */
 	virtual const UClass* GetControllerType() const;
 
+	/** (optional) ops can optionally behave as a 'parent' where child ops must be executed first. */
+	virtual bool CanHaveChildOps() const { return false; };
+
+	/** (optional) ops can optionally behave as a 'child' where they can be parented to ops of the type returned by this function. */
+	virtual const UScriptStruct* GetParentOpType() const { return nullptr; };
+
 #if WITH_EDITOR
+	
 	/** get the nice name to display in the viewport (defaults to "DisplayName" UStruct metadata) */
-	FText GetNiceName() const;
+	FName GetDefaultName() const;
 	
 	/** (optional) override to automate initial setup after being added to the stack */
 	virtual void RunInitialSetup(const UIKRetargeter* InRetargetAsset) {};
@@ -161,6 +168,21 @@ public:
 	virtual bool AreChainSettingsAtDefault(const FName& InChainName) { return true; };
 #endif
 
+	/** Get the name of this op (may be customized by user) */
+	FName GetName() const;
+
+	/** Set the name of this op
+	 * NOTE: InName is assumed to be unique in the stack (this constraint is enforced when setting the name from the controller) */
+	void SetName(const FName InName);
+	
+	/** Set the name of the op this op is "parented" to. This enforces execution order.
+	 * NOTE: InName is assumed to refer to an op that exists.
+	 * NOTE: Renaming the parent op through the controller will auto-update this. */
+	void SetParentOpName(const FName InName);
+
+	/** Get the name of the op this op is parented to (None for root level ops)*/
+	FName GetParentOpName() const;
+
 	/** wholesale copy all settings from the input settings into this op.
 	 * NOTE: InSettings is assumed to be castable to the type returned by GetSettingsType()
 	 * NOTE: be careful calling this on an initialized op as it may invalidate it's runtime state (use SetSettings instead) */
@@ -169,8 +191,17 @@ public:
 private:
 	
 	/** when false, execution of this op is skipped */
+	UPROPERTY()
 	bool bIsEnabled = true;
 
+	/** the text label given to the op, used to refer to it from script */
+	UPROPERTY()
+	FName Name;
+
+	/** (optional) some ops are considered as a group, this is the name of the group parent */
+	UPROPERTY()
+	FName ParentOpName;
+
 protected:
 
 	/** a convenience function to copy all properties from one struct to another while ignoring some */
diff --git a/Engine/Plugins/Animation/IKRig/Source/IKRig/Public/Retargeter/IKRetargetProfile.h b/Engine/Plugins/Animation/IKRig/Source/IKRig/Public/Retargeter/IKRetargetProfile.h
--- a/Engine/Plugins/Animation/IKRig/Source/IKRig/Public/Retargeter/IKRetargetProfile.h
+++ b/Engine/Plugins/Animation/IKRig/Source/IKRig/Public/Retargeter/IKRetargetProfile.h
@@ -28,7 +28,7 @@ struct FRetargetOpProfile
 	FRetargetOpProfile() = default;
 
 	FRetargetOpProfile(
-		int32 InOpIndex,
+		const FName InOpName,
 		const UScriptStruct* InSettingsType,
 		const FIKRetargetOpSettingsBase* InOpSettings);
 
@@ -40,10 +40,10 @@ struct FRetargetOpProfile
 	 * @return true if the supplied op was compatible with the settings, false otherwise */
 	bool CopySettingsToOp(FInstancedStruct& InOutOpStruct, ECopyOpSettingsContext InApplyContext) const;
 	
-	/** The index of the op these settings should be applied to.
-	 * If left at -1, will apply to all ops that use the type of settings stored in 'SettingsToApply' */
+	/** The name of the op these settings should be applied to.
+	 * If left at None, will apply to all ops that use the type of settings stored in 'SettingsToApply' */
 	UPROPERTY(EditAnywhere, BlueprintReadWrite, Category = "IK Retarget Op Profile")
-	int32 OpToApplySettingsTo = INDEX_NONE;
+	FName OpToApplySettingsTo = NAME_None;
 
 	/** The settings to override for a specific op in the retargeter's op stack.
 	 * NOTE: all settings stored here will be applied to the op. However, at runtime each op independently
@@ -108,7 +108,7 @@ public:
 
 	FRetargetOpProfile* FindMatchingOpProfile(const FRetargetOpProfile& OtherOpProfile);
 
-	bool ApplyOpProfilesToOpStruct(FInstancedStruct& InOutOpStruct, const int32 InIndexOfOpInStack, const ECopyOpSettingsContext InCopyContext) const;
+	bool ApplyOpProfilesToOpStruct(FInstancedStruct& InOutOpStruct, const ECopyOpSettingsContext InCopyContext) const;
 	
 	// search all the op profiles and get an array of all the settings of a particular type
 	template <typename T>
diff --git a/Engine/Plugins/Animation/IKRig/Source/IKRig/Public/Retargeter/IKRetargetSettings.h b/Engine/Plugins/Animation/IKRig/Source/IKRig/Public/Retargeter/IKRetargetSettings.h
--- a/Engine/Plugins/Animation/IKRig/Source/IKRig/Public/Retargeter/IKRetargetSettings.h
+++ b/Engine/Plugins/Animation/IKRig/Source/IKRig/Public/Retargeter/IKRetargetSettings.h
@@ -30,22 +30,31 @@ enum class EBasicAxis
 #if WITH_EDITOR
 enum class ERetargetSelectionType : uint8
 {
+	NONE,
 	BONE,
 	CHAIN,
 	MESH,
 	ROOT,
-	NONE
+	OP
 };
 
 struct FIKRetargetDebugDrawState
 {
 	// state of all things you can select in the retarget editor viewport
+
+	// what type of objects was last selected?
+	ERetargetSelectionType LastSelectedType = ERetargetSelectionType::NONE;
+	// BONE
 	bool bIsRootSelected = false;
-	UPrimitiveComponent* SelectedMesh = nullptr;
+	TMap<ERetargetSourceOrTarget, TArray<FName>> SelectedBoneNames;
+	// CHAINS
 	TArray<FName> SelectedChains;
+	// MESH
+	UPrimitiveComponent* SelectedMesh = nullptr;
+	// GOALS
 	TArray<FName> SelectedGoals;
-	TMap<ERetargetSourceOrTarget, TArray<FName>> SelectedBoneNames;
-	ERetargetSelectionType LastSelectedItem = ERetargetSelectionType::NONE;
+	// OPS
+	FName LastSelectedOpName;
 
 	// common viewport drawing color palette for selected things
 	static FLinearColor Muted;
diff --git a/Engine/Plugins/Animation/IKRig/Source/IKRig/Public/Retargeter/IKRetargeter.h b/Engine/Plugins/Animation/IKRig/Source/IKRig/Public/Retargeter/IKRetargeter.h
--- a/Engine/Plugins/Animation/IKRig/Source/IKRig/Public/Retargeter/IKRetargeter.h
+++ b/Engine/Plugins/Animation/IKRig/Source/IKRig/Public/Retargeter/IKRetargeter.h
@@ -176,7 +176,6 @@ public:
     
 		return Results;
 	}
-
 	// Get read-only access to a retarget pose 
 	const FIKRetargetPose* GetCurrentRetargetPose(const ERetargetSourceOrTarget& SourceOrTarget) const;
 	// Get name of the current retarget pose 
@@ -358,11 +357,18 @@ public:
 
 private:
 
+#if WITH_EDITOR
 	// upgrade old assets to new data formats
 	void PostLoadOldSettingsToNew();
 	void PostLoadOldOpsToNewStructOps();
 	void PostLoadConvertEverythingToOps();
-	void CleanAndInitialize();
+#endif
+	void CleanRetargetPoses();
+	void CleanOpStack();
+	
+	// make the provided op name valid and unique
+	// NOTE: if InIndexOfOp is INDEX_None, then it assumes this is a new op
+	FName GetCleanAndUniqueOpName(const FName& InOpName, const int32 InIndexOfOp);
 	
 	// incremented by any edits that require re-initialization
 	UPROPERTY(Transient)
diff --git a/Engine/Plugins/Animation/IKRig/Source/IKRig/Public/Retargeter/RetargetOps/CurveRemapOp.h b/Engine/Plugins/Animation/IKRig/Source/IKRig/Public/Retargeter/RetargetOps/CurveRemapOp.h
--- a/Engine/Plugins/Animation/IKRig/Source/IKRig/Public/Retargeter/RetargetOps/CurveRemapOp.h
+++ b/Engine/Plugins/Animation/IKRig/Source/IKRig/Public/Retargeter/RetargetOps/CurveRemapOp.h
@@ -15,11 +15,11 @@ struct FCurveRemapPair
 	GENERATED_BODY()
 	
 	// The curve name on the SOURCE skeletal mesh to copy animation data from.
-	UPROPERTY(EditAnywhere, Category=Settings)
+	UPROPERTY(EditAnywhere, BlueprintReadWrite, Category=Settings)
 	FName SourceCurve;
 
 	// The curve name on the TARGET skeletal mesh to receive animation data.
-	UPROPERTY(EditAnywhere, Category=Settings)
+	UPROPERTY(EditAnywhere, BlueprintReadWrite, Category=Settings)
 	FName TargetCurve;
 };
 
@@ -32,14 +32,14 @@ struct IKRIG_API FIKRetargetCurveRemapOpSettings : public FIKRetargetOpSettingsB
 	// NOTE: This setting also applies when exporting retargeted animations.
 	// True: all source curves are copied to the target animation instance/asset
 	// False: only remapped curves are copied on the target animation instance/asset
-	UPROPERTY(EditAnywhere, Category="OpSettings")
+	UPROPERTY(EditAnywhere, BlueprintReadWrite, Category="OpSettings")
 	bool bCopyAllSourceCurves = true;
 	
 	// Add pairs of Source/Target curve names to remap. While retargeting, the animation from the source curves
 	// will be redirected to the curves on the target skeletal meshes. Can be used to drive, blendshapes or other downstream systems.
 	// NOTE: By default the IK Retargeter will automatically copy all equivalently named curves from the source to the
 	// target. Remapping is only necessary when the target curve name(s) are different.
-	UPROPERTY(EditAnywhere, Category="OpSettings")
+	UPROPERTY(EditAnywhere, BlueprintReadWrite, Category="OpSettings")
 	TArray<FCurveRemapPair> CurvesToRemap;
 };
 
diff --git a/Engine/Plugins/Animation/IKRig/Source/IKRig/Public/Retargeter/RetargetOps/FKChainsOp.h b/Engine/Plugins/Animation/IKRig/Source/IKRig/Public/Retargeter/RetargetOps/FKChainsOp.h
--- a/Engine/Plugins/Animation/IKRig/Source/IKRig/Public/Retargeter/RetargetOps/FKChainsOp.h
+++ b/Engine/Plugins/Animation/IKRig/Source/IKRig/Public/Retargeter/RetargetOps/FKChainsOp.h
@@ -192,7 +192,7 @@ struct FFKChainDebugData
 };
 #endif
 
-USTRUCT(BlueprintType, meta = (DisplayName = "Retarget Chains (FK)"))
+USTRUCT(BlueprintType, meta = (DisplayName = "Retarget FK Chains"))
 struct IKRIG_API FIKRetargetFKChainsOp : public FIKRetargetOpBase
 {
 	GENERATED_BODY()
diff --git a/Engine/Plugins/Animation/IKRig/Source/IKRig/Public/Retargeter/RetargetOps/IKChainsOp.h b/Engine/Plugins/Animation/IKRig/Source/IKRig/Public/Retargeter/RetargetOps/IKChainsOp.h
--- a/Engine/Plugins/Animation/IKRig/Source/IKRig/Public/Retargeter/RetargetOps/IKChainsOp.h
+++ b/Engine/Plugins/Animation/IKRig/Source/IKRig/Public/Retargeter/RetargetOps/IKChainsOp.h
@@ -214,6 +214,8 @@ struct IKRIG_API FIKRetargetIKChainsOp : public FIKRetargetOpBase
 
 	virtual const UClass* GetControllerType() const override;
 
+	virtual const UScriptStruct* GetParentOpType() const override;
+
 	UPROPERTY()
 	FIKRetargetIKChainsOpSettings Settings;
 
diff --git a/Engine/Plugins/Animation/IKRig/Source/IKRig/Public/Retargeter/RetargetOps/PelvisMotionOp.h b/Engine/Plugins/Animation/IKRig/Source/IKRig/Public/Retargeter/RetargetOps/PelvisMotionOp.h
--- a/Engine/Plugins/Animation/IKRig/Source/IKRig/Public/Retargeter/RetargetOps/PelvisMotionOp.h
+++ b/Engine/Plugins/Animation/IKRig/Source/IKRig/Public/Retargeter/RetargetOps/PelvisMotionOp.h
@@ -182,7 +182,7 @@ private:
 	
 #if WITH_EDITOR
 public:
-	virtual void RunInitialSetup(const UIKRetargeter* InRetargetAsset);
+	virtual void RunInitialSetup(const UIKRetargeter* InRetargetAsset) override;
 	virtual FText GetWarningMessage() const override;
 	virtual void DebugDraw(
 		FPrimitiveDrawInterface* InPDI,
diff --git a/Engine/Plugins/Animation/IKRig/Source/IKRig/Public/Retargeter/RetargetOps/PinBoneOp.h b/Engine/Plugins/Animation/IKRig/Source/IKRig/Public/Retargeter/RetargetOps/PinBoneOp.h
--- a/Engine/Plugins/Animation/IKRig/Source/IKRig/Public/Retargeter/RetargetOps/PinBoneOp.h
+++ b/Engine/Plugins/Animation/IKRig/Source/IKRig/Public/Retargeter/RetargetOps/PinBoneOp.h
@@ -91,23 +91,23 @@ struct IKRIG_API FIKRetargetPinBoneOpSettings : public FIKRetargetOpSettingsBase
 	GENERATED_BODY()
 	
 	/** A list of bone-pairs to copy transforms between */
-	UPROPERTY(EditAnywhere, Category=Setup, meta=(ReinitializeOnEdit))
+	UPROPERTY(EditAnywhere, BlueprintReadWrite, Category=Setup, meta=(ReinitializeOnEdit))
 	TArray<FPinBoneData> BonesToPin;
 
 	/** Which skeleton to copy animation from. */
-	UPROPERTY(EditAnywhere, Category=Settings)
+	UPROPERTY(EditAnywhere, BlueprintReadWrite, Category=Settings)
 	ERetargetSourceOrTarget SkeletonToCopyFrom = ERetargetSourceOrTarget::Target;
 
 	/** Copy the translation of the bone. */
-	UPROPERTY(EditAnywhere, Category=Settings)
+	UPROPERTY(EditAnywhere, BlueprintReadWrite, Category=Settings)
 	bool bCopyTranslation = true;
 
 	/** Copy the rotation of the bone. */
-	UPROPERTY(EditAnywhere, Category=Settings)
+	UPROPERTY(EditAnywhere, BlueprintReadWrite, Category=Settings)
 	bool bCopyRotation = true;
 	
 	/** Copy the scale of the bone. */
-	UPROPERTY(EditAnywhere, Category=Settings)
+	UPROPERTY(EditAnywhere, BlueprintReadWrite, Category=Settings)
 	bool bCopyScale = true;
 
 	/** Update children bone transforms (recursively). */
@@ -115,19 +115,19 @@ struct IKRIG_API FIKRetargetPinBoneOpSettings : public FIKRetargetOpSettingsBase
 	bool bPropagateToChildren = false;
 
 	/** The method used to calculate the translation of the bone. */
-	UPROPERTY(EditAnywhere, Category=Settings)
+	UPROPERTY(EditAnywhere, BlueprintReadWrite, Category=Settings)
 	EPinBoneTranslationMode TranslationMode = EPinBoneTranslationMode::CopyGlobalPosition;
 
 	/** The method used to calculate the rotation of the bone. */
-	UPROPERTY(EditAnywhere, Category=Settings)
+	UPROPERTY(EditAnywhere, BlueprintReadWrite, Category=Settings)
 	EPinBoneRotationMode RotationMode = EPinBoneRotationMode::CopyGlobalRotation;
 
 	/** A manual offset to apply in global space */
-	UPROPERTY(EditAnywhere, Category=Settings)
+	UPROPERTY(EditAnywhere, BlueprintReadWrite, Category=Settings)
 	FTransform GlobalOffset;
 
 	/** A manual offset to apply in local space */
-	UPROPERTY(EditAnywhere, Category=Settings)
+	UPROPERTY(EditAnywhere, BlueprintReadWrite, Category=Settings)
 	FTransform LocalOffset;
 
 #if WITH_EDITOR
diff --git a/Engine/Plugins/Animation/IKRig/Source/IKRig/Public/Retargeter/RetargetOps/RootMotionGeneratorOp.h b/Engine/Plugins/Animation/IKRig/Source/IKRig/Public/Retargeter/RetargetOps/RootMotionGeneratorOp.h
--- a/Engine/Plugins/Animation/IKRig/Source/IKRig/Public/Retargeter/RetargetOps/RootMotionGeneratorOp.h
+++ b/Engine/Plugins/Animation/IKRig/Source/IKRig/Public/Retargeter/RetargetOps/RootMotionGeneratorOp.h
@@ -45,33 +45,33 @@ struct IKRIG_API FIKRetargetRootMotionOpSettings : public FIKRetargetOpSettingsB
 	// Copy From Source Root: copies the motion from the source root bone, but scales it according to relative height difference.
 	// Generate From Target Pelvis: uses the retargeted Pelvis motion to generate root motion.
 	// NOTE: Generating root motion from the Pelvis 
-	UPROPERTY(EditAnywhere, Category="OpSettings")
+	UPROPERTY(EditAnywhere, BlueprintReadWrite, Category="OpSettings")
 	ERootMotionSource RootMotionSource = ERootMotionSource::CopyFromSourceRoot;
 	
 	// How to set the height of the root bone.
 	// Copy Height From Source: copies the height of the root bone on the source skeleton's root bone.
 	// Snap To Ground: sets the root bone height to the ground plane (Z=0).
-	UPROPERTY(EditAnywhere, Category="OpSettings")
+	UPROPERTY(EditAnywhere, BlueprintReadWrite, Category="OpSettings")
 	ERootMotionHeightSource RootHeightSource = ERootMotionHeightSource::CopyHeightFromSource;
 
 	// Will transform all children of the target root that are not themselves part of a retarget chain.
-	UPROPERTY(EditAnywhere, Category="OpSettings")
+	UPROPERTY(EditAnywhere, BlueprintReadWrite, Category="OpSettings")
 	bool bPropagateToNonRetargetedChildren = true;
 
 	// Applies only when generating root motion from the Pelvis.
 	// Maintains the offset between the root and pelvis as recorded in the target retarget pose.
 	// If false, the root bone is placed directly under the Pelvis bone.
-	UPROPERTY(EditAnywhere, Category="OpSettings")
+	UPROPERTY(EditAnywhere, BlueprintReadWrite, Category="OpSettings")
 	bool bMaintainOffsetFromPelvis = true;
 
 	// Applies only when generating root motion from the Pelvis.
 	// When true, the applied offset will be rotated by the Pelvis.
 	// (NOTE: This may cause unwanted rotations, for example if Pelvis Yaw is animated.)
-	UPROPERTY(EditAnywhere, Category="OpSettings")
+	UPROPERTY(EditAnywhere, BlueprintReadWrite, Category="OpSettings")
 	bool bRotateWithPelvis = false;
 
 	// A manual offset to apply in global space to the root bone.
-	UPROPERTY(EditAnywhere, Category=Settings)
+	UPROPERTY(EditAnywhere, BlueprintReadWrite, Category=Settings)
 	FTransform GlobalOffset;
 
 	/** FBoneReference needs a skeleton for the UI widget. */
diff --git a/Engine/Plugins/Animation/IKRig/Source/IKRig/Public/Retargeter/RetargetOps/RunIKRigOp.h b/Engine/Plugins/Animation/IKRig/Source/IKRig/Public/Retargeter/RetargetOps/RunIKRigOp.h
--- a/Engine/Plugins/Animation/IKRig/Source/IKRig/Public/Retargeter/RetargetOps/RunIKRigOp.h
+++ b/Engine/Plugins/Animation/IKRig/Source/IKRig/Public/Retargeter/RetargetOps/RunIKRigOp.h
@@ -14,11 +14,11 @@ struct IKRIG_API FIKRetargetRunIKRigOpSettings : public FIKRetargetOpSettingsBas
 	GENERATED_BODY()
 
 	/** The IK Rig asset to run when this op is executed. */
-	UPROPERTY(EditAnywhere, Category = "IK Rig Asset", meta=(ReinitializeOnEdit))
+	UPROPERTY(EditAnywhere, BlueprintReadWrite, Category = "IK Rig Asset", meta=(ReinitializeOnEdit))
 	TObjectPtr<const UIKRigDefinition> IKRigAsset = nullptr;
 
 	/** Goals in this list will be excluded from the rig */
-	UPROPERTY(EditAnywhere, Category = "IK Rig Asset")
+	UPROPERTY(EditAnywhere, BlueprintReadWrite, Category = "IK Rig Asset")
 	TArray<FName> ExcludedGoals;
 	
 	// Draw IK goal locations. 
@@ -47,7 +47,7 @@ struct FRunIKRigOpGoalDebugData
 };
 #endif
 
-USTRUCT(BlueprintType, meta = (DisplayName = "Solve IK Goals"))
+USTRUCT(BlueprintType, meta = (DisplayName = "Run IK Rig"))
 struct IKRIG_API FIKRetargetRunIKRigOp : public FIKRetargetOpBase
 {
 	GENERATED_BODY()
@@ -73,6 +73,8 @@ struct IKRIG_API FIKRetargetRunIKRigOp : public FIKRetargetOpBase
 	virtual const UScriptStruct* GetType() const override;
 
 	virtual const UClass* GetControllerType() const override;
+	
+	virtual bool CanHaveChildOps() const override { return true; };
 
 	UPROPERTY()
 	FIKRetargetRunIKRigOpSettings Settings;
diff --git a/Engine/Plugins/Animation/IKRig/Source/IKRig/Public/Retargeter/RetargetOps/SpeedPlantingOp.h b/Engine/Plugins/Animation/IKRig/Source/IKRig/Public/Retargeter/RetargetOps/SpeedPlantingOp.h
--- a/Engine/Plugins/Animation/IKRig/Source/IKRig/Public/Retargeter/RetargetOps/SpeedPlantingOp.h
+++ b/Engine/Plugins/Animation/IKRig/Source/IKRig/Public/Retargeter/RetargetOps/SpeedPlantingOp.h
@@ -100,6 +100,8 @@ struct IKRIG_API FIKRetargetSpeedPlantingOp : public FIKRetargetOpBase
 	
 	virtual const UClass* GetControllerType() const override;
 
+	virtual const UScriptStruct* GetParentOpType() const override;
+
 	/** return array of names of all the speed curves referenced by this op. */
 	TArray<FName> GetRequiredSpeedCurves() const;
 
diff --git a/Engine/Plugins/Animation/IKRig/Source/IKRig/Public/Retargeter/RetargetOps/StrideWarpingOp.h b/Engine/Plugins/Animation/IKRig/Source/IKRig/Public/Retargeter/RetargetOps/StrideWarpingOp.h
--- a/Engine/Plugins/Animation/IKRig/Source/IKRig/Public/Retargeter/RetargetOps/StrideWarpingOp.h
+++ b/Engine/Plugins/Animation/IKRig/Source/IKRig/Public/Retargeter/RetargetOps/StrideWarpingOp.h
@@ -115,6 +115,8 @@ struct IKRIG_API FIKRetargetStrideWarpingOp : public FIKRetargetOpBase
 	virtual const UScriptStruct* GetType() const override;
 
 	virtual const UClass* GetControllerType() const override;
+	
+	virtual const UScriptStruct* GetParentOpType() const override;
 
 	UPROPERTY()
 	FIKRetargetStrideWarpingOpSettings Settings;
diff --git a/Engine/Plugins/Animation/IKRig/Source/IKRigEditor/Private/RetargetEditor/IKRetargetEditorController.cpp b/Engine/Plugins/Animation/IKRig/Source/IKRigEditor/Private/RetargetEditor/IKRetargetEditorController.cpp
--- a/Engine/Plugins/Animation/IKRig/Source/IKRigEditor/Private/RetargetEditor/IKRetargetEditorController.cpp
+++ b/Engine/Plugins/Animation/IKRig/Source/IKRigEditor/Private/RetargetEditor/IKRetargetEditorController.cpp
@@ -161,7 +161,7 @@ void FIKRetargetEditorController::Initialize(TSharedPtr<FIKRetargetEditor> InEdi
 
 	Selection.SelectedBoneNames.Add(ERetargetSourceOrTarget::Source);
 	Selection.SelectedBoneNames.Add(ERetargetSourceOrTarget::Target);
-	Selection.LastSelectedItem = ERetargetSelectionType::NONE;
+	Selection.LastSelectedType = ERetargetSelectionType::NONE;
 
 	// clean the asset before editing
 	AssetController->CleanAsset();
@@ -353,7 +353,7 @@ void FIKRetargetEditorController::UpdateSkeletalMeshComponents() const
 bool FIKRetargetEditorController::GetCameraTargetForSelection(FSphere& OutTarget) const
 {
 	// center the view on the last selected item
-	switch (GetLastSelectedItemType())
+	switch (Selection.LastSelectedType)
 	{
 	case ERetargetSelectionType::BONE:
 		{
@@ -453,6 +453,7 @@ bool FIKRetargetEditorController::GetCameraTargetForSelection(FSphere& OutTarget
 		}
 	case ERetargetSelectionType::ROOT:
 	case ERetargetSelectionType::MESH:
+	case ERetargetSelectionType::OP:
 	case ERetargetSelectionType::NONE:
 	default:
 		// frame both meshes
@@ -670,8 +671,8 @@ void FIKRetargetEditorController::RefreshDetailsView() const
 	// refresh the details panel, cannot assume tab is not closed
 	if (DetailsView.IsValid())
 	{
-		int32 SelectedOp = OpStackView->GetSelectedItemIndex();
-		ShowDetailsForOp(SelectedOp);
+		const int32 LastSelectedOpIndex = AssetController->GetIndexOfOpByName(Selection.LastSelectedOpName);
+		ShowDetailsForOp(LastSelectedOpIndex);
 	}
 }
 
@@ -735,7 +736,7 @@ void FIKRetargetEditorController::SetDetailsObjects(const TArray<UObject*>& Deta
 
 void FIKRetargetEditorController::ShowDetailsForOp(const int32 OpIndex) const
 {
-	if (AssetController->GetRetargetOpAtIndex(OpIndex) == nullptr)
+	if (AssetController->GetOpAtIndex(OpIndex) == nullptr)
 	{
 		DetailsView->SetObject(AssetController->GetAsset());
 		return;
@@ -745,7 +746,7 @@ void FIKRetargetEditorController::ShowDetailsForOp(const int32 OpIndex) const
 	auto MemoryProvider = [Asset, OpIndex]() -> uint8*
 	{
 		UIKRetargeterController* Controller = UIKRetargeterController::GetController(Asset);
-		if (FIKRetargetOpBase* OpToEdit = Controller->GetRetargetOpAtIndex(OpIndex))
+		if (FIKRetargetOpBase* OpToEdit = Controller->GetOpAtIndex(OpIndex))
 		{
 			return reinterpret_cast<uint8*>(OpToEdit->GetSettings());
 		}
@@ -806,7 +807,7 @@ void FIKRetargetEditorController::SetSourceOrTargetMode(ERetargetSourceOrTarget
 	}
 
 	// make sure details panel updates with selected bone on OTHER skeleton
-	if (Selection.LastSelectedItem == ERetargetSelectionType::BONE)
+	if (Selection.LastSelectedType == ERetargetSelectionType::BONE)
 	{
 		const TArray<FName>& SelectedBones = GetSelectedBones();
 		EditBoneSelection(SelectedBones, ESelectionEdit::Replace);
@@ -828,7 +829,7 @@ void FIKRetargetEditorController::EditBoneSelection(
 		return;
 	}
 
-	Selection.LastSelectedItem = ERetargetSelectionType::BONE;
+	Selection.LastSelectedType = ERetargetSelectionType::BONE;
 	
 	SetRootSelected(false);
 	
@@ -897,7 +898,7 @@ void FIKRetargetEditorController::EditChainSelection(
 	// deselect others
 	SetRootSelected(false);
 
-	Selection.LastSelectedItem = ERetargetSelectionType::CHAIN;
+	Selection.LastSelectedType = ERetargetSelectionType::CHAIN;
 	
 	// update selection set based on edit mode
 	switch (EditMode)
@@ -953,7 +954,7 @@ void FIKRetargetEditorController::SetRootSelected(const bool bIsSelected)
 		return;
 	}
 
-	Selection.LastSelectedItem = ERetargetSelectionType::ROOT;
+	Selection.LastSelectedType = ERetargetSelectionType::ROOT;
 }
 
 void FIKRetargetEditorController::CleanSelection(ERetargetSourceOrTarget SourceOrTarget)
@@ -998,20 +999,36 @@ void FIKRetargetEditorController::ClearSelection(const bool bKeepBoneSelection)
 		Selection.SelectedBoneNames[ERetargetSourceOrTarget::Target].Reset();
 	}
 
-	Selection.LastSelectedItem = ERetargetSelectionType::NONE;
+	Selection.LastSelectedType = ERetargetSelectionType::NONE;
 
 	RefreshDetailsView();
 }
 
+void FIKRetargetEditorController::SetOpSelected(const int32 InOpIndex)
+{
+	const FName NameOfOp = AssetController->GetOpName(InOpIndex);
+	if (!ensureAlways(NameOfOp != NAME_None))
+	{
+		return;
+	}
+	Selection.LastSelectedType = ERetargetSelectionType::OP;
+	Selection.LastSelectedOpName = NameOfOp;
+	ShowDetailsForOp(InOpIndex);
+}
+
+FName FIKRetargetEditorController::GetLastSelectedOpName() const
+{
+	return Selection.LastSelectedOpName;
+}
+
 FIKRetargetOpBase* FIKRetargetEditorController::GetSelectedOp() const
 {
 	if (!OpStackView.IsValid())
 	{
 		return nullptr;
 	}
-
-	const int32 SelectedOpIndex = OpStackView->GetSelectedItemIndex();
-	return AssetController->GetRetargetOpAtIndex(SelectedOpIndex);
+	
+	return AssetController->GetRetargetOpByName(Selection.LastSelectedOpName);
 }
 
 int32 FIKRetargetEditorController::GetSelectedOpIndex() const
@@ -1020,8 +1037,8 @@ int32 FIKRetargetEditorController::GetSelectedOpIndex() const
 	{
 		return INDEX_NONE;
 	}
-	
-	return OpStackView->GetSelectedItemIndex();
+
+	return AssetController->GetIndexOfOpByName(Selection.LastSelectedOpName);
 }
 
 void FIKRetargetEditorController::SetRetargeterMode(ERetargeterOutputMode Mode)
@@ -1416,8 +1433,7 @@ bool FIKRetargetEditorController::CanRenamePose() const
 {
 	// cannot rename default pose
 	const bool bNotUsingDefaultPose = AssetController->GetCurrentRetargetPoseName(GetSourceOrTarget()) != UIKRetargeter::GetDefaultPoseName();
-	// cannot rename pose while editing
-	return bNotUsingDefaultPose && !IsEditingPose();
+	return bNotUsingDefaultPose;
 }
 
 void FIKRetargetEditorController::AddReferencedObjects(FReferenceCollector& Collector)
diff --git a/Engine/Plugins/Animation/IKRig/Source/IKRigEditor/Private/RetargetEditor/IKRetargeterController.cpp b/Engine/Plugins/Animation/IKRig/Source/IKRigEditor/Private/RetargetEditor/IKRetargeterController.cpp
--- a/Engine/Plugins/Animation/IKRig/Source/IKRigEditor/Private/RetargetEditor/IKRetargeterController.cpp
+++ b/Engine/Plugins/Animation/IKRig/Source/IKRigEditor/Private/RetargetEditor/IKRetargeterController.cpp
@@ -4,12 +4,10 @@
 
 #include "ScopedTransaction.h"
 #include "Algo/LevenshteinDistance.h"
-#include "Engine/AssetManager.h"
 #include "Engine/SkeletalMesh.h"
 #include "Retargeter/IKRetargeter.h"
 #include "RigEditor/IKRigController.h"
 #include "IKRigEditor.h"
-#include "Components/Widget.h"
 #include "RetargetEditor/IKRetargeterPoseGenerator.h"
 #include "Retargeter/IKRetargetOps.h"
 #include "Retargeter/RetargetOps/FKChainsOp.h"
@@ -61,11 +59,6 @@ UIKRetargeter* UIKRetargeterController::GetAsset() const
 	return Asset;
 }
 
-TObjectPtr<UIKRetargeter>& UIKRetargeterController::GetAssetPtr()
-{
-	return Asset;
-}
-
 void UIKRetargeterController::CleanAsset() const
 {
 	FScopeLock Lock(&ControllerLock);
@@ -174,12 +167,8 @@ int32 UIKRetargeterController::AddRetargetOp(const FString InIKRetargetOpType, c
 		UE_LOG(LogIKRigEditor, Warning, TEXT("Specified retarget op type was not found, %s."), *InIKRetargetOpType);
 		return INDEX_NONE;
 	}
-	
-	FScopedTransaction Transaction(LOCTEXT("AddRetargetOp_Label", "Add Retarget Op"));
-	FScopedReinitializeIKRetargeter Reinitialize(this);
-	Asset->Modify();
 
-	return AddRetargetOp(OpType, bAllowDuplicates);
+	return AddRetargetOp(OpType, NAME_None, bAllowDuplicates, bRunInitialSetup);
 }
 
 bool UIKRetargeterController::RemoveRetargetOp(const int32 OpIndex) const
@@ -195,13 +184,25 @@ bool UIKRetargeterController::RemoveRetargetOp(const int32 OpIndex) const
 	FScopedTransaction Transaction(LOCTEXT("RemoveRetargetOp_Label", "Remove Retarget Op"));
 	FScopedReinitializeIKRetargeter Reinitialize(this);
 	Asset->Modify();
-	Asset->RetargetOps.RemoveAt(OpIndex);
+
+	TArray<int32> IndicesToRemove = GetChildOpIndices(OpIndex);
+	IndicesToRemove.Add(OpIndex);
+	Algo::Reverse(IndicesToRemove); // high to low (children are always before parent)
+	for (const int32 OpToRemove : IndicesToRemove)
+	{
+		Asset->RetargetOps.RemoveAt(OpToRemove);
+	}
+	
 	return true;
 }
 
 bool UIKRetargeterController::RemoveAllOps() const
 {
 	check(Asset)
+	if (Asset->RetargetOps.IsEmpty())
+	{
+		return false;
+	}
 
 	FScopedTransaction Transaction(LOCTEXT("RemoveAllRetargetOps_Label", "Remove All Retarget Ops"));
 	FScopedReinitializeIKRetargeter Reinitialize(this);
@@ -210,6 +211,101 @@ bool UIKRetargeterController::RemoveAllOps() const
 	return true;
 }
 
+FName UIKRetargeterController::SetOpName(const FName InName, const int32 InOpIndex) const
+{
+	FIKRetargetOpBase* Op = GetOpAtIndex(InOpIndex);
+	if (!Op)
+	{
+		return NAME_None;
+	}
+	
+	FScopedTransaction Transaction(LOCTEXT("SetOpName_Label", "Rename Retarget Op"));
+	FScopedReinitializeIKRetargeter Reinitialize(this);
+	Asset->Modify();
+
+	const FName OldOpName = GetOpName(InOpIndex);
+	FName UniqueName = GetUniqueOpName(InName, InOpIndex);
+	Op->SetName(UniqueName);
+	
+	// update any children pointing at the old name
+	if (OldOpName != NAME_None)
+	{
+		for (FInstancedStruct& OpStruct : Asset->RetargetOps)
+		{
+			FIKRetargetOpBase& OtherOp = OpStruct.GetMutable<FIKRetargetOpBase>();
+			if (OtherOp.GetParentOpName() == OldOpName)
+			{
+				OtherOp.SetParentOpName(UniqueName);
+			}
+		}
+	}
+	
+	return UniqueName;
+}
+
+FName UIKRetargeterController::GetOpName(const int32 InOpIndex) const
+{
+	const FIKRetargetOpBase* Op = GetOpAtIndex(InOpIndex);
+	if (!Op)
+	{
+		return NAME_None;
+	}
+	
+	return Op->GetName();
+}
+
+bool UIKRetargeterController::SetParentOpByName(
+	const FName InChildOpName,
+	const FName InParentOpName) const
+{
+	FIKRetargetOpBase* ChildOp = GetRetargetOpByName(InChildOpName);
+	if (!ChildOp)
+	{
+		return false; // child not found
+	}
+
+	FIKRetargetOpBase* ParentOp = GetRetargetOpByName(InParentOpName);
+	if (!ParentOp)
+	{
+		return false; // parent not found
+	}
+
+	if (ChildOp->GetParentOpType() != ParentOp->GetType())
+	{
+		return false; // wrong type of parent
+	}
+
+	FScopedTransaction Transaction(LOCTEXT("SetOpParent_Label", "Set Op Parent"));
+	FScopedReinitializeIKRetargeter Reinitialize(this);
+	Asset->Modify();
+	
+	ChildOp->SetParentOpName(InParentOpName);
+	
+	Asset->CleanOpStack();
+	
+	return true;
+}
+
+FName UIKRetargeterController::GetParentOpByName(const FName InOpName) const
+{
+	FIKRetargetOpBase* Op = GetRetargetOpByName(InOpName);
+	if (!Op)
+	{
+		return NAME_None; // child not found
+	}
+	
+	return Op->GetParentOpName();
+}
+
+int32 UIKRetargeterController::GetIndexOfOpByName(const FName InOpName) const
+{
+	return Asset->RetargetOps.IndexOfByPredicate([&](const FInstancedStruct& InOpStruct)
+	{
+		const FIKRetargetOpBase* Op = InOpStruct.GetPtr<FIKRetargetOpBase>();
+		return Op->GetName().IsEqual(InOpName, ENameCase::IgnoreCase); 
+	});
+}
+
 void UIKRetargeterController::AddDefaultOps(const bool bRunInitialSetup) const
 {
 	FScopedTransaction Transaction(LOCTEXT("AddDefaultOps_Label", "Add Default Ops"));
@@ -218,16 +314,16 @@ void UIKRetargeterController::AddDefaultOps(const bool bRunInitialSetup) const
 
 	// add set of default ops for basic retargeting
 	constexpr bool bAllowDuplicateOps = false;
-	AddRetargetOp(FIKRetargetPelvisMotionOp::StaticStruct(), bAllowDuplicateOps, bRunInitialSetup);
-	AddRetargetOp(FIKRetargetFKChainsOp::StaticStruct(), bAllowDuplicateOps, bRunInitialSetup);
-	AddRetargetOp(FIKRetargetIKChainsOp::StaticStruct(), bAllowDuplicateOps, bRunInitialSetup);
-	AddRetargetOp(FIKRetargetRunIKRigOp::StaticStruct(), bAllowDuplicateOps, bRunInitialSetup);
-	AddRetargetOp(FIKRetargetRootMotionOp::StaticStruct(), bAllowDuplicateOps, bRunInitialSetup);
+	AddRetargetOp(FIKRetargetPelvisMotionOp::StaticStruct(), NAME_None, bAllowDuplicateOps, bRunInitialSetup);
+	AddRetargetOp(FIKRetargetFKChainsOp::StaticStruct(), NAME_None,bAllowDuplicateOps, bRunInitialSetup);
+	AddRetargetOp(FIKRetargetIKChainsOp::StaticStruct(), NAME_None,bAllowDuplicateOps, bRunInitialSetup);
+	AddRetargetOp(FIKRetargetRunIKRigOp::StaticStruct(), NAME_None,bAllowDuplicateOps, bRunInitialSetup);
+	AddRetargetOp(FIKRetargetRootMotionOp::StaticStruct(), NAME_None,bAllowDuplicateOps, bRunInitialSetup);
 }
 
 void UIKRetargeterController::RunOpInitialSetup(const int32 InOpIndex) const
 {
-	FIKRetargetOpBase* Op = GetRetargetOpAtIndex(InOpIndex);
+	FIKRetargetOpBase* Op = GetOpAtIndex(InOpIndex);
 	if (!Op)
 	{
 		UE_LOG(LogIKRigEditor, Warning, TEXT("Retarget Op initial setup skipped. Invalid index, %d."), InOpIndex);
@@ -240,18 +336,104 @@ void UIKRetargeterController::RunOpInitialSetup(const int32 InOpIndex) const
 	Op->RunInitialSetup(GetAsset());
 }
 
-FIKRetargetOpBase* UIKRetargeterController::GetRetargetOpAtIndex(int32 Index) const
+FIKRetargetOpBase* UIKRetargeterController::GetOpAtIndex(const int32 InOpIndex) const
 {
 	check(Asset)
 
-	if (Asset->RetargetOps.IsValidIndex(Index))
+	if (Asset->RetargetOps.IsValidIndex(InOpIndex))
 	{
-		return Asset->RetargetOps[Index].GetMutablePtr<FIKRetargetOpBase>();
+		return Asset->RetargetOps[InOpIndex].GetMutablePtr<FIKRetargetOpBase>();
 	}
 	
 	return nullptr;
 }
 
+FIKRetargetOpBase* UIKRetargeterController::GetRetargetOpByName(const FName InOpName) const
+{
+	return GetOpAtIndex(GetIndexOfOpByName(InOpName));
+}
+
+TArray<int32> UIKRetargeterController::GetChildOpIndices(const int32 InOpIndex) const
+{
+	check(Asset)
+	
+	if (!Asset->RetargetOps.IsValidIndex(InOpIndex))
+	{
+		return TArray<int32>{};
+	}
+
+	const FName InOpName = GetOpName(InOpIndex);
+
+	TArray<int32> ChildrenOpIndices;
+	for (int32 OtherOpIndex = 0; OtherOpIndex < Asset->RetargetOps.Num(); OtherOpIndex++)
+	{
+		const FIKRetargetOpBase* OtherOp = Asset->RetargetOps[OtherOpIndex].GetPtr<FIKRetargetOpBase>();
+		if (OtherOp->GetParentOpName() == InOpName)
+		{
+			ChildrenOpIndices.Add(OtherOpIndex);
+		}
+	}
+	
+	return MoveTemp(ChildrenOpIndices);
+}
+
+bool UIKRetargeterController::GetCanOpHaveChildren(const int32 InOpIndex) const
+{
+	if (FIKRetargetOpBase* Op = GetOpAtIndex(InOpIndex))
+	{
+		return Op->CanHaveChildOps();
+	}
+	return false;
+}
+
+int32 UIKRetargeterController::GetParentOpIndex(const int32 InOpIndex) const
+{
+	check(Asset)
+
+	if (FIKRetargetOpBase* Op = GetOpAtIndex(InOpIndex))
+	{
+		return GetIndexOfOpByName(Op->GetParentOpName());
+	}
+	
+	return INDEX_NONE;	
+}
+
+FName UIKRetargeterController::GetUniqueOpName(const FName InNameToMakeUnique, int32 InOpIndexToIgnore) const
+{
+	// check if any other op using name
+	auto OpNameInUse = [this, InOpIndexToIgnore](const FName InOpNameToCheck)
+	{
+		for (int32 OpIndex=0; OpIndex<Asset->RetargetOps.Num(); ++OpIndex)
+		{
+			if (OpIndex == InOpIndexToIgnore)
+			{
+				continue;
+			}
+		
+			const FIKRetargetOpBase* Op = Asset->RetargetOps[OpIndex].GetPtr<FIKRetargetOpBase>();
+			if (Op->GetName() == InOpNameToCheck)
+			{
+				return true;
+			}
+		}
+		return false;
+	};
+
+	if (!OpNameInUse(InNameToMakeUnique))
+	{
+		return InNameToMakeUnique;
+	}
+	
+	// keep concatenating an incremented integer suffix until name is unique
+	int32 Number = InNameToMakeUnique.GetNumber() + 1;
+	while(OpNameInUse(FName(InNameToMakeUnique, Number)))
+	{
+		Number++;
+	}
+
+	return FName(InNameToMakeUnique, Number);
+}
+
 int32 UIKRetargeterController::GetIndexOfRetargetOp(FIKRetargetOpBase* RetargetOp) const
 {
 	check(Asset)
@@ -292,11 +474,8 @@ int32 UIKRetargeterController::GetNumRetargetOps() const
 
 bool UIKRetargeterController::MoveRetargetOpInStack(int32 OpToMoveIndex, int32 TargetIndex) const
 {
-	// adjust TargetIndex if necessary (since RemoveAt shifts elements)
-	if (OpToMoveIndex < TargetIndex)
-	{
-		TargetIndex--;  // compensate for the removed element shifting indices down
-	}
+	// ensure target is within range
+	TargetIndex = FMath::Clamp(TargetIndex, 0, Asset->RetargetOps.Num() - 1);
 	
 	if (!Asset->RetargetOps.IsValidIndex(OpToMoveIndex))
 	{
@@ -322,13 +501,14 @@ bool UIKRetargeterController::MoveRetargetOpInStack(int32 OpToMoveIndex, int32 T
 
 	// extract the element to move
 	FInstancedStruct MovedOp = MoveTemp(Asset->RetargetOps[OpToMoveIndex]);
-
 	// Remove the element without shrinking the array
 	Asset->RetargetOps.RemoveAt(OpToMoveIndex, 1, EAllowShrinking::No);
-
 	// Insert at the corrected TargetIndex
 	Asset->RetargetOps.Insert(MoveTemp(MovedOp), TargetIndex);
 
+	// enforce ordering constraints
+	Asset->CleanOpStack();
+	
 	return true;
 }
 
@@ -343,7 +523,7 @@ bool UIKRetargeterController::SetRetargetOpEnabled(int32 RetargetOpIndex, bool b
 	FScopedTransaction Transaction(LOCTEXT("SetRetargetOpEnabled_Label", "Enable/Disable Op"));
 	FScopedReinitializeIKRetargeter Reinitialize(this);
 	Asset->Modify();
-	GetRetargetOpAtIndex(RetargetOpIndex)->SetEnabled(bIsEnabled);
+	GetOpAtIndex(RetargetOpIndex)->SetEnabled(bIsEnabled);
 	return true;
 }
 
@@ -355,12 +535,12 @@ bool UIKRetargeterController::GetRetargetOpEnabled(int32 RetargetOpIndex) const
 		return false;
 	}
 
-	return GetRetargetOpAtIndex(RetargetOpIndex)->IsEnabled();
+	return GetOpAtIndex(RetargetOpIndex)->IsEnabled();
 }
 
 UIKRetargetOpControllerBase* UIKRetargeterController::GetOpController(int32 OpIndex)
 {
-	if (FIKRetargetOpBase* OpAtIndex = GetRetargetOpAtIndex(OpIndex))
+	if (FIKRetargetOpBase* OpAtIndex = GetOpAtIndex(OpIndex))
 	{
 		return OpAtIndex->GetController(this);
 	}
@@ -397,7 +577,11 @@ bool UIKRetargeterController::ChainHasNonDefaultSettings(const FName InTargetCha
 	return false;
 }
 
-int32 UIKRetargeterController::AddRetargetOp(UScriptStruct* InRetargetOpType, const bool bAllowDuplicates, const bool bRunInitialSetup) const
+int32 UIKRetargeterController::AddRetargetOp(
+	const UScriptStruct* InRetargetOpType,
+	const FName InParentOpName,
+	const bool bAllowDuplicates,
+	const bool bRunInitialSetup) const
 {
 	check(Asset)
 
@@ -426,15 +610,32 @@ int32 UIKRetargeterController::AddRetargetOp(UScriptStruct* InRetargetOpType, co
 		}
 	}
 
+	FScopedTransaction Transaction(LOCTEXT("AddRetargetOp_Label", "Add Retarget Op"));
+	FScopedReinitializeIKRetargeter Reinitialize(this);
+	Asset->Modify();
+
 	// add the op
-	const int32 NewOpIndex = Asset->RetargetOps.Emplace(InRetargetOpType);
+	int32 NewOpIndex = Asset->RetargetOps.Emplace(InRetargetOpType);
+
+	// give a unique default name
+	FIKRetargetOpBase* NewOp = GetOpAtIndex(NewOpIndex);
+	const FName NewOpName = GetUniqueOpName(NewOp->GetDefaultName(), NewOpIndex);
+	NewOp->SetName(NewOpName);
+
+	// assign to parent
+	NewOp->SetParentOpName(InParentOpName);
+
+	// enforce correct execution order
+	Asset->CleanOpStack();
 
 	// optionally run the initial setup
 	if (bRunInitialSetup)
 	{
-		FIKRetargetOpBase* NewOp = GetRetargetOpAtIndex(NewOpIndex);
 		NewOp->RunInitialSetup(GetAsset());
 	}
+
+	// in the unlikely event that cleaning the op stack reshuffled the indices, get the index again
+	NewOpIndex = GetIndexOfOpByName(NewOpName);
 	
 	return NewOpIndex;
 }
@@ -464,18 +665,6 @@ FName UIKRetargeterController::GetRetargetRootBone(const ERetargetSourceOrTarget
 	return IKRig ? IKRig->GetRetargetRoot() : FName("None");
 }
 
-void UIKRetargeterController::GetChainNames(const ERetargetSourceOrTarget SourceOrTarget, TArray<FName>& OutNames) const
-{
-	if (const UIKRigDefinition* IKRig = GetIKRig(SourceOrTarget))
-	{
-		const TArray<FBoneChain>& Chains = IKRig->GetRetargetChains();
-		for (const FBoneChain& Chain : Chains)
-		{
-			OutNames.Add(Chain.ChainName);
-		}
-	}
-}
-
 void UIKRetargeterController::CleanPoseList(const ERetargetSourceOrTarget SourceOrTarget) const
 {
 	const UIKRigDefinition* IKRig = GetIKRig(SourceOrTarget);
@@ -518,6 +707,18 @@ UIKRigStructViewer* UIKRetargeterController::GetStructViewer()
 
 void UIKRetargeterController::AutoMapChains(const EAutoMapChainType AutoMapType, const bool bForceRemap) const
 {
+	auto GetChainNames = [this](const ERetargetSourceOrTarget SourceOrTarget, TArray<FName>& OutNames)
+	{
+		if (const UIKRigDefinition* IKRig = GetIKRig(SourceOrTarget))
+		{
+			const TArray<FBoneChain>& Chains = IKRig->GetRetargetChains();
+			for (const FBoneChain& Chain : Chains)
+			{
+				OutNames.Add(Chain.ChainName);
+			}
+		}
+	};
+	
 	FScopeLock Lock(&ControllerLock);
 	FScopedTransaction Transaction(LOCTEXT("AutoMapRetargetChains", "Auto-Map Retarget Chains"));
 	FScopedReinitializeIKRetargeter Reinitialize(this);
diff --git a/Engine/Plugins/Animation/IKRig/Source/IKRigEditor/Private/RetargetEditor/IKRetargeterPoseGenerator.cpp b/Engine/Plugins/Animation/IKRig/Source/IKRigEditor/Private/RetargetEditor/IKRetargeterPoseGenerator.cpp
--- a/Engine/Plugins/Animation/IKRig/Source/IKRigEditor/Private/RetargetEditor/IKRetargeterPoseGenerator.cpp
+++ b/Engine/Plugins/Animation/IKRig/Source/IKRigEditor/Private/RetargetEditor/IKRetargeterPoseGenerator.cpp
@@ -608,7 +608,7 @@ bool FRetargetAutoPoseGenerator::CheckReadyToAlignBones() const
 	// we have to initialize a processor to "resolve" the retarget asset setup onto actual skeletons to do the alignment
 	USkeletalMesh* SourceSkeletalMesh = Controller->GetPreviewMesh(ERetargetSourceOrTarget::Source);
 	USkeletalMesh* TargetSkeletalMesh = Controller->GetPreviewMesh(ERetargetSourceOrTarget::Target);
-	UIKRetargeter* RetargeterAsset = Controller->GetAssetPtr();
+	UIKRetargeter* RetargeterAsset = Controller->GetAsset();
 	constexpr bool bSuppressWarnings = true;
 	FRetargetProfile RetargetProfile;
 	RetargetProfile.FillProfileWithAssetSettings(RetargeterAsset);
diff --git a/Engine/Plugins/Animation/IKRig/Source/IKRigEditor/Private/RetargetEditor/SRetargetOpStack.cpp b/Engine/Plugins/Animation/IKRig/Source/IKRigEditor/Private/RetargetEditor/SRetargetOpStack.cpp
--- a/Engine/Plugins/Animation/IKRig/Source/IKRigEditor/Private/RetargetEditor/SRetargetOpStack.cpp
+++ b/Engine/Plugins/Animation/IKRig/Source/IKRigEditor/Private/RetargetEditor/SRetargetOpStack.cpp
@@ -5,160 +5,211 @@
 #include "UObject/UObjectIterator.h"
 #include "SPositiveActionButton.h"
 #include "RetargetEditor/IKRetargetEditorController.h"
+#include "RetargetEditor/IKRetargetEditorStyle.h"
 #include "Retargeter/IKRetargetOps.h"
 #include "RigEditor/IKRigEditorStyle.h"
+#include "Widgets/Text/SInlineEditableTextBlock.h"
 
 #define LOCTEXT_NAMESPACE "SRetargetOpStack"
 
+float SRetargetOpItem::OpHorizontalPadding = 6.0f;;
+float SRetargetOpItem::OpVerticalPadding = 3.0f;
+
 TSharedRef<ITableRow> FRetargetOpStackElement::MakeListRowWidget(
 	const TSharedRef<STableViewBase>& InOwnerTable,
 	TSharedRef<FRetargetOpStackElement> InStackElement,
-	TSharedPtr<SRetargetOpStack> InOpStackWidget)
+	TSharedPtr<SRetargetOpList> InOpListWidget)
 {
-	return SNew(SRetargetOpStackItem, InOwnerTable, InStackElement, InOpStackWidget);
+	if (InStackElement.Get().bCanHaveChildren)
+	{
+		return SNew(SParentRetargetOpItem, InOwnerTable)
+			.InStackElement(InStackElement)
+			.InOpListWidget(InOpListWidget);	
+	}
+	
+	return SNew(SRetargetOpSingleItem, InOwnerTable)
+			.InStackElement(InStackElement)
+			.InOpListWidget(InOpListWidget);
 }
 
-void SRetargetOpStackItem::Construct(
-	const FArguments& InArgs,
-	const TSharedRef<STableViewBase>& OwnerTable,
-	TSharedRef<FRetargetOpStackElement> InStackElement,
-	TSharedPtr<SRetargetOpStack> InOpStackWidget)
+FName FRetargetOpStackElement::GetName() const
 {
-	StackElement = InStackElement;
-	OpStackWidget = InOpStackWidget;
-	
-	STableRow<TSharedPtr<FRetargetOpStackElement>>::Construct(
-        STableRow<TSharedPtr<FRetargetOpStackElement>>::FArguments()
-        .OnDragDetected(InOpStackWidget.Get(), &SRetargetOpStack::OnDragDetected)
-        .OnCanAcceptDrop(InOpStackWidget.Get(), &SRetargetOpStack::OnCanAcceptDrop)
-        .OnAcceptDrop(InOpStackWidget.Get(), &SRetargetOpStack::OnAcceptDrop)
-        .Content()
+	FIKRetargetEditorController* EditorController = OpListWidget.Pin()->EditorController.Pin().Get();
+	return EditorController->AssetController->GetOpName(IndexInStack);
+}
+
+const UScriptStruct* FRetargetOpStackElement::GetType() const
+{
+	const FIKRetargetEditorController* EditorController = OpListWidget.Pin()->EditorController.Pin().Get();
+	const FInstancedStruct* OpStruct = EditorController->AssetController->GetRetargetOpStructAtIndex(IndexInStack);
+	return OpStruct ? OpStruct->GetScriptStruct() : nullptr;
+}
+
+void SRetargetOpItem::Construct(const FArguments& InArgs)
+{
+	ListView = InArgs._InOpListWidget;
+    Element = InArgs._InStackElement;
+
+    ChildSlot
+    [
+        SNew(SHorizontalBox)
+	    
+        // drag Icon
+        + SHorizontalBox::Slot()
+        .MaxWidth(18)
+        .AutoWidth()
+        .HAlign(HAlign_Left)
+        .VAlign(VAlign_Center)
+        .Padding(FMargin(OpHorizontalPadding, OpVerticalPadding))
+        [
+            SNew(SImage)
+            .Image(FIKRigEditorStyle::Get().GetBrush("IKRig.DragSolver"))
+        ]
+
+        // enable Checkbox
+        + SHorizontalBox::Slot()
+        .AutoWidth()
+        .HAlign(HAlign_Left)
+        .VAlign(VAlign_Center)
+        .Padding(FMargin(OpHorizontalPadding, OpVerticalPadding))
         [
-            SNew(SHorizontalBox)
-            + SHorizontalBox::Slot()
-            .AutoWidth()
-            .HAlign(HAlign_Left)
-            .Padding(3)
+            SNew(SCheckBox)
+            .IsChecked_Lambda([this]() -> ECheckBoxState
+            {
+                FIKRetargetEditorController* EditorController = ListView.Pin()->EditorController.Pin().Get();
+            	bool bEnabled = EditorController->AssetController->GetRetargetOpEnabled(Element.Pin()->GetIndexInStack());
+                return bEnabled ? ECheckBoxState::Checked : ECheckBoxState::Unchecked;
+            })
+            .OnCheckStateChanged_Lambda([this](ECheckBoxState InCheckBoxState)
+            {
+            	FIKRetargetEditorController* EditorController = ListView.Pin()->EditorController.Pin().Get();
+            	const bool bIsChecked = InCheckBoxState == ECheckBoxState::Checked;
+            	EditorController->AssetController->SetRetargetOpEnabled( Element.Pin()->GetIndexInStack(), bIsChecked);
+				EditorController->ReinitializeRetargeterNoUIRefresh();
+            })
+        ]
+
+    	// display index
+		+ SHorizontalBox::Slot()
+		.AutoWidth()
+		.HAlign(HAlign_Left)
+		.VAlign(VAlign_Center)
+		.Padding(0, OpVerticalPadding)
+		[
+			SNew(STextBlock)
+			.Text(this, &SRetargetOpItem::GetIndexAsText)
+		]
+
+        // display Name
+        + SHorizontalBox::Slot()
+        .AutoWidth()
+        .HAlign(HAlign_Left)
+        .VAlign(VAlign_Center)
+        .Padding(FMargin(OpHorizontalPadding, OpVerticalPadding))
+        [
+			SAssignNew(EditNameWidget, SInlineEditableTextBlock)
+			.Text(this, &SRetargetOpItem::GetName)
+        	.OnVerifyTextChanged_Lambda([](const FText& InText, FText& OutErrorMessage)
+			{
+        		static FString IllegalNameCharacters = "^<>:\"/\\|?*";
+				return FName::IsValidXName(InText.ToString(), FString(IllegalNameCharacters), &OutErrorMessage);
+			})
+			.OnTextCommitted(this, &SRetargetOpItem::OnNameCommitted)
+			.MultiLine(false)
+        ]
+
+        // spacer
+        + SHorizontalBox::Slot()
+        .HAlign(HAlign_Fill)
+        .FillWidth(1.0f)
+        [
+            SNew(SSpacer)
+            .Size(FVector2D(0.0f, 1.0f))
+        ]
+
+    	// add sub-op button
+		+ SHorizontalBox::Slot()
+		.AutoWidth()
+		.HAlign(HAlign_Right)
+		.VAlign(VAlign_Center)
+		.Padding(FMargin(OpHorizontalPadding, OpVerticalPadding))
+		[
+			SNew(SPositiveActionButton)
+			.Visibility(Element.Pin()->GetCanHaveChildren() ? EVisibility::Visible : EVisibility::Collapsed)
+			.Icon(FAppStyle::Get().GetBrush("Icons.Plus"))
+			.ToolTipText(LOCTEXT("AddChildToolTip", "Add a child op to run before this op."))
+			.OnGetMenuContent_Lambda([this]()
+			{
+				return ListView.Pin()->CreateAddNewOpMenu();	
+			})
+		]
+
+        // spin Box
+        + SHorizontalBox::Slot()
+        .AutoWidth()
+        .HAlign(HAlign_Right)
+        .VAlign(VAlign_Center)
+        .Padding(FMargin(OpHorizontalPadding, OpVerticalPadding))
+        [
+            SNew(SSpinBox<double>)
+	        .Value_Lambda([this]() -> double
+	        {
+		        return 1.0f; // TODO add support for alpha value on op
+	        })
+	        .IsEnabled(false) // TODO enable this once alpha is supported
+            .MinValue(TOptional<double>())
+            .MaxValue(TOptional<double>())
+            .Style(&FAppStyle::Get().GetWidgetStyle<FSpinBoxStyle>("Sequencer.HyperlinkSpinBox"))
+        ]
+
+        // delete Button
+        + SHorizontalBox::Slot()
+        .AutoWidth()
+        .HAlign(HAlign_Right)
+        .VAlign(VAlign_Center)
+        .Padding(FMargin(OpHorizontalPadding, OpVerticalPadding))
+        [
+            SNew(SButton)
+            .ContentPadding(FMargin(0,0))
+            .ToolTipText(LOCTEXT("DeleteOp", "Delete retarget op and remove from stack."))
+            .OnClicked_Lambda([this]() -> FReply
+            {
+            	ListView.Pin()->DeleteRetargetOp(Element.Pin());
+                return FReply::Handled();
+            })
+            .Content()
             [
-            	SNew(SHorizontalBox)
-	            + SHorizontalBox::Slot()
-	            .MaxWidth(18)
-				.FillWidth(1.0)
-				.HAlign(HAlign_Left)
-				.VAlign(VAlign_Center)
-	            [
-	                SNew(SImage)
-	                .Image(FIKRigEditorStyle::Get().GetBrush("IKRig.DragSolver"))
-	            ]
-
-	            + SHorizontalBox::Slot()
-	            .AutoWidth()
-				.HAlign(HAlign_Left)
-				.VAlign(VAlign_Center)
-				.Padding(3.0f, 1.0f)
-				[
-					SNew(SCheckBox)
-					.IsChecked_Lambda([InOpStackWidget, InStackElement]() -> ECheckBoxState
-					{
-						bool bEnabled = true;
-						if (InOpStackWidget.IsValid() &&
-							InOpStackWidget->EditorController.IsValid() &&
-							InOpStackWidget->EditorController.Pin().IsValid())
-						{
-							const TObjectPtr<UIKRetargeterController> AssetController = InOpStackWidget->EditorController.Pin()->AssetController;
-							bEnabled = AssetController->GetRetargetOpEnabled(InStackElement->IndexInStack);
-						}
-						return bEnabled ? ECheckBoxState::Checked : ECheckBoxState::Unchecked;
-					})
-					.OnCheckStateChanged_Lambda([InOpStackWidget, InStackElement](ECheckBoxState InCheckBoxState)
-					{
-						if (InOpStackWidget.IsValid() &&
-							InOpStackWidget->EditorController.IsValid() &&
-							InOpStackWidget->EditorController.Pin().IsValid())
-						{
-							const bool bIsChecked = InCheckBoxState == ECheckBoxState::Checked;
-							InOpStackWidget->EditorController.Pin()->AssetController->SetRetargetOpEnabled(InStackElement->IndexInStack, bIsChecked);
-							InOpStackWidget->EditorController.Pin()->ReinitializeRetargeterNoUIRefresh();
-						}
-					})
-				]
-	     
-				+ SHorizontalBox::Slot()
-				.AutoWidth()
-				.HAlign(HAlign_Left)
-				.VAlign(VAlign_Center)
-				.Padding(3.0f, 1.0f)
-	            [
-					SNew(STextBlock)
-					.Text(InStackElement->DisplayName)
-					.IsEnabled_Lambda([this]()
-					{
-						return IsOpEnabled();
-					})
-					.TextStyle(FAppStyle::Get(), "NormalText.Important")
-	            ]
-
-	            + SHorizontalBox::Slot()
-	            .AutoWidth()
-				.HAlign(HAlign_Left)
-				.VAlign(VAlign_Center)
-				.Padding(3.0f, 1.0f)
-				[
-					SNew(STextBlock)
-					.Text_Lambda([this]()
-					{
-						FText Message;
-						GetWarningMessage(Message);
-						return Message;
-					})
-					.TextStyle(FAppStyle::Get(), "NormalText.Subdued")
-				]
+                SNew(SImage)
+                .Image(FAppStyle::Get().GetBrush("Icons.Delete"))
+                .ColorAndOpacity(FSlateColor::UseForeground())
             ]
+        ]
+    ];
 
-			+ SHorizontalBox::Slot()
-			.FillWidth(1)
-			.HAlign(HAlign_Right)
-			.VAlign(VAlign_Center)
-			.Padding(3)
-			[
-				SNew(SButton)
-				.ToolTipText(LOCTEXT("DeleteOp", "Delete retarget op and remove from stack."))
-				.OnClicked_Lambda([InOpStackWidget, InStackElement]() -> FReply
-				{
-					InOpStackWidget.Get()->DeleteRetargetOp(InStackElement);
-					return FReply::Handled();
-				})
-				.Content()
-				[
-					SNew(SImage)
-					.Image(FAppStyle::Get().GetBrush("Icons.Delete"))
-					.ColorAndOpacity(FSlateColor::UseForeground())
-				]
-			]
-        ], OwnerTable);	
+	// bind the rename callback to enter text editing mode when item is slow double-clicked
+	Element.Pin()->OnRenameRequested.BindSP(EditNameWidget.Get(), &SInlineEditableTextBlock::EnterEditingMode);
 }
 
-bool SRetargetOpStackItem::GetWarningMessage(FText& Message) const
+void SRetargetOpItem::OnNameCommitted(const FText& InText, ETextCommit::Type InCommitType) const
 {
-	const FIKRetargetProcessor* Processor = OpStackWidget.Pin()->EditorController.Pin()->GetRetargetProcessor();
-	if (!(Processor && Processor->IsInitialized()))
-	{
-		return false;
-	}
+	FName NewName = FName(InText.ToString());
+	FIKRetargetEditorController* EditorController = ListView.Pin()->EditorController.Pin().Get();
+	EditorController->AssetController->SetOpName(NewName, Element.Pin()->GetIndexInStack());
+	ListView.Pin()->RebuildList();
+}
 
-	const TArray<FInstancedStruct>& RetargetOps = Processor->GetRetargetOps();
-	const int32 OpIndex = StackElement.Pin()->IndexInStack;
-	if (RetargetOps.IsValidIndex(OpIndex))
-	{
-		const FIKRetargetOpBase* Op = RetargetOps[OpIndex].GetPtr<FIKRetargetOpBase>();
-		Message = Op->GetWarningMessage();
-		return true;
-	}
-	
-	return false;
+FText SRetargetOpItem::GetName() const
+{
+	FName OpName = ListView.Pin()->EditorController.Pin()->AssetController->GetOpName(Element.Pin()->GetIndexInStack());
+	return FText::FromName(OpName);
+}
+
+FText SRetargetOpItem::GetIndexAsText() const
+{
+	return FText::Format(FText::FromString(TEXT("{0}: ")), Element.Pin()->GetIndexInStack() + 1);
 }
 
-bool SRetargetOpStackItem::IsOpEnabled() const
+bool SRetargetOpItem::IsOpEnabled() const
 {
 	if (const FIKRetargetOpBase* Op = GetRetargetOp())
 	{
@@ -168,135 +219,258 @@ bool SRetargetOpStackItem::IsOpEnabled() const
 	return false;
 }
 
-FIKRetargetOpBase* SRetargetOpStackItem::GetRetargetOp() const
+FIKRetargetOpBase* SRetargetOpItem::GetRetargetOp() const
 {
-	if (!(StackElement.IsValid() && OpStackWidget.IsValid()))
+	if (!(Element.IsValid() && ListView.IsValid()))
 	{
 		return nullptr;
 	}
-	if (!OpStackWidget.Pin()->EditorController.IsValid())
+	FIKRetargetEditorController* EditorController = ListView.Pin()->EditorController.Pin().Get();
+	if (!EditorController)
 	{
 		return nullptr;
 	}
-	const int32 OpIndex = StackElement.Pin()->IndexInStack;
-	return OpStackWidget.Pin()->EditorController.Pin()->AssetController->GetRetargetOpAtIndex(OpIndex);
+	const int32 OpIndex = Element.Pin()->GetIndexInStack();
+	return EditorController->AssetController->GetOpAtIndex(OpIndex);
 }
 
-void SRetargetOpStack::Construct(
-	const FArguments& InArgs,
-	const TWeakPtr<FIKRetargetEditorController>& InEditorController)
+void SRetargetOpList::Construct(const FArguments& InArgs)
 {
-	EditorController = InEditorController;
-	EditorController.Pin()->SetOpStackView(SharedThis(this));
+	EditorController = InArgs._InEditorController;
+	ParentElement = InArgs._InParentElement;
+
+	CacheOpTypeMetaData();
 	
-	CommandList = MakeShared<FUICommandList>();
+	SListView::FArguments SuperArgs;
+	SuperArgs.ListItemsSource(&Elements);
+	SuperArgs.SelectionMode(ESelectionMode::Single);
+	SuperArgs.OnGenerateRow(this, &SRetargetOpList::MakeListRowWidget);
+	SuperArgs.IsEnabled(this, &SRetargetOpList::IsEnabled);
+	SuperArgs.OnSelectionChanged(this, &SRetargetOpList::OnSelectionChanged);
+	SuperArgs.OnMouseButtonClick(this, &SRetargetOpList::OnItemClicked);
+	
+	SListView::Construct(SuperArgs);
+}
 
-	ChildSlot
-	[
-		SNew(SVerticalBox)
-		+SVerticalBox::Slot()
-		.AutoHeight()
-		.VAlign(VAlign_Top)
-		.Padding(0.0f)
-		[
-			SNew(SBorder)
-			.Padding(0.0f)
-			.BorderImage(FAppStyle::GetBrush("DetailsView.CategoryTop"))
-			.BorderBackgroundColor(FLinearColor(0.6f, 0.6f, 0.6f, 1.0f))
-			[
-				SNew(SVerticalBox)
-				+ SVerticalBox::Slot()
-				.AutoHeight()
-				.VAlign(VAlign_Top)
-				[
-					SNew(SHorizontalBox)
-					+ SHorizontalBox::Slot()
-					.VAlign(VAlign_Center)
-					.HAlign(HAlign_Left)
-					.FillWidth(1.f)
-					.Padding(3.0f, 1.0f)
-					[
-						SNew(SPositiveActionButton)
-						.Icon(FAppStyle::Get().GetBrush("Icons.Plus"))
-						.Text(LOCTEXT("AddNewRetargetOpLabel", "Add New Retarget Op"))
-						.ToolTipText(LOCTEXT("AddNewToolTip", "Add a new operation to run as part of the retargeter."))
-						.IsEnabled(this, &SRetargetOpStack::IsAddOpEnabled)
-						.OnGetMenuContent(this, &SRetargetOpStack::CreateAddNewMenuWidget)
-					]
-				]
-			]
-		]
+void SRetargetOpList::CacheOpTypeMetaData()
+{
+	AllOpsMetaData.Reset();
 
-		+SVerticalBox::Slot()
-		.Padding(0.0f)
-		[
-			SAssignNew( ListView, SOpStackViewListType )
-			.SelectionMode(ESelectionMode::Single)
-			.IsEnabled(this, &SRetargetOpStack::IsAddOpEnabled)
-			.ListItemsSource( &ListViewItems )
-			.OnGenerateRow( this, &SRetargetOpStack::MakeListRowWidget )
-			.OnSelectionChanged(this, &SRetargetOpStack::OnSelectionChanged)
-		]
-	];
+	// instantiate all types of ops to cache their meta data
+	TArray<FInstancedStruct> TempInstancedOps;
+	const UStruct* BaseStruct = FIKRetargetOpBase::StaticStruct();
+	for (TObjectIterator<UStruct> StructIt; StructIt; ++StructIt)
+	{
+		UStruct* Struct = *StructIt;
+		if (!Struct->IsChildOf(BaseStruct) || Struct == BaseStruct)
+		{
+			continue;
+		}
 
-	RefreshStackView();
+		UScriptStruct* ScriptStruct = Cast<UScriptStruct>(Struct);
+		if (!ScriptStruct)
+		{
+			continue;
+		}
+		
+		TempInstancedOps.Emplace(ScriptStruct);
+	}
+
+	// what type of op is the parent of this list? (determines which type of children can be added)
+	// this is nullptr if we are the top level (top-level ops don't need a parent)
+	const UScriptStruct* ListType = ParentElement.IsValid() ? ParentElement.Pin()->GetType() : nullptr;
+
+	// store meta data for each type of op
+	for (const FInstancedStruct& TempOpStruct : TempInstancedOps)
+	{
+		const FIKRetargetOpBase* TempOp = TempOpStruct.GetPtr<FIKRetargetOpBase>();
+		const UScriptStruct* TempOpParentType = TempOp->GetParentOpType();
+
+		// filter ops from the sub-lists if they require a different parent type
+		if (ListType != nullptr && TempOpParentType != ListType)
+		{
+			continue;
+		}
+
+		// filter ops from main list if they require a special parent type
+		if (ListType == nullptr && TempOpParentType != nullptr)
+		{
+			continue;
+		}
+		
+		FIKRetargetOpMetaData OpMetaData;
+		OpMetaData.NiceName = TempOp->GetDefaultName();
+		OpMetaData.OpType = TempOpStruct.GetScriptStruct();
+		OpMetaData.ParentType = TempOp->GetParentOpType();
+		AllOpsMetaData.Add(OpMetaData);
+	}
 }
 
-int32 SRetargetOpStack::GetSelectedItemIndex() const
+TSharedRef<ITableRow> SRetargetOpList::MakeListRowWidget(
+	TSharedPtr<FRetargetOpStackElement> InElement,
+	const TSharedRef<STableViewBase>& OwnerTable)
 {
-	if (!ListView.IsValid())
+	return InElement->MakeListRowWidget(OwnerTable,InElement.ToSharedRef(),SharedThis(this));
+}
+
+FReply SRetargetOpList::OnKeyDown(const FGeometry& MyGeometry, const FKeyEvent& InKeyEvent)
+{
+	TArray<TSharedPtr<FRetargetOpStackElement>> SelectedOps = GetSelectedItems();
+	if (!SelectedOps.IsEmpty() && InKeyEvent.GetKey() == EKeys::Delete)
 	{
-		return INDEX_NONE;
+		DeleteRetargetOp(SelectedOps.Last());
+		return FReply::Handled();
 	}
 	
-	const TArray<TSharedPtr<FRetargetOpStackElement>>& SelectedItems = ListView.Get()->GetSelectedItems();
-	if (SelectedItems.IsEmpty())
+	return SListView::OnKeyDown(MyGeometry, InKeyEvent);
+}
+
+FReply SRetargetOpList::OnDragDetected(const FGeometry& MyGeometry, const FPointerEvent& MouseEvent)
+{
+	const TArray<TSharedPtr<FRetargetOpStackElement>> CurrentSelectedItems = GetSelectedItems();
+	if (CurrentSelectedItems.Num() != 1)
 	{
-		return INDEX_NONE;
+		return FReply::Unhandled();
+	}
+	
+	if (MouseEvent.IsMouseButtonDown(EKeys::LeftMouseButton))
+	{
+		const TSharedPtr<FRetargetOpStackElement> DraggedElement = CurrentSelectedItems[0];
+		const TSharedRef<FRetargetOpStackDragDropOp> DragDropOp = FRetargetOpStackDragDropOp::New(DraggedElement);
+		return FReply::Handled().BeginDragDrop(DragDropOp);
 	}
 
-	return SelectedItems[0]->IndexInStack;
+	return FReply::Unhandled();
 }
 
-TSharedRef<SWidget> SRetargetOpStack::CreateAddNewMenuWidget()
+TOptional<EItemDropZone> SRetargetOpList::OnCanAcceptDrop(
+	const FDragDropEvent& DragDropEvent,
+	EItemDropZone DropZone,
+	TSharedPtr<FRetargetOpStackElement> TargetElement)
 {
-	constexpr bool bCloseMenuAfterSelection = true;
-	FMenuBuilder MenuBuilder(bCloseMenuAfterSelection, nullptr);
+	TOptional<EItemDropZone> ReturnedDropZone;
+	const TSharedPtr<FRetargetOpStackDragDropOp> DragDropOp = DragDropEvent.GetOperationAs<FRetargetOpStackDragDropOp>();
+	if (!DragDropOp.IsValid())
+	{
+		return ReturnedDropZone;
+	}
 
-	MenuBuilder.BeginSection("AddNewRetargetOp", LOCTEXT("AddOperations", "Add New Retarget Op"));
+	const TSharedPtr<FRetargetOpStackElement>& DraggedElement = DragDropOp.Get()->Element.Pin();
+	if (!DraggedElement)
+	{
+		return ReturnedDropZone;
+	}
 
-	// add menu option to create each retarget op type
-	const UStruct* BaseStruct = FIKRetargetOpBase::StaticStruct();
-	for (TObjectIterator<UStruct> StructIt; StructIt; ++StructIt)
+	// if this is child element, it can only be re-ordered with siblings
+	if (const FRetargetOpStackElement* ParentOfDraggedItem = DraggedElement->GetParent())
 	{
-		UStruct* Struct = *StructIt;
-		if (!Struct->IsChildOf(BaseStruct) || Struct == BaseStruct)
+		// only allow dropping a child element on a sibling
+		if (!ParentOfDraggedItem->GetChildren().Contains(TargetElement.Get()))
 		{
-			continue;
+			return ReturnedDropZone;
 		}
-
-		if (UScriptStruct* ScriptStruct = Cast<UScriptStruct>(Struct))
+	}
+	else
+	{
+		// only allow dropping a top-level element on another top-level element
+		if (TargetElement->GetParent() != nullptr)
 		{
-			// try to get the "nice name" from metadata
-			FString TypeName = ScriptStruct->GetMetaData(TEXT("DisplayName"));
-		
-			// if no "DisplayName" metadata is found, fall back to the struct name
-			if (TypeName.IsEmpty())
-			{
-				TypeName = ScriptStruct->GetName();
-			}
-
-			FUIAction Action = FUIAction(FExecuteAction::CreateSP(this, &SRetargetOpStack::AddNewRetargetOp, ScriptStruct));
-			MenuBuilder.AddMenuEntry(FText::FromString(TypeName), FText::GetEmpty(), FSlateIcon(), Action);
+			return ReturnedDropZone;
 		}
 	}
+
+	// validate index to move to
+	const int32 IndexToMoveTo = FRetargetOpStackDragDropOp::GetIndexToMoveTo(DraggedElement, TargetElement, DropZone);
+
+	// GetIndexToMoveTo returns index none for invalid moves
+	if (IndexToMoveTo == INDEX_NONE)
+	{
+		return ReturnedDropZone;
+	}
 	
-	MenuBuilder.EndSection();
+	// all good!
+	ReturnedDropZone = DropZone;
+	return ReturnedDropZone;
+}
 
-	return MenuBuilder.MakeWidget();
+FReply SRetargetOpList::OnAcceptDrop(
+	const FDragDropEvent& DragDropEvent,
+	EItemDropZone DropZone,
+	TSharedPtr<FRetargetOpStackElement> TargetElement)
+{
+	const TSharedPtr<FRetargetOpStackDragDropOp> DragDropOp = DragDropEvent.GetOperationAs<FRetargetOpStackDragDropOp>();
+	if (!DragDropOp.IsValid())
+	{
+		return FReply::Unhandled();
+	}
+	
+	const TSharedPtr<FIKRetargetEditorController> Controller = EditorController.Pin();
+	if (!Controller.IsValid())
+	{
+		return FReply::Handled();
+	}
+
+	const TSharedPtr<FRetargetOpStackElement>& DraggedElement = DragDropOp.Get()->Element.Pin();
+	const int32 IndexToMoveTo = FRetargetOpStackDragDropOp::GetIndexToMoveTo(DraggedElement, TargetElement, DropZone);
+	if (IndexToMoveTo == INDEX_NONE)
+	{
+		// don't do anything if the drop location is invalid
+		return FReply::Unhandled();
+	}
+	
+	const UIKRetargeterController* AssetController = Controller->AssetController;
+	const bool bWasReparented = AssetController->MoveRetargetOpInStack(DraggedElement->GetIndexInStack(),IndexToMoveTo);
+	if (bWasReparented)
+	{
+		RestoreSelectionAndRefreshList();
+	}
+	
+	return FReply::Handled();
+}
+
+FReply SRetargetOpList::OnFocusReceived(const FGeometry& MyGeometry, const FFocusEvent& InFocusEvent)
+{
+	LastClickCycles = FPlatformTime::Cycles();
+	return SListView::OnFocusReceived(MyGeometry, InFocusEvent);
+}
+
+void SRetargetOpList::OnSelectionChanged(TSharedPtr<FRetargetOpStackElement> InItem, ESelectInfo::Type SelectInfo) const
+{
+	if (!EditorController.IsValid())
+	{
+		return;
+	}
+	
+	if (SelectInfo != ESelectInfo::Direct && InItem.IsValid())
+	{
+		EditorController.Pin()->SetOpSelected(InItem->GetIndexInStack());
+		return;
+	}
+
+	EditorController.Pin()->SetDetailsObject(EditorController.Pin()->AssetController->GetAsset());
+}
+
+void SRetargetOpList::OnItemClicked(TSharedPtr<FRetargetOpStackElement> InItem)
+{
+	// to rename an item, you have to select it first, then click on it again within a time limit (slow double click)
+	const bool ClickedOnSameItem = LastSelectedElement.Pin().Get() == InItem.Get();
+	const uint32 CurrentCycles = FPlatformTime::Cycles();
+	const double SecondsPassed = static_cast<double>(CurrentCycles - LastClickCycles) * FPlatformTime::GetSecondsPerCycle();
+	if (ClickedOnSameItem && SecondsPassed > 0.25f && SecondsPassed < 0.75f)
+	{
+		RegisterActiveTimer(0.f,
+			FWidgetActiveTimerDelegate::CreateLambda([this](double, float)
+			{
+				RequestRenameSelectedOp();
+				return EActiveTimerReturnType::Stop;
+			}));
+	}
+
+	LastClickCycles = CurrentCycles;
+	LastSelectedElement = InItem;
 }
 
-bool SRetargetOpStack::IsAddOpEnabled() const
+bool SRetargetOpList::IsEnabled() const
 {
 	if (!EditorController.IsValid())
 	{
@@ -311,7 +485,83 @@ bool SRetargetOpStack::IsAddOpEnabled() const
 	return false;
 }
 
-void SRetargetOpStack::AddNewRetargetOp(UScriptStruct* ScriptStruct)
+void SRetargetOpList::RestoreSelectionAndRefreshList()
+{
+	auto GetElementOfLastSelectedOp = [this]() -> TSharedPtr<FRetargetOpStackElement>
+	{
+		TSharedPtr<FRetargetOpStackElement> ElementToSelect;
+
+		if (!EditorController.IsValid())
+		{
+			return TSharedPtr<FRetargetOpStackElement>();
+		}
+		
+		const FName LastSelectedOpName = EditorController.Pin()->GetLastSelectedOpName();
+		if (LastSelectedOpName == NAME_None)
+		{
+			return TSharedPtr<FRetargetOpStackElement>();
+		}
+
+		for (TSharedPtr<FRetargetOpStackElement>& Element : Elements)
+		{
+			if (Element->GetName() == LastSelectedOpName)
+			{
+				return Element;
+			}
+		}
+
+		return TSharedPtr<FRetargetOpStackElement>();
+	};
+
+	// restore selection to the last selected element
+	TSharedPtr<FRetargetOpStackElement> ElementToRestore = GetElementOfLastSelectedOp();
+	if (ElementToRestore.IsValid())
+	{
+		// restore selection
+		SetSelection(ElementToRestore);
+		// 
+		EditorController.Pin()->SetOpSelected(ElementToRestore->GetIndexInStack());
+	}
+	else
+	{
+		// clear selection otherwise
+		ClearSelection();
+	}
+		
+	// refresh the list
+	RequestListRefresh();
+}
+
+void SRetargetOpList::RequestRenameSelectedOp() const
+{
+	if (!LastSelectedElement.IsValid())
+	{
+		return;
+	}
+
+	LastSelectedElement.Pin()->OnRenameRequested.ExecuteIfBound();
+}
+
+TSharedRef<SWidget> SRetargetOpList::CreateAddNewOpMenu()
+{
+	constexpr bool bCloseMenuAfterSelection = true;
+	FMenuBuilder MenuBuilder(bCloseMenuAfterSelection, nullptr);
+
+	MenuBuilder.BeginSection("AddNewRetargetOp", LOCTEXT("AddOperations", "Add New Retarget Op"));
+
+	// add menu option to create each retarget op type
+	for (FIKRetargetOpMetaData& OpMetaData : AllOpsMetaData)
+	{
+		FUIAction Action = FUIAction(FExecuteAction::CreateSP(this, &SRetargetOpList::AddNewRetargetOp, OpMetaData.OpType));
+		MenuBuilder.AddMenuEntry(FText::FromName(OpMetaData.NiceName), FText::GetEmpty(), FSlateIcon(), Action);
+	}
+	
+	MenuBuilder.EndSection();
+
+	return MenuBuilder.MakeWidget();
+}
+
+void SRetargetOpList::AddNewRetargetOp(const UScriptStruct* ScriptStruct)
 {
 	if (!EditorController.IsValid())
 	{
@@ -323,60 +573,240 @@ void SRetargetOpStack::AddNewRetargetOp(UScriptStruct* ScriptStruct)
 	{
 		return;
 	}
+
+	// add a new op to the stack
+	constexpr bool bAllowDuplicates = true;
+	constexpr bool bRunInitialSetup = true;
+	const FName ParentOpName = ParentElement.IsValid() ? ParentElement.Pin()->GetName() : NAME_None;
+	int32 NewOpIndex = AssetController->AddRetargetOp(ScriptStruct, ParentOpName, bAllowDuplicates, bRunInitialSetup);
+
+	RestoreSelectionAndRefreshList();
+}
+
+void SRetargetOpList::DeleteRetargetOp(TSharedPtr<FRetargetOpStackElement> OpToDelete)
+{
+	if (!OpToDelete.IsValid())
+	{
+		return;
+	}
+	
+	EditorController.Pin()->AssetController->RemoveRetargetOp(OpToDelete->GetIndexInStack());
+	
+	RestoreSelectionAndRefreshList();
+}
+
+void SParentRetargetOpItem::Construct(const FArguments& InArgs, const TSharedRef<STableViewBase>& OwnerTable)
+{
+	ParentStackElement = InArgs._InStackElement;
+	OpListWidget = InArgs._InOpListWidget;
 	
-	int32 NewOpIndex = AssetController->AddRetargetOp(ScriptStruct);
+	// create the list view ahead of time so it can be passed to the op item
+	ChildrenListView = MakeShared<SRetargetOpList>();
+	ChildrenListView->Construct(
+	SRetargetOpList::FArguments()
+		.InEditorController(OpListWidget.Pin()->EditorController)
+		.InParentElement(ParentStackElement));
 	
-	// ops have a chance at initial setup when first added to a stack
-	FIKRetargetProcessor* Processor = EditorController.Pin()->GetRetargetProcessor();
-	FIKRetargetOpBase* NewOp = AssetController->GetRetargetOpAtIndex(NewOpIndex);
-	if (NewOp && Processor)
+	STableRow::Construct(
+		STableRow::FArguments()
+		.OnDragDetected(OpListWidget.Pin().Get(), &SRetargetOpList::OnDragDetected)
+		.OnCanAcceptDrop(OpListWidget.Pin().Get(), &SRetargetOpList::OnCanAcceptDrop)
+		.OnAcceptDrop(OpListWidget.Pin().Get(), &SRetargetOpList::OnAcceptDrop)
+		.Padding(FMargin(6.f, 3.f))
+		.Content()
+		[
+			SNew(SBorder)
+			.BorderImage(FIKRetargetEditorStyle::Get().GetBrush("IKRetarget.OpBorder"))
+			.Padding(FMargin(SRetargetOpItem::OpHorizontalPadding, SRetargetOpItem::OpVerticalPadding))
+			[
+				SNew(SVerticalBox)
+				+SVerticalBox::Slot()
+				.AutoHeight()
+				.VAlign(VAlign_Top)
+				.Padding(0.0f)
+				[
+					SNew(SRetargetOpItem)
+					.InOpListWidget(ChildrenListView)
+					.InStackElement(ParentStackElement)
+				]
+
+				+SVerticalBox::Slot()
+				.Padding(5.0f)
+				[
+					SNew(SBorder)
+					.Visibility_Lambda([this]()
+					{
+						return ChildrenListView->Elements.IsEmpty() ? EVisibility::Collapsed : EVisibility::Visible;
+					})
+					.BorderImage(FIKRetargetEditorStyle::Get().GetBrush("IKRetarget.OpGroupBorder"))
+					.Padding(0.f, SRetargetOpItem::OpVerticalPadding)
+					[
+						ChildrenListView.ToSharedRef()
+					]
+				]
+			]
+			
+		], OwnerTable);
+
+	RefreshListView();
+}
+
+void SParentRetargetOpItem::RefreshListView() const
+{
+	const FIKRetargetEditorController* Controller = OpListWidget.Pin()->EditorController.Pin().Get();
+	if (!Controller)
 	{
-		NewOp->RunInitialSetup(AssetController->GetAsset());
-		Processor->SetNeedsInitialized();
+		return; 
 	}
 
-	RefreshStackView();
+	// generate all list elements
+	ChildrenListView->Elements.Reset();
+	UIKRetargeterController* AssetController = Controller->AssetController;
+	TArray<int32> ChildIndices = AssetController->GetChildOpIndices(ParentStackElement.Pin()->GetIndexInStack());
+	for (int32 ChildOpIndex : ChildIndices)
+	{
+		// create a new op element for the child
+		constexpr bool bCanOpHaveChildren = false; // we only allow 1 level of nesting
+		TSharedPtr<FRetargetOpStackElement> ChildElement = FRetargetOpStackElement::Make(ChildOpIndex, bCanOpHaveChildren, ChildrenListView);
+		ChildrenListView->Elements.Add(ChildElement);
+		
+		// store pointer to parent element on child
+		ChildElement->SetParent(ParentStackElement.Pin().Get());
+		
+		// store pointer to child element on parent
+		ParentStackElement.Pin()->AddChild(ChildElement.Get());
+	}
+
+	ChildrenListView->RestoreSelectionAndRefreshList();
 }
 
-void SRetargetOpStack::DeleteRetargetOp(TSharedPtr<FRetargetOpStackElement> OpToDelete)
+void SRetargetOpSingleItem::Construct(const FArguments& InArgs, const TSharedRef<STableViewBase>& OwnerTable)
 {
-	const TSharedPtr<FIKRetargetEditorController> Controller = EditorController.Pin();
-	if (!Controller.IsValid())
+	StackElement = InArgs._InStackElement;
+	OpListWidget = InArgs._InOpListWidget;
+	
+	STableRow::Construct(
+        STableRow::FArguments()
+        .OnDragDetected(OpListWidget.Pin().Get(), &SRetargetOpList::OnDragDetected)
+        .OnCanAcceptDrop(OpListWidget.Pin().Get(), &SRetargetOpList::OnCanAcceptDrop)
+        .OnAcceptDrop(OpListWidget.Pin().Get(), &SRetargetOpList::OnAcceptDrop)
+        .Padding(FMargin(6.f, 3.f))
+        .Content()
+        [
+        	SNew(SBorder)
+			.BorderImage(FIKRetargetEditorStyle::Get().GetBrush("IKRetarget.OpBorder"))
+			.Padding(FMargin(SRetargetOpItem::OpHorizontalPadding, SRetargetOpItem::OpVerticalPadding))
+			[
+				SNew(SRetargetOpItem)
+				.InOpListWidget(OpListWidget)
+				.InStackElement(StackElement)
+			]
+        ], OwnerTable);	
+}
+
+bool SRetargetOpSingleItem::GetWarningMessage(FText& Message) const
+{
+	const FIKRetargetProcessor* Processor = OpListWidget.Pin()->EditorController.Pin()->GetRetargetProcessor();
+	if (!(Processor && Processor->IsInitialized()))
 	{
-		return; 
+		return false;
+	}
+
+	const TArray<FInstancedStruct>& RetargetOps = Processor->GetRetargetOps();
+	const int32 OpIndex = StackElement.Pin()->GetIndexInStack();
+	if (RetargetOps.IsValidIndex(OpIndex))
+	{
+		const FIKRetargetOpBase* Op = RetargetOps[OpIndex].GetPtr<FIKRetargetOpBase>();
+		Message = Op->GetWarningMessage();
+		return true;
 	}
 	
-	if (!OpToDelete.IsValid())
+	return false;
+}
+
+bool SRetargetOpSingleItem::IsOpEnabled() const
+{
+	if (const FIKRetargetOpBase* Op = GetRetargetOp())
 	{
-		return;
+		return Op->IsEnabled();
 	}
 	
-	Controller->AssetController->RemoveRetargetOp(OpToDelete->IndexInStack);
-	RefreshStackView();
+	return false;
 }
 
-void SRetargetOpStack::OnSelectionChanged(
-	TSharedPtr<FRetargetOpStackElement> InItem,
-	ESelectInfo::Type SelectInfo)
+FIKRetargetOpBase* SRetargetOpSingleItem::GetRetargetOp() const
 {
-	if (SelectInfo != ESelectInfo::Direct && InItem.IsValid())
+	if (!(StackElement.IsValid() && OpListWidget.IsValid()))
 	{
-		EditorController.Pin()->LastSelectedOpIndex = InItem->IndexInStack;
-		EditorController.Pin()->ShowDetailsForOp(InItem->IndexInStack);
+		return nullptr;
 	}
+	FIKRetargetEditorController* EditorController = OpListWidget.Pin()->EditorController.Pin().Get();
+	if (!EditorController)
+	{
+		return nullptr;
+	}
+	const int32 OpIndex = StackElement.Pin()->GetIndexInStack();
+	return EditorController->AssetController->GetOpAtIndex(OpIndex);
 }
 
-TSharedRef<ITableRow> SRetargetOpStack::MakeListRowWidget(
-	TSharedPtr<FRetargetOpStackElement> InElement,
-	const TSharedRef<STableViewBase>& OwnerTable)
+void SRetargetOpStack::Construct(
+	const FArguments& InArgs,
+	const TWeakPtr<FIKRetargetEditorController>& InEditorController)
 {
-	return InElement->MakeListRowWidget(
-		OwnerTable,
-		InElement.ToSharedRef(),
-		SharedThis(this));
+	EditorController = InEditorController;
+	EditorController.Pin()->SetOpStackView(SharedThis(this));
+
+	ChildSlot
+	[
+		SNew(SVerticalBox)
+		+SVerticalBox::Slot()
+		.AutoHeight()
+		.VAlign(VAlign_Top)
+		.Padding(0.0f)
+		[
+			SNew(SBorder)
+			.Padding(0.0f)
+			.BorderImage(FAppStyle::GetBrush("DetailsView.CategoryTop"))
+			.BorderBackgroundColor(FLinearColor(0.6f, 0.6f, 0.6f, 1.0f))
+			[
+				SNew(SVerticalBox)
+				+ SVerticalBox::Slot()
+				.AutoHeight()
+				.VAlign(VAlign_Top)
+				[
+					SNew(SHorizontalBox)
+					+ SHorizontalBox::Slot()
+					.VAlign(VAlign_Center)
+					.HAlign(HAlign_Left)
+					.FillWidth(1.f)
+					.Padding(6.0f, 4.0f)
+					[
+						SNew(SPositiveActionButton)
+						.Icon(FAppStyle::Get().GetBrush("Icons.Plus"))
+						.Text(LOCTEXT("AddNewRetargetOpLabel", "Add New Retarget Op"))
+						.ToolTipText(LOCTEXT("AddNewToolTip", "Add a new operation to run as part of the retargeter."))
+						.OnGetMenuContent_Lambda([this]()
+						{
+							return ListView->CreateAddNewOpMenu();	
+						})
+					]
+				]
+			]
+		]
+
+		+SVerticalBox::Slot()
+		.Padding(0.0f)
+		[
+			SAssignNew( ListView, SRetargetOpList )
+			.InEditorController(EditorController)
+			.InParentElement(nullptr)
+		]
+	];
+
+	RefreshStackView();
 }
 
-void SRetargetOpStack::RefreshStackView()
+void SRetargetOpStack::RefreshStackView() const
 {
 	const TSharedPtr<FIKRetargetEditorController> Controller = EditorController.Pin();
 	if (!Controller.IsValid())
@@ -384,36 +814,30 @@ void SRetargetOpStack::RefreshStackView()
 		return; 
 	}
 
-	// generate all list items
-	ListViewItems.Reset();
+	// empty old elements
+	ListView->Elements.Reset();
+
+	// rebuild list of top level elements
 	UIKRetargeterController* AssetController = Controller->AssetController;
 	const int32 NumOps = AssetController->GetNumRetargetOps();
 	for (int32 OpIndex=0; OpIndex<NumOps; ++OpIndex)
 	{
-		const FIKRetargetOpBase* Op = AssetController->GetRetargetOpAtIndex(OpIndex);
-		FString UniqueName =  FString::FromInt(OpIndex+1) + " - " + Op->GetNiceName().ToString();
-		TSharedPtr<FRetargetOpStackElement> StackElement = FRetargetOpStackElement::Make(FText::FromString(UniqueName), OpIndex);
-		ListViewItems.Add(StackElement);
-	}
+		// skip ops that are children of another op
+		const int32 ParentOpIndex = AssetController->GetParentOpIndex(OpIndex);
+		if (ParentOpIndex != INDEX_NONE)
+		{
+			continue;
+		}
 
-	if (NumOps && ListViewItems.IsValidIndex(EditorController.Pin()->LastSelectedOpIndex))
-	{
-		// restore selection
-		ListView->SetSelection(ListViewItems[EditorController.Pin()->LastSelectedOpIndex]);
-	}
-	else
-	{
-		// clear selection otherwise
-		ListView->ClearSelection();
-	}
+		// is this a group with children
+		const bool bCanOpHaveChildren = AssetController->GetCanOpHaveChildren(OpIndex);
 
-	ListView->RequestListRefresh();
-}
+		// add op to main stack
+		TSharedPtr<FRetargetOpStackElement> StackElement = FRetargetOpStackElement::Make(OpIndex, bCanOpHaveChildren, ListView);
+		ListView->Elements.Add(StackElement);
+	}
 
-FReply SRetargetOpStack::OnKeyDown(const FGeometry& MyGeometry, const FKeyEvent& InKeyEvent)
-{
-	// TODO delete here
-	return SCompoundWidget::OnKeyDown(MyGeometry, InKeyEvent);
+	ListView->RestoreSelectionAndRefreshList();
 }
 
 TSharedRef<FRetargetOpStackDragDropOp> FRetargetOpStackDragDropOp::New(TWeakPtr<FRetargetOpStackElement> InElement)
@@ -431,71 +855,47 @@ TSharedPtr<SWidget> FRetargetOpStackDragDropOp::GetDefaultDecorator() const
 		.BorderImage(FAppStyle::GetBrush("Menu.Background"))
 		[
 			SNew(STextBlock)
-			.Text(Element.Pin()->DisplayName)
+			.Text(FText::FromName(Element.Pin()->GetName()))
 		];
 }
 
-FReply SRetargetOpStack::OnDragDetected(const FGeometry& MyGeometry, const FPointerEvent& MouseEvent)
-{
-	const TArray<TSharedPtr<FRetargetOpStackElement>> SelectedItems = ListView.Get()->GetSelectedItems();
-	if (SelectedItems.Num() != 1)
-	{
-		return FReply::Unhandled();
-	}
-	
-	if (MouseEvent.IsMouseButtonDown(EKeys::LeftMouseButton))
-	{
-		const TSharedPtr<FRetargetOpStackElement> DraggedElement = SelectedItems[0];
-		const TSharedRef<FRetargetOpStackDragDropOp> DragDropOp = FRetargetOpStackDragDropOp::New(DraggedElement);
-		return FReply::Handled().BeginDragDrop(DragDropOp);
-	}
-
-	return FReply::Unhandled();
-}
-
-TOptional<EItemDropZone> SRetargetOpStack::OnCanAcceptDrop(
-	const FDragDropEvent& DragDropEvent,
-	EItemDropZone DropZone,
-	TSharedPtr<FRetargetOpStackElement> TargetItem)
+int32 FRetargetOpStackDragDropOp::GetIndexToMoveTo(
+	const TSharedPtr<FRetargetOpStackElement>& InDraggedElement,
+	const TSharedPtr<FRetargetOpStackElement>& InTargetElement,
+	const EItemDropZone InDropZone)
 {
-	TOptional<EItemDropZone> ReturnedDropZone;
-	
-	const TSharedPtr<FRetargetOpStackDragDropOp> DragDropOp = DragDropEvent.GetOperationAs<FRetargetOpStackDragDropOp>();
-	if (DragDropOp.IsValid())
+	// disallow dropping on self
+	const int32 DraggedItemIndex = InDraggedElement->GetIndexInStack();
+	const int32 TargetItemIndex = InTargetElement->GetIndexInStack();
+	if (DraggedItemIndex == TargetItemIndex)
 	{
-		ReturnedDropZone = DropZone;
+		return INDEX_NONE;
 	}
-	
-	return ReturnedDropZone;
-}
 
-FReply SRetargetOpStack::OnAcceptDrop(
-	const FDragDropEvent& DragDropEvent,
-	EItemDropZone DropZone,
-	TSharedPtr<FRetargetOpStackElement> TargetItem)
-{
-	const TSharedPtr<FRetargetOpStackDragDropOp> DragDropOp = DragDropEvent.GetOperationAs<FRetargetOpStackDragDropOp>();
-	if (!DragDropOp.IsValid())
+	// disallow dropping in a place that would not change the order (like below the one above, or above the one below)
+	int32 IndexToMoveDragItemTo = INDEX_NONE;
+	bool bDraggingDown = DraggedItemIndex <= TargetItemIndex;
+	switch (InDropZone)
 	{
-		return FReply::Unhandled();
+	case EItemDropZone::AboveItem:
+		IndexToMoveDragItemTo = bDraggingDown ? TargetItemIndex - 1 : TargetItemIndex;
+		break;
+	case EItemDropZone::BelowItem:
+		IndexToMoveDragItemTo = bDraggingDown ? TargetItemIndex : TargetItemIndex + 1;
+		break;
+	case EItemDropZone::OntoItem:
+		IndexToMoveDragItemTo = TargetItemIndex;
+		break;
+	default:
+		checkNoEntry();
 	}
 	
-	const TSharedPtr<FIKRetargetEditorController> Controller = EditorController.Pin();
-	if (!Controller.IsValid())
+	if (DraggedItemIndex == IndexToMoveDragItemTo)
 	{
-		return FReply::Handled();
+		return INDEX_NONE;
 	}
 
-	const FRetargetOpStackElement& DraggedElement = *DragDropOp.Get()->Element.Pin().Get();
-	const UIKRetargeterController* AssetController = Controller->AssetController;
-	int32 TargetIndexOffset = DropZone == EItemDropZone::BelowItem ? 1 : -1;
-	const bool bWasReparented = AssetController->MoveRetargetOpInStack(DraggedElement.IndexInStack, TargetItem.Get()->IndexInStack + TargetIndexOffset);
-	if (bWasReparented)
-	{
-		RefreshStackView();
-	}
-	
-	return FReply::Handled();
+	return IndexToMoveDragItemTo;
 }
 
 #undef LOCTEXT_NAMESPACE
diff --git a/Engine/Plugins/Animation/IKRig/Source/IKRigEditor/Public/RetargetEditor/IKRetargetEditorController.h b/Engine/Plugins/Animation/IKRig/Source/IKRigEditor/Public/RetargetEditor/IKRetargetEditorController.h
--- a/Engine/Plugins/Animation/IKRig/Source/IKRigEditor/Public/RetargetEditor/IKRetargetEditorController.h
+++ b/Engine/Plugins/Animation/IKRig/Source/IKRigEditor/Public/RetargetEditor/IKRetargetEditorController.h
@@ -209,13 +209,19 @@ public:
 
 	// ------------------------- SELECTION -----------------------------
 
+	// access to all selection state for the editor
+	const FIKRetargetDebugDrawState& GetSelectionState() const { return Selection; };
+
 	// SELECTION - BONES (viewport or hierarchy view)
 	void EditBoneSelection(
 		const TArray<FName>& InBoneNames,
 		ESelectionEdit EditMode,
 		const bool bFromHierarchyView = false);
 	const TArray<FName>& GetSelectedBones() const {return Selection.SelectedBoneNames[CurrentlyEditingSourceOrTarget]; };
-	const FIKRetargetDebugDrawState& GetSelectionState() const { return Selection; };
+	void SetRootSelected(const bool bIsSelected);
+	bool GetRootSelected() const { return Selection.bIsRootSelected; };
+	bool IsEditingPoseWithAnyBoneSelected() const;
+	TArray<FName> GetSelectedBonesAndChildren() const;
 	// END bone selection
 
 	// SELECTION - CHAINS (viewport or chains view)
@@ -226,24 +232,19 @@ public:
 	const TArray<FName>& GetSelectedChains() const {return Selection.SelectedChains; };
 	// END chain selection
 
-	void SetRootSelected(const bool bIsSelected);
-	bool GetRootSelected() const { return Selection.bIsRootSelected; };
-
-	void CleanSelection(ERetargetSourceOrTarget SourceOrTarget);
-	void ClearSelection(const bool bKeepBoneSelection=false);
-	ERetargetSelectionType GetLastSelectedItemType() const { return Selection.LastSelectedItem; };
-
-	// op stack selection
+	// SELECTION - OP STACK
+	void SetOpSelected(const int32 InOpIndex);
+	FName GetLastSelectedOpName() const;
 	FIKRetargetOpBase* GetSelectedOp() const;
 	int32 GetSelectedOpIndex() const;
-	int32 LastSelectedOpIndex = 0;
+	// END op selection
 
+	// GENERAL SELECTION
+	void CleanSelection(ERetargetSourceOrTarget SourceOrTarget);
+	void ClearSelection(const bool bKeepBoneSelection=false);
 	// to frame selection when pressing "f" in viewport
 	bool GetCameraTargetForSelection(FSphere& OutTarget) const;
 
-	// check if any bone is selected
-	bool IsEditingPoseWithAnyBoneSelected() const;
-
 	// ------------------------- END SELECTION -----------------------------
 
 	// determine if bone in the specified skeleton is part of the retarget (in a mapped chain)
@@ -309,9 +310,7 @@ public:
 	void UpdateSkeletalMeshComponents() const;
 
 private:
-
-	TArray<FName> GetSelectedBonesAndChildren() const;
-
+	
 	// modal dialog to ask user if they want to fix root bones that are "on the ground"
 	bool PromptToFixRootHeight(ERetargetSourceOrTarget SourceOrTarget) const;
 	bool bAskedToFixRoot = false;
diff --git a/Engine/Plugins/Animation/IKRig/Source/IKRigEditor/Public/RetargetEditor/IKRetargetEditorStyle.h b/Engine/Plugins/Animation/IKRig/Source/IKRigEditor/Public/RetargetEditor/IKRetargetEditorStyle.h
--- a/Engine/Plugins/Animation/IKRig/Source/IKRigEditor/Public/RetargetEditor/IKRetargetEditorStyle.h
+++ b/Engine/Plugins/Animation/IKRig/Source/IKRigEditor/Public/RetargetEditor/IKRetargetEditorStyle.h
@@ -5,6 +5,8 @@
 #include "Styling/SlateStyle.h"
 #include "Styling/SlateStyleRegistry.h"
 #include "Styling/SlateStyleMacros.h"
+#include "Styling/StyleColors.h"
+#include "Styling/CoreStyle.h"
 
 class FIKRetargetEditorStyle final	: public FSlateStyleSet
 {
@@ -31,6 +33,13 @@ public:
 		Set("IKRetarget.ShowRetargetPose", new IMAGE_BRUSH_SVG("Slate/ShowRetargetPose", Icon64x64));
 		Set("IKRetarget.AutoAlign", new IMAGE_BRUSH_SVG("Slate/AutoRetargetPose", Icon16x16));
 
+		FSlateColor OpOutlineColor = FSlateColor(FLinearColor(0.1843f, 0.1843f, 0.1843f, 0.5f));
+		Set("IKRetarget.OpBorder", new FSlateRoundedBoxBrush(
+			FStyleColors::Header, CoreStyleConstants::InputFocusRadius,
+			OpOutlineColor, CoreStyleConstants::InputFocusThickness));
+
+		Set("IKRetarget.OpGroupBorder", new FSlateRoundedBoxBrush(FStyleColors::Recessed, CoreStyleConstants::InputFocusRadius));
+
 		SetContentRoot(FPaths::EngineContentDir() / TEXT("Editor/Slate"));
 		Set( "IKRetarget.Viewport.Border", new BOX_BRUSH( "Old/Window/ViewportDebugBorder", 0.8f, FLinearColor(1.0f,1.0f,1.0f,1.0f) ) );
 		
diff --git a/Engine/Plugins/Animation/IKRig/Source/IKRigEditor/Public/RetargetEditor/IKRetargeterController.h b/Engine/Plugins/Animation/IKRig/Source/IKRigEditor/Public/RetargetEditor/IKRetargeterController.h
--- a/Engine/Plugins/Animation/IKRig/Source/IKRigEditor/Public/RetargetEditor/IKRetargeterController.h
+++ b/Engine/Plugins/Animation/IKRig/Source/IKRigEditor/Public/RetargetEditor/IKRetargeterController.h
@@ -46,61 +46,50 @@ public:
 	virtual void PostInitProperties() override;
 	
 	// Get access to the retargeter asset.
-	// Warning: Do not make modifications to the asset directly. Use this API instead. 
+	// Warning: Do not make modifications to the asset directly. Using the controller API guarantees correctness.
 	UIKRetargeter* GetAsset() const;
-	TObjectPtr<UIKRetargeter>& GetAssetPtr();
   
 private:
 	
 	// The actual asset that this Controller modifies. This is the only field this class should have.
 	TObjectPtr<UIKRetargeter> Asset = nullptr;
 
+public:
+
 	//
 	// GENERAL PUBLIC/SCRIPTING API
 	//
-
-public:
 	
-	/** Use this to get the controller for the given retargeter asset */
+	/** Get the controller for the given retargeter asset
+	 * @param InRetargeterAsset an IK Retarget asset
+	 * @return the controller with an API for modifying the given retarget asset */
 	UFUNCTION(BlueprintCallable, Category="IK Retarget", meta = (BlueprintThreadSafe, DisplayName = "Get IK Retarget Controller"))
 	static UIKRetargeterController* GetController(const UIKRetargeter* InRetargeterAsset);
 	
-	// Set the IK Rig to use as the source or target (to copy animation FROM/TO) 
+	/** Set the IK Rig to use as the source or target (to copy animation FROM/TO)
+	 * @param SourceOrTarget an enum specifying either "Source" or "Target"
+	 * @param the IK Rig asset to apply as either the "Source" (to copy FROM) or "Target" (to copy TO)*/
 	UFUNCTION(BlueprintCallable, BlueprintPure=false, Category="IK Retarget", meta = (BlueprintThreadSafe, DisplayName = "Set IK Rig"))
 	void SetIKRig(const ERetargetSourceOrTarget SourceOrTarget, UIKRigDefinition* IKRig) const;
 	
-	// Get either source or target IK Rig 
+	/** Get either source or target IK Rig
+	 * @param SourceOrTarget an enum specifying either "Source" or "Target"
+	 * @return the IK Rig asset associated with either the Source or Target of the IK Retargeter */ 
 	UFUNCTION(BlueprintCallable, Category="IK Retarget", meta = (BlueprintThreadSafe, DisplayName = "Get IK Rig"))
 	const UIKRigDefinition* GetIKRig(const ERetargetSourceOrTarget SourceOrTarget) const;
 
-	// Set the preview skeletal mesh for either source or target
+	/** Set the preview skeletal mesh for either source or target
+	* @param SourceOrTarget an enum specifying either "Source" or "Target"
+	* @param InPreviewMesh a skeletal mesh asset to use as the preview mesh in the retarget editor (may be used on other meshes at runtime) */
 	UFUNCTION(BlueprintCallable, BlueprintPure=false, Category="IK Retarget", meta = (BlueprintThreadSafe))
 	void SetPreviewMesh(const ERetargetSourceOrTarget SourceOrTarget, USkeletalMesh* InPreviewMesh) const;
 	
-	// Get the preview skeletal mesh
+	/** Get the preview skeletal mesh
+	 * @param SourceOrTarget an enum specifying either "Source" or "Target"
+	 * @return the skeletal mesh asset currently being used as a preview mesh */
 	UFUNCTION(BlueprintCallable, Category="IK Retarget", meta = (BlueprintThreadSafe))
 	USkeletalMesh* GetPreviewMesh(const ERetargetSourceOrTarget SourceOrTarget) const;
 
-	//
-	// GET/SET SETTINGS PUBLIC/SCRIPTING API
-	//
-	
-	// TODO
-	
-	// Tell each op to reset any settings associated with this chain
-	UFUNCTION(BlueprintCallable, Category="IK Retarget")
-	void ResetChainSettingsInAllOps(const FName InTargetChainName) const;
-
-	//
-	// GET/SET SETTINGS C++ API
-	//
-
-	// get access to the mapping between source/target chains
-	const FRetargetChainMapping& GetChainMapping() const;
-	
-	// Ask if any op has settings for the given chain that are not at the default
-	bool ChainHasNonDefaultSettings(const FName InTargetChainName) const;
-
 	//
 	// RETARGET OPS PUBLIC/SCRIPTING API
 	//
@@ -115,14 +104,51 @@ public:
 		const bool bAllowDuplicates = true,
 		const bool bRunInitialSetup = true) const;
 
-	// Remove the retarget op at the given stack index. 
+	/** Remove the retarget op at the given stack index
+	 * NOTE: if this op is a parent, all it's children will be removed as well
+	 * @param InOpIndex: the index of the op to remove
+	 * @return true if op was found and removed, false otherwise */ 
 	UFUNCTION(BlueprintCallable, BlueprintPure=false, Category="IK Retarget")
-	bool RemoveRetargetOp(const int32 OpIndex) const;
+	bool RemoveRetargetOp(const int32 InOpIndex) const;
 
-	// Remove all ops in the stack.
+	/** Remove all ops in the stack.
+	 * @return true if any ops were removed, false otherwise */
 	UFUNCTION(BlueprintCallable, Category="IK Retarget")
 	bool RemoveAllOps() const;
 
+	/** Set the name of the op at the given index in the stack.
+	 * @param InName the new name to use
+	 * @param InOpIndex the index of the op to be renamed
+	 * @return the new name of the op (after it's made unique) */
+	UFUNCTION(BlueprintCallable, Category="IK Retarget")
+	FName SetOpName(const FName InName, const int32 InOpIndex) const;
+	
+	/** Get the name of the op at the given index in the stack.
+	 * @param InOpIndex the index of the op to get the name of
+	 * @return the name of the op or None if the index was invalid */
+	UFUNCTION(BlueprintCallable, Category="IK Retarget")
+	FName GetOpName(const int32 InOpIndex) const ;
+
+	/** Set the name of the op to parent this op to
+	 * Children ops are forced to execute before their parent
+	 * @param InChildOpName the op to be parented
+	 * @param InParentOpName the name of the op to parent to
+	 * @return true if both parent and child were found, and the parent relationship was allowed */
+	UFUNCTION(BlueprintCallable, Category="IK Retarget")
+	bool SetParentOpByName(const FName InChildOpName, const FName InParentOpName) const;
+	
+	/** Get the name of the parent op for the given op.
+	 * @param InOpName the name of the op to get the parent of
+	 * @return the name of the parent op or None if the op was not found or did not have a parent */
+	UFUNCTION(BlueprintCallable, Category="IK Retarget")
+	FName GetParentOpByName(const FName InOpName) const ;
+
+	/** Get the index of an op.
+	 * @param InOpName the name of the op to get the index of
+	 * @return the integer index of the op in the stack or -1 if op not found. */
+	UFUNCTION(BlueprintCallable, Category="IK Retarget")
+	int32 GetIndexOfOpByName(const FName InOpName) const;
+
 	/** Automatically add basic retargeting operations
 	 * Adds these ops in the following order: Pelvis Motion, FK Chains, IK Chains, IK Solve and Root Motion
 	 * If any of these ops are already present, they will not be re-added.
@@ -131,44 +157,69 @@ public:
 	void AddDefaultOps(const bool bRunInitialSetup=true) const;
 
 	/** Force the op to run the initial setup. This is normally run when an Op is added through the editor UI.
-	 * NOTE: ops may not have custom initial setup routines; the exact behavior is op-dependent */
+	 * NOTE: ops may not have custom initial setup routines; the exact behavior is op-dependent
+	 * @param InOpIndex the index of the op to run setup on */
 	UFUNCTION(BlueprintCallable, Category="IK Retarget", meta = (BlueprintThreadSafe))
 	void RunOpInitialSetup(const int32 InOpIndex) const;
 
-	// Get the number of Ops in the stack.
+	// Tell each op to reset any settings associated with this chain
+	UFUNCTION(BlueprintCallable, Category="IK Retarget")
+	void ResetChainSettingsInAllOps(const FName InTargetChainName) const;
+
+	/** Get the number of Ops in the stack.
+	 * @return int, the number of ops */
 	UFUNCTION(BlueprintCallable, Category="IK Retarget")
 	int32 GetNumRetargetOps() const;
 
-	// Move the retarget op at the given index to the target index. 
+	/** Move the retarget op at the given index to the target index.
+	 * NOTE: due to constraints on execution order, the actual index may differ from what is requested
+	 * @InOpToMoveIndex the index of the op to be moved
+	 * @InTargetIndex the index where the op should be moved to */ 
 	UFUNCTION(BlueprintCallable, BlueprintPure=false, Category="IK Retarget")
-	bool MoveRetargetOpInStack(int32 OpToMoveIndex, int32 TargetIndex) const;
+	bool MoveRetargetOpInStack(int32 InOpToMoveIndex, int32 InTargetIndex) const;
 
-	// Set enabled/disabled status of the given retarget operation. 
+	/** Toggle an op on/off.
+	 * @param InRetargetOpIndex the index of the op to modify
+	 * @param bIsEnabled if true, turns the op On, else Off
+	 * @return true if op was found at index */ 
 	UFUNCTION(BlueprintCallable, BlueprintPure=false, Category="IK Retarget")
-	bool SetRetargetOpEnabled(int32 RetargetOpIndex, bool bIsEnabled) const;
+	bool SetRetargetOpEnabled(int32 InRetargetOpIndex, bool bIsEnabled) const;
 
-	// Get enabled status of the given Op. 
+	/** Get enabled status of the given Op.
+	 * @param InRetargetOpIndex the index of the op to get the enabled state for
+	 * @return true if op is enabled, false if op is disabled or not found */
 	UFUNCTION(BlueprintCallable, Category="IK Retarget")
-	bool GetRetargetOpEnabled(int32 RetargetOpIndex) const;
+	bool GetRetargetOpEnabled(int32 InRetargetOpIndex) const;
 	
 	/** Get a controller to get/set the settings for a given op in the stack
-	 * NOTE: the returned UIKRetargetOpControllerBase* must be cast to the type specific to your op.
-	 * @param InOpIndex: the index of the op in the stack to return a controller for */
+	 * NOTE: the returned UIKRetargetOpControllerBase* should be cast to the type specific to your op to get full functionality for that op.
+	 * @param InOpIndex: the index of the op in the stack to return a controller for
+	 * @return a pointer to the base op controller UClass or null if the op was not found */
 	UFUNCTION(BlueprintCallable, Category="IK Rig")
 	UIKRetargetOpControllerBase* GetOpController(int32 InOpIndex);
 
 	//
-	// GENERAL C++ ONLY API
+	// RETARGET OPS C++ ONLY API
 	//
 	
 	// add a new retarget op to the stack
 	int32 AddRetargetOp(
-		UScriptStruct* InRetargetOpType,
+		const UScriptStruct* InRetargetOpType,
+		const FName InParentOpName = NAME_None,
 		const bool bAllowDuplicates = true,
 		const bool bRunInitialSetup = true) const;
-	
-	// Get access to the given retarget operation.
-	FIKRetargetOpBase* GetRetargetOpAtIndex(int32 Index) const;
+
+	// Gather up all the ops that are children of this op
+	TArray<int32> GetChildOpIndices(const int32 InOpIndex) const;
+
+	// Returns true if this op type can contain child ops (ops forced to be ordered before the parent)
+	bool GetCanOpHaveChildren(const int32 InOpIndex) const;
+
+	// Get the index of the parent op or INDEX_NONE if it's root level
+	int32 GetParentOpIndex(const int32 InOpIndex) const;
+
+	// Return NameToMakeUnique with a numbered suffix that makes it unique in the stack
+	FName GetUniqueOpName(const FName NameToMakeUnique, int32 OpIndexToIgnore) const;
 
 	// Get access to the first retarget operation of the given type
 	template <typename T>
@@ -177,12 +228,22 @@ public:
 		return Asset->GetFirstRetargetOpOfType<T>();
 	}
 
-	// Get access to the given retarget operation. 
-	int32 GetIndexOfRetargetOp(FIKRetargetOpBase* RetargetOp) const;
+	// Get access to the given retarget operation by index.
+	FIKRetargetOpBase* GetOpAtIndex(const int32 InOpIndex) const;
+
+	// Get access to the given retarget operation by name.
+	FIKRetargetOpBase* GetRetargetOpByName(const FName InOpName) const;
 
 	// Get access to the given retarget operation.
 	FInstancedStruct* GetRetargetOpStructAtIndex(int32 Index) const;
 
+	// Get the index of a given retarget operation. 
+	int32 GetIndexOfRetargetOp(FIKRetargetOpBase* RetargetOp) const;
+
+	//
+	// GENERAL C++ ONLY API
+	//
+	
 	// Ensures all internal data is compatible with assigned meshes and ready to edit.
 	// - Removes bones from retarget poses that are no longer in skeleton
 	// - Removes chain settings for chains that are no longer in target IK Rig
@@ -204,21 +265,33 @@ public:
 	// RETARGET CHAIN MAPPING PUBLIC/SCRIPTING API
 	//
 	
-	// Use fuzzy string search to find "best" Source chain to map to each Target chain
+	/** Use string comparision to find "best" Source chain to map to each Target chain or clear the mappings
+	 * @param AutoMapType an enum specifying the type of mapping to perform (ie, Exact, Fuzzy or Clear)
+	 * @param bForceRemap if false, will only remap those mappings that unset/None */
 	UFUNCTION(BlueprintCallable, BlueprintPure=false, Category="IK Retarget", meta = (BlueprintThreadSafe))
 	void AutoMapChains(const EAutoMapChainType AutoMapType, const bool bForceRemap) const;
 
-	// Assign a source chain to the given target chain. Animation will be copied from the source to the target.
+	/** Assign a source chain to the given target chain. Animation will be copied from the source to the target.
+	 * @param InSourceChainName the name of the source retarget chain to assign
+	 * @param InTargetChainName the name of the target chain to map */
 	UFUNCTION(BlueprintCallable, BlueprintPure=false, Category="IK Retarget", meta = (BlueprintThreadSafe))
-	bool SetSourceChain(FName SourceChainName, FName TargetChainName) const;
+	bool SetSourceChain(FName InSourceChainName, FName InTargetChainName) const;
 
-	// Get the name of the source chain mapped to a given target chain (the chain animation is copied FROM).
+	/** Get the name of the source chain mapped to a given target chain (the chain animation is copied FROM).
+	 * @param InTargetChainName the name of the target retarget chain to get a mapping for
+	 * @return the name of the source chain mapped to the given target chain (or None if not mapped) */
 	UFUNCTION(BlueprintCallable, Category="IK Retarget", meta = (BlueprintThreadSafe))
-	FName GetSourceChain(const FName& TargetChainName) const;
+	FName GetSourceChain(const FName& InTargetChainName) const;
 
 	//
 	// RETARGET CHAIN MAPPING C++ ONLY API
 	//
+
+	// get access to the mapping between source/target chains
+	const FRetargetChainMapping& GetChainMapping() const;
+	
+	// Ask if any op has settings for the given chain that are not at the default
+	bool ChainHasNonDefaultSettings(const FName InTargetChainName) const;
 	
 	// Get whether the given chain's IK goal is connected to a solver 
 	bool IsChainGoalConnectedToASolver(const FName& GoalName) const;
@@ -231,87 +304,115 @@ public:
 	
 	// Call this when IK Rig chain is removed. 
 	void HandleRetargetChainRemoved(UIKRigDefinition* IKRig, const FName& InChainRemoved) const;
-	
-private:
-	
-	// Get names of all the bone chains.
-	void GetChainNames(const ERetargetSourceOrTarget SourceOrTarget, TArray<FName>& OutNames) const;
-
-public:
 
 	//
 	// RETARGET POSE PUBLIC/SCRIPTING API
 	//
 	
-	// Add new retarget pose. Returns the name of the new retarget pose.
+	/** Add new retarget pose.
+	 * @param NewPoseName The name to assign to the new retarget pose
+	 * @param SourceOrTarget Specifies whether the pose applies to the source or target skeleton
+	 * @return The name of the newly created retarget pose */
 	UFUNCTION(BlueprintCallable, BlueprintPure=false, Category="IK Retarget", meta =(BlueprintThreadSafe))
 	FName CreateRetargetPose(const FName& NewPoseName, const ERetargetSourceOrTarget SourceOrTarget) const;
 
-	// Remove a retarget pose. Returns true if the pose was found and removed.
+	/** Remove a retarget pose.
+	 * @param PoseToRemove The name of the retarget pose to remove
+	 * @param SourceOrTarget Specifies whether the pose belongs to the source or target skeleton
+	 * @return True if the pose was found and successfully removed, false otherwise */
 	UFUNCTION(BlueprintCallable, BlueprintPure=false, Category="IK Retarget", meta =(BlueprintThreadSafe))
 	bool RemoveRetargetPose(const FName& PoseToRemove, const ERetargetSourceOrTarget SourceOrTarget) const;
 
-	// Duplicate a retarget pose. Returns the name of the new, duplicate pose (or NAME_None if PoseToDuplicate is not found).
+	/** Duplicate a retarget pose.
+	 * @param PoseToDuplicate The name of the retarget pose to duplicate
+	 * @param NewName The name to assign to the duplicated pose
+	 * @param SourceOrTarget Specifies whether the pose belongs to the source or target skeleton
+	 * @return The name of the new duplicated pose, or NAME_None if PoseToDuplicate is not found */
 	UFUNCTION(BlueprintCallable, BlueprintPure=false, Category="IK Retarget", meta =(BlueprintThreadSafe))
 	FName DuplicateRetargetPose(const FName PoseToDuplicate, const FName NewName, const ERetargetSourceOrTarget SourceOrTarget) const;
-	
-	// Rename current retarget pose. Returns true if a pose was renamed.
+
+	/** Rename current retarget pose.
+	 * @param OldPoseName The current name of the retarget pose to rename
+	 * @param NewPoseName The new name to assign to the retarget pose
+	 * @param SourceOrTarget Specifies whether the pose belongs to the source or target skeleton
+	 * @return True if the pose was found and successfully renamed, false otherwise */
 	UFUNCTION(BlueprintCallable, BlueprintPure=false, Category="IK Retarget", meta =(BlueprintThreadSafe))
 	bool RenameRetargetPose(const FName OldPoseName, const FName NewPoseName, const ERetargetSourceOrTarget SourceOrTarget) const;
 
-	// Change which retarget pose is used by the retargeter at runtime
+	/** Change which retarget pose is used by the retargeter at runtime.
+	 * @param CurrentPose The name of the retarget pose to set as current
+	 * @param SourceOrTarget Specifies whether the pose applies to the source or target skeleton
+	 * @return True if the current pose was successfully set, false otherwise */
 	UFUNCTION(BlueprintCallable, BlueprintPure=false, Category="IK Retarget", meta =(BlueprintThreadSafe))
 	bool SetCurrentRetargetPose(FName CurrentPose, const ERetargetSourceOrTarget SourceOrTarget) const;
-	
-	// Get the current retarget pose
+
+	/** Get the name of the current retarget pose.
+	 * @param SourceOrTarget Specifies whether to retrieve the pose from the source or target skeleton
+	 * @return The name of the current retarget pose */
 	UFUNCTION(BlueprintCallable, Category="IK Retarget", meta =(BlueprintThreadSafe))
-    FName GetCurrentRetargetPoseName(const ERetargetSourceOrTarget SourceOrTarget) const;
-	
-	// Get access to array of retarget poses
+	FName GetCurrentRetargetPoseName(const ERetargetSourceOrTarget SourceOrTarget) const;
+
+	/** Get access to array of retarget poses.
+	 * @param SourceOrTarget Specifies whether to retrieve poses from the source or target skeleton
+	 * @return A reference to the map of retarget pose names to their pose data */
 	UFUNCTION(BlueprintCallable, Category="IK Retarget", meta =(BlueprintThreadSafe))
 	TMap<FName, FIKRetargetPose>& GetRetargetPoses(const ERetargetSourceOrTarget SourceOrTarget) const;
-	
-	// Get the current retarget pose
+
+	/** Get the current retarget pose.
+	 * @param SourceOrTarget Specifies whether to retrieve the pose from the source or target skeleton
+	 * @return A reference to the current retarget pose data */
 	UFUNCTION(BlueprintCallable, Category="IK Retarget", meta =(BlueprintThreadSafe))
 	FIKRetargetPose& GetCurrentRetargetPose(const ERetargetSourceOrTarget SourceOrTarget) const;
 
-	// Reset a retarget pose for the specified bones.
-	// If BonesToReset is Empty, will removes all stored deltas, returning pose to reference pose
+	/** Reset a retarget pose for the specified bones.
+	 * @param PoseToReset The name of the retarget pose to reset
+	 * @param BonesToReset The array of bone names to reset; if empty, resets all bones to the reference pose
+	 * @param SourceOrTarget Specifies whether the pose belongs to the source or target skeleton */
 	UFUNCTION(BlueprintCallable, BlueprintPure=false, Category="IK Retarget", meta =(BlueprintThreadSafe))
 	void ResetRetargetPose(const FName& PoseToReset, const TArray<FName>& BonesToReset, const ERetargetSourceOrTarget SourceOrTarget) const;
-	
-	// Set a delta rotation for a given bone for the current retarget pose
+
+	/** Set a delta rotation for a given bone in the current retarget pose.
+	 * @param BoneName The name of the bone to apply the rotation offset to
+	 * @param RotationOffset The quaternion representing the rotation offset to apply
+	 * @param SkeletonMode Specifies whether the bone belongs to the source or target skeleton */
 	UFUNCTION(BlueprintCallable, BlueprintPure=false, Category="IK Retarget", meta =(BlueprintThreadSafe))
-	void SetRotationOffsetForRetargetPoseBone(
-		const FName& BoneName,
-		const FQuat& RotationOffset,
-		const ERetargetSourceOrTarget SkeletonMode) const;
-	
-	// Get a delta rotation for a given bone for the current retarget pose
+	void SetRotationOffsetForRetargetPoseBone(const FName& BoneName, const FQuat& RotationOffset, const ERetargetSourceOrTarget SkeletonMode) const;
+
+	/** Get a delta rotation for a given bone in the current retarget pose.
+	 * @param BoneName The name of the bone to retrieve the rotation offset for
+	 * @param SourceOrTarget Specifies whether the bone belongs to the source or target skeleton
+	 * @return The quaternion representing the rotation offset for the specified bone */
 	UFUNCTION(BlueprintCallable, BlueprintPure=false, Category="IK Retarget", meta =(BlueprintThreadSafe))
-	FQuat GetRotationOffsetForRetargetPoseBone(
-		const FName& BoneName,
-		const ERetargetSourceOrTarget SourceOrTarget) const;
-	
-	// Set the translation offset on the retarget root bone for the current retarget pose
+	FQuat GetRotationOffsetForRetargetPoseBone(const FName& BoneName, const ERetargetSourceOrTarget SourceOrTarget) const;
+
+	/** Set the translation offset on the retarget pelvis bone for the current retarget pose.
+	 * @param TranslationOffset The vector representing the translation offset to apply to the pelvis bone
+	 * @param SourceOrTarget Specifies whether the pelvis bone belongs to the source or target skeleton */
 	UFUNCTION(BlueprintCallable, BlueprintPure=false, Category="IK Retarget", meta =(BlueprintThreadSafe))
-	void SetRootOffsetInRetargetPose(
-		const FVector& TranslationOffset,
-		const ERetargetSourceOrTarget SourceOrTarget) const;
+	void SetRootOffsetInRetargetPose(const FVector& TranslationOffset, const ERetargetSourceOrTarget SourceOrTarget) const;
 
-	// Get the translation offset on the retarget root bone for the current retarget pose
+	/** Get the translation offset on the retarget pelvis bone for the current retarget pose.
+	 * @param SourceOrTarget Specifies whether the pelvis bone belongs to the source or target skeleton
+	 * @return The vector representing the translation offset of the pelvis bone */
 	UFUNCTION(BlueprintCallable, Category="IK Retarget", meta = (BlueprintThreadSafe))
 	FVector GetRootOffsetInRetargetPose(const ERetargetSourceOrTarget SourceOrTarget) const;
 
-	// Automatically align all bones in mapped chains and store in the current retarget pose.
+	/** Automatically align all bones in mapped chains and store in the current retarget pose.
+	 * @param SourceOrTarget Specifies whether to align bones in the source or target skeleton
+	 * @param Method The method to use for automatic alignment (defaults to ChainToChain) */
 	UFUNCTION(BlueprintCallable, Category="IK Retarget")
 	void AutoAlignAllBones(ERetargetSourceOrTarget SourceOrTarget, const ERetargetAutoAlignMethod Method = ERetargetAutoAlignMethod::ChainToChain) const;
 
-	// Automatically align an array of bones and store in the current retarget pose. Bones not in a mapped chain will be ignored.
+	/** Automatically align an array of bones and store in the current retarget pose.
+	 * @param BonesToAlign The array of bone names to align; bones not in mapped chains are ignored
+	 * @param Method The method to use for automatic alignment
+	 * @param SourceOrTarget Specifies whether to align bones in the source or target skeleton */
 	UFUNCTION(BlueprintCallable, Category="IK Retarget")
 	void AutoAlignBones(const TArray<FName>& BonesToAlign, const ERetargetAutoAlignMethod Method, ERetargetSourceOrTarget SourceOrTarget) const;
 
-	// Moves the entire skeleton vertically until the specified bone is the same height off the ground as in the reference pose.
+	/** Moves the entire skeleton vertically until the specified bone is the same height off the ground as in the reference pose.
+	 * @param ReferenceBone The name of the bone to use as a reference for ground snapping
+	 * @param SourceOrTarget Specifies whether the skeleton belongs to the source or target */
 	UFUNCTION(BlueprintCallable, Category="IK Retarget")
 	void SnapBoneToGround(FName ReferenceBone, ERetargetSourceOrTarget SourceOrTarget);
 
diff --git a/Engine/Plugins/Animation/IKRig/Source/IKRigEditor/Public/RetargetEditor/SRetargetOpStack.h b/Engine/Plugins/Animation/IKRig/Source/IKRigEditor/Public/RetargetEditor/SRetargetOpStack.h
--- a/Engine/Plugins/Animation/IKRig/Source/IKRigEditor/Public/RetargetEditor/SRetargetOpStack.h
+++ b/Engine/Plugins/Animation/IKRig/Source/IKRigEditor/Public/RetargetEditor/SRetargetOpStack.h
@@ -9,116 +9,268 @@
 #include "Widgets/Views/STableViewBase.h"
 #include "Widgets/Views/STableRow.h"
 #include "Widgets/Views/SListView.h"
+#include "Widgets/Text/SInlineEditableTextBlock.h"
 #include "Framework/Commands/UICommandList.h"
 
+class SRetargetOpStack;
 struct FIKRetargetOpBase;
 class FIKRetargetEditorController;
-class SRetargetOpStack;
+class SRetargetOpList;
 
+// SListView item-type for a single retarget op
 class FRetargetOpStackElement
 {
-	
 public:
 
-	TSharedRef<ITableRow> MakeListRowWidget(
+	static TSharedRef<ITableRow> MakeListRowWidget(
 		const TSharedRef<STableViewBase>& InOwnerTable,
         TSharedRef<FRetargetOpStackElement> InStackElement,
-        TSharedPtr<SRetargetOpStack> InSolverStack);
+        TSharedPtr<SRetargetOpList> InOpListWidget);
 
-	static TSharedRef<FRetargetOpStackElement> Make(FText DisplayName, int32 SolverIndex)
+	static TSharedRef<FRetargetOpStackElement> Make(
+		const int32 InOpIndex,
+		const bool InCanHaveChildren,
+		const TSharedPtr<SRetargetOpList>& InOpListWidget)
 	{
-		return MakeShareable(new FRetargetOpStackElement(DisplayName, SolverIndex));
+		FRetargetOpStackElement* NewElement = new FRetargetOpStackElement(InOpIndex, InCanHaveChildren, InOpListWidget);
+		return MakeShareable(NewElement);
 	}
+
+	FName GetName() const;
 	
-	FText DisplayName;
-	int32 IndexInStack;
+	const UScriptStruct* GetType() const;
+	
+	int32 GetIndexInStack() const { return IndexInStack; };
+	
+	bool GetCanHaveChildren() const { return bCanHaveChildren; };
+	
+	const TArray<FRetargetOpStackElement*>& GetChildren() const { return ChildElements; };
+
+	void AddChild(FRetargetOpStackElement* InChild) { ChildElements.Add(InChild); };
+	
+	const FRetargetOpStackElement* GetParent() const { return ParentElement; };
+	
+	void SetParent(FRetargetOpStackElement* InParentElement) { ParentElement = InParentElement; };
 
+	DECLARE_DELEGATE(FOnRenameRequested);
+	FOnRenameRequested OnRenameRequested;
+	
 private:
+
+	int32 IndexInStack;
+	bool bCanHaveChildren;
+	TArray<FRetargetOpStackElement*> ChildElements;
+	FRetargetOpStackElement* ParentElement = nullptr;
+	TWeakPtr<SRetargetOpList> OpListWidget;
+	
 	// hidden constructors, always use Make above
 	FRetargetOpStackElement() {}
-	FRetargetOpStackElement(FText InDisplayName, int32 StackIndex) :
-		DisplayName(InDisplayName), IndexInStack(StackIndex){}
+	FRetargetOpStackElement(
+		const int32 InIndexInStack,
+		const bool InHasChildren,
+		const TSharedPtr<SRetargetOpList>& InOpListWidget) :
+		IndexInStack(InIndexInStack),
+		bCanHaveChildren(InHasChildren),
+		OpListWidget(InOpListWidget){}
+};
+
+// meta data about an op type
+// used to create filtered menus for creating ops with compatible child types in parent-op menus
+struct FIKRetargetOpMetaData
+{
+	FName NiceName;
+	const UScriptStruct* OpType = nullptr;
+	const UScriptStruct* ParentType = nullptr;
 };
 
-class SRetargetOpStackItem : public STableRow<TSharedPtr<FRetargetOpStackElement>>
+// an SListView customized for retarget ops
+// Supports:
+// - slow double-click rename of list items
+// - drag/drop reordering w/ nested ops
+// - building a menu of compatible ops to add to the list
+// - deleting ops
+class SRetargetOpList : public SListView<TSharedPtr<FRetargetOpStackElement>>
 {
 public:
 	
-	void Construct(
-        const FArguments& InArgs,
-        const TSharedRef<STableViewBase>& OwnerTable,
-        TSharedRef<FRetargetOpStackElement> InStackElement,
-        TSharedPtr<SRetargetOpStack> InOpStackWidget);
+	SLATE_BEGIN_ARGS(SRetargetOpList) {}
+	SLATE_ARGUMENT(TWeakPtr<FIKRetargetEditorController>, InEditorController)
+	SLATE_ARGUMENT(TWeakPtr<FRetargetOpStackElement>, InParentElement)
+	SLATE_END_ARGS()
+
+	virtual ~SRetargetOpList() {}
+
+	void Construct(const FArguments& InArgs);
+	
+	TSharedRef<ITableRow> MakeListRowWidget(TSharedPtr<FRetargetOpStackElement> InElement, const TSharedRef<STableViewBase>& OwnerTable);
+	
+	bool IsEnabled() const;
+
+	void RestoreSelectionAndRefreshList();
+	
+	// drag and drop operations
+	FReply OnDragDetected(
+		const FGeometry& MyGeometry,
+		const FPointerEvent& MouseEvent);
+	TOptional<EItemDropZone> OnCanAcceptDrop(
+		const FDragDropEvent& DragDropEvent,
+		EItemDropZone DropZone,
+		TSharedPtr<FRetargetOpStackElement> TargetElement);
+	FReply OnAcceptDrop(
+		const FDragDropEvent& DragDropEvent,
+		EItemDropZone DropZone,
+		TSharedPtr<FRetargetOpStackElement> TargetElement);
+	// END drag and drop operations
+	
+	// slow-click renaming
+	virtual FReply OnFocusReceived(const FGeometry& MyGeometry, const FFocusEvent& InFocusEvent) override;
+	void OnSelectionChanged(TSharedPtr<FRetargetOpStackElement> InItem, ESelectInfo::Type SelectInfo) const;
+	void OnItemClicked(TSharedPtr<FRetargetOpStackElement> InItem);
+	void RequestRenameSelectedOp() const;
+	// END slow-click renaming
+
+	// add new op menu
+	TSharedRef<SWidget> CreateAddNewOpMenu();
+	// menu command callback for adding a new op
+	void AddNewRetargetOp(const UScriptStruct* ScriptStruct);
+	// END op stack menu
+	
+	// delete op from stack
+	virtual FReply OnKeyDown(const FGeometry& MyGeometry, const FKeyEvent& InKeyEvent);
+	void DeleteRetargetOp(TSharedPtr<FRetargetOpStackElement> OpToDelete);
+
+	// the elements contained in this list
+	TArray< TSharedPtr<FRetargetOpStackElement> > Elements;
+
+	// editor controller
+	TWeakPtr<FIKRetargetEditorController> EditorController;
+	
+	// the type of op that is the parent of this list of ops
+	// if this is null, then it's assumed to be the top-level op stack
+	TWeakPtr<FRetargetOpStackElement> ParentElement;
+
+private:
+	
+	// slow double-click rename state
+	uint32 LastClickCycles = 0;
+	TWeakPtr<FRetargetOpStackElement> LastSelectedElement;
+	
+	// a map of op types to the parent op type (cached at startup)
+	TArray<FIKRetargetOpMetaData> AllOpsMetaData;
+	void CacheOpTypeMetaData();
+};
+
+// builds an UI element representing a single op
+class SRetargetOpItem : public SCompoundWidget
+{
+public:
+	SLATE_BEGIN_ARGS(SRetargetOpItem) {}
+	SLATE_ARGUMENT(TWeakPtr<SRetargetOpList>, InOpListWidget)
+	SLATE_ARGUMENT(TWeakPtr<FRetargetOpStackElement>, InStackElement)
+	SLATE_END_ARGS()
+
+	void Construct(const FArguments& InArgs);
+
+	// RENAMING op
+	void OnNameCommitted(const FText& InText, ETextCommit::Type InCommitType) const;
+	FText GetName() const;
+	FText GetIndexAsText() const;
+	// ENd renaming op
+
+	static float OpHorizontalPadding;
+	static float OpVerticalPadding;
+	
+private:
+	
+	bool IsOpEnabled() const;
+	
+	FIKRetargetOpBase* GetRetargetOp() const;
+
+	TWeakPtr<SRetargetOpList> ListView;
+	TWeakPtr<FRetargetOpStackElement> Element;
+	TSharedPtr<SInlineEditableTextBlock> EditNameWidget;
+};
+
+// an SRetargetOpList item representing a single op with children op beneath it
+// (contains an SRetgargetOpList to house children ops)
+class SParentRetargetOpItem : public STableRow<TSharedPtr<FRetargetOpStackElement>>
+{
+public:
+
+	SLATE_BEGIN_ARGS(SParentRetargetOpItem) {}
+	SLATE_ARGUMENT(TWeakPtr<FRetargetOpStackElement>, InStackElement)
+	SLATE_ARGUMENT(TWeakPtr<SRetargetOpList>, InOpListWidget)
+	SLATE_END_ARGS()
+	
+	void Construct(const FArguments& InArgs, const TSharedRef<STableViewBase>& OwnerTable);
+	
+	void RefreshListView() const;
+
+private:
+	
+	TWeakPtr<FRetargetOpStackElement> ParentStackElement;
+	TWeakPtr<SRetargetOpList> OpListWidget;
+	TSharedPtr<SRetargetOpList> ChildrenListView;
+};
+
+// an SRetargetOpList item representing a single op with no children
+class SRetargetOpSingleItem : public STableRow<TSharedPtr<FRetargetOpStackElement>>
+{
+public:
+	
+	SLATE_BEGIN_ARGS(SRetargetOpSingleItem) {}
+	SLATE_ARGUMENT(TWeakPtr<FRetargetOpStackElement>, InStackElement)
+	SLATE_ARGUMENT(TWeakPtr<SRetargetOpList>, InOpListWidget)
+	SLATE_END_ARGS()
+	
+	void Construct(const FArguments& InArgs, const TSharedRef<STableViewBase>& OwnerTable);
 
 	bool GetWarningMessage(FText& Message) const;
 
 	bool IsOpEnabled() const;
 
 private:
+	
 	FIKRetargetOpBase* GetRetargetOp() const;
+	
 	TWeakPtr<FRetargetOpStackElement> StackElement;
-	TWeakPtr<SRetargetOpStack> OpStackWidget;
+	TWeakPtr<SRetargetOpList> OpListWidget;
 };
 
 class FRetargetOpStackDragDropOp : public FDecoratedDragDropOp
 {
 public:
+	
 	DRAG_DROP_OPERATOR_TYPE(FRetargetOpStackDragDropOp, FDecoratedDragDropOp)
-    static TSharedRef<FRetargetOpStackDragDropOp> New(TWeakPtr<FRetargetOpStackElement> InElement);
+
+	static TSharedRef<FRetargetOpStackDragDropOp> New(TWeakPtr<FRetargetOpStackElement> InElement);
+	
 	virtual TSharedPtr<SWidget> GetDefaultDecorator() const override;
+	
+	static int32 GetIndexToMoveTo(
+		const TSharedPtr<FRetargetOpStackElement>& InDraggedElement,
+		const TSharedPtr<FRetargetOpStackElement>& InTargetElement,
+		const EItemDropZone InDropZone);
+	
 	TWeakPtr<FRetargetOpStackElement> Element;
 };
 
-typedef SListView< TSharedPtr<FRetargetOpStackElement> > SOpStackViewListType;
+// top level view of a stack of ik retargeter ops
+// shows the main menu for adding ops to the retargeting stack, and a list view of top-level ops
 class SRetargetOpStack : public SCompoundWidget, public FEditorUndoClient
 {
 public:
+	
 	SLATE_BEGIN_ARGS(SRetargetOpStack) {}
+	SLATE_ARGUMENT(TWeakPtr<FIKRetargetEditorController>, InEditorController)
 	SLATE_END_ARGS()
 
 	void Construct(const FArguments& InArgs, const TWeakPtr<FIKRetargetEditorController>& InEditorController);
-
-	int32 GetSelectedItemIndex() const;
+	
+	void RefreshStackView() const;
 
 private:
-	// menu for adding new commands
-	TSharedPtr<FUICommandList> CommandList;
 	
-	// editor controller
 	TWeakPtr<FIKRetargetEditorController> EditorController;
-	
-	// the op stack list view
-	TSharedPtr<SOpStackViewListType> ListView;
-	TArray< TSharedPtr<FRetargetOpStackElement> > ListViewItems;
-
-	// add new op menu
-	TSharedRef<SWidget> CreateAddNewMenuWidget();
-	bool IsAddOpEnabled() const;
-	// END op stack menu
-	
-	// menu command callback for adding a new op
-	void AddNewRetargetOp(UScriptStruct* ScriptStruct);
-	// delete op from stack
-	void DeleteRetargetOp(TSharedPtr<FRetargetOpStackElement> OpToDelete);
-	// when an op is selected on in the stack view
-	void OnSelectionChanged(TSharedPtr<FRetargetOpStackElement> InItem, ESelectInfo::Type SelectInfo);
-
-	// list view generate row callback
-	TSharedRef<ITableRow> MakeListRowWidget(TSharedPtr<FRetargetOpStackElement> InElement, const TSharedRef<STableViewBase>& OwnerTable);
-
-	// call to refresh the stack view
-	void RefreshStackView();
-
-	// SWidget interface
-	virtual FReply OnKeyDown(const FGeometry& MyGeometry, const FKeyEvent& InKeyEvent);
-	// END SWidget interface
-	
-	// drag and drop operations
-    FReply OnDragDetected(const FGeometry& MyGeometry, const FPointerEvent& MouseEvent);
-    TOptional<EItemDropZone> OnCanAcceptDrop(const FDragDropEvent& DragDropEvent, EItemDropZone DropZone, TSharedPtr<FRetargetOpStackElement> TargetItem);
-    FReply OnAcceptDrop(const FDragDropEvent& DragDropEvent, EItemDropZone DropZone, TSharedPtr<FRetargetOpStackElement> TargetItem);
-	// END drag and drop operations
-	
-	friend SRetargetOpStackItem;
-	friend FIKRetargetEditorController;
+	TSharedPtr<SRetargetOpList> ListView;
 };
