[IK Retarget] Fixes crash when loading old retargeter with no target IK Rig.

diff --git a/Engine/Plugins/Animation/IKRig/Source/IKRig/Private/Retargeter/IKRetargeter.cpp b/Engine/Plugins/Animation/IKRig/Source/IKRig/Private/Retargeter/IKRetargeter.cpp
--- a/Engine/Plugins/Animation/IKRig/Source/IKRig/Private/Retargeter/IKRetargeter.cpp
+++ b/Engine/Plugins/Animation/IKRig/Source/IKRig/Private/Retargeter/IKRetargeter.cpp
@@ -414,12 +414,15 @@ void UIKRetargeter::PostLoadConvertEverythingToOps()
 
 	// record all chains that have an IK goal assigned
 	TArray<FName> ChainsWithIK;
-	const TArray<FBoneChain>& RetargetBoneChains = TargetIKRigAsset->GetRetargetChains();
-	for (const FBoneChain& Chain : RetargetBoneChains)
+	if (IsValid(TargetIKRigAsset))
 	{
-		if (Chain.IKGoalName != NAME_None)
+		const TArray<FBoneChain>& RetargetBoneChains =TargetIKRigAsset->GetRetargetChains();
+		for (const FBoneChain& Chain : RetargetBoneChains)
 		{
-			ChainsWithIK.Add(Chain.ChainName);
+			if (Chain.IKGoalName != NAME_None)
+			{
+				ChainsWithIK.Add(Chain.ChainName);
+			}
 		}
 	}
 	
