[IK Retargeter] Fixes regression, no default retarget pose on new retarget assets.

diff --git a/Engine/Plugins/Animation/IKRig/Source/IKRig/Private/Retargeter/IKRetargeter.cpp b/Engine/Plugins/Animation/IKRig/Source/IKRig/Private/Retargeter/IKRetargeter.cpp
--- a/Engine/Plugins/Animation/IKRig/Source/IKRig/Private/Retargeter/IKRetargeter.cpp
+++ b/Engine/Plugins/Animation/IKRig/Source/IKRig/Private/Retargeter/IKRetargeter.cpp
@@ -205,6 +205,9 @@ UIKRetargeter::UIKRetargeter(const FObjectInitializer& ObjectInitializer)
 	OpStack_DEPRECATED = CreateDefaultSubobject<UDEPRECATED_RetargetOpStack>(TEXT("PostSettings"));
 	OpStack_DEPRECATED->SetFlags(RF_Transactional);
 	PRAGMA_ENABLE_DEPRECATION_WARNINGS
+
+	// we need this to ensure that new retargeters always have a default retarget pose
+	CleanRetargetPoses();
 }
 
 const UIKRigDefinition* UIKRetargeter::GetIKRig(ERetargetSourceOrTarget SourceOrTarget) const
