[IK Retarget] Fix crash when source mesh component is the same as target mesh component.

diff --git a/Engine/Plugins/Animation/IKRig/Source/IKRig/Private/AnimNodes/AnimNode_RetargetPoseFromMesh.cpp b/Engine/Plugins/Animation/IKRig/Source/IKRig/Private/AnimNodes/AnimNode_RetargetPoseFromMesh.cpp
--- a/Engine/Plugins/Animation/IKRig/Source/IKRig/Private/AnimNodes/AnimNode_RetargetPoseFromMesh.cpp
+++ b/Engine/Plugins/Animation/IKRig/Source/IKRig/Private/AnimNodes/AnimNode_RetargetPoseFromMesh.cpp
@@ -343,7 +343,15 @@ bool FAnimNode_RetargetPoseFromMesh::EnsureProcessorIsInitialized(const USkeleta
 	const bool bCopyingFromOtherComponent = RetargetFrom != ERetargetSourceMode::SourcePosePin;
 	if (bCopyingFromOtherComponent && !SourceMeshComponent.IsValid())
 	{
-		return false; // can't do anything if we don't have a source mesh component
+		if (!SourceMeshComponent.IsValid())
+		{
+			return false; // can't do anything if we don't have a source mesh component
+		}
+
+		if (!ensureMsgf(SourceMeshComponent != TargetMeshComponent, TEXT("Cannot use target component as source.")))
+		{
+			return false; // we do not support retargeting between the same component
+		}
 	}
 
 	// check that both a source and target mesh exist
