[IK Retarget] Fixes crash when renaming ops rapidly.

diff --git a/Engine/Plugins/Animation/IKRig/Source/IKRigEditor/Private/RetargetEditor/SRetargetOpStack.cpp b/Engine/Plugins/Animation/IKRig/Source/IKRigEditor/Private/RetargetEditor/SRetargetOpStack.cpp
--- a/Engine/Plugins/Animation/IKRig/Source/IKRigEditor/Private/RetargetEditor/SRetargetOpStack.cpp
+++ b/Engine/Plugins/Animation/IKRig/Source/IKRigEditor/Private/RetargetEditor/SRetargetOpStack.cpp
@@ -195,7 +195,6 @@ void SRetargetOpItem::OnNameCommitted(const FText& InText, ETextCommit::Type InC
 	FName NewName = FName(InText.ToString());
 	FIKRetargetEditorController* EditorController = ListView.Pin()->EditorController.Pin().Get();
 	EditorController->AssetController->SetOpName(NewName, Element.Pin()->GetIndexInStack());
-	ListView.Pin()->RebuildList();
 }
 
 FText SRetargetOpItem::GetName() const
