[IK Retarget] Fixing wrong default for warping direction source in deprecated global settings.

diff --git a/Engine/Plugins/Animation/IKRig/Source/IKRig/Public/Retargeter/IKRetargetDeprecated.h b/Engine/Plugins/Animation/IKRig/Source/IKRig/Public/Retargeter/IKRetargetDeprecated.h
--- a/Engine/Plugins/Animation/IKRig/Source/IKRig/Public/Retargeter/IKRetargetDeprecated.h
+++ b/Engine/Plugins/Animation/IKRig/Source/IKRig/Public/Retargeter/IKRetargetDeprecated.h
@@ -167,7 +167,7 @@ public:
 	UPROPERTY(EditAnywhere, BlueprintReadWrite, Category=DeprecatedSettings, meta=(DeprecatedProperty))
 	bool bWarping = false;
 	UPROPERTY(EditAnywhere, BlueprintReadWrite, Category=DeprecatedSettings, meta=(DeprecatedProperty))
-	EWarpingDirectionSource DirectionSource = EWarpingDirectionSource::RootBone;
+	EWarpingDirectionSource DirectionSource = EWarpingDirectionSource::Goals;
 	UPROPERTY(EditAnywhere, BlueprintReadWrite, Category=DeprecatedSettings, meta=(DeprecatedProperty))
 	EBasicAxis ForwardDirection = EBasicAxis::Y;
 	UPROPERTY(EditAnywhere, BlueprintReadWrite, Category=DeprecatedSettings, meta=(DeprecatedProperty))
