[IK Rig] IK Rig goal "Additive" mode broken in editor.

diff --git a/Engine/Plugins/Animation/IKRig/Source/IKRig/Private/Rig/IKRigDataTypes.cpp b/Engine/Plugins/Animation/IKRig/Source/IKRig/Private/Rig/IKRigDataTypes.cpp
--- a/Engine/Plugins/Animation/IKRig/Source/IKRig/Private/Rig/IKRigDataTypes.cpp
+++ b/Engine/Plugins/Animation/IKRig/Source/IKRig/Private/Rig/IKRigDataTypes.cpp
@@ -35,30 +35,26 @@ void FIKRigGoalContainer::SetIKGoal(const UIKRigEffectorGoal* InEffectorGoal)
 		bRigNeedsInitialized = true;
 		return;
 	}
+	
+	// copy transform
+	Goal->Position = InEffectorGoal->CurrentTransform.GetTranslation();
+	Goal->Rotation = InEffectorGoal->CurrentTransform.Rotator();
+    Goal->PositionAlpha = InEffectorGoal->PositionAlpha;
+    Goal->RotationAlpha = InEffectorGoal->RotationAlpha;
+	Goal->PositionSpace = EIKRigGoalSpace::Component;
+	Goal->RotationSpace = EIKRigGoalSpace::Component;
 
-	// goals in editor have "preview mode" which allows them to be specified relative to their
-	// initial pose to simulate an additive runtime behavior
+	// goals in editor have "preview mode" which allows them to be additive relative to the goal initial position
 #if WITH_EDITOR
 	if (InEffectorGoal->PreviewMode == EIKRigGoalPreviewMode::Additive)
 	{
-		// transform to be interpreted as an offset, relative to input pose
 		Goal->Position = InEffectorGoal->CurrentTransform.GetTranslation() - InEffectorGoal->InitialTransform.GetTranslation();
 		const FQuat RelativeRotation = InEffectorGoal->CurrentTransform.GetRotation() * InEffectorGoal->InitialTransform.GetRotation().Inverse();
 		Goal->Rotation = RelativeRotation.Rotator();
 		Goal->PositionSpace = EIKRigGoalSpace::Additive;
 		Goal->RotationSpace = EIKRigGoalSpace::Additive;
-	}else
-#endif
-	{
-		// transform to be interpreted as absolute and in component space
-		Goal->Position = InEffectorGoal->CurrentTransform.GetTranslation();
-		Goal->Rotation = InEffectorGoal->CurrentTransform.Rotator();
-		Goal->PositionSpace = EIKRigGoalSpace::Component;
-		Goal->RotationSpace = EIKRigGoalSpace::Component;
 	}
-	
-    Goal->PositionAlpha = InEffectorGoal->PositionAlpha;
-    Goal->RotationAlpha = InEffectorGoal->RotationAlpha;
+#endif
 }
 
 const FIKRigGoal* FIKRigGoalContainer::FindGoalByName(const FName& GoalName) const
diff --git a/Engine/Plugins/Animation/IKRig/Source/IKRig/Private/Rig/IKRigProcessor.cpp b/Engine/Plugins/Animation/IKRig/Source/IKRig/Private/Rig/IKRigProcessor.cpp
--- a/Engine/Plugins/Animation/IKRig/Source/IKRig/Private/Rig/IKRigProcessor.cpp
+++ b/Engine/Plugins/Animation/IKRig/Source/IKRig/Private/Rig/IKRigProcessor.cpp
@@ -415,7 +415,10 @@ const FIKRigSolverBase* FIKRigProcessor::GetSolver(const int32 InSolverIndex) co
 
 void FIKRigProcessor::CopyAllSettingsFromAsset(const UIKRigDefinition* SourceAsset)
 {
-	check(SourceAsset)
+	if (!IsValid(SourceAsset))
+	{
+		return;
+	}
 	
 	if (!bInitialized)
 	{
@@ -426,17 +429,7 @@ void FIKRigProcessor::CopyAllSettingsFromAsset(const UIKRigDefinition* SourceAss
 	const TArray<UIKRigEffectorGoal*>& AssetGoals =  SourceAsset->GetGoalArray();
 	for (const UIKRigEffectorGoal* AssetGoal : AssetGoals)
 	{
-		if (FIKRigGoal* Goal = GoalContainer.FindGoalByName(AssetGoal->GoalName))
-		{
-			Goal->Position = AssetGoal->CurrentTransform.GetTranslation();
-			Goal->Rotation = AssetGoal->CurrentTransform.GetRotation().Rotator();
-			Goal->RotationAlpha = AssetGoal->RotationAlpha;
-			Goal->PositionAlpha = AssetGoal->PositionAlpha;
-		}
-		else
-		{
-			SetIKGoal(AssetGoal);
-		}
+		SetIKGoal(AssetGoal);
 	}
 
 	// copy solver settings
