[IK Retarget] Added API to get/set bones from python, on ops that use FBoneReference

diff --git a/Engine/Plugins/Animation/IKRig/Source/IKRig/Private/Retargeter/RetargetOps/PelvisMotionOp.cpp b/Engine/Plugins/Animation/IKRig/Source/IKRig/Private/Retargeter/RetargetOps/PelvisMotionOp.cpp
--- a/Engine/Plugins/Animation/IKRig/Source/IKRig/Private/Retargeter/RetargetOps/PelvisMotionOp.cpp
+++ b/Engine/Plugins/Animation/IKRig/Source/IKRig/Private/Retargeter/RetargetOps/PelvisMotionOp.cpp
@@ -352,4 +352,28 @@ void UIKRetargetPelvisMotionController::SetSettings(FIKRetargetPelvisMotionOpSet
 	OpToControl->SetSettings(&InSettings);
 }
 
+void UIKRetargetPelvisMotionController::SetSourcePelvisBone(const FName InSourcePelvisBone)
+{
+	FIKRetargetPelvisMotionOp* Op = reinterpret_cast<FIKRetargetPelvisMotionOp*>(OpToControl);
+	Op->Settings.SourcePelvisBone.BoneName = InSourcePelvisBone;
+}
+
+FName UIKRetargetPelvisMotionController::GetSourcePelvisBone()
+{
+	FIKRetargetPelvisMotionOp* Op = reinterpret_cast<FIKRetargetPelvisMotionOp*>(OpToControl);
+	return Op->Settings.SourcePelvisBone.BoneName;
+}
+
+void UIKRetargetPelvisMotionController::SetTargetPelvisBone(const FName InTargetPelvisBone)
+{
+	FIKRetargetPelvisMotionOp* Op = reinterpret_cast<FIKRetargetPelvisMotionOp*>(OpToControl);
+	Op->Settings.TargetPelvisBone.BoneName = InTargetPelvisBone;
+}
+
+FName UIKRetargetPelvisMotionController::GetTargetPelvisBone()
+{
+	FIKRetargetPelvisMotionOp* Op = reinterpret_cast<FIKRetargetPelvisMotionOp*>(OpToControl);
+	return Op->Settings.TargetPelvisBone.BoneName;
+}
+
 #undef LOCTEXT_NAMESPACE
diff --git a/Engine/Plugins/Animation/IKRig/Source/IKRig/Private/Retargeter/RetargetOps/PinBoneOp.cpp b/Engine/Plugins/Animation/IKRig/Source/IKRig/Private/Retargeter/RetargetOps/PinBoneOp.cpp
--- a/Engine/Plugins/Animation/IKRig/Source/IKRig/Private/Retargeter/RetargetOps/PinBoneOp.cpp
+++ b/Engine/Plugins/Animation/IKRig/Source/IKRig/Private/Retargeter/RetargetOps/PinBoneOp.cpp
@@ -420,4 +420,44 @@ void UIKRetargetPinBoneController::SetSettings(FIKRetargetPinBoneOpSettings InSe
 	OpToControl->SetSettings(&InSettings);
 }
 
+void UIKRetargetPinBoneController::ClearAllBonePairs()
+{
+	FIKRetargetPinBoneOp* Op = reinterpret_cast<FIKRetargetPinBoneOp*>(OpToControl);
+	Op->Settings.BonesToPin.Reset();
+}
+
+void UIKRetargetPinBoneController::SetBonePair(const FName InBoneToCopyFrom, const FName InBoneToCopyTo)
+{
+	FIKRetargetPinBoneOp* Op = reinterpret_cast<FIKRetargetPinBoneOp*>(OpToControl);
+
+	// update existing pair with new bone to copy from (if there is one)
+	for (FPinBoneData& BonePair : Op->Settings.BonesToPin)
+	{
+		if (BonePair.BoneToCopyTo == InBoneToCopyTo)
+		{
+			BonePair.BoneToCopyFrom = InBoneToCopyFrom;
+			return;
+		}
+	}
+
+	// needs new bone pair
+	FPinBoneData NewBonePair;
+	NewBonePair.BoneToCopyFrom.BoneName = InBoneToCopyFrom;
+	NewBonePair.BoneToCopyTo.BoneName = InBoneToCopyTo;
+	Op->Settings.BonesToPin.Add(NewBonePair);
+}
+
+TMap<FName, FName> UIKRetargetPinBoneController::GetAllBonePairs()
+{
+	TMap<FName, FName> AllBonePairs;
+	
+	FIKRetargetPinBoneOp* Op = reinterpret_cast<FIKRetargetPinBoneOp*>(OpToControl);
+	for (FPinBoneData& BonePair : Op->Settings.BonesToPin)
+	{
+		AllBonePairs.Add(BonePair.BoneToCopyTo.BoneName, BonePair.BoneToCopyFrom.BoneName);
+	}
+	
+	return MoveTemp(AllBonePairs);
+}
+
 #undef LOCTEXT_NAMESPACE
diff --git a/Engine/Plugins/Animation/IKRig/Source/IKRig/Private/Retargeter/RetargetOps/RootMotionGeneratorOp.cpp b/Engine/Plugins/Animation/IKRig/Source/IKRig/Private/Retargeter/RetargetOps/RootMotionGeneratorOp.cpp
--- a/Engine/Plugins/Animation/IKRig/Source/IKRig/Private/Retargeter/RetargetOps/RootMotionGeneratorOp.cpp
+++ b/Engine/Plugins/Animation/IKRig/Source/IKRig/Private/Retargeter/RetargetOps/RootMotionGeneratorOp.cpp
@@ -342,4 +342,28 @@ void UIKRetargetRootMotionController::SetSettings(FIKRetargetRootMotionOpSetting
 	OpToControl->SetSettings(&InSettings);
 }
 
+void UIKRetargetRootMotionController::SetSourceRootBone(const FName InSourceRootBone)
+{
+	FIKRetargetRootMotionOp* Op = reinterpret_cast<FIKRetargetRootMotionOp*>(OpToControl);
+	Op->Settings.SourceRoot.BoneName = InSourceRootBone;
+}
+
+FName UIKRetargetRootMotionController::GetSourceRootBone()
+{
+	FIKRetargetRootMotionOp* Op = reinterpret_cast<FIKRetargetRootMotionOp*>(OpToControl);
+	return Op->Settings.SourceRoot.BoneName;
+}
+
+void UIKRetargetRootMotionController::SetTargetRootBone(const FName InTargetRootBone)
+{
+	FIKRetargetRootMotionOp* Op = reinterpret_cast<FIKRetargetRootMotionOp*>(OpToControl);
+	Op->Settings.TargetRoot.BoneName = InTargetRootBone;
+}
+
+FName UIKRetargetRootMotionController::GetTargetRootBone()
+{
+	FIKRetargetRootMotionOp* Op = reinterpret_cast<FIKRetargetRootMotionOp*>(OpToControl);
+	return Op->Settings.TargetRoot.BoneName;
+}
+
 #undef LOCTEXT_NAMESPACE
diff --git a/Engine/Plugins/Animation/IKRig/Source/IKRig/Public/Retargeter/RetargetOps/PelvisMotionOp.h b/Engine/Plugins/Animation/IKRig/Source/IKRig/Public/Retargeter/RetargetOps/PelvisMotionOp.h
--- a/Engine/Plugins/Animation/IKRig/Source/IKRig/Public/Retargeter/RetargetOps/PelvisMotionOp.h
+++ b/Engine/Plugins/Animation/IKRig/Source/IKRig/Public/Retargeter/RetargetOps/PelvisMotionOp.h
@@ -154,6 +154,9 @@ struct IKRIG_API FIKRetargetPelvisMotionOp : public FIKRetargetOpBase
 
 	FVector GetPelvisTranslationOffset() const;
 
+	UPROPERTY()
+	FIKRetargetPelvisMotionOpSettings Settings;
+
 private:
 
 	void Reset();
@@ -172,9 +175,6 @@ private:
 	
 	void DecodePose(TArray<FTransform> &OutTargetGlobalPose);
 
-	UPROPERTY()
-	FIKRetargetPelvisMotionOpSettings Settings;
-
 	// transient work data
 	FPelvisSource Source;
 	FPelvisTarget Target;
@@ -210,6 +210,26 @@ public:
 	 * @param InSettings a FIKRetargetPelvisMotionOpSettings struct containing all the settings to apply to this op */
 	UFUNCTION(BlueprintCallable, Category = Settings)
 	void SetSettings(FIKRetargetPelvisMotionOpSettings InSettings);
+
+	/* Set the pelvis bone for the source.
+	 * @param InSourcePelvisBone the name of the pelvis bone on the source skeleton. */
+	UFUNCTION(BlueprintCallable, Category = Settings)
+	void SetSourcePelvisBone(const FName InSourcePelvisBone);
+
+	/* Get the pelvis bone for the source.
+	 * @return the name of the pelvis bone on the source skeleton. */
+	UFUNCTION(BlueprintCallable, Category = Settings)
+	FName GetSourcePelvisBone();
+
+	/* Set the pelvis bone for the target.
+	 * @param InTargetPelvisBone the name of the pelvis bone on the target skeleton. */
+	UFUNCTION(BlueprintCallable, Category = Settings)
+	void SetTargetPelvisBone(const FName InTargetPelvisBone);
+
+	/* Get the pelvis bone for the target.
+	 * @return the name of the pelvis bone on the target skeleton. */
+	UFUNCTION(BlueprintCallable, Category = Settings)
+	FName GetTargetPelvisBone();
 };
 
 #if WITH_EDITOR
diff --git a/Engine/Plugins/Animation/IKRig/Source/IKRig/Public/Retargeter/RetargetOps/PinBoneOp.h b/Engine/Plugins/Animation/IKRig/Source/IKRig/Public/Retargeter/RetargetOps/PinBoneOp.h
--- a/Engine/Plugins/Animation/IKRig/Source/IKRig/Public/Retargeter/RetargetOps/PinBoneOp.h
+++ b/Engine/Plugins/Animation/IKRig/Source/IKRig/Public/Retargeter/RetargetOps/PinBoneOp.h
@@ -196,6 +196,22 @@ public:
 	 * @param InSettings a FIKRetargetPinBoneOpSettings struct containing all the settings to apply to this op */
 	UFUNCTION(BlueprintCallable, Category = Settings)
 	void SetSettings(FIKRetargetPinBoneOpSettings InSettings);
+
+	/* Clear all the bone pairs */
+	UFUNCTION(BlueprintCallable, Category = BonePairs)
+	void ClearAllBonePairs();
+
+	/* Clear all the bone pairs
+	 * @param InBoneToCopyFrom the name of the bone to copy animation from
+	 * @param InBoneToCopyTo the name of the bone to apply animation to
+	 * NOTE: if the bone to copy to is already present in the op, it will be updated with the new bone to copy from. */
+	UFUNCTION(BlueprintCallable, Category = BonePairs)
+	void SetBonePair(const FName InBoneToCopyFrom, const FName InBoneToCopyTo);
+
+	/* Get all the bone pairs currently stored in the op.
+	 * @return a map with target bones as keys and source bones as values. */
+	UFUNCTION(BlueprintCallable, Category = BonePairs)
+	TMap<FName,FName> GetAllBonePairs();
 };
 
 //
diff --git a/Engine/Plugins/Animation/IKRig/Source/IKRig/Public/Retargeter/RetargetOps/RootMotionGeneratorOp.h b/Engine/Plugins/Animation/IKRig/Source/IKRig/Public/Retargeter/RetargetOps/RootMotionGeneratorOp.h
--- a/Engine/Plugins/Animation/IKRig/Source/IKRig/Public/Retargeter/RetargetOps/RootMotionGeneratorOp.h
+++ b/Engine/Plugins/Animation/IKRig/Source/IKRig/Public/Retargeter/RetargetOps/RootMotionGeneratorOp.h
@@ -185,6 +185,26 @@ public:
 	 * @param InSettings a FIKRetargetRootMotionOpSettings struct containing all the settings to apply to this op */
 	UFUNCTION(BlueprintCallable, Category = Settings)
 	void SetSettings(FIKRetargetRootMotionOpSettings InSettings);
+
+	/* Set the root bone for the source.
+	 * @param InSourceRootBone the name of the root bone on the source skeleton. */
+	UFUNCTION(BlueprintCallable, Category = Settings)
+	void SetSourceRootBone(const FName InSourceRootBone);
+
+	/* Get the root bone for the source.
+	 * @return the name of the root bone on the source skeleton. */
+	UFUNCTION(BlueprintCallable, Category = Settings)
+	FName GetSourceRootBone();
+
+	/* Set the root bone for the target.
+	 * @param InTargetRootBone the name of the root bone on the target skeleton. */
+	UFUNCTION(BlueprintCallable, Category = Settings)
+	void SetTargetRootBone(const FName InTargetRootBone);
+
+	/* Get the root bone for the target.
+	 * @return the name of the root bone on the target skeleton. */
+	UFUNCTION(BlueprintCallable, Category = Settings)
+	FName GetTargetRootBone();
 };
 
 //
