[IK Retargeter] Adding support for chain mapping per-op

diff --git a/Engine/Plugins/Animation/IKRig/Source/IKRig/Private/Retargeter/IKRetargetOps.cpp b/Engine/Plugins/Animation/IKRig/Source/IKRig/Private/Retargeter/IKRetargetOps.cpp
--- a/Engine/Plugins/Animation/IKRig/Source/IKRig/Private/Retargeter/IKRetargetOps.cpp
+++ b/Engine/Plugins/Animation/IKRig/Source/IKRig/Private/Retargeter/IKRetargetOps.cpp
@@ -9,6 +9,10 @@ void FIKRetargetOpBase::SetSettings(const FIKRetargetOpSettingsBase* InSettings)
 	// copy everything by default
 	static TArray<FName> PropertiesToIgnore;
 	CopySettingsRaw(InSettings, PropertiesToIgnore);
+
+#if WITH_EDITOR
+	GetSettings()->OwningOpName = GetName();
+#endif
 }
 
 #if WITH_EDITOR
@@ -63,6 +67,9 @@ void FIKRetargetOpBase::SetName(const FName InName)
 {
 	ensure(InName != NAME_None);
 	Name = InName;
+#if WITH_EDITOR
+	GetSettings()->OwningOpName = GetName();
+#endif
 }
 
 void FIKRetargetOpBase::SetParentOpName(const FName InName)
diff --git a/Engine/Plugins/Animation/IKRig/Source/IKRig/Private/Retargeter/IKRetargetProcessor.cpp b/Engine/Plugins/Animation/IKRig/Source/IKRig/Private/Retargeter/IKRetargetProcessor.cpp
--- a/Engine/Plugins/Animation/IKRig/Source/IKRig/Private/Retargeter/IKRetargetProcessor.cpp
+++ b/Engine/Plugins/Animation/IKRig/Source/IKRig/Private/Retargeter/IKRetargetProcessor.cpp
@@ -23,7 +23,6 @@
 
 void FRetargetSkeleton::Initialize(
 	const USkeletalMesh* InSkeletalMesh,
-	const TArray<FBoneChain>& BoneChains,
 	const FName InRetargetPoseName,
 	const FIKRetargetPose* RetargetPose,
 	const FName RetargetRootBone,
@@ -51,23 +50,6 @@ void FRetargetSkeleton::Initialize(
 	GenerateRetargetPose(InRetargetPoseName, RetargetPose, RetargetRootBone, ScaleFactor);
 }
 
-void FRetargetSkeleton::FillPerBoneChainNames(const TArray<FResolvedBoneChain>& InBoneChains)
-{
-	// record array of FName per-bone specifying which chain each bone belongs to
-
-	// set all to NAME_None by default
-	ChainThatContainsBone.Init(NAME_None, BoneNames.Num());
-
-	// spin through all bones on all the resolved chains on this skeleton
-	for (const FResolvedBoneChain& BoneChain : InBoneChains)
-	{
-		for (const int32 BoneIndex : BoneChain.BoneIndices)
-		{
-			ChainThatContainsBone[BoneIndex] = BoneChain.ChainName;
-		}
-	}
-}
-
 void FRetargetSkeleton::Reset()
 {
 	BoneNames.Reset();
@@ -80,7 +62,7 @@ void FRetargetSkeleton::Reset()
 void FRetargetSkeleton::GenerateRetargetPose(
 	const FName InRetargetPoseName,
 	const FIKRetargetPose* InRetargetPose,
-	const FName RetargetRootBone,
+	const FName PelvisBoneName,
 	const double ScaleFactor)
 {
 	// record the name of the retarget pose (prevents re-initialization if profile swaps it)
@@ -102,19 +84,13 @@ void FRetargetSkeleton::GenerateRetargetPose(
 		RetargetGlobalPose[BoneIndex].SetScale3D(FVector::OneVector);
 	}
 
-	// no retarget pose specified (will use default pose from skeletal mesh with no offsets)
-	if (RetargetRootBone == NAME_None)
-	{
-		return;
-	}
-	
-	// apply root translation offset
-	const int32 RootBoneIndex = FindBoneIndexByName(RetargetRootBone);
-	if (RootBoneIndex != INDEX_NONE)
+	// apply pelvis translation offset
+	const int32 PelvisBoneIndex = FindBoneIndexByName(PelvisBoneName);
+	if (PelvisBoneIndex != INDEX_NONE)
 	{
-		FTransform& RootTransform = RetargetGlobalPose[RootBoneIndex];
-		RootTransform.AddToTranslation(InRetargetPose->GetRootTranslationDelta());
-		UpdateLocalTransformOfSingleBone(RootBoneIndex, RetargetLocalPose, RetargetGlobalPose);
+		FTransform& PelvisTransform = RetargetGlobalPose[PelvisBoneIndex];
+		PelvisTransform.AddToTranslation(InRetargetPose->GetRootTranslationDelta());
+		UpdateLocalTransformOfSingleBone(PelvisBoneIndex, RetargetLocalPose, RetargetGlobalPose);
 	}
 
 	// apply retarget pose offsets (retarget pose is stored as offset relative to reference pose)
@@ -361,7 +337,6 @@ int32 FRetargetSkeleton::GetParentIndex(const int32 BoneIndex) const
 
 void FTargetSkeleton::Initialize(
 	const USkeletalMesh* InSkeletalMesh,
-	const TArray<FBoneChain>& BoneChains,
 	const FName InRetargetPoseName,
 	const FIKRetargetPose* RetargetPose,
 	const FName RetargetRootBone)
@@ -369,7 +344,7 @@ void FTargetSkeleton::Initialize(
 	Reset();
 
 	constexpr double TargetSkeletonScaleFactor = 1.0;
-	FRetargetSkeleton::Initialize(InSkeletalMesh, BoneChains, InRetargetPoseName, RetargetPose, RetargetRootBone, TargetSkeletonScaleFactor);
+	FRetargetSkeleton::Initialize(InSkeletalMesh, InRetargetPoseName, RetargetPose, RetargetRootBone, TargetSkeletonScaleFactor);
 
 	// make storage for per-bone "Is Retargeted" flag (used for hierarchy updates)
 	// these are bones that are in a target chain that is mapped to a source chain (ie, will actually be retargeted)
@@ -764,6 +739,7 @@ void FResolvedBoneChain::CalculateBoneParameters(FIKRigLogger& Log)
 
 bool FRetargeterBoneChains::Initialize(
 	const UIKRetargeter* InRetargetAsset,
+	const TArray<const UIKRigDefinition*>& InTargetIKRigs,
 	const FRetargetSkeleton& InSourceSkeleton,
 	const FRetargetSkeleton& InTargetSkeleton,
 	FIKRigLogger& InLog)
@@ -772,12 +748,14 @@ bool FRetargeterBoneChains::Initialize(
 
 	// get source and target IK Rig assets
 	const UIKRigDefinition* SourceIKRig = InRetargetAsset->GetIKRig(ERetargetSourceOrTarget::Source);
-	const UIKRigDefinition* TargetIKRig = InRetargetAsset->GetIKRig(ERetargetSourceOrTarget::Target);
-	if (!ensure(SourceIKRig && TargetIKRig))
+	if (!ensure(SourceIKRig && !InTargetIKRigs.IsEmpty()))
 	{
 		InLog.LogError(LOCTEXT("MissingIKRig", "IK Retargeter either a source or target IK Rig asset."));
 		return false;
 	}
+
+	// first element is always the default
+	DefaultTargetIKRig = InTargetIKRigs[0];
 	
 	// load ALL bone chains on source and target (ops may use them even if they aren't mapped)
 	auto LoadChains = [](
@@ -796,8 +774,16 @@ bool FRetargeterBoneChains::Initialize(
 			}
 		}
 	};
+
+	// load chains for source
 	LoadChains(InSourceSkeleton, SourceIKRig->GetRetargetChains(), InLog, SourceBoneChains);
-	LoadChains(InTargetSkeleton, TargetIKRig->GetRetargetChains(), InLog, TargetBoneChains);
+
+	// load chains for all target IK rigs
+	for (const UIKRigDefinition* TargetIKRig : InTargetIKRigs)
+	{
+		TArray<FResolvedBoneChain>& BoneChains = TargetBoneChains.Add(TargetIKRig);
+		LoadChains(InTargetSkeleton, TargetIKRig->GetRetargetChains(), InLog, BoneChains);
+	}
 
 	// sort the chains based on their StartBone's index
 	auto ChainsSorter = [this](const FResolvedBoneChain& A, const FResolvedBoneChain& B)
@@ -812,107 +798,107 @@ bool FRetargeterBoneChains::Initialize(
 		return IndexA < IndexB;
 	};
 	SourceBoneChains.Sort(ChainsSorter);
-	TargetBoneChains.Sort(ChainsSorter);
-
-	// load the chain mapping for each target chain (loads INDEX_NONE for any unmapped chains)
-	bool bFoundAtLeastOneMappedChain = true;
-	const FRetargetChainMapping& ChainNameMappings = InRetargetAsset->GetChainMapping();
-	for (int32 TargetBoneChainIndex = 0; TargetBoneChainIndex<TargetBoneChains.Num(); TargetBoneChainIndex++)
+	for (TTuple<const UIKRigDefinition*, TArray<FResolvedBoneChain>>& TargetChains : TargetBoneChains)
 	{
-		const FResolvedBoneChain& ResolvedTargetChain = TargetBoneChains[TargetBoneChainIndex];
-		FName SourceChainName = ChainNameMappings.GetChainMappedTo(ResolvedTargetChain.ChainName, ERetargetSourceOrTarget::Target);
-		if (SourceChainName == NAME_None)
-		{
-			ChainMapping.Add(INDEX_NONE);
-			continue;
-		}
+		TargetChains.Value.Sort(ChainsSorter);
+	}
 
-		// find the index of the source chain to map to
-		int32 SourceChainIndex = SourceBoneChains.IndexOfByPredicate([&SourceChainName](const FResolvedBoneChain& Value)
-		{
-			return Value.ChainName == SourceChainName;
-		});
+	return true;
+}
 
-		// log a warning if the user has mapped a target chain to a source chain that doesn't exist (ill formed data)
-		if (SourceChainName!=NAME_None && SourceChainIndex==INDEX_NONE)
-		{
-			InLog.LogWarning( FText::Format(
-			LOCTEXT("MissingTargetChain", "IK Retargeter has a target bone chain: {0} that is mapped to a non-existant source chain."),
-			FText::FromString(ResolvedTargetChain.ChainName.ToString())));
-		}
+const TArray<FResolvedBoneChain>* FRetargeterBoneChains::GetResolvedBoneChains(
+	ERetargetSourceOrTarget SourceOrTarget,
+	const UIKRigDefinition* InTargetIKRig) const
+{
+	if (SourceOrTarget == ERetargetSourceOrTarget::Source)
+	{
+		return &SourceBoneChains;
+	}
 
-		bFoundAtLeastOneMappedChain &= SourceChainIndex != INDEX_NONE;
-		ChainMapping.Add(SourceChainIndex);
+	if (InTargetIKRig)
+	{
+		return TargetBoneChains.Find(InTargetIKRig);
 	}
 
-	bIsInitialized = bFoundAtLeastOneMappedChain;
-	return bIsInitialized;
+	return TargetBoneChains.Find(DefaultTargetIKRig);
 }
 
-const TArray<FResolvedBoneChain>& FRetargeterBoneChains::GetResolvedBoneChains(ERetargetSourceOrTarget SourceOrTarget) const
+const TMap<const UIKRigDefinition*, TArray<FResolvedBoneChain>>& FRetargeterBoneChains::GetAllResolvedTargetBoneChains() const
 {
-	return SourceOrTarget == ERetargetSourceOrTarget::Source ? SourceBoneChains : TargetBoneChains;
+	return TargetBoneChains;
 }
 
-int32 FRetargeterBoneChains::GetIndexOfBoneChainByName(const FName InChainName, const ERetargetSourceOrTarget SourceOrTarget) const
+const FResolvedBoneChain* FRetargeterBoneChains::GetResolvedBoneChainByName(
+	const FName InChainName,
+	const ERetargetSourceOrTarget SourceOrTarget,
+	const UIKRigDefinition* TargetIKRig) const
 {
-	const TArray<FResolvedBoneChain>& ResolvedBoneChains = GetResolvedBoneChains(SourceOrTarget);
-	return ResolvedBoneChains.IndexOfByPredicate([InChainName](const FResolvedBoneChain& Value)
+	const TArray<FResolvedBoneChain>* ResolvedBoneChains = GetResolvedBoneChains(SourceOrTarget, TargetIKRig);
+	if (!ResolvedBoneChains)
+	{
+		return nullptr;
+	}
+	
+	return ResolvedBoneChains->FindByPredicate([InChainName](const FResolvedBoneChain& Value)
 	{
 		return Value.ChainName == InChainName;
 	});
 }
 
-const FResolvedBoneChain* FRetargeterBoneChains::GetSourceChainMappedToTargetChain(const FName InTargetChainName) const
+void FRetargeterBoneChains::UpdatePoseFromSkeleton(
+	const FRetargetSkeleton& InSkeleton,
+	const ERetargetSourceOrTarget SourceOrTarget)
 {
-	const int32 TargetChainIndex = GetIndexOfBoneChainByName(InTargetChainName, ERetargetSourceOrTarget::Target);
-	if (ChainMapping.IsValidIndex(TargetChainIndex))
+	// update source poses
+	if (SourceOrTarget == ERetargetSourceOrTarget::Source)
 	{
-		const int32 SourceChainIndex = ChainMapping[TargetChainIndex];
-		if (SourceBoneChains.IsValidIndex(SourceChainIndex))
+		for (FResolvedBoneChain& BoneChain : SourceBoneChains)
 		{
-			return &SourceBoneChains[SourceChainIndex];
+			BoneChain.UpdatePoseFromSkeleton(InSkeleton);
 		}
+		return;
 	}
-	return nullptr;
-}
 
-const FResolvedBoneChain* FRetargeterBoneChains::GetResolvedBoneChainByName(
-	const FName InChainName,
-	const ERetargetSourceOrTarget SourceOrTarget) const
-{
-	const TArray<FResolvedBoneChain>& ResolvedBoneChains = SourceOrTarget == ERetargetSourceOrTarget::Source ? SourceBoneChains : TargetBoneChains;
-	const FResolvedBoneChain* FoundChain = ResolvedBoneChains.FindByPredicate([InChainName](const FResolvedBoneChain& Value)
+	// update target poses
+	for (TTuple<const UIKRigDefinition*, TArray<FResolvedBoneChain>>& Pair : TargetBoneChains)
 	{
-		return Value.ChainName == InChainName;
-	});
-	return FoundChain;
+		for (FResolvedBoneChain& BoneChain : Pair.Value)
+		{
+			BoneChain.UpdatePoseFromSkeleton(InSkeleton);
+		}
+	}
 }
 
-const FResolvedBoneChain* FRetargeterBoneChains::GetResolvedBoneChainByIndex(
-	const int32 InChainIndex,
-	const ERetargetSourceOrTarget SourceOrTarget) const
+TSet<FName> FRetargeterBoneChains::GetChainsThatContainBone(int32 InBoneIndex, ERetargetSourceOrTarget SourceOrTarget) const
 {
-	const TArray<FResolvedBoneChain>& ResolvedBoneChains = SourceOrTarget == ERetargetSourceOrTarget::Source ? SourceBoneChains : TargetBoneChains;
-	if (ResolvedBoneChains.IsValidIndex(InChainIndex))
+	TSet<FName> ChainsThatContainBone;
+	
+	if (SourceOrTarget == ERetargetSourceOrTarget::Source)
 	{
-		return &ResolvedBoneChains[InChainIndex];
+		// source
+		for (const FResolvedBoneChain& SourceChain : SourceBoneChains)
+		{
+			if (SourceChain.BoneIndices.Contains(InBoneIndex))
+			{
+				ChainsThatContainBone.Add(SourceChain.ChainName);
+			}
+		}
+		return MoveTemp(ChainsThatContainBone);
 	}
-
-	return nullptr;
-}
-
-const TArray<int32>& FRetargeterBoneChains::GetChainMappings() const
-{
-	return ChainMapping;
-}
-
-void FRetargeterBoneChains::UpdatePoseFromSkeleton(const FRetargetSkeleton& InSkeleton, const ERetargetSourceOrTarget SourceOrTarget)
-{
-	TArray<FResolvedBoneChain>& BoneChainsToUpdate = SourceOrTarget == ERetargetSourceOrTarget::Source ? SourceBoneChains : TargetBoneChains;
-	for (FResolvedBoneChain& BoneChain : BoneChainsToUpdate)
+	else
 	{
-		BoneChain.UpdatePoseFromSkeleton(InSkeleton);
+		// target
+		for (const TTuple<const UIKRigDefinition*, TArray<FResolvedBoneChain>>& TargetPair : TargetBoneChains)
+		{
+			for (const FResolvedBoneChain& TargetChain : TargetPair.Value)
+			{
+				if (TargetChain.BoneIndices.Contains(InBoneIndex))
+				{
+					ChainsThatContainBone.Add(TargetChain.ChainName);
+				}
+			}
+		}
+		return MoveTemp(ChainsThatContainBone);
 	}
 }
 
@@ -920,8 +906,6 @@ void FRetargeterBoneChains::Reset()
 {
 	SourceBoneChains.Reset();
 	TargetBoneChains.Reset();
-	ChainMapping.Reset();
-	bIsInitialized = false;
 }
 
 void FBasePoseCopier::CopyPoseFromSource(
@@ -1103,7 +1087,6 @@ void FIKRetargetProcessor::Initialize(
 	// initialize source skeleton
 	SourceSkeleton.Initialize(
 		SourceSkeletalMesh,
-		SourceIKRig->GetRetargetChains(),
 		RetargeterAsset->GetCurrentRetargetPoseName(ERetargetSourceOrTarget::Source),
 		RetargeterAsset->GetCurrentRetargetPose(ERetargetSourceOrTarget::Source),
 		SourceIKRig->GetRetargetRoot(),
@@ -1112,18 +1095,10 @@ void FIKRetargetProcessor::Initialize(
 	// initialize target skeleton
 	TargetSkeleton.Initialize(
 		TargetSkeletalMesh,
-		TargetIKRig->GetRetargetChains(),
 		RetargeterAsset->GetCurrentRetargetPoseName(ERetargetSourceOrTarget::Target),
 		RetargeterAsset->GetCurrentRetargetPose(ERetargetSourceOrTarget::Target),
 		TargetIKRig->GetRetargetRoot());
 
-	// resolve all bone chains onto the skeletons
-	AllBoneChains.Initialize(RetargeterAsset, SourceSkeleton, TargetSkeleton, Log);
-
-	// record which chain each bone belongs to (must be done AFTER BoneChains.Initialize())
-	SourceSkeleton.FillPerBoneChainNames(AllBoneChains.GetResolvedBoneChains(ERetargetSourceOrTarget::Source));
-	TargetSkeleton.FillPerBoneChainNames(AllBoneChains.GetResolvedBoneChains(ERetargetSourceOrTarget::Target));
-
 	// setup the goal container
 	ResetIKRigGoalContainer();
 
@@ -1183,12 +1158,8 @@ bool FIKRetargetProcessor::InitializeOpStack(
 		// apply the profile settings to the op (may be different than what was loaded from the asset)
 		InRetargetProfile.ApplyOpProfilesToOpStruct(OpStack[NewOpIndex], ECopyOpSettingsContext::PreInitialize);
 		
-		// initialize it
-		FIKRetargetOpBase& NewOp = OpStack[NewOpIndex].GetMutable<FIKRetargetOpBase>();
-		NewOp.Initialize(*this, SourceSkeleton, TargetSkeleton, Log);
-		
-		
 #if WITH_EDITOR
+		FIKRetargetOpBase& NewOp = OpStack[NewOpIndex].GetMutable<FIKRetargetOpBase>();
 		FIKRetargetOpBase* AssetOp = const_cast<FIKRetargetOpBase*>(AssetOpStruct.GetPtr<FIKRetargetOpBase>());
 		FIKRetargetOpSettingsBase* AssetOpSettings = AssetOp->GetSettings();
 		// asset instance gets reference to running editor instance
@@ -1197,6 +1168,30 @@ bool FIKRetargetProcessor::InitializeOpStack(
 		AssetOpSettings->SourceSkeletonAsset = SourceSkeleton.SkeletalMesh->GetSkeleton();
 		AssetOpSettings->TargetSkeletonAsset = TargetSkeleton.SkeletalMesh->GetSkeleton();
 #endif
+
+	}
+
+	// collect all the IK Rigs referenced by the ops
+	TArray<const UIKRigDefinition*> AllIKRigs;
+	AllIKRigs.Add(RetargeterAsset->GetIKRig(ERetargetSourceOrTarget::Target));
+	for (FInstancedStruct& OpStruct : OpStack)
+	{
+		FIKRetargetOpBase& NewOp = OpStruct.GetMutable<FIKRetargetOpBase>();
+		if (const UIKRigDefinition* TargetIKRigFromOp = NewOp.GetCustomTargetIKRig())
+		{
+			AllIKRigs.AddUnique(TargetIKRigFromOp);
+		}
+	}
+	
+	// resolve all bone chains onto the skeletons
+	AllBoneChains.Initialize(RetargeterAsset, AllIKRigs, SourceSkeleton, TargetSkeleton, Log);
+
+	// initialize the ops
+	for (FInstancedStruct& OpStruct : OpStack)
+	{
+		FIKRetargetOpBase& Op = OpStruct.GetMutable<FIKRetargetOpBase>();
+		const FIKRetargetOpBase* ParentOp = GetRetargetOpByName(Op.GetParentOpName());
+		Op.Initialize(*this, SourceSkeleton, TargetSkeleton, ParentOp, Log);
 	}
 
 	if (!bAllOpsInitialized)
@@ -1352,6 +1347,19 @@ TArray<const FIKRetargetOpBase*> FIKRetargetProcessor::GetRetargetOpsByType(cons
 	return MoveTemp(Ops);
 }
 
+const FIKRetargetOpBase* FIKRetargetProcessor::GetRetargetOpByName(const FName InOpName) const
+{
+	for (const FInstancedStruct& OpStruct : OpStack)
+	{
+		const FIKRetargetOpBase* Op = OpStruct.GetPtr<FIKRetargetOpBase>();
+		if (Op->GetName() == InOpName)
+		{
+			return Op;
+		}
+	}
+	return nullptr;
+}
+
 void FIKRetargetProcessor::OnPlaybackReset()
 {
 	for (FInstancedStruct& OpStruct : OpStack)
@@ -1450,29 +1458,24 @@ bool FIKRetargetProcessor::IsBoneInAMappedChain(
 		return true;
 	}
 	
-	// bones must be in a chain
-	const FName ChainThatContainsBone = Skeleton.ChainThatContainsBone[BoneIndex];
-	if (ChainThatContainsBone == NAME_None)
-	{
-		
-		return false;
-	}
-	
-	// bone must be in a successfully resolved chain
-	const FResolvedBoneChain* Chain = AllBoneChains.GetResolvedBoneChainByName(ChainThatContainsBone, SourceOrTarget);
-	if (!Chain)
+	// bones must be in at least one resolved chain
+	const TSet<FName> ChainsThatContainBone = AllBoneChains.GetChainsThatContainBone(BoneIndex, SourceOrTarget);
+	if (ChainsThatContainBone.IsEmpty())
 	{
 		return false;
 	}
 
 	// bone's chain must be mapped to alternate chain (source-to-target or vice versa)
-	const FName MappedChain = GetChainMappedToChain(Chain->ChainName, SourceOrTarget);
-	if (MappedChain == NAME_None)
+	for (const FName ChainContainingBone : ChainsThatContainBone)
 	{
-		return false;
+		const FName MappedChain = GetFirstChainMappedToChain(ChainContainingBone, SourceOrTarget);
+		if (MappedChain != NAME_None)
+		{
+			return true;
+		}
 	}
 	
-	return true;
+	return false;
 }
 
 int32 FIKRetargetProcessor::GetBoneIndexFromName(
@@ -1496,10 +1499,39 @@ FName FIKRetargetProcessor::GetChainNameForBone(const FName BoneName, const ERet
 	const FName RootBoneName = GetPelvisBone(SourceOrTarget);
 	if (BoneName == RootBoneName)
 	{
-		return FName("Retarget Root");
+		return FName("Retarget Pelvis");
 	}
 
-	return Skeleton.ChainThatContainsBone[BoneIndex];
+	auto GetCombinedName = [](const TSet<FName>& NameSet) -> FName
+	{
+		if (NameSet.Num() == 0)
+		{
+			return NAME_None;
+		}
+    
+		if (NameSet.Num() == 1)
+		{
+			return NameSet.Array()[0];
+		}
+    
+		FString Result;
+		bool bFirst = true;
+    
+		for (const FName& Name : NameSet)
+		{
+			if (!bFirst)
+			{
+				Result += TEXT(", ");
+			}
+			Result += Name.ToString();
+			bFirst = false;
+		}
+    
+		return FName(Result);
+	};
+	
+	const TSet<FName> ChainsThatContainBone = AllBoneChains.GetChainsThatContainBone(BoneIndex, SourceOrTarget);
+	return GetCombinedName(ChainsThatContainBone);
 }
 
 FTransform FIKRetargetProcessor::GetGlobalRetargetPoseAtParam(
@@ -1579,30 +1611,27 @@ FName FIKRetargetProcessor::GetClosestBoneToParam(
 	return Skeleton.BoneNames[Chain->BoneIndices[ChainIndexOfClosestBone]];
 }
 
-FName FIKRetargetProcessor::GetChainMappedToChain(
+FName FIKRetargetProcessor::GetFirstChainMappedToChain(
 	const FName InChainName,
 	const ERetargetSourceOrTarget InSourceOrTarget) const
 {
-	int32 IndexOfChain = GetBoneChains().GetIndexOfBoneChainByName(InChainName, InSourceOrTarget);
-	const TArray<int32>& ChainMappings = GetBoneChains().GetChainMappings();
-	
-	if (InSourceOrTarget == ERetargetSourceOrTarget::Target)
+	if (!ensure(bIsInitialized))
 	{
-		// find the mapped source chain
-		const TArray<FResolvedBoneChain>& SourceChains = GetBoneChains().GetResolvedBoneChains(ERetargetSourceOrTarget::Source);
-		int32 IndexOfMappedSourceChain = ChainMappings[IndexOfChain];
-		return SourceChains.IsValidIndex(IndexOfMappedSourceChain) ? SourceChains[IndexOfMappedSourceChain].ChainName : NAME_None;
+		return NAME_None;
 	}
-	else
+
+	for (const FInstancedStruct& OpStruct : OpStack)
 	{
-		// find the mapped target chain
-		const TArray<FResolvedBoneChain>& TargetChains = GetBoneChains().GetResolvedBoneChains(ERetargetSourceOrTarget::Target);
-		for (int32 TargetChainIndex = 0; TargetChainIndex < ChainMappings.Num(); ++TargetChainIndex)
+		FIKRetargetOpBase* Op = const_cast<FIKRetargetOpBase*>(OpStruct.GetPtr<FIKRetargetOpBase>());
+		const FRetargetChainMapping* ChainMapping = Op->GetChainMapping();
+		if (!ChainMapping)
 		{
-			if (ChainMappings[TargetChainIndex] == IndexOfChain)
-			{
-				return TargetChains.IsValidIndex(TargetChainIndex) ? TargetChains[TargetChainIndex].ChainName : NAME_None;
-			}
+			continue;
+		}
+		const FName MappedChain = ChainMapping->GetChainMappedTo(InChainName, InSourceOrTarget);
+		if (MappedChain != NAME_None)
+		{
+			return MappedChain;
 		}
 	}
 
@@ -1715,14 +1744,9 @@ void FIKRetargetProcessor::UpdateRetargetPoseAtRuntime(
 		return; // retarget pose has not changed since it was initialized
 	}
 	
-	const FName PelvisBoneName = GetPelvisBone(SourceOrTarget);
-	if (PelvisBoneName == NAME_None)
-	{
-		return; // cannot regenerate pose without a pelvis bone
-	}
-
 	// re-generate the retarget pose
 	const double ScaleFactor = SourceOrTarget == ERetargetSourceOrTarget::Source ? SourceScaleFactor : 1.0;
+	const FName PelvisBoneName = GetPelvisBone(SourceOrTarget);
 	RetargetSkeleton.GenerateRetargetPose(NewRetargetPoseName, NewRetargetPose, PelvisBoneName, ScaleFactor);
 
 	// re-load the updated pose into the bone chains
@@ -1731,8 +1755,9 @@ void FIKRetargetProcessor::UpdateRetargetPoseAtRuntime(
 	// re-initialize the op stack (ops sometimes cache retarget poses
 	for (FInstancedStruct& OpStruct : OpStack)
 	{
-		FIKRetargetOpBase* Op = OpStruct.GetMutablePtr<FIKRetargetOpBase>();
-		Op->Initialize(*this, SourceSkeleton, TargetSkeleton, Log);
+		FIKRetargetOpBase& Op = OpStruct.GetMutable<FIKRetargetOpBase>();
+		const FIKRetargetOpBase* ParentOp = GetRetargetOpByName(Op.GetParentOpName());
+		Op.Initialize(*this, SourceSkeleton, TargetSkeleton, ParentOp, Log);
 	}
 }
 
diff --git a/Engine/Plugins/Animation/IKRig/Source/IKRig/Private/Retargeter/IKRetargetProfile.cpp b/Engine/Plugins/Animation/IKRig/Source/IKRig/Private/Retargeter/IKRetargetProfile.cpp
--- a/Engine/Plugins/Animation/IKRig/Source/IKRig/Private/Retargeter/IKRetargetProfile.cpp
+++ b/Engine/Plugins/Animation/IKRig/Source/IKRig/Private/Retargeter/IKRetargetProfile.cpp
@@ -223,6 +223,7 @@ void FRetargetProfile::MergeWithOtherProfile(const FRetargetProfile& OtherProfil
 				const FTargetChainSettings& OtherChainSettings = Pair.Value;
 
 				FRetargetPoleVectorSettings SettingsToMerge;
+				SettingsToMerge.bEnabled = true;
 				SettingsToMerge.TargetChainName = TargetChainName;
 				SettingsToMerge.AlignAlpha = OtherChainSettings.FK.PoleVectorMatching;
 				SettingsToMerge.MaintainOffset = OtherChainSettings.FK.PoleVectorMaintainOffset;
diff --git a/Engine/Plugins/Animation/IKRig/Source/IKRig/Private/Retargeter/IKRetargeter.cpp b/Engine/Plugins/Animation/IKRig/Source/IKRig/Private/Retargeter/IKRetargeter.cpp
--- a/Engine/Plugins/Animation/IKRig/Source/IKRig/Private/Retargeter/IKRetargeter.cpp
+++ b/Engine/Plugins/Animation/IKRig/Source/IKRig/Private/Retargeter/IKRetargeter.cpp
@@ -24,174 +24,6 @@ const FName UIKRetargeter::GetSourcePreviewMeshPropertyName() { return GET_MEMBE
 const FName UIKRetargeter::GetTargetPreviewMeshPropertyName() { return GET_MEMBER_NAME_STRING_CHECKED(UIKRetargeter, TargetPreviewMesh); }
 #endif
 
-bool FRetargetChainMapping::HasChain(const FName& InChainName, ERetargetSourceOrTarget InSourceOrTarget) const
-{
-	if (InSourceOrTarget == ERetargetSourceOrTarget::Target)
-	{
-		return ChainMap.ContainsByPredicate([InChainName](const FRetargetChainPair& InChainPair)
-		{
-			return InChainPair.TargetChainName == InChainName;
-		});
-	}
-	else
-	{
-		return ChainMap.ContainsByPredicate([InChainName](const FRetargetChainPair& InChainPair)
-		{
-			return InChainPair.SourceChainName == InChainName;
-		});
-	}
-}
-
-FName FRetargetChainMapping::GetChainMappedTo(const FName InChainName, ERetargetSourceOrTarget InSourceOrTarget) const
-{
-	if (const FRetargetChainPair* FoundChainPair = FindChainPair(InChainName, InSourceOrTarget))
-	{
-		return InSourceOrTarget == ERetargetSourceOrTarget::Source ? FoundChainPair->TargetChainName : FoundChainPair->SourceChainName;
-	}
-	
-	return NAME_None;
-}
-
-FRetargetChainPair* FRetargetChainMapping::FindChainPair(const FName& InChainName, ERetargetSourceOrTarget InSourceOrTarget) const
-{
-	return const_cast<FRetargetChainPair*>(FindChainPairConst(InChainName, InSourceOrTarget));
-}
-
-const FRetargetChainPair* FRetargetChainMapping::FindChainPairConst(const FName& InChainName, ERetargetSourceOrTarget InSourceOrTarget) const
-{
-	for (const FRetargetChainPair& ChainPair : ChainMap)
-	{
-		const FName& NameToCompare = InSourceOrTarget == ERetargetSourceOrTarget::Source ? ChainPair.SourceChainName : ChainPair.TargetChainName;
-		if (NameToCompare == InChainName)
-		{
-			return &ChainPair;
-		}
-	}
-	return nullptr;
-}
-
-TArray<FRetargetChainPair>& FRetargetChainMapping::GetChainPairs()
-{
-	return ChainMap;
-}
-
-const TArray<FRetargetChainPair>& FRetargetChainMapping::GetChainPairs() const
-{
-	return ChainMap;
-}
-
-void FRetargetChainMapping::SetChainMapping(const FName InTargetChainName, const FName InSourceChainName)
-{
-	if (FRetargetChainPair* FoundChainPair = FindChainPair(InTargetChainName, ERetargetSourceOrTarget::Target))
-	{
-		FoundChainPair->SourceChainName = InSourceChainName;
-		return;
-	}
-
-	ChainMap.Add(FRetargetChainPair(InTargetChainName, InSourceChainName));
-}
-
-void FRetargetChainMapping::CleanMapping(const UIKRigDefinition* SourceIKRig, const UIKRigDefinition* TargetIKRig)
-{
-	// remove null target chains
-	ChainMap.RemoveAll([](const FRetargetChainPair& ChainPair)
-	{
-		return ChainPair.TargetChainName == NAME_None;
-	});
-	
-	if (IsValid(TargetIKRig))
-	{
-		TArray<FName> TargetChainNames = TargetIKRig->GetRetargetChainNames();
-
-		// remove chains in the map that are not in the IK Rig
-		ChainMap.RemoveAll([&TargetChainNames](const FRetargetChainPair& ChainPair)
-		{
-			return !TargetChainNames.Contains(ChainPair.TargetChainName);
-		});
-
-		// add any chains in the rig that are not in the mapping
-		for (const FName& RequiredChain : TargetChainNames)
-		{
-			if (!HasChain(RequiredChain, ERetargetSourceOrTarget::Target))
-			{
-				SetChainMapping(RequiredChain, NAME_None);
-			}
-		}
-	}
-
-	if (IsValid(SourceIKRig))
-	{
-		TArray<FName> SourceChainNames = SourceIKRig->GetRetargetChainNames();
-	
-		// reset any sources that are no longer present to "None"
-		for (FRetargetChainPair& Pair : ChainMap)
-		{
-			if (!SourceChainNames.Contains(Pair.SourceChainName))
-			{
-				Pair.SourceChainName = NAME_None;
-			}
-		}
-	}
-
-	// enforce the chain order based on the StartBone index
-	SortMapping(TargetIKRig);
-}
-
-void FRetargetChainMapping::SortMapping(const UIKRigDefinition* TargetIKRig)
-{
-	if (!IsValid(TargetIKRig))
-	{
-		return;
-	}
-	
-	ChainMap.Sort([TargetIKRig](const FRetargetChainPair& A, const FRetargetChainPair& B)
-	{
-		const TArray<FBoneChain>& BoneChains = TargetIKRig->GetRetargetChains();
-		const FIKRigSkeleton& TargetSkeleton = TargetIKRig->GetSkeleton();
-
-		// look for chains
-		const int32 IndexA = BoneChains.IndexOfByPredicate([&A](const FBoneChain& Chain)
-		{
-			return A.TargetChainName == Chain.ChainName;
-		});
-
-		const int32 IndexB = BoneChains.IndexOfByPredicate([&B](const FBoneChain& Chain)
-		{
-			return B.TargetChainName == Chain.ChainName;
-		});
-
-		// compare their StartBone Index 
-		if (IndexA > INDEX_NONE && IndexB > INDEX_NONE)
-		{
-			const int32 StartBoneIndexA = TargetSkeleton.GetBoneIndexFromName(BoneChains[IndexA].StartBone.BoneName);
-			const int32 StartBoneIndexB = TargetSkeleton.GetBoneIndexFromName(BoneChains[IndexB].StartBone.BoneName);
-
-			if (StartBoneIndexA == StartBoneIndexB)
-			{
-				// fallback to sorting alphabetically
-				return BoneChains[IndexA].ChainName.LexicalLess(BoneChains[IndexB].ChainName);
-			}
-				
-			return StartBoneIndexA < StartBoneIndexB;	
-		}
-
-		// sort them according to the target ik rig if previously failed 
-		return IndexA < IndexB;
-	});
-}
-
-PRAGMA_DISABLE_DEPRECATION_WARNINGS
-void FRetargetChainMapping::LoadFromDeprecatedChainSettings(const TArray<URetargetChainSettings*>& InChainSettings)
-{
-	ChainMap.Empty();
-	
-	for (const URetargetChainSettings* ChainSetting : InChainSettings)
-	{
-		ChainMap.Add(FRetargetChainPair(ChainSetting->TargetChain, ChainSetting->SourceChain));
-	}
-}
-PRAGMA_ENABLE_DEPRECATION_WARNINGS
-
 UIKRetargeter::UIKRetargeter(const FObjectInitializer& ObjectInitializer)
 	: Super(ObjectInitializer)
 {
@@ -293,6 +125,8 @@ void UIKRetargeter::PostLoad()
 	PostLoadOldOpsToNewStructOps();
 
 	PostLoadConvertEverythingToOps();
+
+	PostLoadPutChainMappingInOps();
 #endif
 
 	CleanRetargetPoses();
@@ -409,7 +243,7 @@ void UIKRetargeter::PostLoadConvertEverythingToOps()
 	PRAGMA_DISABLE_DEPRECATION_WARNINGS
 
 	// load old chain settings into new chain map
-	ChainMap.LoadFromDeprecatedChainSettings(ChainSettings_DEPRECATED);
+	ChainMap_DEPRECATED.LoadFromDeprecatedChainSettings(ChainSettings_DEPRECATED);
 	
 	// copy old post-ops to add to the end of the stack
 	TArray<FInstancedStruct> OldPostOps = RetargetOps;
@@ -540,7 +374,7 @@ void UIKRetargeter::PostLoadConvertEverythingToOps()
 			
 			if (Chain->Settings.IK.bAffectedByIKWarping)
 			{
-				Settings->ChainsToStrideWarp.Add(Chain->TargetChain);
+				Settings->ChainSettings.Add(Chain->TargetChain);
 			}
 		}
 	}
@@ -577,7 +411,8 @@ void UIKRetargeter::PostLoadConvertEverythingToOps()
 	{
 		++OpIndex;
 		RetargetOps.EmplaceAt(OpIndex,FIKRetargetRunIKRigOp::StaticStruct());
-		FIKRetargetOpBase* Op = RetargetOps[OpIndex].GetMutablePtr<FIKRetargetOpBase>();
+		FIKRetargetRunIKRigOp* Op = RetargetOps[OpIndex].GetMutablePtr<FIKRetargetRunIKRigOp>();
+		Op->SetName(RunIKOpName);
 		FIKRetargetRunIKRigOpSettings* Settings = reinterpret_cast<FIKRetargetRunIKRigOpSettings*>(Op->GetSettings());
 		Settings->IKRigAsset = TargetIKRigAsset;
 	}
@@ -593,6 +428,7 @@ void UIKRetargeter::PostLoadConvertEverythingToOps()
 			if (Chain->Settings.FK.PoleVectorMatching > 0.0f)
 			{
 				FRetargetPoleVectorSettings ChainToAlign;
+				ChainToAlign.bEnabled = true;
 				ChainToAlign.TargetChainName = Chain->TargetChain;
 				ChainToAlign.AlignAlpha = Chain->Settings.FK.PoleVectorMatching;
 				ChainToAlign.MaintainOffset = Chain->Settings.FK.PoleVectorMaintainOffset;
@@ -607,6 +443,66 @@ void UIKRetargeter::PostLoadConvertEverythingToOps()
 	
 	PRAGMA_ENABLE_DEPRECATION_WARNINGS
 }
+
+void UIKRetargeter::PostLoadPutChainMappingInOps()
+{
+	// only perform this if loading older package version
+	if (GetLinkerCustomVersion(FIKRigObjectVersion::GUID) >= FIKRigObjectVersion::OpsOwnChainMapping)
+	{
+		return;
+	}
+	
+	PRAGMA_DISABLE_DEPRECATION_WARNINGS
+
+	// NOTE: we go through all op types that have been upgraded to support a custom IK Rig asset in order to:
+	// 1. copy the global chain mapping into the op's local chain mapping
+	// 2. call the "OnAddedToStack" virtual to regenerate internal settings structs
+	
+	TArray<FIKRetargetRunIKRigOp*> RunIKOps = GetAllRetargetOpsOfType<FIKRetargetRunIKRigOp>();
+	for (FIKRetargetRunIKRigOp* Op : RunIKOps)
+	{
+		Op->ChainMapping = ChainMap_DEPRECATED;
+		Op->OnAddedToStack(this, nullptr);
+	}
+
+	TArray<FIKRetargetIKChainsOp*> IKChainsOp = GetAllRetargetOpsOfType<FIKRetargetIKChainsOp>();
+	for (FIKRetargetIKChainsOp* Op : IKChainsOp)
+	{
+		const FIKRetargetOpBase* ParentOp = GetRetargetOpByName(Op->GetParentOpName());
+		Op->OnAddedToStack(this, ParentOp);
+	}
+
+	TArray<FIKRetargetStrideWarpingOp*> WarpOp = GetAllRetargetOpsOfType<FIKRetargetStrideWarpingOp>();
+	for (FIKRetargetStrideWarpingOp* Op : WarpOp)
+	{
+		const FIKRetargetOpBase* ParentOp = GetRetargetOpByName(Op->GetParentOpName());
+		Op->OnAddedToStack(this, ParentOp);
+	}
+
+	TArray<FIKRetargetSpeedPlantingOp*> PlantingOp = GetAllRetargetOpsOfType<FIKRetargetSpeedPlantingOp>();
+	for (FIKRetargetSpeedPlantingOp* Op : PlantingOp)
+	{
+		const FIKRetargetOpBase* ParentOp = GetRetargetOpByName(Op->GetParentOpName());
+		Op->OnAddedToStack(this, ParentOp);
+	}
+	
+	TArray<FIKRetargetFKChainsOp*> FKChainOps = GetAllRetargetOpsOfType<FIKRetargetFKChainsOp>();
+	for (FIKRetargetFKChainsOp* Op : FKChainOps)
+	{
+		Op->Settings.IKRigAsset = TargetIKRigAsset;
+		Op->ChainMapping = ChainMap_DEPRECATED;
+		Op->OnAddedToStack(this, nullptr);
+	}
+
+	TArray<FIKRetargetAlignPoleVectorOp*> PoleVectorOps = GetAllRetargetOpsOfType<FIKRetargetAlignPoleVectorOp>();
+	for (FIKRetargetAlignPoleVectorOp* Op : PoleVectorOps)
+	{
+		Op->Settings.IKRigAsset = TargetIKRigAsset;
+		Op->ChainMapping = ChainMap_DEPRECATED;
+		Op->OnAddedToStack(this, nullptr);
+	}
+	PRAGMA_ENABLE_DEPRECATION_WARNINGS
+}
 #endif
 
 FName UIKRetargeter::GetCleanAndUniqueOpName(const FName& InOpName, const int32 InIndexOfOp)
@@ -653,7 +549,7 @@ FName UIKRetargeter::GetCleanAndUniqueOpName(const FName& InOpName, const int32
 		Number++;
 	}
 
-	return OutCleanedName;
+	return FName(OutCleanedName, Number);
 }
 
 void UIKRetargeter::CleanRetargetPoses()
@@ -819,6 +715,40 @@ void UIKRetargeter::CleanOpStack()
 			}
 		}
 	}
+
+	// clean chain mappings inside ops
+	// NOTE: this updates the IK Rig references in the chain mappings and refreshes the list of source/target chains
+	for (FInstancedStruct& OpStruct : RetargetOps)
+	{
+		FIKRetargetOpBase& Op = OpStruct.GetMutable<FIKRetargetOpBase>();
+		FRetargetChainMapping* ChainMapping = Op.GetChainMapping();
+		if (ChainMapping == nullptr)
+		{
+			continue;
+		}
+		const UIKRigDefinition* SourceIKRig = GetIKRig(ERetargetSourceOrTarget::Source);
+		const UIKRigDefinition* TargetIKRig = Op.GetCustomTargetIKRig();
+		if (TargetIKRig == nullptr)
+		{
+			if (const FIKRetargetOpBase* ParentOp = GetRetargetOpByName(Op.GetParentOpName()))
+			{
+				TargetIKRig = ParentOp->GetCustomTargetIKRig();	
+			}
+		}
+		ChainMapping->ReinitializeWithIKRigs(SourceIKRig, TargetIKRig);
+	}
+
+	// give each op a chance to clean its own data based on its parents state
+	for (FInstancedStruct& OpStruct : RetargetOps)
+	{
+		FIKRetargetOpBase& Op = OpStruct.GetMutable<FIKRetargetOpBase>();
+		Op.OnReinitPropertyEdited(nullptr);
+		
+		if (const FIKRetargetOpBase* ParentOp = GetRetargetOpByName(Op.GetParentOpName()))
+		{
+			Op.OnParentReinitPropertyEdited(*ParentOp,nullptr);
+		}
+	}
 	
 	ensureAlwaysMsgf(RetargetOps.Num() == TempOps.Num(), TEXT("Retarget ops were lost during cleaning."));
 };
@@ -889,6 +819,20 @@ void FIKRetargetPose::SortHierarchically(const FIKRigSkeleton& Skeleton)
 	});
 }
 
+const FIKRetargetOpBase* UIKRetargeter::GetRetargetOpByName(const FName& InOpName) const
+{
+	for (const FInstancedStruct& OpStruct : RetargetOps)
+	{
+		const FIKRetargetOpBase& Op = OpStruct.Get<FIKRetargetOpBase>();
+		if (Op.GetName() == InOpName)
+		{
+			return &Op;
+		}
+	}
+	
+	return nullptr;
+}
+
 const FIKRetargetPose* UIKRetargeter::GetCurrentRetargetPose(const ERetargetSourceOrTarget& SourceOrTarget) const
 {
 	return SourceOrTarget == ERetargetSourceOrTarget::Source ? &SourceRetargetPoses[CurrentSourceRetargetPose] : &TargetRetargetPoses[CurrentTargetRetargetPose];
diff --git a/Engine/Plugins/Animation/IKRig/Source/IKRig/Private/Retargeter/RetargetOps/AlignPoleVectorOp.cpp b/Engine/Plugins/Animation/IKRig/Source/IKRig/Private/Retargeter/RetargetOps/AlignPoleVectorOp.cpp
--- a/Engine/Plugins/Animation/IKRig/Source/IKRig/Private/Retargeter/RetargetOps/AlignPoleVectorOp.cpp
+++ b/Engine/Plugins/Animation/IKRig/Source/IKRig/Private/Retargeter/RetargetOps/AlignPoleVectorOp.cpp
@@ -205,6 +205,7 @@ bool FIKRetargetAlignPoleVectorOp::Initialize(
 	const FIKRetargetProcessor& InProcessor,
 	const FRetargetSkeleton& InSourceSkeleton,
 	const FTargetSkeleton& InTargetSkeleton,
+	const FIKRetargetOpBase* InParentOp,
 	FIKRigLogger& InLog)
 {
 	bIsInitialized = false;
@@ -215,7 +216,15 @@ bool FIKRetargetAlignPoleVectorOp::Initialize(
 	const FRetargeterBoneChains& BoneChains = InProcessor.GetBoneChains();
 	for (const FRetargetPoleVectorSettings& ChainSettings : Settings.ChainsToAlign)
 	{
-		const FResolvedBoneChain* TargetBoneChain = BoneChains.GetResolvedBoneChainByName(ChainSettings.TargetChainName, ERetargetSourceOrTarget::Target);
+		if (!ChainSettings.bEnabled)
+		{
+			continue; // if the enabled flag is toggled it will trigger reinit to make a PoleVectorMatcher
+		}
+		
+		const FResolvedBoneChain* TargetBoneChain = BoneChains.GetResolvedBoneChainByName(
+			ChainSettings.TargetChainName,
+			ERetargetSourceOrTarget::Target,
+			Settings.IKRigAsset);
 		if (TargetBoneChain == nullptr)
 		{
 			InLog.LogWarning( FText::Format(
@@ -225,7 +234,8 @@ bool FIKRetargetAlignPoleVectorOp::Initialize(
 		}
 
 		// which source chain was this target chain mapped to?
-		const FResolvedBoneChain* SourceBoneChain = BoneChains.GetSourceChainMappedToTargetChain(ChainSettings.TargetChainName);
+		const FName SourceChainName = ChainMapping.GetChainMappedTo(ChainSettings.TargetChainName, ERetargetSourceOrTarget::Target);
+		const FResolvedBoneChain* SourceBoneChain = BoneChains.GetResolvedBoneChainByName(SourceChainName, ERetargetSourceOrTarget::Source);
 		if (SourceBoneChain == nullptr)
 		{
 			InLog.LogWarning( FText::Format(
@@ -271,6 +281,18 @@ void FIKRetargetAlignPoleVectorOp::Run(
 	}
 }
 
+void FIKRetargetAlignPoleVectorOp::OnAddedToStack(const UIKRetargeter* InRetargetAsset, const FIKRetargetOpBase* InParentOp)
+{
+	// on initial setup, use the default source/target IK rigs
+	const UIKRigDefinition* SourceIKRig = InRetargetAsset->GetIKRig(ERetargetSourceOrTarget::Source);
+	const UIKRigDefinition* TargetIKRig = InRetargetAsset->GetIKRig(ERetargetSourceOrTarget::Target);
+	ApplyIKRigs(SourceIKRig, TargetIKRig);
+
+	// auto map
+	static bool bForceRemap = true;
+	ChainMapping.AutoMapChains(EAutoMapChainType::Fuzzy, bForceRemap);
+}
+
 FIKRetargetOpSettingsBase* FIKRetargetAlignPoleVectorOp::GetSettings()
 {
 	return &Settings;
@@ -308,6 +330,34 @@ const UClass* FIKRetargetAlignPoleVectorOp::GetControllerType() const
 	return UIKRetargetAlignPoleVectorController::StaticClass();
 }
 
+const UIKRigDefinition* FIKRetargetAlignPoleVectorOp::GetCustomTargetIKRig() const
+{
+	return Settings.IKRigAsset;
+}
+
+FRetargetChainMapping* FIKRetargetAlignPoleVectorOp::GetChainMapping()
+{
+	return &ChainMapping;
+}
+
+void FIKRetargetAlignPoleVectorOp::OnTargetChainRenamed(const FName InOldChainName, const FName InNewChainName)
+{
+	for (FRetargetPoleVectorSettings& ChainSettings : Settings.ChainsToAlign)
+	{
+		if (ChainSettings.TargetChainName == InOldChainName)
+		{
+			ChainSettings.TargetChainName = InNewChainName;
+		}
+	}
+}
+
+void FIKRetargetAlignPoleVectorOp::OnReinitPropertyEdited(const FPropertyChangedEvent* InPropertyChangedEvent)
+{
+	const UIKRigDefinition* SourceIKRig = ChainMapping.GetIKRig(ERetargetSourceOrTarget::Source);
+	const UIKRigDefinition* TargetIKRig = Settings.IKRigAsset;
+	ApplyIKRigs(SourceIKRig, TargetIKRig);
+}
+
 #if WITH_EDITOR
 FText FIKRetargetAlignPoleVectorOp::GetWarningMessage() const
 {
@@ -319,7 +369,7 @@ FText FIKRetargetAlignPoleVectorOp::GetWarningMessage() const
 	return FIKRetargetOpBase::GetWarningMessage();
 }
 
-void FIKRetargetAlignPoleVectorOp::SetChainSettingsToDefault(const FName& InChainName)
+void FIKRetargetAlignPoleVectorOp::ResetChainSettingsToDefault(const FName& InChainName)
 {
 	for (FRetargetPoleVectorSettings& ChainToAlign : Settings.ChainsToAlign)
 	{
@@ -345,6 +395,54 @@ bool FIKRetargetAlignPoleVectorOp::AreChainSettingsAtDefault(const FName& InChai
 }
 #endif
 
+void FIKRetargetAlignPoleVectorOp::ApplyIKRigs(const UIKRigDefinition* InSourceIKRig, const UIKRigDefinition* InTargetIKRig)
+{
+	// store IK Rig
+	Settings.IKRigAsset = InTargetIKRig;
+	
+	// update chain mapping
+	ChainMapping.ReinitializeWithIKRigs(InSourceIKRig, InTargetIKRig);
+
+	// update settings only if we have a valid mapping
+	if (!ChainMapping.IsReady())
+	{
+		// don't remove settings, instead we want to preserve existing settings at least until the next valid rig is loaded
+		return;
+	}
+	
+	// get the required target chains
+	const TArray<FBoneChain>& AllTargetChains = Settings.IKRigAsset->GetRetargetChains();
+	TArray<FName> RequiredTargetChains;
+	for (const FBoneChain& ChainToRetarget : AllTargetChains)
+	{
+		RequiredTargetChains.Add(ChainToRetarget.ChainName);
+	}
+	
+	// remove chains that are not required
+	Settings.ChainsToAlign.RemoveAll([&RequiredTargetChains](const FRetargetPoleVectorSettings& InChainSettings)
+	{
+		return !RequiredTargetChains.Contains(InChainSettings.TargetChainName);
+	});
+
+	// add any required chains not already present
+	for (const FName RequiredTargetChain : RequiredTargetChains)
+	{
+		bool bFound = false;
+		for (const FRetargetPoleVectorSettings& ChainToAlign : Settings.ChainsToAlign)
+		{
+			if (ChainToAlign.TargetChainName == RequiredTargetChain)
+			{
+				bFound = true;
+				break;
+			}
+		}
+		if (!bFound)
+		{
+			Settings.ChainsToAlign.Emplace(RequiredTargetChain);
+		}
+	}
+}
+
 FIKRetargetAlignPoleVectorOpSettings UIKRetargetAlignPoleVectorController::GetSettings()
 {
 	return *reinterpret_cast<FIKRetargetAlignPoleVectorOpSettings*>(OpToControl->GetSettings());
diff --git a/Engine/Plugins/Animation/IKRig/Source/IKRig/Private/Retargeter/RetargetOps/CopyBasePoseOp.cpp b/Engine/Plugins/Animation/IKRig/Source/IKRig/Private/Retargeter/RetargetOps/CopyBasePoseOp.cpp
--- a/Engine/Plugins/Animation/IKRig/Source/IKRig/Private/Retargeter/RetargetOps/CopyBasePoseOp.cpp
+++ b/Engine/Plugins/Animation/IKRig/Source/IKRig/Private/Retargeter/RetargetOps/CopyBasePoseOp.cpp
@@ -6,6 +6,7 @@ bool FIKRetargetCopyBasePoseOp::Initialize(
 	const FIKRetargetProcessor& InProcessor,
 	const FRetargetSkeleton& InSourceSkeleton,
 	const FTargetSkeleton& InTargetSkeleton,
+	const FIKRetargetOpBase* InParentOp,
 	FIKRigLogger& Log)
 {
 	bIsInitialized = true;
diff --git a/Engine/Plugins/Animation/IKRig/Source/IKRig/Private/Retargeter/RetargetOps/FKChainsOp.cpp b/Engine/Plugins/Animation/IKRig/Source/IKRig/Private/Retargeter/RetargetOps/FKChainsOp.cpp
--- a/Engine/Plugins/Animation/IKRig/Source/IKRig/Private/Retargeter/RetargetOps/FKChainsOp.cpp
+++ b/Engine/Plugins/Animation/IKRig/Source/IKRig/Private/Retargeter/RetargetOps/FKChainsOp.cpp
@@ -549,6 +549,7 @@ bool FIKRetargetFKChainsOp::Initialize(
 	const FIKRetargetProcessor& InProcessor,
 	const FRetargetSkeleton& InSourceSkeleton,
 	const FTargetSkeleton& InTargetSkeleton,
+	const FIKRetargetOpBase* InParentOp,
 	FIKRigLogger& InLog)
 {
 	bIsInitialized = false;
@@ -560,7 +561,10 @@ bool FIKRetargetFKChainsOp::Initialize(
 	for (const FRetargetFKChainSettings& ChainSettings : Settings.ChainsToRetarget)
 	{
 		FName TargetChainName = ChainSettings.TargetChainName;
-		const FResolvedBoneChain* TargetBoneChain = BoneChains.GetResolvedBoneChainByName(TargetChainName, ERetargetSourceOrTarget::Target);
+		const FResolvedBoneChain* TargetBoneChain = BoneChains.GetResolvedBoneChainByName(
+			TargetChainName,
+			ERetargetSourceOrTarget::Target,
+			Settings.IKRigAsset);
 		if (TargetBoneChain == nullptr)
 		{
 			InLog.LogWarning( FText::Format(
@@ -570,7 +574,8 @@ bool FIKRetargetFKChainsOp::Initialize(
 		}
 
 		// which source chain was this target chain mapped to?
-		const FResolvedBoneChain* SourceBoneChain = BoneChains.GetSourceChainMappedToTargetChain(TargetChainName);
+		const FName SourceChainName = ChainMapping.GetChainMappedTo(ChainSettings.TargetChainName, ERetargetSourceOrTarget::Target);
+		const FResolvedBoneChain* SourceBoneChain = BoneChains.GetResolvedBoneChainByName(SourceChainName, ERetargetSourceOrTarget::Source);
 		if (SourceBoneChain == nullptr)
 		{
 			// this target chain is not mapped to anything (don't spam user about it)
@@ -670,6 +675,18 @@ void FIKRetargetFKChainsOp::PostInitialize(
 	}
 }
 
+void FIKRetargetFKChainsOp::OnAddedToStack(const UIKRetargeter* InRetargetAsset, const FIKRetargetOpBase* InParentOp)
+{
+	// on initial setup, use the default source/target IK rigs
+	const UIKRigDefinition* SourceIKRig = InRetargetAsset->GetIKRig(ERetargetSourceOrTarget::Source);
+	const UIKRigDefinition* TargetIKRig = InRetargetAsset->GetIKRig(ERetargetSourceOrTarget::Target);
+	ApplyIKRigs(SourceIKRig, TargetIKRig);
+
+	// auto map
+	static bool bForceRemap = true;
+	ChainMapping.AutoMapChains(EAutoMapChainType::Fuzzy, bForceRemap);
+}
+
 void FIKRetargetFKChainsOp::CollectRetargetedBones(TSet<int32>& OutRetargetedBones) const
 {
 	// all bones in an FK chain are retargeted
@@ -679,6 +696,34 @@ void FIKRetargetFKChainsOp::CollectRetargetedBones(TSet<int32>& OutRetargetedBon
 	}
 }
 
+const UIKRigDefinition* FIKRetargetFKChainsOp::GetCustomTargetIKRig() const
+{
+	return Settings.IKRigAsset;
+}
+
+FRetargetChainMapping* FIKRetargetFKChainsOp::GetChainMapping()
+{
+	return &ChainMapping;
+}
+
+void FIKRetargetFKChainsOp::OnTargetChainRenamed(const FName InOldChainName, const FName InNewChainName)
+{
+	for (FRetargetFKChainSettings& ChainSettings : Settings.ChainsToRetarget)
+	{
+		if (ChainSettings.TargetChainName == InOldChainName)
+		{
+			ChainSettings.TargetChainName = InNewChainName;
+		}
+	}
+}
+
+void FIKRetargetFKChainsOp::OnReinitPropertyEdited(const FPropertyChangedEvent* InPropertyChangedEvent)
+{
+	const UIKRigDefinition* SourceIKRig = ChainMapping.GetIKRig(ERetargetSourceOrTarget::Source);
+	const UIKRigDefinition* TargetIKRig = Settings.IKRigAsset;
+	ApplyIKRigs(SourceIKRig, TargetIKRig);
+}
+
 FIKRetargetOpSettingsBase* FIKRetargetFKChainsOp::GetSettings()
 {
 	return &Settings;
@@ -721,18 +766,6 @@ const UClass* FIKRetargetFKChainsOp::GetControllerType() const
 }
 
 #if WITH_EDITOR
-void FIKRetargetFKChainsOp::RunInitialSetup(const UIKRetargeter* InRetargetAsset)
-{
-	// copy the target chains
-	if (const UIKRigDefinition* TargetIKRig = InRetargetAsset->GetIKRig(ERetargetSourceOrTarget::Target))
-	{
-		const TArray<FBoneChain>& TargetChains = TargetIKRig->GetRetargetChains();
-		for (const FBoneChain& Chain : TargetChains)
-		{
-			Settings.ChainsToRetarget.Emplace(Chain.ChainName);
-		}
-	}
-}
 
 FCriticalSection FIKRetargetFKChainsOp::DebugDataMutex;
 
@@ -804,7 +837,7 @@ void FIKRetargetFKChainsOp::SaveDebugData(const TArray<FTransform>& OutTargetGlo
 	}
 }
 
-void FIKRetargetFKChainsOp::SetChainSettingsToDefault(const FName& InChainName)
+void FIKRetargetFKChainsOp::ResetChainSettingsToDefault(const FName& InChainName)
 {
 	for (FRetargetFKChainSettings& ChainToRetarget : Settings.ChainsToRetarget)
 	{
@@ -832,6 +865,54 @@ bool FIKRetargetFKChainsOp::AreChainSettingsAtDefault(const FName& InChainName)
 
 #endif
 
+void FIKRetargetFKChainsOp::ApplyIKRigs(const UIKRigDefinition* InSourceIKRig, const UIKRigDefinition* InTargetIKRig)
+{
+	// store IK Rig
+	Settings.IKRigAsset = InTargetIKRig;
+	
+	// update chain mapping
+	ChainMapping.ReinitializeWithIKRigs(InSourceIKRig, InTargetIKRig);
+
+	// update settings only if we have a valid mapping
+	if (!ChainMapping.IsReady())
+	{
+		// don't remove settings, instead we want to preserve existing settings at least until the next valid rig is loaded
+		return;
+	}
+	
+	// get the required target chains
+	const TArray<FBoneChain>& AllTargetChains = Settings.IKRigAsset->GetRetargetChains();
+	TArray<FName> RequiredTargetChains;
+	for (const FBoneChain& ChainToRetarget : AllTargetChains)
+	{
+		RequiredTargetChains.Add(ChainToRetarget.ChainName);
+	}
+	
+	// remove chains that are not required
+	Settings.ChainsToRetarget.RemoveAll([&RequiredTargetChains](const FRetargetFKChainSettings& InChainSettings)
+	{
+		return !RequiredTargetChains.Contains(InChainSettings.TargetChainName);
+	});
+
+	// add any required chains not already present
+	for (const FName RequiredTargetChain : RequiredTargetChains)
+	{
+		bool bFound = false;
+		for (const FRetargetFKChainSettings& ChainToRetarget : Settings.ChainsToRetarget)
+		{
+			if (ChainToRetarget.TargetChainName == RequiredTargetChain)
+			{
+				bFound = true;
+				break;
+			}
+		}
+		if (!bFound)
+		{
+			Settings.ChainsToRetarget.Emplace(RequiredTargetChain);
+		}
+	}
+}
+
 FIKRetargetFKChainsOpSettings UIKRetargetFKChainsController::GetSettings()
 {
 	return *reinterpret_cast<FIKRetargetFKChainsOpSettings*>(OpToControl->GetSettings());
diff --git a/Engine/Plugins/Animation/IKRig/Source/IKRig/Private/Retargeter/RetargetOps/IKChainsOp.cpp b/Engine/Plugins/Animation/IKRig/Source/IKRig/Private/Retargeter/RetargetOps/IKChainsOp.cpp
--- a/Engine/Plugins/Animation/IKRig/Source/IKRig/Private/Retargeter/RetargetOps/IKChainsOp.cpp
+++ b/Engine/Plugins/Animation/IKRig/Source/IKRig/Private/Retargeter/RetargetOps/IKChainsOp.cpp
@@ -191,19 +191,36 @@ bool FIKRetargetIKChainsOp::Initialize(
 	const FIKRetargetProcessor& InProcessor,
 	const FRetargetSkeleton& InSourceSkeleton,
 	const FTargetSkeleton& InTargetSkeleton,
+	const FIKRetargetOpBase* InParentOp,
 	FIKRigLogger& InLog)
 {
 	bIsInitialized = false;
-
 	IKChainRetargeters.Reset();
-	
-	const UIKRigDefinition* TargetIKRig = InProcessor.GetRetargetAsset()->GetIKRig(ERetargetSourceOrTarget::Target);
-	const FRetargeterBoneChains& BoneChains = InProcessor.GetBoneChains();
 
+	// this op requires a parent to supply an IK Rig
+	if (!ensure(InParentOp))
+	{
+		return false;
+	}
+
+	// validate that an IK rig has been assigned
+	const FIKRetargetRunIKRigOp* ParentOp = reinterpret_cast<const FIKRetargetRunIKRigOp*>(InParentOp);
+	if (ParentOp->Settings.IKRigAsset == nullptr)
+	{
+		InLog.LogWarning( FText::Format(
+		LOCTEXT("MissingIKRig", "{0}, is missing an IK rig. No chains can be retargeted."), FText::FromName(GetName())));
+		return false;
+	}
+
+	// go through all chains to retarget and load them
+	const FRetargeterBoneChains& BoneChains = InProcessor.GetBoneChains();
 	for (FRetargetIKChainSettings& ChainSettings : Settings.ChainsToRetarget)
 	{
 		FName TargetChainName = ChainSettings.TargetChainName;
-		const FResolvedBoneChain* TargetBoneChain = BoneChains.GetResolvedBoneChainByName(TargetChainName, ERetargetSourceOrTarget::Target);
+		const FResolvedBoneChain* TargetBoneChain = BoneChains.GetResolvedBoneChainByName(
+			TargetChainName,
+			ERetargetSourceOrTarget::Target,
+			ParentOp->Settings.IKRigAsset);
 
 		// validate that the chain even exists
 		if (TargetBoneChain == nullptr)
@@ -224,7 +241,8 @@ bool FIKRetargetIKChainsOp::Initialize(
 		}
 		
 		// which source chain was this target chain mapped to?
-		const FResolvedBoneChain* SourceBoneChain = BoneChains.GetSourceChainMappedToTargetChain(TargetChainName);
+		const FName SourceChainName = ParentOp->ChainMapping.GetChainMappedTo(TargetChainName, ERetargetSourceOrTarget::Target);
+		const FResolvedBoneChain* SourceBoneChain = BoneChains.GetResolvedBoneChainByName(SourceChainName, ERetargetSourceOrTarget::Source);
 		if (SourceBoneChain == nullptr)
 		{
 			InLog.LogWarning( FText::Format(
@@ -244,7 +262,8 @@ bool FIKRetargetIKChainsOp::Initialize(
 		}
 
 		// warn user if IK goal is not on the END bone of the target chain. It will still work, but may produce bad results.
-		for (const UIKRigEffectorGoal* Goal : TargetIKRig->GetGoalArray())
+		const TArray<UIKRigEffectorGoal*>& AllGoals = ParentOp->Settings.IKRigAsset->GetGoalArray();
+		for (const UIKRigEffectorGoal* Goal : AllGoals)
 		{
 			if (Goal->GoalName != TargetBoneChain->IKGoalName)
 			{
@@ -317,6 +336,11 @@ void FIKRetargetIKChainsOp::Run(
 #endif
 }
 
+void FIKRetargetIKChainsOp::OnAddedToStack(const UIKRetargeter* InRetargetAsset, const FIKRetargetOpBase* InParentOp)
+{
+	RegenerateChainSettings(InParentOp);
+}
+
 FIKRetargetOpSettingsBase* FIKRetargetIKChainsOp::GetSettings()
 {
 	return &Settings;
@@ -363,25 +387,26 @@ const UScriptStruct* FIKRetargetIKChainsOp::GetParentOpType() const
 	return FIKRetargetRunIKRigOp::StaticStruct();
 }
 
-#if WITH_EDITOR
-void FIKRetargetIKChainsOp::RunInitialSetup(const UIKRetargeter* InRetargetAsset)
+void FIKRetargetIKChainsOp::OnTargetChainRenamed(const FName InOldChainName, const FName InNewChainName)
 {
-	// copy the target chains that have IK
-	if (const UIKRigDefinition* TargetIKRig = InRetargetAsset->GetIKRig(ERetargetSourceOrTarget::Target))
+	for (FRetargetIKChainSettings& ChainSettings : Settings.ChainsToRetarget)
 	{
-		const TArray<FBoneChain>& TargetChains = TargetIKRig->GetRetargetChains();
-		for (const FBoneChain& Chain : TargetChains)
+		if (ChainSettings.TargetChainName == InOldChainName)
 		{
-			if (Chain.IKGoalName == NAME_None)
-			{
-				continue;
-			}
-			
-			Settings.ChainsToRetarget.Emplace(Chain.ChainName);
+			ChainSettings.TargetChainName = InNewChainName;
 		}
 	}
 }
 
+void FIKRetargetIKChainsOp::OnParentReinitPropertyEdited(
+	const FIKRetargetOpBase& InParentOp,
+	const FPropertyChangedEvent* InPropertyChangedEvent)
+{
+	RegenerateChainSettings(&InParentOp);
+}
+
+#if WITH_EDITOR
+
 FCriticalSection FIKRetargetIKChainsOp::DebugDataMutex;
 
 void FIKRetargetIKChainsOp::DebugDraw(
@@ -487,7 +512,7 @@ void FIKRetargetIKChainsOp::SaveDebugData(
 	}
 }
 
-void FIKRetargetIKChainsOp::SetChainSettingsToDefault(const FName& InChainName)
+void FIKRetargetIKChainsOp::ResetChainSettingsToDefault(const FName& InChainName)
 {
 	for (FRetargetIKChainSettings& ChainToRetarget : Settings.ChainsToRetarget)
 	{
@@ -515,6 +540,48 @@ bool FIKRetargetIKChainsOp::AreChainSettingsAtDefault(const FName& InChainName)
 
 #endif
 
+void FIKRetargetIKChainsOp::RegenerateChainSettings(const FIKRetargetOpBase* InParentOp)
+{
+	const FIKRetargetRunIKRigOp* ParentOp = reinterpret_cast<const FIKRetargetRunIKRigOp*>(InParentOp);
+	if (!ensure(ParentOp))
+	{
+		return;
+	}
+	
+	// find the target chains that require goal retargeting
+	const TArray<FName> RequiredTargetChains = ParentOp->GetRequiredTargetChains();
+	if (RequiredTargetChains.IsEmpty())
+	{
+		// NOTE: if there's no chains, we don't clear the settings
+		// this allows users to clear and reassign a different rig and potentially retain/restore compatible settings
+		return;
+	}
+
+	// remove chains that are not required
+	Settings.ChainsToRetarget.RemoveAll([&RequiredTargetChains](const FRetargetIKChainSettings& InChainSettings)
+	{
+		return !RequiredTargetChains.Contains(InChainSettings.TargetChainName);
+	});
+	
+	// add any required chains not already present
+	for (FName RequiredTargetChain : RequiredTargetChains)
+	{
+		bool bFound = false;
+		for (const FRetargetIKChainSettings& ChainToRetarget : Settings.ChainsToRetarget)
+		{
+			if (ChainToRetarget.TargetChainName == RequiredTargetChain)
+			{
+				bFound = true;
+				break;
+			}
+		}
+		if (!bFound)
+		{
+			Settings.ChainsToRetarget.Emplace(RequiredTargetChain);
+		}
+	}
+}
+
 FIKRetargetIKChainsOpSettings UIKRetargetIKChainsController::GetSettings()
 {
 	return *reinterpret_cast<FIKRetargetIKChainsOpSettings*>(OpToControl->GetSettings());
diff --git a/Engine/Plugins/Animation/IKRig/Source/IKRig/Private/Retargeter/RetargetOps/PelvisMotionOp.cpp b/Engine/Plugins/Animation/IKRig/Source/IKRig/Private/Retargeter/RetargetOps/PelvisMotionOp.cpp
--- a/Engine/Plugins/Animation/IKRig/Source/IKRig/Private/Retargeter/RetargetOps/PelvisMotionOp.cpp
+++ b/Engine/Plugins/Animation/IKRig/Source/IKRig/Private/Retargeter/RetargetOps/PelvisMotionOp.cpp
@@ -35,6 +35,7 @@ bool FIKRetargetPelvisMotionOp::Initialize(
 	const FIKRetargetProcessor& Processor,
 	const FRetargetSkeleton& SourceSkeleton,
 	const FTargetSkeleton& TargetSkeleton,
+	const FIKRetargetOpBase* InParentOp,
 	FIKRigLogger& Log)
 {
 	bIsInitialized = false;
@@ -91,6 +92,20 @@ void FIKRetargetPelvisMotionOp::Run(
 #endif
 }
 
+void FIKRetargetPelvisMotionOp::OnAddedToStack(const UIKRetargeter* InRetargetAsset, const FIKRetargetOpBase* InParentOp)
+{
+	// copy the source/target pelvis from the default IK Rig
+	if (const UIKRigDefinition* SourceIKRig = InRetargetAsset->GetIKRig(ERetargetSourceOrTarget::Source))
+	{
+		Settings.SourcePelvisBone.BoneName = SourceIKRig->GetRetargetRoot();
+	}
+	
+	if (const UIKRigDefinition* TargetIKRig = InRetargetAsset->GetIKRig(ERetargetSourceOrTarget::Target))
+	{
+		Settings.TargetPelvisBone.BoneName = TargetIKRig->GetRetargetRoot();
+	}
+}
+
 bool FIKRetargetPelvisMotionOp::InitializeSource(
 	const FName SourcePelvisBoneName,
 	const FRetargetSkeleton& SourceSkeleton,
@@ -284,20 +299,6 @@ void FIKRetargetPelvisMotionOp::DecodePose(TArray<FTransform>& OutTargetGlobalPo
 
 #if WITH_EDITOR
 
-void FIKRetargetPelvisMotionOp::RunInitialSetup(const UIKRetargeter* InRetargetAsset)
-{
-	// copy the source/target pelvis from the default IK Rig
-	if (const UIKRigDefinition* SourceIKRig = InRetargetAsset->GetIKRig(ERetargetSourceOrTarget::Source))
-	{
-		Settings.SourcePelvisBone.BoneName = SourceIKRig->GetRetargetRoot();
-	}
-	
-	if (const UIKRigDefinition* TargetIKRig = InRetargetAsset->GetIKRig(ERetargetSourceOrTarget::Target))
-	{
-		Settings.TargetPelvisBone.BoneName = TargetIKRig->GetRetargetRoot();
-	}
-}
-
 FText FIKRetargetPelvisMotionOp::GetWarningMessage() const
 {
 	if (bIsInitialized)
diff --git a/Engine/Plugins/Animation/IKRig/Source/IKRig/Private/Retargeter/RetargetOps/PinBoneOp.cpp b/Engine/Plugins/Animation/IKRig/Source/IKRig/Private/Retargeter/RetargetOps/PinBoneOp.cpp
--- a/Engine/Plugins/Animation/IKRig/Source/IKRig/Private/Retargeter/RetargetOps/PinBoneOp.cpp
+++ b/Engine/Plugins/Animation/IKRig/Source/IKRig/Private/Retargeter/RetargetOps/PinBoneOp.cpp
@@ -90,6 +90,7 @@ bool FIKRetargetPinBoneOp::Initialize(
 	const FIKRetargetProcessor& Processor,
 	const FRetargetSkeleton& SourceSkeleton,
 	const FTargetSkeleton& TargetSkeleton,
+	const FIKRetargetOpBase* InParentOp,
 	FIKRigLogger& Log)
 {
 	bFoundAllBonesToPin = true;
diff --git a/Engine/Plugins/Animation/IKRig/Source/IKRig/Private/Retargeter/RetargetOps/RootMotionGeneratorOp.cpp b/Engine/Plugins/Animation/IKRig/Source/IKRig/Private/Retargeter/RetargetOps/RootMotionGeneratorOp.cpp
--- a/Engine/Plugins/Animation/IKRig/Source/IKRig/Private/Retargeter/RetargetOps/RootMotionGeneratorOp.cpp
+++ b/Engine/Plugins/Animation/IKRig/Source/IKRig/Private/Retargeter/RetargetOps/RootMotionGeneratorOp.cpp
@@ -61,6 +61,7 @@ bool FIKRetargetRootMotionOp::Initialize(
 	const FIKRetargetProcessor& Processor,
 	const FRetargetSkeleton& SourceSkeleton,
 	const FTargetSkeleton& TargetSkeleton,
+	const FIKRetargetOpBase* InParentOp,
 	FIKRigLogger& Log)
 {
 	bIsInitialized = false;
@@ -170,6 +171,29 @@ void FIKRetargetRootMotionOp::PostInitialize(
 	}
 }
 
+void FIKRetargetRootMotionOp::OnAddedToStack(const UIKRetargeter* InRetargetAsset, const FIKRetargetOpBase* InParentOp)
+{
+#if WITH_EDITOR
+	// get the root of the SOURCE skeleton
+	if (const USkeletalMesh* Mesh = InRetargetAsset->GetPreviewMesh(ERetargetSourceOrTarget::Source))
+	{
+		Settings.SourceRoot = Mesh->GetRefSkeleton().GetBoneName(0);
+	}
+	
+	// get the root of the TARGET skeleton
+	if (const USkeletalMesh* Mesh = InRetargetAsset->GetPreviewMesh(ERetargetSourceOrTarget::Target))
+	{
+		Settings.TargetRoot = Mesh->GetRefSkeleton().GetBoneName(0);
+	}
+
+	// get the retarget root bone
+	if (const UIKRigDefinition* IKRig = InRetargetAsset->GetIKRig(ERetargetSourceOrTarget::Target))
+	{
+		Settings.TargetPelvis = IKRig->GetRetargetRoot();
+	}
+#endif
+}
+
 FIKRetargetOpSettingsBase* FIKRetargetRootMotionOp::GetSettings()
 {
 	return &Settings;
@@ -280,26 +304,6 @@ void FIKRetargetRootMotionOp::CopyRootMotionFromSourceRoot(
 }
 
 #if WITH_EDITOR
-void FIKRetargetRootMotionOp::RunInitialSetup(const UIKRetargeter* InRetargetAsset)
-{
-	// get the root of the SOURCE skeleton
-	if (const USkeletalMesh* Mesh = InRetargetAsset->GetPreviewMesh(ERetargetSourceOrTarget::Source))
-	{
-		Settings.SourceRoot = Mesh->GetRefSkeleton().GetBoneName(0);
-	}
-	
-	// get the root of the TARGET skeleton
-	if (const USkeletalMesh* Mesh = InRetargetAsset->GetPreviewMesh(ERetargetSourceOrTarget::Target))
-	{
-		Settings.TargetRoot = Mesh->GetRefSkeleton().GetBoneName(0);
-	}
-
-	// get the retarget root bone
-	if (const UIKRigDefinition* IKRig = InRetargetAsset->GetIKRig(ERetargetSourceOrTarget::Target))
-	{
-		Settings.TargetPelvis = IKRig->GetRetargetRoot();
-	}
-}
 
 FText FIKRetargetRootMotionOp::GetWarningMessage() const
 {
diff --git a/Engine/Plugins/Animation/IKRig/Source/IKRig/Private/Retargeter/RetargetOps/RunIKRigOp.cpp b/Engine/Plugins/Animation/IKRig/Source/IKRig/Private/Retargeter/RetargetOps/RunIKRigOp.cpp
--- a/Engine/Plugins/Animation/IKRig/Source/IKRig/Private/Retargeter/RetargetOps/RunIKRigOp.cpp
+++ b/Engine/Plugins/Animation/IKRig/Source/IKRig/Private/Retargeter/RetargetOps/RunIKRigOp.cpp
@@ -12,6 +12,7 @@ bool FIKRetargetRunIKRigOp::Initialize(
 	const FIKRetargetProcessor& InProcessor,
 	const FRetargetSkeleton& InSourceSkeleton,
 	const FTargetSkeleton& InTargetSkeleton,
+	const FIKRetargetOpBase* InParentOp,
 	FIKRigLogger& InLog)
 {
 	bIsInitialized = false;
@@ -87,6 +88,17 @@ void FIKRetargetRunIKRigOp::Run(
 #endif
 }
 
+void FIKRetargetRunIKRigOp::OnAddedToStack(const UIKRetargeter* InRetargetAsset, const FIKRetargetOpBase* InParentOp)
+{
+	// load the target IK Rig asset to execute
+	Settings.IKRigAsset = InRetargetAsset->GetIKRig(ERetargetSourceOrTarget::Target);
+
+	// initialize the chain mapping
+	ChainMapping.ReinitializeWithIKRigs(InRetargetAsset->GetIKRig(ERetargetSourceOrTarget::Source), Settings.IKRigAsset);
+	static bool bForceRemap = true;
+	ChainMapping.AutoMapChains(EAutoMapChainType::Fuzzy, bForceRemap);
+}
+
 FIKRetargetOpSettingsBase* FIKRetargetRunIKRigOp::GetSettings()
 {
 	return &Settings;
@@ -114,14 +126,23 @@ const UClass* FIKRetargetRunIKRigOp::GetControllerType() const
 	return UIKRetargetRunIKRigController::StaticClass();
 }
 
+const UIKRigDefinition* FIKRetargetRunIKRigOp::GetCustomTargetIKRig() const
+{
+	return Settings.IKRigAsset;
+}
+
+FRetargetChainMapping* FIKRetargetRunIKRigOp::GetChainMapping()
+{
+	return &ChainMapping;
+}
 
-#if WITH_EDITOR
-void FIKRetargetRunIKRigOp::RunInitialSetup(const UIKRetargeter* InRetargetAsset)
+void FIKRetargetRunIKRigOp::OnReinitPropertyEdited(const FPropertyChangedEvent* InPropertyChangedEvent)
 {
-	// load the target IK Rig asset to execute
-	Settings.IKRigAsset = InRetargetAsset->GetIKRig(ERetargetSourceOrTarget::Target);
+	ChainMapping.ReinitializeWithIKRigs(ChainMapping.GetIKRig(ERetargetSourceOrTarget::Source), Settings.IKRigAsset);
 }
 
+#if WITH_EDITOR
+
 FCriticalSection FIKRetargetRunIKRigOp::DebugDataMutex;
 
 void FIKRetargetRunIKRigOp::DebugDraw(
@@ -227,6 +248,35 @@ void FIKRetargetRunIKRigOp::SaveCurrentGoalTransformsIntoDebugData()
 }
 #endif
 
+TArray<FName> FIKRetargetRunIKRigOp::GetRequiredTargetChains() const
+{
+	if (!Settings.IKRigAsset)
+	{
+		return TArray<FName>{};
+	}
+	
+	// find the target chains that children ops should deal with
+	TArray<FName> RequiredTargetChains;
+	const TArray<FBoneChain>& TargetChains = Settings.IKRigAsset->GetRetargetChains();
+	for (const FBoneChain& TargetChain : TargetChains)
+	{
+		if (TargetChain.IKGoalName == NAME_None)
+		{
+			continue; // skip non-IK chains
+		}
+
+		const FName SourceChain = ChainMapping.GetChainMappedTo(TargetChain.ChainName, ERetargetSourceOrTarget::Target);
+		if (SourceChain == NAME_None)
+		{
+			continue; // skip unmapped chains
+		}
+		
+		RequiredTargetChains.Add(TargetChain.ChainName);
+	}
+	
+	return MoveTemp(RequiredTargetChains);
+}
+
 FIKRetargetRunIKRigOpSettings UIKRetargetRunIKRigController::GetSettings()
 {
 	return *reinterpret_cast<FIKRetargetRunIKRigOpSettings*>(OpToControl->GetSettings());
diff --git a/Engine/Plugins/Animation/IKRig/Source/IKRig/Private/Retargeter/RetargetOps/SpeedPlantingOp.cpp b/Engine/Plugins/Animation/IKRig/Source/IKRig/Private/Retargeter/RetargetOps/SpeedPlantingOp.cpp
--- a/Engine/Plugins/Animation/IKRig/Source/IKRig/Private/Retargeter/RetargetOps/SpeedPlantingOp.cpp
+++ b/Engine/Plugins/Animation/IKRig/Source/IKRig/Private/Retargeter/RetargetOps/SpeedPlantingOp.cpp
@@ -14,9 +14,25 @@ bool FIKRetargetSpeedPlantingOp::Initialize(
 	const FIKRetargetProcessor& InProcessor,
 	const FRetargetSkeleton& InSourceSkeleton,
 	const FTargetSkeleton& InTargetSkeleton,
+	const FIKRetargetOpBase* InParentOp,
 	FIKRigLogger& InLog)
 {
 	bIsInitialized = false;
+	
+	// this op requires a parent to supply an IK Rig
+	if (!ensure(InParentOp))
+	{
+		return false;
+	}
+
+	// validate that an IK rig has been assigned
+	const FIKRetargetRunIKRigOp* ParentOp = reinterpret_cast<const FIKRetargetRunIKRigOp*>(InParentOp);
+	if (ParentOp->Settings.IKRigAsset == nullptr)
+	{
+		InLog.LogWarning( FText::Format(
+		LOCTEXT("MissingIKRig", "{0}, is missing an IK rig. No chains can be retargeted. "), FText::FromName(GetName())));
+		return false;
+	}
 
 	GoalsToPlant.Reset();
 
@@ -25,7 +41,10 @@ bool FIKRetargetSpeedPlantingOp::Initialize(
 	for (const FRetargetSpeedPlantingSettings& SettingsForChain : Settings.ChainsToSpeedPlant)
 	{
 		const FName TargetChainName = SettingsForChain.TargetChainName;
-		const FResolvedBoneChain* TargetBoneChain = BoneChains.GetResolvedBoneChainByName(TargetChainName, ERetargetSourceOrTarget::Target);
+		const FResolvedBoneChain* TargetBoneChain = BoneChains.GetResolvedBoneChainByName(
+			TargetChainName,
+			ERetargetSourceOrTarget::Target,
+			ParentOp->Settings.IKRigAsset);
 		if (TargetBoneChain == nullptr)
 		{
 			InLog.LogWarning( FText::Format(
@@ -99,6 +118,11 @@ void FIKRetargetSpeedPlantingOp::Run(
 	}
 }
 
+void FIKRetargetSpeedPlantingOp::OnAddedToStack(const UIKRetargeter* InRetargetAsset, const FIKRetargetOpBase* InParentOp)
+{
+	RegenerateChainSettings(InParentOp);
+}
+
 FIKRetargetOpSettingsBase* FIKRetargetSpeedPlantingOp::GetSettings()
 {
 	return &Settings;
@@ -145,6 +169,24 @@ const UScriptStruct* FIKRetargetSpeedPlantingOp::GetParentOpType() const
 	return FIKRetargetRunIKRigOp::StaticStruct();
 }
 
+void FIKRetargetSpeedPlantingOp::OnTargetChainRenamed(const FName InOldChainName, const FName InNewChainName)
+{
+	for (FRetargetSpeedPlantingSettings& ChainSettings : Settings.ChainsToSpeedPlant)
+	{
+		if (ChainSettings.TargetChainName == InOldChainName)
+		{
+			ChainSettings.TargetChainName = InNewChainName;
+		}
+	}
+}
+
+void FIKRetargetSpeedPlantingOp::OnParentReinitPropertyEdited(
+	const FIKRetargetOpBase& InParentOp,
+	const FPropertyChangedEvent* InPropertyChangedEvent)
+{
+	RegenerateChainSettings(&InParentOp);
+}
+
 void FIKRetargetSpeedPlantingOp::OnPlaybackReset()
 {
 	for (FSpeedPlantingGoalState& GoalToPlant : GoalsToPlant)
@@ -237,6 +279,48 @@ FText FIKRetargetSpeedPlantingOp::GetWarningMessage() const
 }
 #endif
 
+void FIKRetargetSpeedPlantingOp::RegenerateChainSettings(const FIKRetargetOpBase* InParentOp)
+{
+	const FIKRetargetRunIKRigOp* ParentOp = reinterpret_cast<const FIKRetargetRunIKRigOp*>(InParentOp);
+	if (!ensure(ParentOp))
+	{
+		return;
+	}
+	
+	// find the target chains that require goal retargeting
+	const TArray<FName> RequiredTargetChains = ParentOp->GetRequiredTargetChains();
+	if (RequiredTargetChains.IsEmpty())
+	{
+		// NOTE: if there's no chains, we don't clear the settings
+		// this allows users to clear and reassign a different rig and potentially retain/restore compatible settings
+		return;
+	}
+
+	// remove chains that are not required
+	Settings.ChainsToSpeedPlant.RemoveAll([&RequiredTargetChains](const FRetargetSpeedPlantingSettings& InChainSettings)
+	{
+		return !RequiredTargetChains.Contains(InChainSettings.TargetChainName);
+	});
+	
+	// add any required chains not already present
+	for (FName RequiredTargetChain : RequiredTargetChains)
+	{
+		bool bFound = false;
+		for (const FRetargetSpeedPlantingSettings& ChainToSpeedPlant : Settings.ChainsToSpeedPlant)
+		{
+			if (ChainToSpeedPlant.TargetChainName == RequiredTargetChain)
+			{
+				bFound = true;
+				break;
+			}
+		}
+		if (!bFound)
+		{
+			Settings.ChainsToSpeedPlant.Emplace(RequiredTargetChain);
+		}
+	}
+}
+
 FIKRetargetSpeedPlantingOpSettings UIKRetargetSpeedPlantingController::GetSettings()
 {
 	return *reinterpret_cast<FIKRetargetSpeedPlantingOpSettings*>(OpToControl->GetSettings());
diff --git a/Engine/Plugins/Animation/IKRig/Source/IKRig/Private/Retargeter/RetargetOps/StrideWarpingOp.cpp b/Engine/Plugins/Animation/IKRig/Source/IKRig/Private/Retargeter/RetargetOps/StrideWarpingOp.cpp
--- a/Engine/Plugins/Animation/IKRig/Source/IKRig/Private/Retargeter/RetargetOps/StrideWarpingOp.cpp
+++ b/Engine/Plugins/Animation/IKRig/Source/IKRig/Private/Retargeter/RetargetOps/StrideWarpingOp.cpp
@@ -2,6 +2,7 @@
 
 #include "Retargeter/RetargetOps/StrideWarpingOp.h"
 
+#include "IKRigObjectVersion.h"
 #include "Retargeter/IKRetargetProcessor.h"
 #include "Retargeter/RetargetOps/RunIKRigOp.h"
 #include "Rig/Solvers/PointsToRotation.h"
@@ -14,6 +15,11 @@
 #define LOCTEXT_NAMESPACE "StrideWarpingOp"
 
 
+bool FRetargetStrideWarpChainSettings::operator==(const FRetargetStrideWarpChainSettings& Other) const
+{
+	return EnableStrideWarping == Other.EnableStrideWarping;
+}
+
 FVector FIKRetargetStrideWarpingOpSettings::GetAxisVector(const EBasicAxis& Axis)
 {
 	switch (Axis)
@@ -50,20 +56,48 @@ bool FIKRetargetStrideWarpingOp::Initialize(
 	const FIKRetargetProcessor& InProcessor,
 	const FRetargetSkeleton& InSourceSkeleton,
 	const FTargetSkeleton& InTargetSkeleton,
+	const FIKRetargetOpBase* InParentOp,
 	FIKRigLogger& InLog)
 {
 	bIsInitialized = false;
+	GoalsToWarp.Reset();
+
+	// this op requires a parent to supply an IK Rig
+	if (!ensure(InParentOp))
+	{
+		return false;
+	}
+
+	// validate that an IK rig has been assigned
+	const FIKRetargetRunIKRigOp* ParentOp = reinterpret_cast<const FIKRetargetRunIKRigOp*>(InParentOp);
+	if (ParentOp->Settings.IKRigAsset == nullptr)
+	{
+		InLog.LogWarning( FText::Format(
+		LOCTEXT("MissingIKRig", "{0}, is missing an IK rig. No chains can be retargeted. "), FText::FromName(GetName())));
+		return false;
+	}
+
+	// store the target IK rig for querying bone chains
+	TargetIKRig = ParentOp->Settings.IKRigAsset;
 	
 	const FRetargeterBoneChains& BoneChains = InProcessor.GetBoneChains();
 	const FIKRigGoalContainer& GoalContainer = InProcessor.GetIKRigGoalContainer();
-	for (const FName TargetChainName : Settings.ChainsToStrideWarp)
+	for (const FRetargetStrideWarpChainSettings& ChainSettings : Settings.ChainSettings)
 	{
-		const FResolvedBoneChain* TargetBoneChain = BoneChains.GetResolvedBoneChainByName(TargetChainName, ERetargetSourceOrTarget::Target);
+		if (!ChainSettings.EnableStrideWarping)
+		{
+			continue;
+		}
+		
+		const FResolvedBoneChain* TargetBoneChain = BoneChains.GetResolvedBoneChainByName(
+			ChainSettings.TargetChainName,
+			ERetargetSourceOrTarget::Target,
+			TargetIKRig);
 		if (TargetBoneChain == nullptr)
 		{
 			InLog.LogWarning( FText::Format(
 			LOCTEXT("StrideWarpingMissingChain", "Stride Warping Op: chain data is out of sync with IK Rig. Missing target chain, '{0}."),
-			FText::FromName(TargetChainName)));
+			FText::FromName(ChainSettings.TargetChainName)));
 			continue;
 		}
 
@@ -71,7 +105,7 @@ bool FIKRetargetStrideWarpingOp::Initialize(
 		{
 			InLog.LogWarning( FText::Format(
 			LOCTEXT("StrideWarpingChainWithNoGoal", "Stride Warping Op: specified chain does not have an IK goal. Cannot stride warp, '{0}."),
-			FText::FromName(TargetChainName)));
+			FText::FromName(ChainSettings.TargetChainName)));
 			continue;
 		}
 
@@ -169,7 +203,7 @@ void FIKRetargetStrideWarpingOp::Run(
 		{
 			// use chain to determine the body's rotation
 			const FRetargeterBoneChains& BoneChains = InProcessor.GetBoneChains();
-			const FResolvedBoneChain* Chain = BoneChains.GetResolvedBoneChainByName(Settings.DirectionChain, ERetargetSourceOrTarget::Target);
+			const FResolvedBoneChain* Chain = BoneChains.GetResolvedBoneChainByName(Settings.DirectionChain, ERetargetSourceOrTarget::Target, TargetIKRig);
 			if (!Chain)
 			{
 				return;
@@ -288,6 +322,11 @@ void FIKRetargetStrideWarpingOp::Run(
 #endif
 }
 
+void FIKRetargetStrideWarpingOp::OnAddedToStack(const UIKRetargeter* InRetargetAsset, const FIKRetargetOpBase* InParentOp)
+{
+	RegenerateChainSettings(InParentOp);
+}
+
 FIKRetargetOpSettingsBase* FIKRetargetStrideWarpingOp::GetSettings()
 {
 	return &Settings;
@@ -320,6 +359,24 @@ const UScriptStruct* FIKRetargetStrideWarpingOp::GetParentOpType() const
 	return FIKRetargetRunIKRigOp::StaticStruct();
 }
 
+void FIKRetargetStrideWarpingOp::OnTargetChainRenamed(const FName InOldChainName, const FName InNewChainName)
+{
+	for (FRetargetStrideWarpChainSettings& ChainSettings : Settings.ChainSettings)
+	{
+		if (ChainSettings.TargetChainName == InOldChainName)
+		{
+			ChainSettings.TargetChainName = InNewChainName;
+		}
+	}
+}
+
+void FIKRetargetStrideWarpingOp::OnParentReinitPropertyEdited(
+	const FIKRetargetOpBase& InParentOp,
+	const FPropertyChangedEvent* InPropertyChangedEvent)
+{
+	RegenerateChainSettings(&InParentOp);
+}
+
 #if WITH_EDITOR
 
 FCriticalSection FIKRetargetStrideWarpingOp::DebugDataMutex;
@@ -347,27 +404,81 @@ void FIKRetargetStrideWarpingOp::DebugDraw(
 		static_cast<float>(Settings.DebugDrawThickness * InComponentScale));
 }
 
-void FIKRetargetStrideWarpingOp::SetChainSettingsToDefault(const FName& InChainName)
+void FIKRetargetStrideWarpingOp::ResetChainSettingsToDefault(const FName& InChainName)
 {
-	if (Settings.ChainsToStrideWarp.Contains(InChainName))
+	for (FRetargetStrideWarpChainSettings& ChainToStrideWarp : Settings.ChainSettings)
 	{
-		FIKRetargetStrideWarpingOpSettings ResetSettings = FIKRetargetStrideWarpingOpSettings();
-		SetSettings(&ResetSettings);
+		if (ChainToStrideWarp.TargetChainName != InChainName)
+		{
+			continue;
+		}
+
+		// reset
+		const FName ChainName = ChainToStrideWarp.TargetChainName;
+		ChainToStrideWarp = FRetargetStrideWarpChainSettings(ChainName);
+		return;
 	}
 }
 
 bool FIKRetargetStrideWarpingOp::AreChainSettingsAtDefault(const FName& InChainName)
 {
-	if (Settings.ChainsToStrideWarp.Contains(InChainName))
+	for (FRetargetStrideWarpChainSettings& ChainToStrideWarp : Settings.ChainSettings)
 	{
-		FIKRetargetStrideWarpingOpSettings DefaultSettings = FIKRetargetStrideWarpingOpSettings();
-		return Settings == DefaultSettings;
+		if (ChainToStrideWarp.TargetChainName != InChainName)
+		{
+			continue;
+		}
+		
+		FRetargetStrideWarpChainSettings DefaultSettings = FRetargetStrideWarpChainSettings(ChainToStrideWarp.TargetChainName);
+		return ChainToStrideWarp == DefaultSettings;
 	}
 
 	return true;
 }
 #endif
 
+void FIKRetargetStrideWarpingOp::RegenerateChainSettings(const FIKRetargetOpBase* InParentOp)
+{
+	const FIKRetargetRunIKRigOp* ParentOp = reinterpret_cast<const FIKRetargetRunIKRigOp*>(InParentOp);
+	if (!ensure(ParentOp))
+	{
+		return;
+	}
+	
+	// find the target chains that require goal retargeting
+	const TArray<FName> RequiredTargetChains = ParentOp->GetRequiredTargetChains();
+	if (RequiredTargetChains.IsEmpty())
+	{
+		// NOTE: if there's no chains, we don't clear the settings
+		// this allows users to clear and reassign a different rig and potentially retain/restore compatible settings
+		return;
+	}
+
+	// remove chains that are not required
+	Settings.ChainSettings.RemoveAll([&RequiredTargetChains](const FRetargetStrideWarpChainSettings& InChainSettings)
+	{
+		return !RequiredTargetChains.Contains(InChainSettings.TargetChainName);
+	});
+	
+	// add any required chains not already present
+	for (FName RequiredTargetChain : RequiredTargetChains)
+	{
+		bool bFound = false;
+		for (const FRetargetStrideWarpChainSettings& ChainToStrideWarp : Settings.ChainSettings)
+		{
+			if (ChainToStrideWarp.TargetChainName == RequiredTargetChain)
+			{
+				bFound = true;
+				break;
+			}
+		}
+		if (!bFound)
+		{
+			Settings.ChainSettings.Emplace(RequiredTargetChain);
+		}
+	}
+}
+
 FIKRetargetStrideWarpingOpSettings UIKRetargetStrideWarpingController::GetSettings()
 {
 	return *reinterpret_cast<FIKRetargetStrideWarpingOpSettings*>(OpToControl->GetSettings());
diff --git a/Engine/Plugins/Animation/IKRig/Source/IKRig/Public/IKRigObjectVersion.h b/Engine/Plugins/Animation/IKRig/Source/IKRig/Public/IKRigObjectVersion.h
--- a/Engine/Plugins/Animation/IKRig/Source/IKRig/Public/IKRigObjectVersion.h
+++ b/Engine/Plugins/Animation/IKRig/Source/IKRig/Public/IKRigObjectVersion.h
@@ -30,6 +30,9 @@ struct IKRIG_API FIKRigObjectVersion
 		// modularized retargeting internals (converted phases to ops)
 		ModularRetargeterOps,
 
+		// ops control their own chain mapping (no more global chain map on retarget asset)
+		OpsOwnChainMapping,
+
 		// -----<new versions can be added above this line>-------------------------------------------------
 		VersionPlusOne,
 		LatestVersion = VersionPlusOne - 1
diff --git a/Engine/Plugins/Animation/IKRig/Source/IKRig/Public/Retargeter/IKRetargetOps.h b/Engine/Plugins/Animation/IKRig/Source/IKRig/Public/Retargeter/IKRetargetOps.h
--- a/Engine/Plugins/Animation/IKRig/Source/IKRig/Public/Retargeter/IKRetargetOps.h
+++ b/Engine/Plugins/Animation/IKRig/Source/IKRig/Public/Retargeter/IKRetargetOps.h
@@ -4,6 +4,7 @@
 
 #include "IKRetargetOps.generated.h"
 
+class UIKRigDefinition;
 struct FIKRetargetProcessor;
 struct FRetargetSkeleton;
 struct FTargetSkeleton;
@@ -12,6 +13,7 @@ class UIKRetargeter;
 struct FInstancedStruct;
 class USkeleton;
 struct FPoseContext;
+struct FRetargetChainMapping;
 
 #if WITH_EDITOR
 class FPrimitiveDrawInterface;
@@ -39,7 +41,10 @@ struct IKRIG_API FIKRetargetOpSettingsBase
 	
 	virtual ~FIKRetargetOpSettingsBase() {};
 
-#if WITH_EDITOR
+#if WITH_EDITORONLY_DATA
+	/** the op these settings belong to */
+	UPROPERTY(transient)
+	FName OwningOpName;
 	/** the instance of this op currently running in the editor viewport */
 	mutable FIKRetargetOpSettingsBase* EditorInstance;
 	/** a reference to the skeletons for bone selector widgets (can return these from GetSkeleton) */
@@ -72,6 +77,7 @@ public:
 		const FIKRetargetProcessor& InProcessor,
 		const FRetargetSkeleton& InSourceSkeleton,
 		const FTargetSkeleton& InTargetSkeleton,
+		const FIKRetargetOpBase* InParentOp,
 		FIKRigLogger& Log) { bIsInitialized = true; return true; };
 
 	/** (optional) override to evaluate this operation and modify the output pose */
@@ -88,6 +94,9 @@ public:
 		const FTargetSkeleton& InTargetSkeleton,
 		FIKRigLogger& InLog) {};
 	
+	/** (optional) override to automate initial setup after being added to the stack */
+	virtual void OnAddedToStack(const UIKRetargeter* InRetargetAsset, const FIKRetargetOpBase* InParentOp) {};
+	
 	/** get if this op is enabled */
 	bool IsEnabled() const { return bIsEnabled; };
 
@@ -143,13 +152,35 @@ public:
 	/** (optional) ops can optionally behave as a 'child' where they can be parented to ops of the type returned by this function. */
 	virtual const UScriptStruct* GetParentOpType() const { return nullptr; };
 
+	/** (optional) return true to disallow multiple copies of this op in the stack */
+	virtual bool IsSingleton() const { return false; };
+
+	/** (optional) override and supply the target IK Rig this op references
+	 * NOTE: The retarget processor will resolve the bone chains in the rig so that they can be queried by the op. */
+	virtual const UIKRigDefinition* GetCustomTargetIKRig() const { return nullptr; };
+	
+	/** (optional) ops can optionally store their own chain mapping, this allows outside systems to query/edit it. */
+	virtual FRetargetChainMapping* GetChainMapping() { return nullptr; };
+
+	/** (optional) implement this if the op stores chain settings by name to allow the settings to be maintained after a chain is renamed
+	 * NOTE: this is only called when the target IK Rig this op references has a chain that is renamed.
+	 * NOTE: ops do not have to update chain mappings in this function, that is managed for them if GetChainMapping() is implemented */
+	virtual void OnTargetChainRenamed(const FName InOldChainName, const FName InNewChainName) {};
+
+	/** (optional) override to allow ops to react when a property marked "ReinitializeOnEdit" is modified
+	 * NOTE: this is called before Initialize() is called again to give op a chance to auto-configure itself based on its new state
+	 * NOTE: this is called in several places where InPropertyChangedEvent may be null */
+	virtual void OnReinitPropertyEdited(const FPropertyChangedEvent* InPropertyChangedEvent) {};
+
+	/** (optional) override to allow ops to react when a property marked "ReinitializeOnEdit" on the parent op is modified
+	 * NOTE: InPropertyChangedEvent will be null at load time giving op a chance to clean based on parent state.
+	 * NOTE: this is called in several places where InPropertyChangedEvent may be null */
+	virtual void OnParentReinitPropertyEdited(const FIKRetargetOpBase& InParentOp, const FPropertyChangedEvent* InPropertyChangedEvent) {};
+
 #if WITH_EDITOR
 	
 	/** get the nice name to display in the viewport (defaults to "DisplayName" UStruct metadata) */
 	FName GetDefaultName() const;
-	
-	/** (optional) override to automate initial setup after being added to the stack */
-	virtual void RunInitialSetup(const UIKRetargeter* InRetargetAsset) {};
 
 	/** (optional) override to display a warning message in the op stack */
 	virtual FText GetWarningMessage() const;
@@ -162,7 +193,7 @@ public:
 		const FIKRetargetDebugDrawState& InEditorState) const {};
 
 	/** (optional) override to support resetting all settings for a given chain */
-	virtual void SetChainSettingsToDefault(const FName& InChainName) {};
+	virtual void ResetChainSettingsToDefault(const FName& InChainName) {};
 
 	/** (optional) override to support "Reset to Default" UI for chains */
 	virtual bool AreChainSettingsAtDefault(const FName& InChainName) { return true; };
diff --git a/Engine/Plugins/Animation/IKRig/Source/IKRig/Public/Retargeter/IKRetargetProcessor.h b/Engine/Plugins/Animation/IKRig/Source/IKRig/Public/Retargeter/IKRetargetProcessor.h
--- a/Engine/Plugins/Animation/IKRig/Source/IKRig/Public/Retargeter/IKRetargetProcessor.h
+++ b/Engine/Plugins/Animation/IKRig/Source/IKRig/Public/Retargeter/IKRetargetProcessor.h
@@ -48,24 +48,20 @@ struct IKRIG_API FRetargetSkeleton
 	double RetargetPoseScale;				// the scale of the retarget pose this was initialized with (transient)
 	
 	const USkeletalMesh* SkeletalMesh;		// the skeletal mesh this was initialized with
-	TArray<FName> ChainThatContainsBone;	// record which chain is actually controlling each bone
 
 	void Initialize(
 		const USkeletalMesh* InSkeletalMesh,
-		const TArray<FBoneChain>& BoneChains,
 		const FName InRetargetPoseName,
 		const FIKRetargetPose* RetargetPose,
 		const FName RetargetRootBone,
 		const double ScaleFactor);
-
-	void FillPerBoneChainNames(const TArray<FResolvedBoneChain>& InBoneChains);
-
+	
 	void Reset();
 
 	void GenerateRetargetPose(
 		const FName InRetargetPoseName,
 		const FIKRetargetPose* InRetargetPose,
-		const FName RetargetRootBone,
+		const FName PelvisBoneName,
 		const double ScaleFactor);
 
 	int32 FindBoneIndexByName(const FName InName) const;
@@ -127,7 +123,6 @@ struct FTargetSkeleton : public FRetargetSkeleton
 
 	void Initialize(
 		const USkeletalMesh* InSkeletalMesh,
-		const TArray<FBoneChain>& BoneChains,
 		const FName InRetargetPoseName,
 		const FIKRetargetPose* RetargetPose,
 		const FName RetargetRootBone);
@@ -232,30 +227,34 @@ struct IKRIG_API FRetargeterBoneChains
 	/** Load all chains from the IK Rigs and resolved them on the current skeletal meshes*/
 	bool Initialize(
 		const UIKRetargeter* InRetargetAsset,
+		const TArray<const UIKRigDefinition*>& InTargetIKRigs,
 		const FRetargetSkeleton& InSourceSkeleton,
 		const FRetargetSkeleton& InTargetSkeleton,
 		FIKRigLogger& InLog);
 	
-	/** Get read only access to the resolved bone chains on either the source or the target */
-	const TArray<FResolvedBoneChain>& GetResolvedBoneChains(ERetargetSourceOrTarget SourceOrTarget) const;
-
-	/** Get the index of a bone chain by a given name. Returns INDEX_None if not found. */
-	int32 GetIndexOfBoneChainByName(const FName InChainName, const ERetargetSourceOrTarget SourceOrTarget) const;
-
-	/** Given the name of a target chain, returns the source chain that it's mapped to. Returns nullptr if no mapping exists */
-	const FResolvedBoneChain* GetSourceChainMappedToTargetChain(const FName InTargetChainName) const;
-
-	/** Get a fully resolved bone chain by name */
-	const FResolvedBoneChain* GetResolvedBoneChainByName(const FName InChainName, const ERetargetSourceOrTarget SourceOrTarget) const;
-
-	/** Get a fully resolved bone chain by inddex */
-	const FResolvedBoneChain* GetResolvedBoneChainByIndex(const int32 InChainIndex, const ERetargetSourceOrTarget SourceOrTarget) const;
-
-	/** Get read only access to the mapping from each target bone chain to a source chain (or potentially None) */
-	const TArray<int32>& GetChainMappings() const;
+	/** Get read only access to the resolved bone chains on either the source or the target
+	 * If asking for a target chain, must supply the IK Rig it belongs to */
+	const TArray<FResolvedBoneChain>* GetResolvedBoneChains(
+		ERetargetSourceOrTarget SourceOrTarget,
+		const UIKRigDefinition* TargetIKRig = nullptr) const;
+
+	/** Get read only access to all resolved bone chains, for all IK Rigs used on the target */
+	const TMap<const UIKRigDefinition*, TArray<FResolvedBoneChain>>& GetAllResolvedTargetBoneChains() const;
+
+	/** Get a fully resolved bone chain by name
+	 * If asking for a target chain, must supply the IK Rig it belongs to. */
+	const FResolvedBoneChain* GetResolvedBoneChainByName(
+		const FName InChainName,
+		const ERetargetSourceOrTarget SourceOrTarget,
+		const UIKRigDefinition* TargetIKRig = nullptr) const;
 
 	/** Update all transforms stored in the chains. This is needed whenever the retarget pose is modified */
-	void UpdatePoseFromSkeleton(const FRetargetSkeleton& InSkeleton, const ERetargetSourceOrTarget SourceOrTarget);
+	void UpdatePoseFromSkeleton(
+		const FRetargetSkeleton& InSkeleton,
+		const ERetargetSourceOrTarget SourceOrTarget);
+
+	/** Get a list of all chains that contain the bone with the given index */
+	TSet<FName> GetChainsThatContainBone(int32 InBoneIndex, ERetargetSourceOrTarget SourceOrTarget) const;
 
 	/** Reset everything (between initializations) */
 	void Reset();
@@ -265,11 +264,11 @@ private:
 	/** All resolved chains on the source mesh */
 	TArray<FResolvedBoneChain> SourceBoneChains;
 
-	/** All resolved chains on the target mesh */
-	TArray<FResolvedBoneChain> TargetBoneChains;
+	/** All resolved chains on the target mesh for each IK Rig used by the ops */
+	TMap<const UIKRigDefinition*, TArray<FResolvedBoneChain>> TargetBoneChains;
 
-	/** This array is same length as TargetBoneChains and stores index into the mapped source chain for each target chain */
-	TArray<int32> ChainMapping;
+	/** the default IK Rig */
+	const UIKRigDefinition* DefaultTargetIKRig;
 
 	bool bIsInitialized = false;
 };
@@ -404,6 +403,9 @@ public:
 	/** Get list of ops of a given type */
 	TArray<const FIKRetargetOpBase*> GetRetargetOpsByType(const UScriptStruct* OpType) const;
 
+	/** Get an op with the given name */
+	const FIKRetargetOpBase* GetRetargetOpByName(const FName InOpName) const;
+
 	/** Get the first op in the stack of the given type */
 	template <typename T>
 	const T* GetFirstRetargetOpOfType(ERetargetOpsToSearch SourceOfOps = ERetargetOpsToSearch::ProcessorOps) const
@@ -488,7 +490,7 @@ public:
 	/** Get the bone in the chain at the given param. */
 	FName GetClosestBoneToParam(const FName InChainName, const float InParam, const ERetargetSourceOrTarget SourceOrTarget) const;
 	/** Get the chain mapped to this one. */
-	FName GetChainMappedToChain(const FName InChainName, const ERetargetSourceOrTarget InSourceOrTarget) const;
+	FName GetFirstChainMappedToChain(const FName InChainName, const ERetargetSourceOrTarget InSourceOrTarget) const;
 	/** Attach a delegate to be notified whenever this processor is initialized. */
 	FOnRetargeterInitialized& OnRetargeterInitialized(){ return RetargeterInitialized; };
 	/** Run all debug drawing on all ops in the stack */
@@ -669,7 +671,7 @@ public:
 	}
 	FName GetMappedChainName(const FName InChainName, const ERetargetSourceOrTarget SourceOrTarget)
 	{
-		return Processor.GetChainMappedToChain(InChainName, SourceOrTarget);
+		return Processor.GetFirstChainMappedToChain(InChainName, SourceOrTarget);
 	}
 #endif
 
diff --git a/Engine/Plugins/Animation/IKRig/Source/IKRig/Public/Retargeter/IKRetargeter.h b/Engine/Plugins/Animation/IKRig/Source/IKRig/Public/Retargeter/IKRetargeter.h
--- a/Engine/Plugins/Animation/IKRig/Source/IKRig/Public/Retargeter/IKRetargeter.h
+++ b/Engine/Plugins/Animation/IKRig/Source/IKRig/Public/Retargeter/IKRetargeter.h
@@ -2,6 +2,7 @@
 
 #pragma once
 
+#include "IKRetargetChainMapping.h"
 #include "IKRetargetProfile.h"
 #include "Rig/IKRigDefinition.h"
 #include "IKRetargetSettings.h"
@@ -52,72 +53,6 @@ private:
 	friend class UIKRetargeterController;
 };
 
-USTRUCT()
-struct IKRIG_API FRetargetChainPair
-{
-	GENERATED_BODY()
-
-	FRetargetChainPair() = default;
-	FRetargetChainPair(const FName InTargetChainName, const FName InSourceChainName) :
-		TargetChainName(InTargetChainName), SourceChainName(InSourceChainName) {}
-
-	UPROPERTY()
-	FName TargetChainName = NAME_None;
-	
-	UPROPERTY()
-	FName SourceChainName = NAME_None;
-};
-
-USTRUCT()
-struct IKRIG_API FRetargetChainMapping
-{
-	GENERATED_BODY()
-
-	/** Get the name of the bone chain mapped to this chain.
-	 * @param InChainName the name of either a source or target chain
-	 * @param InSourceOrTarget the side that InChainName belongs to (either source or target)
-	 * @return true if the chain is present in the mapping */
-	bool HasChain(const FName& InChainName, ERetargetSourceOrTarget InSourceOrTarget) const;
-
-	/** Get the name of the bone chain mapped to this chain.
-	 * @param InChainName the name of either a source or target chain
-	 * @param InSourceOrTarget the side that InChainName belongs to (either source or target)
-	 * @return the name of the chain that InChainName is mapped to (possibly NAME_None) */
-	FName GetChainMappedTo(const FName InChainName, ERetargetSourceOrTarget InSourceOrTarget) const;
-
-	/** returns the chain pair belonging to the supplied chain or nullptr if not found
-	* @param InChainName the name of either a source or target chain
-	* @param InSourceOrTarget the side that InChainName belongs to (either source or target)*/
-	FRetargetChainPair* FindChainPair(const FName& InChainName, ERetargetSourceOrTarget InSourceOrTarget) const;
-	const FRetargetChainPair* FindChainPairConst(const FName& InChainName, ERetargetSourceOrTarget InSourceOrTarget) const;
-
-	/** Get read-write access to the mapping from target to source chains (by name) */
-	TArray<FRetargetChainPair>& GetChainPairs();
-	const TArray<FRetargetChainPair>& GetChainPairs() const;
-
-	/** Map target chain to a source chain.
-	 * @param InTargetChainName the name of a target chain
-	 * @param InSourceChainName the name of a source chain*/
-	void SetChainMapping(const FName InTargetChainName, const FName InSourceChainName);
-
-	/* synchronizes chain names with given source/target IK Rigs */
-	void CleanMapping(const UIKRigDefinition* SourceIKRig, const UIKRigDefinition* TargetIKRig);
-
-	/* sort mapping hierarchically (root to leaf based on Target IK Rig hierarchy) */
-	void SortMapping(const UIKRigDefinition* TargetIKRig);
-	
-	PRAGMA_DISABLE_DEPRECATION_WARNINGS
-	void LoadFromDeprecatedChainSettings(const TArray<URetargetChainSettings*>& InChainSettings);
-	PRAGMA_ENABLE_DEPRECATION_WARNINGS
-
-private:
-	
-	/** mapping of target to source bone chains by name
-	 * NOTE: this is an array instead of a TMap because it needs to be sorted by hierarchy*/
-	UPROPERTY()
-	TArray<FRetargetChainPair> ChainMap;
-};
-
 UCLASS(BlueprintType)
 class IKRIG_API UIKRetargeter : public UObject
 {
@@ -136,11 +71,7 @@ public:
 	// Get read-only access to preview meshes
 	USkeletalMesh* GetPreviewMesh(ERetargetSourceOrTarget SourceOrTarget) const;
 	#endif
-
-	// Get read/write access to the chain mapping
-	FRetargetChainMapping& GetChainMapping() { return ChainMap; };
-	// Get read-only access to the chain mapping
-	const FRetargetChainMapping& GetChainMapping() const { return ChainMap; };
+	
 	// Get access to the stack of retargeting operations
 	const TArray<FInstancedStruct>& GetRetargetOps() const { return RetargetOps; };
 	/** Get the first op in the stack of the given type */
@@ -176,6 +107,8 @@ public:
     
 		return Results;
 	}
+	// Get a retarget op by name
+	const FIKRetargetOpBase* GetRetargetOpByName(const FName& InOpName) const;
 	// Get read-only access to a retarget pose 
 	const FIKRetargetPose* GetCurrentRetargetPose(const ERetargetSourceOrTarget& SourceOrTarget) const;
 	// Get name of the current retarget pose 
@@ -233,7 +166,16 @@ public:
 	// we cannot provide backwards compatibility.
 	// Instead you can use the new op controller API to get a controller for any retarget op and get/set it's settings.
 	PRAGMA_DISABLE_DEPRECATION_WARNINGS
+
+	// Get read/write access to the chain mapping
+	UE_DEPRECATED(5.6, "Chain mappings are managed by individual ops now. Use the asset or op controllers to edit them.")
+	FRetargetChainMapping& GetChainMapping() { return ChainMap_DEPRECATED; };
+
 	// Get read-only access to the chain mapping
+	UE_DEPRECATED(5.6, "Chain mappings are managed by individual ops now. Use Op controllers to edit them.")
+	const FRetargetChainMapping& GetChainMapping() const { return ChainMap_DEPRECATED; };
+	
+	// Get read-only access to the per-chain settings
 	UE_DEPRECATED(5.6, "Chain settings are now accessed through an Op controller.")
 	const TArray<TObjectPtr<URetargetChainSettings>>& GetAllChainSettings() const { return ChainSettings_DEPRECATED; };
 	
@@ -362,6 +304,7 @@ private:
 	void PostLoadOldSettingsToNew();
 	void PostLoadOldOpsToNewStructOps();
 	void PostLoadConvertEverythingToOps();
+	void PostLoadPutChainMappingInOps();
 #endif
 	void CleanRetargetPoses();
 	void CleanOpStack();
@@ -440,11 +383,7 @@ private:
 	 * executed in serial fashion where output of prior operation is input to the next */
 	UPROPERTY(meta = (ExcludeBaseStruct, BaseStruct = "/Script/IKRig.RetargetOpBase"))
 	TArray<FInstancedStruct> RetargetOps;
-
-	/** mapping of target to source bone chains by name */
-	UPROPERTY()
-	FRetargetChainMapping ChainMap;
-
+	
 	// settings profiles stored in this asset 
 	UPROPERTY()
 	TMap<FName, FRetargetProfile> Profiles;
@@ -487,6 +426,8 @@ private:
 	//
 	// BEGIN deprecated data from 5.6 refactor
 	//
+	UPROPERTY(meta = (DeprecatedProperty))
+	FRetargetChainMapping ChainMap_DEPRECATED;
 	// Settings for how to map source chains to target chains.
 	UPROPERTY(meta = (DeprecatedProperty, DeprecationMessage = "Use FK Chain and IK Chain Op settings instead to affect chain settings."))
 	TArray<TObjectPtr<URetargetChainSettings>> ChainSettings_DEPRECATED;
