[IK Retargeter] Adding support for chain mapping per-op

diff --git a/Engine/Plugins/Animation/IKRig/Source/IKRig/Public/Retargeter/RetargetOps/AlignPoleVectorOp.h b/Engine/Plugins/Animation/IKRig/Source/IKRig/Public/Retargeter/RetargetOps/AlignPoleVectorOp.h
--- a/Engine/Plugins/Animation/IKRig/Source/IKRig/Public/Retargeter/RetargetOps/AlignPoleVectorOp.h
+++ b/Engine/Plugins/Animation/IKRig/Source/IKRig/Public/Retargeter/RetargetOps/AlignPoleVectorOp.h
@@ -1,6 +1,7 @@
 ﻿// Copyright Epic Games, Inc. All Rights Reserved.
 #pragma once
 
+#include "Retargeter/IKRetargetChainMapping.h"
 #include "Retargeter/IKRetargetOps.h"
 
 #include "AlignPoleVectorOp.generated.h"
@@ -65,18 +66,22 @@ struct IKRIG_API FRetargetPoleVectorSettings
 	UPROPERTY(EditAnywhere, BlueprintReadWrite, Category = "Setup", meta=(ReinitializeOnEdit))
 	FName TargetChainName;
 
+	/** Enable pole vector alignment on this chain. */
+	UPROPERTY(EditAnywhere, BlueprintReadWrite, Category = "Align Pole Vector", meta=(ReinitializeOnEdit))
+	bool bEnabled = false;
+
 	/** Range 0 to 1. Default 1. Matches the twist angle of this chain (along the Pole direction) to the source chain.
 	*  At 0, the chain's pole vector direction will be left alone
 	*  At 1, the root bone of the chain will be twist-rotated in the pole direction to match the orientation of the source chain.*/
-	UPROPERTY(EditAnywhere, BlueprintReadWrite, Category = "Align Pole Vector", meta = (UIMin = "0.0", UIMax = "1.0", ClampMin = "0.0", ClampMax = "1.0"))
+	UPROPERTY(EditAnywhere, BlueprintReadWrite, Category = "Align Pole Vector", meta = (EditCondition = "bEnabled", UIMin = "0.0", UIMax = "1.0", ClampMin = "0.0", ClampMax = "1.0"))
 	double AlignAlpha = 1.0f;
 
 	/** Range +/- 180. Default 0. An angular offset, in degrees, for the pole direction of the chain. */
-	UPROPERTY(EditAnywhere, BlueprintReadWrite, Category = "Align Pole Vector", meta = (UIMin = "-180.0", UIMax = "180.0", ClampMin = "-180.0", ClampMax = "180.0"))
+	UPROPERTY(EditAnywhere, BlueprintReadWrite, Category = "Align Pole Vector", meta = (EditCondition = "bEnabled", UIMin = "-180.0", UIMax = "180.0", ClampMin = "-180.0", ClampMax = "180.0"))
 	double StaticAngularOffset = 0.0f;
 
 	/** Default is False. When true, the original angular offset between the source/target pole vectors will be maintained when aligning pole vectors. */
-	UPROPERTY(EditAnywhere, BlueprintReadWrite, Category = "Align Pole Vector")
+	UPROPERTY(EditAnywhere, BlueprintReadWrite, Category = "Align Pole Vector", meta = (EditCondition = "bEnabled"))
 	bool MaintainOffset = false;
 	
 	bool operator==(const FRetargetPoleVectorSettings& Other) const;
@@ -86,6 +91,10 @@ USTRUCT(BlueprintType, meta = (DisplayName = "Pole Vector Alignment Settings"))
 struct IKRIG_API FIKRetargetAlignPoleVectorOpSettings : public FIKRetargetOpSettingsBase
 {
 	GENERATED_BODY()
+
+	/** The target IK Rig asset that contains the bones chains to retarget in this op. */
+	UPROPERTY(EditAnywhere, BlueprintReadWrite, Category = "IK Rig Asset", meta=(ReinitializeOnEdit))
+	TObjectPtr<const UIKRigDefinition> IKRigAsset = nullptr;
 	
 	UPROPERTY(EditAnywhere, BlueprintReadWrite, Category = "Pole Vector Op Settings")
 	TArray<FRetargetPoleVectorSettings> ChainsToAlign;
@@ -102,6 +111,7 @@ struct IKRIG_API FIKRetargetAlignPoleVectorOp : public FIKRetargetOpBase
 		const FIKRetargetProcessor& InProcessor,
 		const FRetargetSkeleton& InSourceSkeleton,
 		const FTargetSkeleton& InTargetSkeleton,
+		const FIKRetargetOpBase* InParentOp,
 		FIKRigLogger& InLog) override;
 	
 	virtual void Run(
@@ -110,6 +120,8 @@ struct IKRIG_API FIKRetargetAlignPoleVectorOp : public FIKRetargetOpBase
 		const TArray<FTransform>& InSourceGlobalPose,
 		TArray<FTransform>& OutTargetGlobalPose) override;
 
+	virtual void OnAddedToStack(const UIKRetargeter* InRetargetAsset, const FIKRetargetOpBase* InParentOp) override;
+
 	virtual FIKRetargetOpSettingsBase* GetSettings() override;
 
 	virtual void SetSettings(const FIKRetargetOpSettingsBase* InSettings) override;
@@ -119,18 +131,33 @@ struct IKRIG_API FIKRetargetAlignPoleVectorOp : public FIKRetargetOpBase
 	virtual const UScriptStruct* GetType() const override;
 
 	virtual const UClass* GetControllerType() const override;
+	
+	virtual const UIKRigDefinition* GetCustomTargetIKRig() const override;
+
+	virtual FRetargetChainMapping* GetChainMapping() override;
+
+	virtual void OnTargetChainRenamed(const FName InOldChainName, const FName InNewChainName) override;
+	
+	virtual void OnReinitPropertyEdited(const FPropertyChangedEvent* InPropertyChangedEvent) override;
 
 #if WITH_EDITOR
 	virtual FText GetWarningMessage() const override;
 
-	virtual void SetChainSettingsToDefault(const FName& InChainName) override;
+	virtual void ResetChainSettingsToDefault(const FName& InChainName) override;
 	virtual bool AreChainSettingsAtDefault(const FName& InChainName) override;
 #endif
 
 	UPROPERTY()
 	FIKRetargetAlignPoleVectorOpSettings Settings;
 
+	/* This op maintains its own chain mapping table to allow per-op mapping */
+	UPROPERTY()
+	FRetargetChainMapping ChainMapping;
+
 private:
+	
+	void ApplyIKRigs(const UIKRigDefinition* InSourceIKRig, const UIKRigDefinition* InTargetIKRig);
+	
 	TArray<FPoleVectorMatcher> PoleVectorMatchers;
 };
 
diff --git a/Engine/Plugins/Animation/IKRig/Source/IKRig/Public/Retargeter/RetargetOps/CopyBasePoseOp.h b/Engine/Plugins/Animation/IKRig/Source/IKRig/Public/Retargeter/RetargetOps/CopyBasePoseOp.h
--- a/Engine/Plugins/Animation/IKRig/Source/IKRig/Public/Retargeter/RetargetOps/CopyBasePoseOp.h
+++ b/Engine/Plugins/Animation/IKRig/Source/IKRig/Public/Retargeter/RetargetOps/CopyBasePoseOp.h
@@ -35,6 +35,7 @@ struct IKRIG_API FIKRetargetCopyBasePoseOp : public FIKRetargetOpBase
 		const FIKRetargetProcessor& InProcessor,
 		const FRetargetSkeleton& InSourceSkeleton,
 		const FTargetSkeleton& InTargetSkeleton,
+		const FIKRetargetOpBase* InParentOp,
 		FIKRigLogger& Log) override;
 	
 	virtual void Run(
@@ -50,6 +51,8 @@ struct IKRIG_API FIKRetargetCopyBasePoseOp : public FIKRetargetOpBase
 	virtual const UScriptStruct* GetSettingsType() const override;
 	
 	virtual const UScriptStruct* GetType() const override;
+
+	virtual bool IsSingleton() const override { return true; };
 	
 	virtual const UClass* GetControllerType() const override;
 
diff --git a/Engine/Plugins/Animation/IKRig/Source/IKRig/Public/Retargeter/RetargetOps/CurveRemapOp.h b/Engine/Plugins/Animation/IKRig/Source/IKRig/Public/Retargeter/RetargetOps/CurveRemapOp.h
--- a/Engine/Plugins/Animation/IKRig/Source/IKRig/Public/Retargeter/RetargetOps/CurveRemapOp.h
+++ b/Engine/Plugins/Animation/IKRig/Source/IKRig/Public/Retargeter/RetargetOps/CurveRemapOp.h
@@ -56,6 +56,7 @@ struct IKRIG_API FIKRetargetCurveRemapOp : public FIKRetargetOpBase
 		const FIKRetargetProcessor& InProcessor,
 		const FRetargetSkeleton& InSourceSkeleton,
 		const FTargetSkeleton& InTargetSkeleton,
+		const FIKRetargetOpBase* InParentOp,
 		FIKRigLogger& Log) override
 	{
 		bIsInitialized = true;
@@ -88,6 +89,8 @@ struct IKRIG_API FIKRetargetCurveRemapOp : public FIKRetargetOpBase
 		return FIKRetargetCurveRemapOp::StaticStruct();
 	}
 
+	virtual bool IsSingleton() const override { return true; };
+
 	virtual void AnimGraphPreUpdateMainThread(USkeletalMeshComponent& SourceMeshComponent, USkeletalMeshComponent& TargetMeshComponent) override;
 
 	virtual void AnimGraphEvaluateAnyThread(FPoseContext& Output) override;
diff --git a/Engine/Plugins/Animation/IKRig/Source/IKRig/Public/Retargeter/RetargetOps/FKChainsOp.h b/Engine/Plugins/Animation/IKRig/Source/IKRig/Public/Retargeter/RetargetOps/FKChainsOp.h
--- a/Engine/Plugins/Animation/IKRig/Source/IKRig/Public/Retargeter/RetargetOps/FKChainsOp.h
+++ b/Engine/Plugins/Animation/IKRig/Source/IKRig/Public/Retargeter/RetargetOps/FKChainsOp.h
@@ -1,6 +1,7 @@
 ﻿// Copyright Epic Games, Inc. All Rights Reserved.
 #pragma once
 
+#include "Retargeter/IKRetargetChainMapping.h"
 #include "Retargeter/IKRetargetOps.h"
 #include "Retargeter/IKRetargetProcessor.h"
 
@@ -165,20 +166,29 @@ USTRUCT(BlueprintType, meta = (DisplayName = "Retarget Chains Settings"))
 struct IKRIG_API FIKRetargetFKChainsOpSettings : public FIKRetargetOpSettingsBase
 {
 	GENERATED_BODY()
-	
+
+	/** The target IK Rig asset that contains the bones chains to retarget in this op. */
+	UPROPERTY(EditAnywhere, BlueprintReadWrite, Category = "IK Rig Asset", meta=(ReinitializeOnEdit))
+	TObjectPtr<const UIKRigDefinition> IKRigAsset = nullptr;
+
+	/** The setting for each chain to retarget */
 	UPROPERTY(EditAnywhere, BlueprintReadWrite, Category = "FK Chain Retargeting", meta=(ReinitializeOnEdit))
 	TArray<FRetargetFKChainSettings> ChainsToRetarget;
 
-	UPROPERTY(EditAnywhere, Category = DebugDraw)
+	/** Debug draw lines on each chain in the viewport */
+	UPROPERTY(EditAnywhere, Category = Debug)
 	bool bDrawChainLines = true;
-	
-	UPROPERTY(EditAnywhere, Category = DebugDraw)
+
+	/** Debug draw spheres on single-bone chains in the viewport */
+	UPROPERTY(EditAnywhere, Category = Debug)
 	bool bDrawSingleBoneChains = true;
 
-	UPROPERTY(EditAnywhere, Category = DebugDraw)
+	/** Debug draw chain line thickness. */
+	UPROPERTY(EditAnywhere, Category = Debug)
 	double ChainDrawThickness = 5.0f;
 
-	UPROPERTY(EditAnywhere, Category = DebugDraw)
+	/** Debug draw size used for single-bone chains. */
+	UPROPERTY(EditAnywhere, Category = Debug)
 	double ChainDrawSize = 1.0f;
 };
 
@@ -201,6 +211,7 @@ struct IKRIG_API FIKRetargetFKChainsOp : public FIKRetargetOpBase
 		const FIKRetargetProcessor& InProcessor,
 		const FRetargetSkeleton& InSourceSkeleton,
 		const FTargetSkeleton& InTargetSkeleton,
+		const FIKRetargetOpBase* InParentOp,
 		FIKRigLogger& InLog) override;
 	
 	virtual void Run(
@@ -213,6 +224,8 @@ struct IKRIG_API FIKRetargetFKChainsOp : public FIKRetargetOpBase
 		const FRetargetSkeleton& SourceSkeleton,
 		const FTargetSkeleton& TargetSkeleton,
 		FIKRigLogger& Log) override;
+
+	virtual void OnAddedToStack(const UIKRetargeter* InRetargetAsset, const FIKRetargetOpBase* InParentOp) override;
 	
 	virtual void CollectRetargetedBones(TSet<int32>& OutRetargetedBones) const override;
 
@@ -226,13 +239,23 @@ struct IKRIG_API FIKRetargetFKChainsOp : public FIKRetargetOpBase
 
 	virtual const UClass* GetControllerType() const override;
 
+	virtual const UIKRigDefinition* GetCustomTargetIKRig() const override;
+
+	virtual FRetargetChainMapping* GetChainMapping() override;
+
+	virtual void OnTargetChainRenamed(const FName InOldChainName, const FName InNewChainName) override;
+
+	virtual void OnReinitPropertyEdited(const FPropertyChangedEvent* InPropertyChangedEvent) override;
+
 	UPROPERTY()
 	FIKRetargetFKChainsOpSettings Settings;
 
+	/* This op maintains its own chain mapping table to allow per-op mapping */
+	UPROPERTY()
+	FRetargetChainMapping ChainMapping;
+
 #if WITH_EDITOR
 public:
-
-	virtual void RunInitialSetup(const UIKRetargeter* InRetargetAsset) override;
 	
 	virtual void DebugDraw(
 		FPrimitiveDrawInterface* InPDI,
@@ -244,11 +267,13 @@ public:
 	TArray<FFKChainDebugData> AllChainsDebugData;
 	static FCriticalSection DebugDataMutex;
 
-	virtual void SetChainSettingsToDefault(const FName& InChainName) override;
+	virtual void ResetChainSettingsToDefault(const FName& InChainName) override;
 	virtual bool AreChainSettingsAtDefault(const FName& InChainName) override;
 #endif
 
 private:
+
+	void ApplyIKRigs(const UIKRigDefinition* InSourceIKRig, const UIKRigDefinition* InTargetIKRig);
 	
 	/** The Source/Target pairs of Bone Chains retargeted using the FK algorithm */
 	TArray<FChainPairFK> ChainPairsFK;
diff --git a/Engine/Plugins/Animation/IKRig/Source/IKRig/Public/Retargeter/RetargetOps/IKChainsOp.h b/Engine/Plugins/Animation/IKRig/Source/IKRig/Public/Retargeter/RetargetOps/IKChainsOp.h
--- a/Engine/Plugins/Animation/IKRig/Source/IKRig/Public/Retargeter/RetargetOps/IKChainsOp.h
+++ b/Engine/Plugins/Animation/IKRig/Source/IKRig/Public/Retargeter/RetargetOps/IKChainsOp.h
@@ -7,6 +7,7 @@
 
 #define LOCTEXT_NAMESPACE "IKChainOp"
 
+struct FIKRetargetRunIKRigOp;
 struct FResolvedBoneChain;
 struct FIKRetargetPelvisMotionOp;
 
@@ -160,19 +161,19 @@ struct IKRIG_API FIKRetargetIKChainsOpSettings : public FIKRetargetOpSettingsBas
 	TArray<FRetargetIKChainSettings> ChainsToRetarget;
 	
 	// Draw final IK goal locations. 
-	UPROPERTY(EditAnywhere, Category = DebugDraw)
+	UPROPERTY(EditAnywhere, Category = Debug)
 	bool bDrawFinalGoals = true;
 
 	// Draw goal locations from source skeleton. 
-	UPROPERTY(EditAnywhere, Category = DebugDraw)
+	UPROPERTY(EditAnywhere, Category = Debug)
 	bool bDrawSourceLocations = true;
 
 	// Adjust size of goal debug drawing in viewport
-	UPROPERTY(EditAnywhere, Category = DebugDraw)
+	UPROPERTY(EditAnywhere, Category = Debug)
 	double GoalDrawSize = 5.0;
 	
 	// Adjust thickness of goal debug drawing in viewport
-	UPROPERTY(EditAnywhere, Category = DebugDraw)
+	UPROPERTY(EditAnywhere, Category = Debug)
 	double GoalDrawThickness = 1.0;
 };
 
@@ -196,6 +197,7 @@ struct IKRIG_API FIKRetargetIKChainsOp : public FIKRetargetOpBase
 		const FIKRetargetProcessor& InProcessor,
 		const FRetargetSkeleton& InSourceSkeleton,
 		const FTargetSkeleton& InTargetSkeleton,
+		const FIKRetargetOpBase* InParentOp,
 		FIKRigLogger& InLog) override;
 	
 	virtual void Run(
@@ -204,6 +206,8 @@ struct IKRIG_API FIKRetargetIKChainsOp : public FIKRetargetOpBase
 		const TArray<FTransform>& InSourceGlobalPose,
 		TArray<FTransform>& OutTargetGlobalPose) override;
 
+	virtual void OnAddedToStack(const UIKRetargeter* InRetargetAsset, const FIKRetargetOpBase* InParentOp) override;
+
 	virtual FIKRetargetOpSettingsBase* GetSettings() override;
 	
 	virtual void SetSettings(const FIKRetargetOpSettingsBase* InSettings) override;
@@ -216,12 +220,15 @@ struct IKRIG_API FIKRetargetIKChainsOp : public FIKRetargetOpBase
 
 	virtual const UScriptStruct* GetParentOpType() const override;
 
+	virtual void OnTargetChainRenamed(const FName InOldChainName, const FName InNewChainName) override;
+
+	virtual void OnParentReinitPropertyEdited(const FIKRetargetOpBase& InParentOp, const FPropertyChangedEvent* InPropertyChangedEvent) override;
+
 	UPROPERTY()
 	FIKRetargetIKChainsOpSettings Settings;
 
 #if WITH_EDITOR
 public:
-	virtual void RunInitialSetup(const UIKRetargeter* InRetargetAsset) override;
 	virtual void DebugDraw(
 		FPrimitiveDrawInterface* InPDI,
 		const FTransform& InComponentTransform,
@@ -235,11 +242,13 @@ public:
 	FVector DebugRootModification;
 	static FCriticalSection DebugDataMutex;
 
-	virtual void SetChainSettingsToDefault(const FName& InChainName) override;
+	virtual void ResetChainSettingsToDefault(const FName& InChainName) override;
 	virtual bool AreChainSettingsAtDefault(const FName& InChainName) override;
 #endif
 
 private:
+
+	void RegenerateChainSettings(const FIKRetargetOpBase* InParentOp);
 	
 	TArray<FIKChainRetargeter> IKChainRetargeters;
 };
diff --git a/Engine/Plugins/Animation/IKRig/Source/IKRig/Public/Retargeter/RetargetOps/PelvisMotionOp.h b/Engine/Plugins/Animation/IKRig/Source/IKRig/Public/Retargeter/RetargetOps/PelvisMotionOp.h
--- a/Engine/Plugins/Animation/IKRig/Source/IKRig/Public/Retargeter/RetargetOps/PelvisMotionOp.h
+++ b/Engine/Plugins/Animation/IKRig/Source/IKRig/Public/Retargeter/RetargetOps/PelvisMotionOp.h
@@ -126,6 +126,7 @@ struct IKRIG_API FIKRetargetPelvisMotionOp : public FIKRetargetOpBase
 		const FIKRetargetProcessor& InProcessor,
 		const FRetargetSkeleton& InSourceSkeleton,
 		const FTargetSkeleton& InTargetSkeleton,
+		const FIKRetargetOpBase* InParentOp,
 		FIKRigLogger& InLog) override;
 	
 	virtual void Run(
@@ -134,6 +135,8 @@ struct IKRIG_API FIKRetargetPelvisMotionOp : public FIKRetargetOpBase
 		const TArray<FTransform>& InSourceGlobalPose,
 		TArray<FTransform>& OutTargetGlobalPose) override;
 
+	virtual void OnAddedToStack(const UIKRetargeter* InRetargetAsset, const FIKRetargetOpBase* InParentOp) override;
+
 	virtual FIKRetargetOpSettingsBase* GetSettings() override;
 
 	virtual void SetSettings(const FIKRetargetOpSettingsBase* InSettings) override;
@@ -182,7 +185,6 @@ private:
 	
 #if WITH_EDITOR
 public:
-	virtual void RunInitialSetup(const UIKRetargeter* InRetargetAsset) override;
 	virtual FText GetWarningMessage() const override;
 	virtual void DebugDraw(
 		FPrimitiveDrawInterface* InPDI,
diff --git a/Engine/Plugins/Animation/IKRig/Source/IKRig/Public/Retargeter/RetargetOps/RootMotionGeneratorOp.h b/Engine/Plugins/Animation/IKRig/Source/IKRig/Public/Retargeter/RetargetOps/RootMotionGeneratorOp.h
--- a/Engine/Plugins/Animation/IKRig/Source/IKRig/Public/Retargeter/RetargetOps/RootMotionGeneratorOp.h
+++ b/Engine/Plugins/Animation/IKRig/Source/IKRig/Public/Retargeter/RetargetOps/RootMotionGeneratorOp.h
@@ -30,49 +30,49 @@ struct IKRIG_API FIKRetargetRootMotionOpSettings : public FIKRetargetOpSettingsB
 	GENERATED_BODY()
 
 	// The root of the source skeleton.
-	UPROPERTY(EditAnywhere, Category="OpSettings", meta=(ReinitializeOnEdit))
+	UPROPERTY(EditAnywhere, Category="Bones", meta=(ReinitializeOnEdit))
 	FBoneReference SourceRoot;
 
 	// The root of the target skeleton.
-	UPROPERTY(EditAnywhere, Category="OpSettings", meta=(ReinitializeOnEdit))
+	UPROPERTY(EditAnywhere, Category="Bones", meta=(ReinitializeOnEdit))
 	FBoneReference TargetRoot;
-
-	// The pelvis of the target skeleton.
-	UPROPERTY(EditAnywhere, Category="OpSettings", meta=(ReinitializeOnEdit))
-	FBoneReference TargetPelvis;
-
+	
 	// Where to copy the root motion from.
 	// Copy From Source Root: copies the motion from the source root bone, but scales it according to relative height difference.
 	// Generate From Target Pelvis: uses the retargeted Pelvis motion to generate root motion.
 	// NOTE: Generating root motion from the Pelvis 
-	UPROPERTY(EditAnywhere, BlueprintReadWrite, Category="OpSettings")
+	UPROPERTY(EditAnywhere, BlueprintReadWrite, Category="Motion Source")
 	ERootMotionSource RootMotionSource = ERootMotionSource::CopyFromSourceRoot;
+
+	// The pelvis of the target skeleton. Only used when "Root Motion Source" is set to target pelvis.
+	UPROPERTY(EditAnywhere, Category="Motion Source", meta=(ReinitializeOnEdit))
+	FBoneReference TargetPelvis;
+
+	// Applies only when generating root motion from the Pelvis.
+	// When true, the applied offset will be rotated by the Pelvis.
+	// (NOTE: This may cause unwanted rotations, for example if Pelvis Yaw is animated.)
+	UPROPERTY(EditAnywhere, BlueprintReadWrite, Category="Motion Source")
+	bool bRotateWithPelvis = false;
 	
 	// How to set the height of the root bone.
 	// Copy Height From Source: copies the height of the root bone on the source skeleton's root bone.
 	// Snap To Ground: sets the root bone height to the ground plane (Z=0).
-	UPROPERTY(EditAnywhere, BlueprintReadWrite, Category="OpSettings")
+	UPROPERTY(EditAnywhere, BlueprintReadWrite, Category="Height Source")
 	ERootMotionHeightSource RootHeightSource = ERootMotionHeightSource::CopyHeightFromSource;
-
-	// Will transform all children of the target root that are not themselves part of a retarget chain.
-	UPROPERTY(EditAnywhere, BlueprintReadWrite, Category="OpSettings")
-	bool bPropagateToNonRetargetedChildren = true;
+	
+	// A manual offset to apply in global space to the root bone.
+	UPROPERTY(EditAnywhere, BlueprintReadWrite, Category="Offset")
+	FTransform GlobalOffset;
 
 	// Applies only when generating root motion from the Pelvis.
 	// Maintains the offset between the root and pelvis as recorded in the target retarget pose.
 	// If false, the root bone is placed directly under the Pelvis bone.
-	UPROPERTY(EditAnywhere, BlueprintReadWrite, Category="OpSettings")
+	UPROPERTY(EditAnywhere, BlueprintReadWrite, Category="Offset")
 	bool bMaintainOffsetFromPelvis = true;
 
-	// Applies only when generating root motion from the Pelvis.
-	// When true, the applied offset will be rotated by the Pelvis.
-	// (NOTE: This may cause unwanted rotations, for example if Pelvis Yaw is animated.)
-	UPROPERTY(EditAnywhere, BlueprintReadWrite, Category="OpSettings")
-	bool bRotateWithPelvis = false;
-
-	// A manual offset to apply in global space to the root bone.
-	UPROPERTY(EditAnywhere, BlueprintReadWrite, Category=Settings)
-	FTransform GlobalOffset;
+	// Will transform all children of the target root that are not themselves part of a retarget chain.
+	UPROPERTY(EditAnywhere, BlueprintReadWrite, Category="Update Children")
+	bool bPropagateToNonRetargetedChildren = true;
 
 	/** FBoneReference needs a skeleton for the UI widget. */
 #if WITH_EDITOR
@@ -111,6 +111,7 @@ struct IKRIG_API FIKRetargetRootMotionOp : public FIKRetargetOpBase
 		const FIKRetargetProcessor& Processor,
 		const FRetargetSkeleton& SourceSkeleton,
 		const FTargetSkeleton& TargetSkeleton,
+		const FIKRetargetOpBase* InParentOp,
 		FIKRigLogger& Log) override;
 	
 	virtual void Run(
@@ -125,6 +126,8 @@ struct IKRIG_API FIKRetargetRootMotionOp : public FIKRetargetOpBase
 		const FTargetSkeleton& TargetSkeleton,
 		FIKRigLogger& Log) override;
 	
+	virtual void OnAddedToStack(const UIKRetargeter* InRetargetAsset, const FIKRetargetOpBase* InParentOp) override;
+	
 	virtual FIKRetargetOpSettingsBase* GetSettings() override;
 
 	virtual void SetSettings(const FIKRetargetOpSettingsBase* InSettings) override;
@@ -139,7 +142,6 @@ struct IKRIG_API FIKRetargetRootMotionOp : public FIKRetargetOpBase
 	FIKRetargetRootMotionOpSettings Settings;
 
 #if WITH_EDITOR
-	virtual void RunInitialSetup(const UIKRetargeter* InRetargetAsset) override;
 	virtual FText GetWarningMessage() const override;
 #endif
 
diff --git a/Engine/Plugins/Animation/IKRig/Source/IKRig/Public/Retargeter/RetargetOps/RunIKRigOp.h b/Engine/Plugins/Animation/IKRig/Source/IKRig/Public/Retargeter/RetargetOps/RunIKRigOp.h
--- a/Engine/Plugins/Animation/IKRig/Source/IKRig/Public/Retargeter/RetargetOps/RunIKRigOp.h
+++ b/Engine/Plugins/Animation/IKRig/Source/IKRig/Public/Retargeter/RetargetOps/RunIKRigOp.h
@@ -1,6 +1,7 @@
 ﻿// Copyright Epic Games, Inc. All Rights Reserved.
 #pragma once
 
+#include "Retargeter/IKRetargetChainMapping.h"
 #include "Retargeter/IKRetargetOps.h"
 #include "Rig/IKRigProcessor.h"
 
@@ -22,19 +23,19 @@ struct IKRIG_API FIKRetargetRunIKRigOpSettings : public FIKRetargetOpSettingsBas
 	TArray<FName> ExcludedGoals;
 	
 	// Draw IK goal locations. 
-	UPROPERTY(EditAnywhere, Category = DebugDraw)
+	UPROPERTY(EditAnywhere, Category = Debug)
 	bool bDrawGoals = true;
 
 	// Draw locations of the source bone (pre-solve)
-	UPROPERTY(EditAnywhere, Category = DebugDraw)
+	UPROPERTY(EditAnywhere, Category = Debug)
 	bool bDrawGoalBoneLocations = true;
 
 	// Adjust size of goal debug drawing in viewport
-	UPROPERTY(EditAnywhere, Category = DebugDraw)
+	UPROPERTY(EditAnywhere, Category = Debug)
 	double GoalDrawSize = 5.0;
 	
 	// Adjust thickness of goal debug drawing in viewport
-	UPROPERTY(EditAnywhere, Category = DebugDraw)
+	UPROPERTY(EditAnywhere, Category = Debug)
 	double GoalDrawThickness = 1.0;
 };
 
@@ -56,6 +57,7 @@ struct IKRIG_API FIKRetargetRunIKRigOp : public FIKRetargetOpBase
 		const FIKRetargetProcessor& InProcessor,
 		const FRetargetSkeleton& InSourceSkeleton,
 		const FTargetSkeleton& InTargetSkeleton,
+		const FIKRetargetOpBase* InParentOp,
 		FIKRigLogger& InLog) override;
 	
 	virtual void Run(
@@ -64,6 +66,8 @@ struct IKRIG_API FIKRetargetRunIKRigOp : public FIKRetargetOpBase
 		const TArray<FTransform>& InSourceGlobalPose,
 		TArray<FTransform>& OutTargetGlobalPose) override;
 
+	virtual void OnAddedToStack(const UIKRetargeter* InRetargetAsset, const FIKRetargetOpBase* InParentOp) override;
+
 	virtual FIKRetargetOpSettingsBase* GetSettings() override;
 	
 	virtual void SetSettings(const FIKRetargetOpSettingsBase* InSettings) override;
@@ -76,11 +80,20 @@ struct IKRIG_API FIKRetargetRunIKRigOp : public FIKRetargetOpBase
 	
 	virtual bool CanHaveChildOps() const override { return true; };
 
+	virtual const UIKRigDefinition* GetCustomTargetIKRig() const override;
+
+	virtual FRetargetChainMapping* GetChainMapping() override;
+	
+	virtual void OnReinitPropertyEdited(const FPropertyChangedEvent* InPropertyChangedEvent) override;
+
 	UPROPERTY()
 	FIKRetargetRunIKRigOpSettings Settings;
+	
+	/* This op maintains its own chain mapping table to allow per-op mapping */
+	UPROPERTY()
+	FRetargetChainMapping ChainMapping;
 
 #if WITH_EDITOR
-	virtual void RunInitialSetup(const UIKRetargeter* InRetargetAsset) override;
 	virtual void DebugDraw(
 		FPrimitiveDrawInterface* InPDI,
 		const FTransform& InComponentTransform,
@@ -92,6 +105,8 @@ struct IKRIG_API FIKRetargetRunIKRigOp : public FIKRetargetOpBase
 	static FCriticalSection DebugDataMutex;
 #endif
 
+	TArray<FName> GetRequiredTargetChains() const;
+	
 private:
 
 	/** The IK Rig processor for running IK on the target */
diff --git a/Engine/Plugins/Animation/IKRig/Source/IKRig/Public/Retargeter/RetargetOps/SpeedPlantingOp.h b/Engine/Plugins/Animation/IKRig/Source/IKRig/Public/Retargeter/RetargetOps/SpeedPlantingOp.h
--- a/Engine/Plugins/Animation/IKRig/Source/IKRig/Public/Retargeter/RetargetOps/SpeedPlantingOp.h
+++ b/Engine/Plugins/Animation/IKRig/Source/IKRig/Public/Retargeter/RetargetOps/SpeedPlantingOp.h
@@ -12,6 +12,9 @@ USTRUCT(BlueprintType)
 struct IKRIG_API FRetargetSpeedPlantingSettings
 {
 	GENERATED_BODY()
+	
+	FRetargetSpeedPlantingSettings() = default;
+	FRetargetSpeedPlantingSettings(const FName InTargetChainName) : TargetChainName(InTargetChainName) {}
 
 	/** The name of the target chain to plant the IK on.
 	 * NOTE: this chain must have an IK Goal assigned to it! */
@@ -28,7 +31,7 @@ struct IKRIG_API FIKRetargetSpeedPlantingOpSettings : public FIKRetargetOpSettin
 {
 	GENERATED_BODY()
 	
-	UPROPERTY(EditAnywhere, BlueprintReadWrite, Category = "Speed Planting", meta=(ReinitializeOnEdit))
+	UPROPERTY(EditAnywhere, BlueprintReadWrite, Category = "Setup", meta=(ReinitializeOnEdit))
 	TArray<FRetargetSpeedPlantingSettings> ChainsToSpeedPlant;
 
 	/** Range 0 to 100. Default 15. The maximum speed a source bone can be moving while being considered 'planted'.
@@ -76,6 +79,7 @@ struct IKRIG_API FIKRetargetSpeedPlantingOp : public FIKRetargetOpBase
 		const FIKRetargetProcessor& InProcessor,
 		const FRetargetSkeleton& InSourceSkeleton,
 		const FTargetSkeleton& InTargetSkeleton,
+		const FIKRetargetOpBase* InParentOp,
 		FIKRigLogger& InLog) override;
 	
 	virtual void Run(
@@ -84,6 +88,8 @@ struct IKRIG_API FIKRetargetSpeedPlantingOp : public FIKRetargetOpBase
 		const TArray<FTransform>& InSourceGlobalPose,
 		TArray<FTransform>& OutTargetGlobalPose) override;
 
+	virtual void OnAddedToStack(const UIKRetargeter* InRetargetAsset, const FIKRetargetOpBase* InParentOp) override;
+
 	virtual FIKRetargetOpSettingsBase* GetSettings() override;
 
 	virtual void SetSettings(const FIKRetargetOpSettingsBase* InSettings) override;
@@ -102,6 +108,10 @@ struct IKRIG_API FIKRetargetSpeedPlantingOp : public FIKRetargetOpBase
 
 	virtual const UScriptStruct* GetParentOpType() const override;
 
+	virtual void OnTargetChainRenamed(const FName InOldChainName, const FName InNewChainName) override;
+
+	virtual void OnParentReinitPropertyEdited(const FIKRetargetOpBase& InParentOp, const FPropertyChangedEvent* InPropertyChangedEvent) override;
+
 	/** return array of names of all the speed curves referenced by this op. */
 	TArray<FName> GetRequiredSpeedCurves() const;
 
@@ -114,6 +124,10 @@ public:
 #endif
 
 	bool ResetThisTick;
+
+private:
+	
+	void RegenerateChainSettings(const FIKRetargetOpBase* InParentOp);
 	
 	TArray<FSpeedPlantingGoalState> GoalsToPlant;
 };
diff --git a/Engine/Plugins/Animation/IKRig/Source/IKRig/Public/Retargeter/RetargetOps/StrideWarpingOp.h b/Engine/Plugins/Animation/IKRig/Source/IKRig/Public/Retargeter/RetargetOps/StrideWarpingOp.h
--- a/Engine/Plugins/Animation/IKRig/Source/IKRig/Public/Retargeter/RetargetOps/StrideWarpingOp.h
+++ b/Engine/Plugins/Animation/IKRig/Source/IKRig/Public/Retargeter/RetargetOps/StrideWarpingOp.h
@@ -17,14 +17,33 @@ enum class EWarpingDirectionSource
 	RootBone
 };
 
+USTRUCT(BlueprintType)
+struct IKRIG_API FRetargetStrideWarpChainSettings
+{
+	GENERATED_BODY()
+
+	FRetargetStrideWarpChainSettings() = default;
+	FRetargetStrideWarpChainSettings(const FName InTargetChainName) : TargetChainName(InTargetChainName) {}
+
+	/** The name of the TARGET chain with an IK goal to warp. */
+	UPROPERTY(EditAnywhere, BlueprintReadWrite, Category = "Setup", meta=(ReinitializeOnEdit))
+	FName TargetChainName;
+	
+	/** Whether to warp the location of the IK goal on this chain. Default is true. */
+	UPROPERTY(EditAnywhere, BlueprintReadWrite, Category = "Setup", meta=(ReinitializeOnEdit))
+	bool EnableStrideWarping = true;
+
+	bool operator==(const FRetargetStrideWarpChainSettings& Other) const;
+};
+
 USTRUCT(BlueprintType, meta = (DisplayName = "Stride Warp Settings"))
 struct IKRIG_API FIKRetargetStrideWarpingOpSettings : public FIKRetargetOpSettingsBase
 {
 	GENERATED_BODY()
-	
+
 	/** The chains to apply stride warping to */
 	UPROPERTY(EditAnywhere, BlueprintReadWrite, Category = Setup, meta=(ReinitializeOnEdit))
-	TArray<FName> ChainsToStrideWarp;
+	TArray<FRetargetStrideWarpChainSettings> ChainSettings;
 
 	/** Defines the source used to determine the forward direction as the character animates in world space. Default is "Goals".
 	 * This method uses various positions on the character to create a "best fit" global rotation that approximates the facing direction of the character over time.
@@ -98,6 +117,7 @@ struct IKRIG_API FIKRetargetStrideWarpingOp : public FIKRetargetOpBase
 		const FIKRetargetProcessor& InProcessor,
 		const FRetargetSkeleton& InSourceSkeleton,
 		const FTargetSkeleton& InTargetSkeleton,
+		const FIKRetargetOpBase* InParentOp,
 		FIKRigLogger& InLog) override;
 	
 	virtual void Run(
@@ -106,6 +126,8 @@ struct IKRIG_API FIKRetargetStrideWarpingOp : public FIKRetargetOpBase
 		const TArray<FTransform>& InSourceGlobalPose,
 		TArray<FTransform>& OutTargetGlobalPose) override;
 
+	virtual void OnAddedToStack(const UIKRetargeter* InRetargetAsset, const FIKRetargetOpBase* InParentOp) override;
+
 	virtual FIKRetargetOpSettingsBase* GetSettings() override;
 
 	virtual void SetSettings(const FIKRetargetOpSettingsBase* InSettings) override;
@@ -117,6 +139,10 @@ struct IKRIG_API FIKRetargetStrideWarpingOp : public FIKRetargetOpBase
 	virtual const UClass* GetControllerType() const override;
 	
 	virtual const UScriptStruct* GetParentOpType() const override;
+	
+	virtual void OnTargetChainRenamed(const FName InOldChainName, const FName InNewChainName) override;
+
+	virtual void OnParentReinitPropertyEdited(const FIKRetargetOpBase& InParentOp, const FPropertyChangedEvent* InPropertyChangedEvent) override;
 
 	UPROPERTY()
 	FIKRetargetStrideWarpingOpSettings Settings;
@@ -130,11 +156,14 @@ struct IKRIG_API FIKRetargetStrideWarpingOp : public FIKRetargetOpBase
 	FTransform DebugStrideWarpingFrame;
 	static FCriticalSection DebugDataMutex;
 
-	virtual void SetChainSettingsToDefault(const FName& InChainName) override;
+	virtual void ResetChainSettingsToDefault(const FName& InChainName) override;
 	virtual bool AreChainSettingsAtDefault(const FName& InChainName) override;
 #endif
 
+	void RegenerateChainSettings(const FIKRetargetOpBase* InParentOp);
+	
 	TArray<FStrideWarpGoalData> GoalsToWarp;
+	TObjectPtr<const UIKRigDefinition> TargetIKRig;
 };
 
 /* The blueprint/python API for editing a Stride Warping Op */
diff --git a/Engine/Plugins/Animation/IKRig/Source/IKRigEditor/Private/IKRigEditor.cpp b/Engine/Plugins/Animation/IKRig/Source/IKRigEditor/Private/IKRigEditor.cpp
--- a/Engine/Plugins/Animation/IKRig/Source/IKRigEditor/Private/IKRigEditor.cpp
+++ b/Engine/Plugins/Animation/IKRig/Source/IKRigEditor/Private/IKRigEditor.cpp
@@ -15,6 +15,7 @@
 #include "RetargetEditor/IKRetargetDetails.h"
 #include "RetargetEditor/IKRetargetEditPoseMode.h"
 #include "RetargetEditor/IKRetargeterThumbnailRenderer.h"
+#include "RetargetEditor/IKRetargetOpDetails.h"
 #include "Retargeter/IKRetargetOps.h"
 #include "RigEditor/IKRigCommands.h"
 #include "RigEditor/IKRigEditMode.h"
@@ -58,17 +59,48 @@ void FIKRigEditor::StartupModule()
 	// register detail customizations
 	FPropertyEditorModule& PropertyEditorModule = FModuleManager::LoadModuleChecked<FPropertyEditorModule>("PropertyEditor");
 	// custom instanced struct details
-	PropertyEditorModule.RegisterCustomClassLayout(UIKRigStructViewer::StaticClass()->GetFName(), FOnGetDetailCustomizationInstance::CreateStatic(&FStructViewerCustomization::MakeInstance));
+	PropertyEditorModule.RegisterCustomClassLayout(
+		UIKRigStructViewer::StaticClass()->GetFName(),
+		FOnGetDetailCustomizationInstance::CreateStatic(&FIKRigStructViewerCustomization::MakeInstance));
 	ClassesToUnregisterOnShutdown.Add(UIKRigStructViewer::StaticClass()->GetFName());
+	PropertyEditorModule.RegisterCustomClassLayout(
+		UIKRigStructWrapperBase::StaticClass()->GetFName(),
+		FOnGetDetailCustomizationInstance::CreateStatic(&FIKRigStructWrapperCustomization::MakeInstance));
+	ClassesToUnregisterOnShutdown.Add(UIKRigStructWrapperBase::StaticClass()->GetFName());
 	// custom IK rig bone details
-	PropertyEditorModule.RegisterCustomClassLayout(UIKRigBoneDetails::StaticClass()->GetFName(), FOnGetDetailCustomizationInstance::CreateStatic(&FIKRigGenericDetailCustomization::MakeInstance));
+	PropertyEditorModule.RegisterCustomClassLayout(
+		UIKRigBoneDetails::StaticClass()->GetFName(),
+		FOnGetDetailCustomizationInstance::CreateStatic(&FIKRigGenericDetailCustomization::MakeInstance));
 	ClassesToUnregisterOnShutdown.Add(UIKRigBoneDetails::StaticClass()->GetFName());
 	// custom IK rig goal details
-	PropertyEditorModule.RegisterCustomClassLayout(UIKRigEffectorGoal::StaticClass()->GetFName(), FOnGetDetailCustomizationInstance::CreateStatic(&FIKRigGenericDetailCustomization::MakeInstance));
+	PropertyEditorModule.RegisterCustomClassLayout(
+		UIKRigEffectorGoal::StaticClass()->GetFName(),
+		FOnGetDetailCustomizationInstance::CreateStatic(&FIKRigGenericDetailCustomization::MakeInstance));
 	ClassesToUnregisterOnShutdown.Add(UIKRigEffectorGoal::StaticClass()->GetFName());
 	// custom retargeter bone details
-	PropertyEditorModule.RegisterCustomClassLayout(UIKRetargetBoneDetails::StaticClass()->GetFName(), FOnGetDetailCustomizationInstance::CreateStatic(&FIKRetargetBoneDetailCustomization::MakeInstance));
+	PropertyEditorModule.RegisterCustomClassLayout(
+		UIKRetargetBoneDetails::StaticClass()->GetFName(),
+		FOnGetDetailCustomizationInstance::CreateStatic(&FIKRetargetBoneDetailCustomization::MakeInstance));
 	ClassesToUnregisterOnShutdown.Add(UIKRetargetBoneDetails::StaticClass()->GetFName());
+	// retarget op details
+	PropertyEditorModule.RegisterCustomPropertyTypeLayout(
+		"IKRetargetFKChainsOpSettings",
+		FOnGetPropertyTypeCustomizationInstance::CreateStatic(&FChainsFKOpCustomization::MakeInstance));
+	PropertyEditorModule.RegisterCustomPropertyTypeLayout(
+		"IKRetargetRunIKRigOpSettings",
+		FOnGetPropertyTypeCustomizationInstance::CreateStatic(&FRunIKRigOpCustomization::MakeInstance));
+	PropertyEditorModule.RegisterCustomPropertyTypeLayout(
+		"IKRetargetIKChainsOpSettings",
+		FOnGetPropertyTypeCustomizationInstance::CreateStatic(&FIKChainOpCustomization::MakeInstance));
+	PropertyEditorModule.RegisterCustomPropertyTypeLayout(
+		"IKRetargetStrideWarpingOpSettings",
+		FOnGetPropertyTypeCustomizationInstance::CreateStatic(&FStrideWarpOpCustomization::MakeInstance));
+	PropertyEditorModule.RegisterCustomPropertyTypeLayout(
+		"IKRetargetSpeedPlantingOpSettings",
+		FOnGetPropertyTypeCustomizationInstance::CreateStatic(&FSpeedPlantOpCustomization::MakeInstance));
+	PropertyEditorModule.RegisterCustomPropertyTypeLayout(
+		"IKRetargetAlignPoleVectorOpSettings",
+		FOnGetPropertyTypeCustomizationInstance::CreateStatic(&FPoleVectorOpCustomization::MakeInstance));
 
 	// register a thumbnail renderer for the assets
 	UThumbnailManager::Get().RegisterCustomRenderer(UIKRigDefinition::StaticClass(), UIKRigThumbnailRenderer::StaticClass());
diff --git a/Engine/Plugins/Animation/IKRig/Source/IKRigEditor/Private/RetargetEditor/IKRetargetAnimInstance.cpp b/Engine/Plugins/Animation/IKRig/Source/IKRigEditor/Private/RetargetEditor/IKRetargetAnimInstance.cpp
--- a/Engine/Plugins/Animation/IKRig/Source/IKRigEditor/Private/RetargetEditor/IKRetargetAnimInstance.cpp
+++ b/Engine/Plugins/Animation/IKRig/Source/IKRigEditor/Private/RetargetEditor/IKRetargetAnimInstance.cpp
@@ -88,7 +88,7 @@ void FAnimNode_PreviewRetargetPose::Evaluate_AnyThread(FPoseContext& Output)
 
 	// apply root translation offset from the retarget pose (done in global space)
 	const UIKRetargeterController* Controller = UIKRetargeterController::GetController(IKRetargeterAsset);
-	const FName RetargetRootBoneName = Controller->GetRetargetRootBone(SourceOrTarget);
+	const FName RetargetRootBoneName = Controller->GetPelvisBone(SourceOrTarget);
 	const int32 RootBoneIndex = RefSkeleton.FindBoneIndex(RetargetRootBoneName);
 	if (RootBoneIndex != INDEX_NONE)
 	{
diff --git a/Engine/Plugins/Animation/IKRig/Source/IKRigEditor/Private/RetargetEditor/IKRetargetApplicationMode.cpp b/Engine/Plugins/Animation/IKRig/Source/IKRigEditor/Private/RetargetEditor/IKRetargetApplicationMode.cpp
--- a/Engine/Plugins/Animation/IKRig/Source/IKRigEditor/Private/RetargetEditor/IKRetargetApplicationMode.cpp
+++ b/Engine/Plugins/Animation/IKRig/Source/IKRigEditor/Private/RetargetEditor/IKRetargetApplicationMode.cpp
@@ -8,7 +8,6 @@
 #include "PersonaTabs.h"
 #include "RetargetEditor/IKRetargetAssetBrowserTabSummoner.h"
 
-#include "RetargetEditor/IKRetargetChainTabSummoner.h"
 #include "RetargetEditor/IKRetargetEditor.h"
 #include "RetargetEditor/IKRetargetHierarchyTabSummoner.h"
 #include "RetargetEditor/IKRetargetOpStackTabSummoner.h"
@@ -39,7 +38,6 @@ FIKRetargetApplicationMode::FIKRetargetApplicationMode(
 	TabFactories.RegisterFactory(PersonaModule.CreateAdvancedPreviewSceneTabFactory(InHostingApp, IKRetargetEditor->GetPersonaToolkit()->GetPreviewScene()));
 
 	// register custom tabs
-	TabFactories.RegisterFactory(MakeShared<FIKRetargetChainTabSummoner>(IKRetargetEditor));
 	TabFactories.RegisterFactory(MakeShared<FIKRetargetAssetBrowserTabSummoner>(IKRetargetEditor));
 	TabFactories.RegisterFactory(MakeShared<FIKRetargetOutputLogTabSummoner>(IKRetargetEditor));
 	TabFactories.RegisterFactory(MakeShared<FIKRetargetHierarchyTabSummoner>(IKRetargetEditor));
@@ -102,24 +100,11 @@ FIKRetargetApplicationMode::FIKRetargetApplicationMode(
 				)
 				->Split
 				(
-					FTabManager::NewSplitter()
+				FTabManager::NewStack()
 					->SetSizeCoefficient(0.2f)
-					->SetOrientation(Orient_Vertical)
-					->Split
-					(
-					FTabManager::NewStack()
-						->SetSizeCoefficient(0.7f)
-						->AddTab(FPersonaTabs::DetailsID, ETabState::OpenedTab)
-						->AddTab(FPersonaTabs::AdvancedPreviewSceneSettingsID, ETabState::OpenedTab)
-						->SetForegroundTab(FPersonaTabs::DetailsID)
-					)
-					->Split
-					(
-					FTabManager::NewStack()
-						->SetSizeCoefficient(0.3f)
-						->AddTab(FIKRetargetChainTabSummoner::TabID, ETabState::OpenedTab)
-						->SetForegroundTab(FIKRetargetChainTabSummoner::TabID)
-					)
+					->AddTab(FPersonaTabs::DetailsID, ETabState::OpenedTab)
+					->AddTab(FPersonaTabs::AdvancedPreviewSceneSettingsID, ETabState::OpenedTab)
+					->SetForegroundTab(FPersonaTabs::DetailsID)
 					
 				)
 			)
diff --git a/Engine/Plugins/Animation/IKRig/Source/IKRigEditor/Private/RetargetEditor/IKRetargetDetails.cpp b/Engine/Plugins/Animation/IKRig/Source/IKRigEditor/Private/RetargetEditor/IKRetargetDetails.cpp
--- a/Engine/Plugins/Animation/IKRig/Source/IKRigEditor/Private/RetargetEditor/IKRetargetDetails.cpp
+++ b/Engine/Plugins/Animation/IKRig/Source/IKRigEditor/Private/RetargetEditor/IKRetargetDetails.cpp
@@ -704,7 +704,7 @@ bool UIKRetargetBoneDetails::IsRootBone() const
 		return false;
 	}
 
-	const FName RootBone = Controller->AssetController->GetRetargetRootBone(Controller->GetSourceOrTarget());
+	const FName RootBone = Controller->AssetController->GetPelvisBone(Controller->GetSourceOrTarget());
 	return SelectedBone == RootBone;
 }
 
@@ -770,7 +770,7 @@ void FIKRetargetBoneDetailCustomization::CustomizeDetails(IDetailLayoutBuilder&
 	
 	const bool bIsEditingPose = Controller.IsEditingPose();
 
-	const FName CurrentRootName = AssetController->GetRetargetRootBone(Controller.GetSourceOrTarget());
+	const FName CurrentRootName = AssetController->GetPelvisBone(Controller.GetSourceOrTarget());
 	const bool bIsRootSelected =  Bones[0]->SelectedBone == CurrentRootName;
 
 	FIKRetargetTransformUIData UIData;
diff --git a/Engine/Plugins/Animation/IKRig/Source/IKRigEditor/Private/RetargetEditor/IKRetargetEditPoseMode.cpp b/Engine/Plugins/Animation/IKRig/Source/IKRigEditor/Private/RetargetEditor/IKRetargetEditPoseMode.cpp
--- a/Engine/Plugins/Animation/IKRig/Source/IKRigEditor/Private/RetargetEditor/IKRetargetEditPoseMode.cpp
+++ b/Engine/Plugins/Animation/IKRig/Source/IKRigEditor/Private/RetargetEditor/IKRetargetEditPoseMode.cpp
@@ -475,7 +475,7 @@ bool FIKRetargetEditPoseMode::IsRootSelected() const
 	}
 	
 	const TArray<FName>& SelectedBones = Controller->GetSelectedBones();
-	const FName& RootName = Controller->AssetController->GetRetargetRootBone(SourceOrTarget);
+	const FName& RootName = Controller->AssetController->GetPelvisBone(SourceOrTarget);
 	return SelectedBones.Contains(RootName);
 }
 
@@ -493,7 +493,7 @@ bool FIKRetargetEditPoseMode::IsOnlyRootSelected() const
 		return false;
 	}
 
-	return Controller->AssetController->GetRetargetRootBone(SourceOrTarget) == SelectedBones[0];
+	return Controller->AssetController->GetPelvisBone(SourceOrTarget) == SelectedBones[0];
 }
 
 void FIKRetargetEditPoseMode::Tick(FEditorViewportClient* ViewportClient, float DeltaTime)
diff --git a/Engine/Plugins/Animation/IKRig/Source/IKRigEditor/Private/RetargetEditor/IKRetargetEditorController.cpp b/Engine/Plugins/Animation/IKRig/Source/IKRigEditor/Private/RetargetEditor/IKRetargetEditorController.cpp
--- a/Engine/Plugins/Animation/IKRig/Source/IKRigEditor/Private/RetargetEditor/IKRetargetEditorController.cpp
+++ b/Engine/Plugins/Animation/IKRig/Source/IKRigEditor/Private/RetargetEditor/IKRetargetEditorController.cpp
@@ -32,6 +32,7 @@
 #include "RigEditor/IKRigController.h"
 #include "ScopedTransaction.h"
 #include "RigEditor/IKRigDetailCustomizations.h"
+#include "RigEditor/IKRigStructViewer.h"
 #include "Styling/AppStyle.h"
 
 #define LOCTEXT_NAMESPACE "IKRetargetEditorController"
@@ -660,7 +661,6 @@ void FIKRetargetEditorController::RefreshAllViews() const
 {
 	Editor.Pin()->RegenerateMenusAndToolbars();
 	RefreshDetailsView();
-	RefreshChainsView();
 	RefreshAssetBrowserView();
 	RefreshHierarchyView();
 	RefreshOpStackView();
@@ -671,17 +671,7 @@ void FIKRetargetEditorController::RefreshDetailsView() const
 	// refresh the details panel, cannot assume tab is not closed
 	if (DetailsView.IsValid())
 	{
-		const int32 LastSelectedOpIndex = AssetController->GetIndexOfOpByName(Selection.LastSelectedOpName);
-		ShowDetailsForOp(LastSelectedOpIndex);
-	}
-}
-
-void FIKRetargetEditorController::RefreshChainsView() const
-{
-	// refesh chains view, cannot assume tab is not closed
-	if (ChainsView.IsValid())
-	{
-		ChainsView.Get()->RefreshView();
+		DetailsView->ForceRefresh();
 	}
 }
 
@@ -736,7 +726,7 @@ void FIKRetargetEditorController::SetDetailsObjects(const TArray<UObject*>& Deta
 
 void FIKRetargetEditorController::ShowDetailsForOp(const int32 OpIndex) const
 {
-	if (AssetController->GetOpAtIndex(OpIndex) == nullptr)
+	if (AssetController->GetRetargetOpByIndex(OpIndex) == nullptr)
 	{
 		DetailsView->SetObject(AssetController->GetAsset());
 		return;
@@ -746,20 +736,24 @@ void FIKRetargetEditorController::ShowDetailsForOp(const int32 OpIndex) const
 	auto MemoryProvider = [Asset, OpIndex]() -> uint8*
 	{
 		UIKRetargeterController* Controller = UIKRetargeterController::GetController(Asset);
-		if (FIKRetargetOpBase* OpToEdit = Controller->GetOpAtIndex(OpIndex))
+		if (FIKRetargetOpBase* OpToEdit = Controller->GetRetargetOpByIndex(OpIndex))
 		{
 			return reinterpret_cast<uint8*>(OpToEdit->GetSettings());
 		}
 		
 		return nullptr;
 	};
-	
-	FInstancedStruct* InstancedStruct = AssetController->GetRetargetOpStructAtIndex(OpIndex);
-	const FIKRetargetOpBase* OpToEdit = InstancedStruct->GetPtr<FIKRetargetOpBase>();
-	const UScriptStruct* TypeToEdit = OpToEdit->GetSettingsType();
 
+	FIKRetargetOpBase* Op = AssetController->GetRetargetOpByIndex(OpIndex);
+
+	FIKRigStructToView StructToView;
+	StructToView.Type = Op->GetSettingsType();
+	StructToView.MemoryProvider = MemoryProvider;
+	StructToView.Owner = AssetController->GetAsset();
+	StructToView.UniqueName = Op->GetName();
+	
 	UIKRigStructViewer* StructViewer = AssetController->GetStructViewer();
-	StructViewer->SetStructToView(TypeToEdit, MemoryProvider, AssetController->GetAsset());
+	StructViewer->SetStructToView(StructToView);
 	
 	constexpr bool bForceRefresh = true; // must forcibly refresh because struct viewer is a recycled UObject
 	DetailsView->SetObject(StructViewer, bForceRefresh);
@@ -928,22 +922,6 @@ void FIKRetargetEditorController::EditChainSelection(
 	default:
 		checkNoEntry();
 	}
-
-	// update chains view with selected chains
-	if (ChainsView)
-	{
-		if (!bFromChainsView)
-		{
-			ChainsView->RefreshView();
-		}
-	}
-
-	// TODO
-	// We could get the chain settings struct from the currently selected op + chain map and display that
-	// but this needs more discussion with UX team on how we want to proceed.
-
-	// no chains selected, then show asset settings in the details view
-	SetDetailsObject(AssetController->GetAsset());
 }
 
 void FIKRetargetEditorController::SetRootSelected(const bool bIsSelected)
@@ -985,11 +963,7 @@ void FIKRetargetEditorController::ClearSelection(const bool bKeepBoneSelection)
 	SetRootSelected(false);
 	
 	// deselect all chains
-	if (ChainsView.IsValid())
-	{
-		ChainsView->ClearSelection();
-		Selection.SelectedChains.Empty();
-	}
+	Selection.SelectedChains.Empty();
 
 	// clear bone selection
 	if (!bKeepBoneSelection)
@@ -1001,7 +975,7 @@ void FIKRetargetEditorController::ClearSelection(const bool bKeepBoneSelection)
 
 	Selection.LastSelectedType = ERetargetSelectionType::NONE;
 
-	RefreshDetailsView();
+	SetDetailsObject(AssetController->GetAsset());
 }
 
 void FIKRetargetEditorController::SetOpSelected(const int32 InOpIndex)
@@ -1016,9 +990,13 @@ void FIKRetargetEditorController::SetOpSelected(const int32 InOpIndex)
 	ShowDetailsForOp(InOpIndex);
 }
 
-FName FIKRetargetEditorController::GetLastSelectedOpName() const
+FName FIKRetargetEditorController::GetSelectedOpName() const
 {
-	return Selection.LastSelectedOpName;
+	if (Selection.LastSelectedType == ERetargetSelectionType::OP)
+	{
+		return Selection.LastSelectedOpName;	
+	}
+	return NAME_None;
 }
 
 FIKRetargetOpBase* FIKRetargetEditorController::GetSelectedOp() const
@@ -1522,22 +1500,43 @@ void FIKRetargetEditorController::RenderSkeleton(FPrimitiveDrawInterface* PDI, E
 			BoneColors.Init(DefaultColor, RefSkeleton.GetNum());
 		}
 
+		auto AddBoneColorsForBonesInChain = [&BoneColors](
+			const TArray<FName>& InChainsToFind,
+			const TArray<FResolvedBoneChain>& InBoneChains,
+			const FLinearColor& InHighlightedColor)
+		{
+			for (const FResolvedBoneChain& BoneChain : InBoneChains)
+			{
+				if (!InChainsToFind.Contains(BoneChain.ChainName))
+				{
+					continue;
+				}
+				
+				for (const int32& BoneIndex : BoneChain.BoneIndices)
+				{
+					BoneColors[BoneIndex] = InHighlightedColor;
+				}
+			}
+		};
+
 		// highlight selected chains in blue
 		const TArray<FName>& SelectedChainNames = GetSelectedChains();
 		const FLinearColor HighlightedColor = FLinearColor::Blue;
-		const bool bIsSource = InSourceOrTarget == ERetargetSourceOrTarget::Source;
 		if (const FIKRetargetProcessor* Processor = GetRetargetProcessor())
 		{
 			const FRetargeterBoneChains& AllBoneChains = Processor->GetBoneChains();
-			const TArray<FResolvedBoneChain>& BoneChains = AllBoneChains.GetResolvedBoneChains(InSourceOrTarget);
-			for (const FResolvedBoneChain& BoneChain : BoneChains)
+
+			if (InSourceOrTarget == ERetargetSourceOrTarget::Source)
+			{
+				const TArray<FResolvedBoneChain>* BoneChains = AllBoneChains.GetResolvedBoneChains(ERetargetSourceOrTarget::Source);
+				AddBoneColorsForBonesInChain(SelectedChainNames, *BoneChains, HighlightedColor);
+			}
+			else
 			{
-				if (SelectedChainNames.Contains(BoneChain.ChainName))
+				const TMap<const UIKRigDefinition*, TArray<FResolvedBoneChain>>& TargetBoneChainMap = AllBoneChains.GetAllResolvedTargetBoneChains();
+				for (const TPair<const UIKRigDefinition*, TArray<FResolvedBoneChain>>& TargetBoneChainPair : TargetBoneChainMap)
 				{
-					for (const int32& BoneIndex : BoneChain.BoneIndices)
-					{
-						BoneColors[BoneIndex] = HighlightedColor;
-					}	
+					AddBoneColorsForBonesInChain(SelectedChainNames, TargetBoneChainPair.Value, HighlightedColor);
 				}
 			}
 		}
@@ -1610,7 +1609,7 @@ void FIKRetargetEditorController::FixZeroHeightRetargetRoot(ERetargetSourceOrTar
 
 	const FName CurrentRetargetPoseName = AssetController->GetCurrentRetargetPoseName(SourceOrTarget);
 	FIKRetargetPose& CurrentRetargetPose = AssetController->GetCurrentRetargetPose(SourceOrTarget);
-	const FName RetargetRootBoneName = AssetController->GetRetargetRootBone(SourceOrTarget);
+	const FName RetargetRootBoneName = AssetController->GetPelvisBone(SourceOrTarget);
 	if (RetargetRootBoneName == NAME_None)
 	{
 		return;
@@ -1621,7 +1620,6 @@ void FIKRetargetEditorController::FixZeroHeightRetargetRoot(ERetargetSourceOrTar
 	FRetargetSkeleton DummySkeleton;
 	DummySkeleton.Initialize(
 		SkeletalMesh,
-		TArray<FBoneChain>(),
 		CurrentRetargetPoseName,
 		&CurrentRetargetPose,
 		RetargetRootBoneName,
diff --git a/Engine/Plugins/Animation/IKRig/Source/IKRigEditor/Private/RetargetEditor/IKRetargetOpDetails.cpp b/Engine/Plugins/Animation/IKRig/Source/IKRigEditor/Private/RetargetEditor/IKRetargetOpDetails.cpp
new file mode 100644
--- /dev/null
+++ b/Engine/Plugins/Animation/IKRig/Source/IKRigEditor/Private/RetargetEditor/IKRetargetOpDetails.cpp
@@ -0,0 +1,536 @@
+// Copyright Epic Games, Inc. All Rights Reserved.
+
+#include "RetargetEditor/IKRetargetOpDetails.h"
+
+#include "IDetailChildrenBuilder.h"
+#include "IPropertyUtilities.h"
+#include "RetargetEditor/IKRetargeterController.h"
+#include "RetargetEditor/SIKRetargetChainMapList.h"
+#include "Retargeter/IKRetargeter.h"
+#include "Retargeter/RetargetOps/AlignPoleVectorOp.h"
+#include "Retargeter/RetargetOps/FKChainsOp.h"
+#include "Retargeter/RetargetOps/IKChainsOp.h"
+#include "Retargeter/RetargetOps/PelvisMotionOp.h"
+#include "Retargeter/RetargetOps/SpeedPlantingOp.h"
+#include "RigEditor/IKRigStructViewer.h"
+
+#define LOCTEXT_NAMESPACE "IKRetargetOpDetails"
+
+bool FIKRetargetOpBaseSettingsCustomization::LoadAndValidateStructToCustomize(
+	const TSharedRef<IPropertyHandle>& InStructPropertyHandle,
+	const IPropertyTypeCustomizationUtils& InStructCustomizationUtils)
+{
+	SelectedObjects = InStructCustomizationUtils.GetPropertyUtilities().Get()->GetSelectedObjects();
+	if (!ensure(!SelectedObjects.IsEmpty()))
+	{
+		return false;
+	}
+
+	void* StructMemory = nullptr;
+	FPropertyAccess::Result Result = InStructPropertyHandle->GetValueData(StructMemory);
+	if (!ensure(Result == FPropertyAccess::Success && StructMemory))
+	{
+		return false;
+	}
+
+	FIKRetargetOpSettingsBase* SettingsBeingCustomized = static_cast<FIKRetargetOpSettingsBase*>(StructMemory);
+	OpName = SettingsBeingCustomized->OwningOpName;
+	if (!ensure(OpName != NAME_None))
+	{
+		return false;
+	}
+	
+	StructViewer = Cast<UIKRigStructViewer>(SelectedObjects[0].Get());
+	RetargetAsset = CastChecked<UIKRetargeter>(StructViewer->GetStructOwner());
+	AssetController = UIKRetargeterController::GetController(RetargetAsset);
+
+	return true;
+}
+
+void FIKRetargetOpBaseSettingsCustomization::AddChildPropertyToGroup(
+	IDetailGroup& InGroup,
+	const TSharedRef<IPropertyHandle>& InParentPropertyHandle,
+	const FName& InChildPropertyName)
+{
+	TSharedPtr<IPropertyHandle> PropertyHandle = InParentPropertyHandle->GetChildHandle(InChildPropertyName);
+	if (!ensure(PropertyHandle.IsValid()))
+	{
+		return;
+	}
+	
+	InGroup.AddPropertyRow(PropertyHandle.ToSharedRef());
+}
+
+void FIKRetargetOpBaseSettingsCustomization::AddChildPropertiesToCategoryGroups(
+	const TSharedRef<IPropertyHandle>& InParentPropertyHandle,
+	IDetailChildrenBuilder& ChildBuilder)
+{
+	uint32 NumChildren;
+	InParentPropertyHandle->GetNumChildren(NumChildren);
+
+	// map of category names to property groups
+	TMap<FName, IDetailGroup*> CategoryGroups;
+	for (uint32 ChildIndex = 0; ChildIndex < NumChildren; ++ChildIndex)
+	{
+		TSharedPtr<IPropertyHandle> ChildHandle = InParentPropertyHandle->GetChildHandle(ChildIndex);
+		if (!ChildHandle.IsValid())
+		{
+			continue;
+		}
+		
+		FProperty* ChildProperty = ChildHandle->GetProperty();
+		if (!ChildProperty)
+		{
+			continue;
+		}
+
+		// get category from metadata
+		FString CategoryName = ChildProperty->GetMetaData(TEXT("Category"));
+		if (CategoryName.IsEmpty())
+		{
+			// if no category specified, add directly to builder
+			ChildBuilder.AddProperty(ChildHandle.ToSharedRef());
+			continue;
+		}
+
+		// convert to FName for map key
+		FName CategoryFName(*CategoryName);
+
+		// Get or create the group for this category
+		IDetailGroup* Group;
+		if (IDetailGroup** ExistingGroup = CategoryGroups.Find(CategoryFName))
+		{
+			Group = *ExistingGroup;
+		}
+		else
+		{
+			// create new group
+			FText CategoryText = FText::FromString(CategoryName);
+			Group = &ChildBuilder.AddGroup(CategoryFName, CategoryText);
+			CategoryGroups.Add(CategoryFName, Group);
+		}
+
+		// sdd the property to its category group
+		Group->AddPropertyRow(ChildHandle.ToSharedRef());
+	}
+}
+
+void FIKRetargetOpBaseSettingsCustomization::AddChildPropertiesInCategory(
+	const TSharedRef<IPropertyHandle>& InParentPropertyHandle,
+	IDetailChildrenBuilder& InChildBuilder,
+	const FName& InCategoryName,
+	const TArray<FName>& InPropertiesToIgnore)
+{
+	uint32 NumChildren;
+	InParentPropertyHandle->GetNumChildren(NumChildren);
+
+	IDetailGroup* CategoryGroup = nullptr;
+	for (uint32 ChildIndex = 0; ChildIndex < NumChildren; ++ChildIndex)
+	{
+		TSharedPtr<IPropertyHandle> ChildHandle = InParentPropertyHandle->GetChildHandle(ChildIndex);
+		if (!ChildHandle.IsValid())
+		{
+			continue;
+		}
+
+		FProperty* ChildProperty = ChildHandle->GetProperty();
+		if (!ChildProperty)
+		{
+			continue;
+		}
+
+		if (InPropertiesToIgnore.Contains(ChildProperty->GetName()))
+		{
+			continue; // filtered out by name
+		}
+		
+		FString CategoryName = ChildProperty->GetMetaData(TEXT("Category"));
+		if (CategoryName.IsEmpty() || InCategoryName != CategoryName)
+		{
+			continue;  // not in the category we're looking for
+		}
+
+		// create group if not created yet
+		if (!CategoryGroup)
+		{
+			static bool bStartExpanded = true;
+			CategoryGroup = &InChildBuilder.AddGroup(InCategoryName, FText::FromName(InCategoryName), bStartExpanded);
+		}
+
+		// Add the property to the category group
+		IDetailPropertyRow& NewRow = CategoryGroup->AddPropertyRow(ChildHandle.ToSharedRef());
+		NewRow.ShouldAutoExpand(true);
+	}
+
+	if (CategoryGroup)
+	{
+		CategoryGroup->ToggleExpansion(true);
+	}
+}
+
+void FChainsFKOpCustomization::CustomizeChildren(
+	TSharedRef<IPropertyHandle> StructPropertyHandle,
+	IDetailChildrenBuilder& ChildBuilder,
+	IPropertyTypeCustomizationUtils& StructCustomizationUtils)
+{
+	if (!LoadAndValidateStructToCustomize(StructPropertyHandle, StructCustomizationUtils))
+	{
+		return;
+	}
+
+	// add the IK Rig asset input field
+	TSharedPtr<IPropertyHandle> IKRigHandle = StructPropertyHandle->GetChildHandle(GET_MEMBER_NAME_CHECKED(FIKRetargetFKChainsOpSettings, IKRigAsset));
+	ChildBuilder.AddProperty(IKRigHandle.ToSharedRef());
+
+	// configure the chain map list view for FK chains
+	FChainMapListConfig ChainListConfig;
+	ChainListConfig.OpWithChainSettings = OpName;
+	ChainListConfig.OpWithChainMapping = OpName;
+	ChainListConfig.Controller = AssetController;
+	ChainListConfig.bEnableGoalColumn = false;
+	ChainListConfig.bEnableChainMapping = true;
+	ChainListConfig.ChainSettingsGetterFunc = [this](const FName InTargetChainName) -> UObject*
+	{
+		// this lambda fetches the FK chain settings struct for the given target chain in the given op (by name)
+		auto MemoryProvider = [this, InTargetChainName]() -> uint8*
+		{
+			if (FIKRetargetOpBase* Op = AssetController->GetRetargetOpByName(OpName))
+			{
+				FIKRetargetFKChainsOpSettings* OpSettings = reinterpret_cast<FIKRetargetFKChainsOpSettings*>(Op->GetSettings());
+				for (FRetargetFKChainSettings& ChainToRetarget : OpSettings->ChainsToRetarget)
+				{
+					if (ChainToRetarget.TargetChainName == InTargetChainName)
+					{
+						return reinterpret_cast<uint8*>(&ChainToRetarget);
+					}
+				}
+					
+			}
+			return nullptr;
+		};
+			
+		FIKRigStructToView StructToView;
+		StructToView.Owner = AssetController->GetAsset();
+		StructToView.Type = FRetargetFKChainSettings::StaticStruct();
+		StructToView.MemoryProvider = MemoryProvider;
+		StructToView.UniqueName = OpName;
+		
+		URetargetFKChainSettingsWrapper* Wrapper = NewObject<URetargetFKChainSettingsWrapper>(GetTransientPackage(), NAME_None, RF_Standalone);
+		const FName SettingsPropertyName = GET_MEMBER_NAME_CHECKED(URetargetFKChainSettingsWrapper, Settings);
+		Wrapper->InitializeWithRetargeter(StructToView, SettingsPropertyName, AssetController);
+		Wrapper->SetPropertyHidden(GET_MEMBER_NAME_CHECKED(FRetargetFKChainSettings, TargetChainName), true);
+		return Wrapper;
+	};
+
+	// add the chain mapping list
+	const FText ChainMapRowName = FText::FromString("Chain Map");
+	ChildBuilder.AddCustomRow(ChainMapRowName)
+	.WholeRowContent()
+	[
+		SNew(SIKRetargetChainMapList).InChainMapListConfig(ChainListConfig)
+	];
+
+	// add the debug properties
+	AddChildPropertiesInCategory(StructPropertyHandle, ChildBuilder, FName("Debug"));
+}
+
+void FRunIKRigOpCustomization::CustomizeChildren(
+	TSharedRef<IPropertyHandle> StructPropertyHandle,
+	IDetailChildrenBuilder& ChildBuilder,
+	IPropertyTypeCustomizationUtils& StructCustomizationUtils)
+{
+	if (!LoadAndValidateStructToCustomize(StructPropertyHandle, StructCustomizationUtils))
+	{
+		return;
+	}
+
+	// add the IK Rig asset input field
+	TSharedPtr<IPropertyHandle> IKRigHandle = StructPropertyHandle->GetChildHandle(GET_MEMBER_NAME_CHECKED(FIKRetargetFKChainsOpSettings, IKRigAsset));
+	ChildBuilder.AddProperty(IKRigHandle.ToSharedRef());
+
+	// show only IK chains by default
+	FChainMapFilterOptions ShowOnlyIKFilter;
+	ShowOnlyIKFilter.bNeverShowChainsWithoutIK = true;
+
+	// configure chain map list for the IK op
+	// NOTE: no settings memory provider because this op only maps chains, it doesn't store any settings for them
+	FChainMapListConfig ChainListConfig;
+	ChainListConfig.OpWithChainSettings = NAME_None;
+	ChainListConfig.OpWithChainMapping = OpName;
+	ChainListConfig.Controller = AssetController;
+	ChainListConfig.bEnableGoalColumn = true;
+	ChainListConfig.bEnableChainMapping = true;
+	ChainListConfig.Filter = ShowOnlyIKFilter;
+
+	// add the chain mapping list
+	FText ChainMapRowName = FText::FromString("Chain Map");
+	ChildBuilder.AddCustomRow(ChainMapRowName)
+	.WholeRowContent()
+	[
+		SNew(SIKRetargetChainMapList).InChainMapListConfig(ChainListConfig)
+	];
+
+	// add the debug properties
+	AddChildPropertiesInCategory(StructPropertyHandle, ChildBuilder, FName("Debug"));
+}
+
+void FIKChainOpCustomization::CustomizeChildren(
+	TSharedRef<IPropertyHandle> StructPropertyHandle,
+	IDetailChildrenBuilder& ChildBuilder,
+	IPropertyTypeCustomizationUtils& StructCustomizationUtils)
+{
+	if (!LoadAndValidateStructToCustomize(StructPropertyHandle, StructCustomizationUtils))
+	{
+		return;
+	}
+	
+	// show only IK chains
+	FChainMapFilterOptions ShowOnlyIKFilter;
+	ShowOnlyIKFilter.bNeverShowChainsWithoutIK = true;
+
+	// configure chain map list for the IK op
+	FChainMapListConfig ChainListConfig;
+	ChainListConfig.OpWithChainSettings = OpName;
+	ChainListConfig.OpWithChainMapping = AssetController->GetParentOpByName(OpName); // use chain mapping from parent
+	ChainListConfig.Controller = AssetController;
+	ChainListConfig.bEnableGoalColumn = true;
+	ChainListConfig.bEnableChainMapping = false;
+	ChainListConfig.Filter = ShowOnlyIKFilter;
+	ChainListConfig.ChainSettingsGetterFunc = [this](const FName InTargetChainName) -> UObject*
+	{
+		// this lambda fetches the chain settings struct for the given target chain in the given op (by name)
+		auto MemoryProvider = [this, InTargetChainName]() -> uint8*
+		{
+			if (FIKRetargetOpBase* Op = AssetController->GetRetargetOpByName(OpName))
+			{
+				FIKRetargetIKChainsOpSettings* OpSettings = reinterpret_cast<FIKRetargetIKChainsOpSettings*>(Op->GetSettings());
+				for (FRetargetIKChainSettings& ChainToRetarget : OpSettings->ChainsToRetarget)
+				{
+					if (ChainToRetarget.TargetChainName == InTargetChainName)
+					{
+						return reinterpret_cast<uint8*>(&ChainToRetarget);
+					}
+				}
+			}
+			return nullptr;
+		};
+			
+		FIKRigStructToView StructToView;
+		StructToView.Owner = AssetController->GetAsset();
+		StructToView.Type = FRetargetIKChainSettings::StaticStruct();
+		StructToView.MemoryProvider = MemoryProvider;
+		StructToView.UniqueName = OpName;
+		
+		URetargetIKChainSettingsWrapper* Wrapper = NewObject<URetargetIKChainSettingsWrapper>(GetTransientPackage(), NAME_None, RF_Standalone);
+		const FName SettingsPropertyName = GET_MEMBER_NAME_CHECKED(URetargetIKChainSettingsWrapper, Settings);
+		Wrapper->InitializeWithRetargeter(StructToView, SettingsPropertyName, AssetController);
+		Wrapper->SetPropertyHidden(GET_MEMBER_NAME_CHECKED(FRetargetIKChainSettings, TargetChainName), true);
+		return Wrapper;
+	};
+
+	// add the chain mapping list
+	FText ChainMapRowName = FText::FromString("Chain Map");
+	ChildBuilder.AddCustomRow(ChainMapRowName)
+	.WholeRowContent()
+	[
+		SNew(SIKRetargetChainMapList).InChainMapListConfig(ChainListConfig)
+	];
+
+	// add the debug properties
+	AddChildPropertiesInCategory(StructPropertyHandle, ChildBuilder, FName("Debug"));
+}
+
+void FStrideWarpOpCustomization::CustomizeChildren(
+	TSharedRef<IPropertyHandle> StructPropertyHandle,
+	IDetailChildrenBuilder& ChildBuilder,
+	IPropertyTypeCustomizationUtils& StructCustomizationUtils)
+{
+	if (!LoadAndValidateStructToCustomize(StructPropertyHandle, StructCustomizationUtils))
+	{
+		return;
+	}
+	
+	// show only IK chains
+	FChainMapFilterOptions ShowOnlyIKFilter;
+	ShowOnlyIKFilter.bNeverShowChainsWithoutIK = true;
+
+	// configure chain map list for the IK op
+	FChainMapListConfig ChainListConfig;
+	ChainListConfig.OpWithChainSettings = OpName;
+	ChainListConfig.OpWithChainMapping = AssetController->GetParentOpByName(OpName); // use chain mapping from parent
+	ChainListConfig.Controller = AssetController;
+	ChainListConfig.bEnableGoalColumn = true;
+	ChainListConfig.bEnableChainMapping = false;
+	ChainListConfig.Filter = ShowOnlyIKFilter;
+	ChainListConfig.ChainSettingsGetterFunc = [this](const FName InTargetChainName) -> UObject*
+	{
+		// this lambda fetches the chain settings struct for the given target chain in the given op (by name)
+		auto MemoryProvider = [this, InTargetChainName]() -> uint8*
+		{
+			if (FIKRetargetOpBase* Op = AssetController->GetRetargetOpByName(OpName))
+			{
+				FIKRetargetStrideWarpingOpSettings* OpSettings = reinterpret_cast<FIKRetargetStrideWarpingOpSettings*>(Op->GetSettings());
+				for (FRetargetStrideWarpChainSettings& ChainToWarp : OpSettings->ChainSettings)
+				{
+					if (ChainToWarp.TargetChainName == InTargetChainName)
+					{
+						return reinterpret_cast<uint8*>(&ChainToWarp);
+					}
+				}
+			}
+			return nullptr;
+		};
+			
+		FIKRigStructToView StructToView;
+		StructToView.Owner = AssetController->GetAsset();
+		StructToView.Type = FRetargetStrideWarpChainSettings::StaticStruct();
+		StructToView.MemoryProvider = MemoryProvider;
+		StructToView.UniqueName = OpName;
+		
+		URetargetStrideWarpSettingsWrapper* Wrapper = NewObject<URetargetStrideWarpSettingsWrapper>(GetTransientPackage(), NAME_None, RF_Standalone);
+		const FName SettingsPropertyName = GET_MEMBER_NAME_CHECKED(URetargetStrideWarpSettingsWrapper, Settings);
+		Wrapper->InitializeWithRetargeter(StructToView, SettingsPropertyName, AssetController);
+		Wrapper->SetPropertyHidden(GET_MEMBER_NAME_CHECKED(FRetargetStrideWarpChainSettings, TargetChainName), true);
+		return Wrapper;
+	};
+
+	// add the chain mapping list
+	FText ChainMapRowName = FText::FromString("Chain Map");
+	ChildBuilder.AddCustomRow(ChainMapRowName)
+	.WholeRowContent()
+	[
+		SNew(SIKRetargetChainMapList).InChainMapListConfig(ChainListConfig)
+	];
+
+	// add the warping properties
+	AddChildPropertiesInCategory(StructPropertyHandle, ChildBuilder, FName("Warping"));
+	// add the debug properties
+	AddChildPropertiesInCategory(StructPropertyHandle, ChildBuilder, FName("Debug"));
+}
+
+void FSpeedPlantOpCustomization::CustomizeChildren(
+	TSharedRef<IPropertyHandle> StructPropertyHandle,
+	IDetailChildrenBuilder& ChildBuilder,
+	IPropertyTypeCustomizationUtils& StructCustomizationUtils)
+{
+	if (!LoadAndValidateStructToCustomize(StructPropertyHandle, StructCustomizationUtils))
+	{
+		return;
+	}
+	
+	// show only IK chains
+	FChainMapFilterOptions ShowOnlyIKFilter;
+	ShowOnlyIKFilter.bNeverShowChainsWithoutIK = true;
+
+	// configure chain map list for the IK op
+	FChainMapListConfig ChainListConfig;
+	ChainListConfig.OpWithChainSettings = OpName;
+	ChainListConfig.OpWithChainMapping = AssetController->GetParentOpByName(OpName); // use chain mapping from parent
+	ChainListConfig.Controller = AssetController;
+	ChainListConfig.bEnableGoalColumn = true;
+	ChainListConfig.bEnableChainMapping = false;
+	ChainListConfig.Filter = ShowOnlyIKFilter;
+	ChainListConfig.ChainSettingsGetterFunc = [this](const FName InTargetChainName) -> UObject*
+	{
+		// this lambda fetches the chain settings struct for the given target chain in the given op (by name)
+		auto MemoryProvider = [this, InTargetChainName]() -> uint8*
+		{
+			if (FIKRetargetOpBase* Op = AssetController->GetRetargetOpByName(OpName))
+			{
+				FIKRetargetSpeedPlantingOpSettings* OpSettings = reinterpret_cast<FIKRetargetSpeedPlantingOpSettings*>(Op->GetSettings());
+				for (FRetargetSpeedPlantingSettings& ChainToWarp : OpSettings->ChainsToSpeedPlant)
+				{
+					if (ChainToWarp.TargetChainName == InTargetChainName)
+					{
+						return reinterpret_cast<uint8*>(&ChainToWarp);
+					}
+				}
+			}
+			return nullptr;
+		};
+			
+		FIKRigStructToView StructToView;
+		StructToView.Owner = AssetController->GetAsset();
+		StructToView.Type = FRetargetSpeedPlantingSettings::StaticStruct();
+		StructToView.MemoryProvider = MemoryProvider;
+		StructToView.UniqueName = OpName;
+		
+		URetargetSpeedPlantSettingsWrapper* Wrapper = NewObject<URetargetSpeedPlantSettingsWrapper>(GetTransientPackage(), NAME_None, RF_Standalone);
+		const FName SettingsPropertyName = GET_MEMBER_NAME_CHECKED(URetargetSpeedPlantSettingsWrapper, Settings);
+		Wrapper->InitializeWithRetargeter(StructToView, SettingsPropertyName, AssetController);
+		Wrapper->SetPropertyHidden(GET_MEMBER_NAME_CHECKED(FRetargetSpeedPlantingSettings, TargetChainName), true);
+		return Wrapper;
+	};
+
+	// add the chain mapping list
+	FText ChainMapRowName = FText::FromString("Chain Map");
+	ChildBuilder.AddCustomRow(ChainMapRowName)
+	.WholeRowContent()
+	[
+		SNew(SIKRetargetChainMapList).InChainMapListConfig(ChainListConfig)
+	];
+
+	// add speed planting properties
+	AddChildPropertiesInCategory(StructPropertyHandle, ChildBuilder, FName("Speed Planting"));
+}
+
+void FPoleVectorOpCustomization::CustomizeChildren(
+	TSharedRef<IPropertyHandle> StructPropertyHandle,
+	IDetailChildrenBuilder& ChildBuilder,
+	IPropertyTypeCustomizationUtils& StructCustomizationUtils)
+{
+	if (!LoadAndValidateStructToCustomize(StructPropertyHandle, StructCustomizationUtils))
+	{
+		return;
+	}
+
+	// configure chain map list for the IK op
+	FChainMapListConfig ChainListConfig;
+	ChainListConfig.OpWithChainSettings = OpName;
+	ChainListConfig.OpWithChainMapping = OpName;
+	ChainListConfig.Controller = AssetController;
+	ChainListConfig.bEnableGoalColumn = true;
+	ChainListConfig.bEnableChainMapping = true;
+	ChainListConfig.ChainSettingsGetterFunc = [this](const FName InTargetChainName) -> UObject*
+	{
+		// this lambda fetches the chain settings struct for the given target chain in the given op (by name)
+		auto MemoryProvider = [this, InTargetChainName]() -> uint8*
+		{
+			if (FIKRetargetOpBase* Op = AssetController->GetRetargetOpByName(OpName))
+			{
+				FIKRetargetAlignPoleVectorOpSettings* OpSettings = reinterpret_cast<FIKRetargetAlignPoleVectorOpSettings*>(Op->GetSettings());
+				for (FRetargetPoleVectorSettings& ChainToAlign : OpSettings->ChainsToAlign)
+				{
+					if (ChainToAlign.TargetChainName == InTargetChainName)
+					{
+						return reinterpret_cast<uint8*>(&ChainToAlign);
+					}
+				}
+			}
+			return nullptr;
+		};
+			
+		FIKRigStructToView StructToView;
+		StructToView.Owner = AssetController->GetAsset();
+		StructToView.Type = FRetargetPoleVectorSettings::StaticStruct();
+		StructToView.MemoryProvider = MemoryProvider;
+		StructToView.UniqueName = OpName;
+		
+		UPoleVectorSettingsWrapper* Wrapper = NewObject<UPoleVectorSettingsWrapper>(GetTransientPackage(), NAME_None, RF_Standalone);
+		const FName SettingsPropertyName = GET_MEMBER_NAME_CHECKED(UPoleVectorSettingsWrapper, Settings);
+		Wrapper->InitializeWithRetargeter(StructToView, SettingsPropertyName, AssetController);
+		Wrapper->SetPropertyHidden(GET_MEMBER_NAME_CHECKED(FRetargetPoleVectorSettings, TargetChainName), true);
+		return Wrapper;
+	};
+
+	// add the chain mapping list
+	FText ChainMapRowName = FText::FromString("Chain Map");
+	ChildBuilder.AddCustomRow(ChainMapRowName)
+	.WholeRowContent()
+	[
+		SNew(SIKRetargetChainMapList).InChainMapListConfig(ChainListConfig)
+	];
+}
+
+
+#undef LOCTEXT_NAMESPACE
diff --git a/Engine/Plugins/Animation/IKRig/Source/IKRigEditor/Private/RetargetEditor/IKRetargetPoseExporter.cpp b/Engine/Plugins/Animation/IKRig/Source/IKRigEditor/Private/RetargetEditor/IKRetargetPoseExporter.cpp
--- a/Engine/Plugins/Animation/IKRig/Source/IKRigEditor/Private/RetargetEditor/IKRetargetPoseExporter.cpp
+++ b/Engine/Plugins/Animation/IKRig/Source/IKRigEditor/Private/RetargetEditor/IKRetargetPoseExporter.cpp
@@ -571,7 +571,7 @@ FReply FIKRetargetPoseExporter::OnImportPoseFromSequence()
 	const FReferenceSkeleton& RefSkeleton = Mesh->GetRefSkeleton();
 	const TArray<FTransform>& RefPose = RefSkeleton.GetRefBonePose();
 	int32 NumBones = RefSkeleton.GetNum();
-	const FName RootBoneName = ControllerPtr->AssetController->GetRetargetRootBone(ControllerPtr->GetSourceOrTarget());
+	const FName RootBoneName = ControllerPtr->AssetController->GetPelvisBone(ControllerPtr->GetSourceOrTarget());
 	for (int32 BoneIndex = 0; BoneIndex < NumBones; BoneIndex++)
 	{
 		const FName& BoneName = RefSkeleton.GetBoneName(BoneIndex);
