[IK Retargeter] Adding support for chain mapping per-op

diff --git a/Engine/Plugins/Animation/IKRig/Source/IKRigEditor/Private/RetargetEditor/IKRetargeterController.cpp b/Engine/Plugins/Animation/IKRig/Source/IKRigEditor/Private/RetargetEditor/IKRetargeterController.cpp
--- a/Engine/Plugins/Animation/IKRig/Source/IKRigEditor/Private/RetargetEditor/IKRetargeterController.cpp
+++ b/Engine/Plugins/Animation/IKRig/Source/IKRigEditor/Private/RetargetEditor/IKRetargeterController.cpp
@@ -3,7 +3,6 @@
 #include "RetargetEditor/IKRetargeterController.h"
 
 #include "ScopedTransaction.h"
-#include "Algo/LevenshteinDistance.h"
 #include "Engine/SkeletalMesh.h"
 #include "Retargeter/IKRetargeter.h"
 #include "RigEditor/IKRigController.h"
@@ -15,7 +14,7 @@
 #include "Retargeter/RetargetOps/PelvisMotionOp.h"
 #include "Retargeter/RetargetOps/RootMotionGeneratorOp.h"
 #include "Retargeter/RetargetOps/RunIKRigOp.h"
-#include "RigEditor/IKRigDetailCustomizations.h"
+#include "RigEditor/IKRigStructViewer.h"
 
 #include UE_INLINE_GENERATED_CPP_BY_NAME(IKRetargeterController)
 
@@ -42,10 +41,7 @@ UIKRetargeterController* UIKRetargeterController::GetController(const UIKRetarge
 UIKRetargeterController::UIKRetargeterController()
 {
 	StructViewer = CreateDefaultSubobject<UIKRigStructViewer>(TEXT("RetargetSettingsViewer"));
-	StructViewer->OnStructPropertyEdited().AddLambda([this]()
-	{
-		BroadcastNeedsReinitialized();
-	});
+	StructViewer->OnStructNeedsReinit().AddUObject(this, &UIKRetargeterController::OnOpPropertyChanged);
 }
 
 void UIKRetargeterController::PostInitProperties()
@@ -62,7 +58,8 @@ UIKRetargeter* UIKRetargeterController::GetAsset() const
 void UIKRetargeterController::CleanAsset() const
 {
 	FScopeLock Lock(&ControllerLock);
-	Asset->GetChainMapping().CleanMapping(GetIKRig(ERetargetSourceOrTarget::Source), GetIKRig(ERetargetSourceOrTarget::Target));
+
+	CleanChainMaps();
 	CleanPoseList(ERetargetSourceOrTarget::Source);
 	CleanPoseList(ERetargetSourceOrTarget::Target);
 }
@@ -90,10 +87,24 @@ void UIKRetargeterController::SetIKRig(const ERetargetSourceOrTarget SourceOrTar
 		SetAskedToFixRootHeightForMesh(GetPreviewMesh(SourceOrTarget), false);
 	}
 
-	Asset->GetChainMapping().CleanMapping(GetIKRig(ERetargetSourceOrTarget::Source), GetIKRig(ERetargetSourceOrTarget::Target));
-	
-	constexpr bool bForceRemap = false;
-	AutoMapChains(EAutoMapChainType::Exact, bForceRemap);
+	// update ops with new source
+	// NOTE: we do NOT auto-update the target IK rig as this may be overridden
+	if (SourceOrTarget == ERetargetSourceOrTarget::Source)
+	{
+		const int32 NumOps = GetNumRetargetOps();
+		for (int32 OpIndex = 0; OpIndex < NumOps; ++OpIndex)
+		{
+			FIKRetargetOpBase* Op = GetRetargetOpByIndex(OpIndex);
+			FRetargetChainMapping* ChainMapping = Op->GetChainMapping();
+			if (!ChainMapping)
+			{
+				continue;
+			}
+
+			const UIKRigDefinition* TargetIKRig = GetTargetIKRigForOp(Op->GetName());
+			ChainMapping->ReinitializeWithIKRigs(Asset->SourceIKRigAsset, TargetIKRig);
+		}
+	}
 
 	// update any editors attached to this asset
 	BroadcastIKRigReplaced(SourceOrTarget);
@@ -106,6 +117,23 @@ const UIKRigDefinition* UIKRetargeterController::GetIKRig(const ERetargetSourceO
 	return Asset->GetIKRig(SourceOrTarget);
 }
 
+TArray<UIKRigDefinition*> UIKRetargeterController::GetAllTargetIKRigs() const
+{
+	TArray<UIKRigDefinition*> AllTargetIKRigs;
+	
+	const int32 NumOps = GetNumRetargetOps();
+	for (int32 OpIndex = 0; OpIndex < NumOps; ++OpIndex)
+	{
+		FIKRetargetOpBase* Op = GetRetargetOpByIndex(OpIndex);
+		if (UIKRigDefinition* TargetIKRig = const_cast<UIKRigDefinition*>(Op->GetCustomTargetIKRig()))
+		{
+			AllTargetIKRigs.AddUnique(TargetIKRig);
+		}
+	}
+	
+	return MoveTemp(AllTargetIKRigs);
+}
+
 UIKRigDefinition* UIKRetargeterController::GetIKRigWriteable(const ERetargetSourceOrTarget SourceOrTarget) const
 {
 	FScopeLock Lock(&ControllerLock);
@@ -159,7 +187,7 @@ USkeletalMesh* UIKRetargeterController::GetPreviewMesh(const ERetargetSourceOrTa
 	return IKRig->GetPreviewMesh();
 }
 
-int32 UIKRetargeterController::AddRetargetOp(const FString InIKRetargetOpType, const bool bAllowDuplicates, const bool bRunInitialSetup) const
+int32 UIKRetargeterController::AddRetargetOp(const FString InIKRetargetOpType) const
 {
 	UScriptStruct* OpType = FindObject<UScriptStruct>(nullptr, *InIKRetargetOpType);
 	if (!OpType)
@@ -168,7 +196,7 @@ int32 UIKRetargeterController::AddRetargetOp(const FString InIKRetargetOpType, c
 		return INDEX_NONE;
 	}
 
-	return AddRetargetOp(OpType, NAME_None, bAllowDuplicates, bRunInitialSetup);
+	return AddRetargetOp(OpType);
 }
 
 bool UIKRetargeterController::RemoveRetargetOp(const int32 OpIndex) const
@@ -213,7 +241,7 @@ bool UIKRetargeterController::RemoveAllOps() const
 
 FName UIKRetargeterController::SetOpName(const FName InName, const int32 InOpIndex) const
 {
-	FIKRetargetOpBase* Op = GetOpAtIndex(InOpIndex);
+	FIKRetargetOpBase* Op = GetRetargetOpByIndex(InOpIndex);
 	if (!Op)
 	{
 		return NAME_None;
@@ -245,7 +273,7 @@ FName UIKRetargeterController::SetOpName(const FName InName, const int32 InOpInd
 
 FName UIKRetargeterController::GetOpName(const int32 InOpIndex) const
 {
-	const FIKRetargetOpBase* Op = GetOpAtIndex(InOpIndex);
+	const FIKRetargetOpBase* Op = GetRetargetOpByIndex(InOpIndex);
 	if (!Op)
 	{
 		return NAME_None;
@@ -306,24 +334,24 @@ int32 UIKRetargeterController::GetIndexOfOpByName(const FName InOpName) const
 	});
 }
 
-void UIKRetargeterController::AddDefaultOps(const bool bRunInitialSetup) const
+void UIKRetargeterController::AddDefaultOps() const
 {
 	FScopedTransaction Transaction(LOCTEXT("AddDefaultOps_Label", "Add Default Ops"));
 	FScopedReinitializeIKRetargeter Reinitialize(this);
 	Asset->Modify();
 
 	// add set of default ops for basic retargeting
-	constexpr bool bAllowDuplicateOps = false;
-	AddRetargetOp(FIKRetargetPelvisMotionOp::StaticStruct(), NAME_None, bAllowDuplicateOps, bRunInitialSetup);
-	AddRetargetOp(FIKRetargetFKChainsOp::StaticStruct(), NAME_None,bAllowDuplicateOps, bRunInitialSetup);
-	AddRetargetOp(FIKRetargetIKChainsOp::StaticStruct(), NAME_None,bAllowDuplicateOps, bRunInitialSetup);
-	AddRetargetOp(FIKRetargetRunIKRigOp::StaticStruct(), NAME_None,bAllowDuplicateOps, bRunInitialSetup);
-	AddRetargetOp(FIKRetargetRootMotionOp::StaticStruct(), NAME_None,bAllowDuplicateOps, bRunInitialSetup);
+	AddRetargetOp(FIKRetargetPelvisMotionOp::StaticStruct());
+	AddRetargetOp(FIKRetargetFKChainsOp::StaticStruct());
+	const int32 RunIKIndex = AddRetargetOp(FIKRetargetRunIKRigOp::StaticStruct());
+	const FName RunIKOpName = GetOpName(RunIKIndex);
+	AddRetargetOp(FIKRetargetIKChainsOp::StaticStruct(), RunIKOpName);
+	AddRetargetOp(FIKRetargetRootMotionOp::StaticStruct());
 }
 
 void UIKRetargeterController::RunOpInitialSetup(const int32 InOpIndex) const
 {
-	FIKRetargetOpBase* Op = GetOpAtIndex(InOpIndex);
+	FIKRetargetOpBase* Op = GetRetargetOpByIndex(InOpIndex);
 	if (!Op)
 	{
 		UE_LOG(LogIKRigEditor, Warning, TEXT("Retarget Op initial setup skipped. Invalid index, %d."), InOpIndex);
@@ -333,10 +361,12 @@ void UIKRetargeterController::RunOpInitialSetup(const int32 InOpIndex) const
 	FScopedTransaction Transaction(LOCTEXT("RunOpInitialSetup_Label", "Run Op Initial Setup"));
 	FScopedReinitializeIKRetargeter Reinitialize(this);
 	Asset->Modify();
-	Op->RunInitialSetup(GetAsset());
+
+	FIKRetargetOpBase* ParentOp = GetRetargetOpByIndex(GetParentOpIndex(InOpIndex));
+	Op->OnAddedToStack(GetAsset(), ParentOp);
 }
 
-FIKRetargetOpBase* UIKRetargeterController::GetOpAtIndex(const int32 InOpIndex) const
+FIKRetargetOpBase* UIKRetargeterController::GetRetargetOpByIndex(const int32 InOpIndex) const
 {
 	check(Asset)
 
@@ -350,7 +380,7 @@ FIKRetargetOpBase* UIKRetargeterController::GetOpAtIndex(const int32 InOpIndex)
 
 FIKRetargetOpBase* UIKRetargeterController::GetRetargetOpByName(const FName InOpName) const
 {
-	return GetOpAtIndex(GetIndexOfOpByName(InOpName));
+	return GetRetargetOpByIndex(GetIndexOfOpByName(InOpName));
 }
 
 TArray<int32> UIKRetargeterController::GetChildOpIndices(const int32 InOpIndex) const
@@ -379,7 +409,7 @@ TArray<int32> UIKRetargeterController::GetChildOpIndices(const int32 InOpIndex)
 
 bool UIKRetargeterController::GetCanOpHaveChildren(const int32 InOpIndex) const
 {
-	if (FIKRetargetOpBase* Op = GetOpAtIndex(InOpIndex))
+	if (FIKRetargetOpBase* Op = GetRetargetOpByIndex(InOpIndex))
 	{
 		return Op->CanHaveChildOps();
 	}
@@ -390,7 +420,7 @@ int32 UIKRetargeterController::GetParentOpIndex(const int32 InOpIndex) const
 {
 	check(Asset)
 
-	if (FIKRetargetOpBase* Op = GetOpAtIndex(InOpIndex))
+	if (FIKRetargetOpBase* Op = GetRetargetOpByIndex(InOpIndex))
 	{
 		return GetIndexOfOpByName(Op->GetParentOpName());
 	}
@@ -455,6 +485,32 @@ int32 UIKRetargeterController::GetIndexOfRetargetOp(FIKRetargetOpBase* RetargetO
 	return INDEX_NONE;
 }
 
+void UIKRetargeterController::OnOpPropertyChanged(const FName& InOpName, const FPropertyChangedEvent& InPropertyChangedEvent) const
+{
+	const int32 OpIndex = GetIndexOfOpByName(InOpName);
+	if (!ensure(OpIndex != INDEX_NONE))
+	{
+		return; // should not get callback for unknown op
+	}
+
+	// notify op of property edit
+	FIKRetargetOpBase* Op = GetRetargetOpByIndex(OpIndex);
+	Op->OnReinitPropertyEdited(&InPropertyChangedEvent);
+
+	// notify children of parent property edit
+	TArray<int32> ChildOps = GetChildOpIndices(OpIndex);
+	for (int32 ChildOpIndex = 0; ChildOpIndex < ChildOps.Num(); ++ChildOpIndex)
+	{
+		if (FIKRetargetOpBase* ChildOp = GetRetargetOpByIndex(ChildOpIndex))
+		{
+			ChildOp->OnParentReinitPropertyEdited(*Op, &InPropertyChangedEvent);
+		}
+	}
+
+	// reinitialize
+	BroadcastNeedsReinitialized();
+}
+
 FInstancedStruct* UIKRetargeterController::GetRetargetOpStructAtIndex(int32 Index) const
 {
 	check(Asset);
@@ -523,7 +579,7 @@ bool UIKRetargeterController::SetRetargetOpEnabled(int32 RetargetOpIndex, bool b
 	FScopedTransaction Transaction(LOCTEXT("SetRetargetOpEnabled_Label", "Enable/Disable Op"));
 	FScopedReinitializeIKRetargeter Reinitialize(this);
 	Asset->Modify();
-	GetOpAtIndex(RetargetOpIndex)->SetEnabled(bIsEnabled);
+	GetRetargetOpByIndex(RetargetOpIndex)->SetEnabled(bIsEnabled);
 	return true;
 }
 
@@ -535,12 +591,12 @@ bool UIKRetargeterController::GetRetargetOpEnabled(int32 RetargetOpIndex) const
 		return false;
 	}
 
-	return GetOpAtIndex(RetargetOpIndex)->IsEnabled();
+	return GetRetargetOpByIndex(RetargetOpIndex)->IsEnabled();
 }
 
 UIKRetargetOpControllerBase* UIKRetargeterController::GetOpController(int32 OpIndex)
 {
-	if (FIKRetargetOpBase* OpAtIndex = GetOpAtIndex(OpIndex))
+	if (FIKRetargetOpBase* OpAtIndex = GetRetargetOpByIndex(OpIndex))
 	{
 		return OpAtIndex->GetController(this);
 	}
@@ -555,33 +611,105 @@ void UIKRetargeterController::ResetChainSettingsInAllOps(const FName InTargetCha
 	for (int32 OpIndex=0; OpIndex<Asset->RetargetOps.Num(); ++OpIndex)
 	{
 		FIKRetargetOpBase& Op = Asset->RetargetOps[OpIndex].GetMutable<FIKRetargetOpBase>();
-		Op.SetChainSettingsToDefault(InTargetChainName);
+		Op.ResetChainSettingsToDefault(InTargetChainName);
 	}
 }
 
 const FRetargetChainMapping& UIKRetargeterController::GetChainMapping() const
 {
+	PRAGMA_DISABLE_DEPRECATION_WARNINGS
 	return Asset->GetChainMapping();
+	PRAGMA_ENABLE_DEPRECATION_WARNINGS
 }
 
-bool UIKRetargeterController::ChainHasNonDefaultSettings(const FName InTargetChainName) const
+void UIKRetargeterController::CleanChainMaps(const FName InOpName) const
 {
-	for (int32 OpIndex=0; OpIndex<Asset->RetargetOps.Num(); ++OpIndex)
+	const UIKRigDefinition* SourceIKRig = GetIKRig(ERetargetSourceOrTarget::Source);
+	
+	auto CleanChainMapInOp = [SourceIKRig](FIKRetargetOpBase* InOp)
 	{
-		FIKRetargetOpBase& Op = Asset->RetargetOps[OpIndex].GetMutable<FIKRetargetOpBase>();
-		if (!Op.AreChainSettingsAtDefault(InTargetChainName))
+		if (!ensure(InOp))
 		{
-			return true;
+			return;
+		}
+
+		FRetargetChainMapping* ChainMapping = InOp->GetChainMapping();
+		if (!ChainMapping)
+		{
+			return; // not all ops maintain their own chain mapping
 		}
+
+		// clean the mapping
+		ChainMapping->ReinitializeWithIKRigs(SourceIKRig, InOp->GetCustomTargetIKRig());
+
+		// force the op to regenerate chain settings if needed
+		InOp->OnReinitPropertyEdited(nullptr);
+	};
+
+	// single op
+	if (InOpName != NAME_None)
+	{
+		CleanChainMapInOp(GetRetargetOpByName(InOpName));
+		return;
 	}
-	return false;
+
+	// all ops
+	const int32 NumOps = GetNumRetargetOps();
+	for (int32 OpIndex = 0; OpIndex < NumOps; ++OpIndex)
+	{
+		CleanChainMapInOp(GetRetargetOpByIndex(OpIndex));
+	}
+}
+
+bool UIKRetargeterController::AreChainSettingsAtDefault(const FName InTargetChainName, const FName InOpName) const
+{
+	FIKRetargetOpBase* Op = GetRetargetOpByName(InOpName);
+	if (!Op)
+	{
+		return false;
+	}
+
+	return Op->AreChainSettingsAtDefault(InTargetChainName);
 }
 
-int32 UIKRetargeterController::AddRetargetOp(
-	const UScriptStruct* InRetargetOpType,
-	const FName InParentOpName,
-	const bool bAllowDuplicates,
-	const bool bRunInitialSetup) const
+void UIKRetargeterController::ResetChainSettingsToDefault(const FName InTargetChainName, const FName InOpName) const
+{
+	FIKRetargetOpBase* Op = GetRetargetOpByName(InOpName);
+	if (!Op)
+	{
+		return;
+	}
+	
+	FScopeLock Lock(&ControllerLock);
+	FScopedTransaction Transaction(LOCTEXT("ResetChainSettings", "Reset Chain Settings"));
+	FScopedReinitializeIKRetargeter Reinitialize(this);
+	Asset->Modify();
+	
+	Op->ResetChainSettingsToDefault(InTargetChainName);
+}
+
+const UIKRigDefinition* UIKRetargeterController::GetTargetIKRigForOp(const FName InOpName) const
+{
+	FIKRetargetOpBase* Op = GetRetargetOpByName(InOpName);
+	if (!Op)
+	{
+		return nullptr;
+	}
+
+	if (const UIKRigDefinition* TargetIKRig = Op->GetCustomTargetIKRig())
+	{
+		return TargetIKRig;
+	}
+
+	if (FIKRetargetOpBase* ParentOp = GetRetargetOpByName(Op->GetParentOpName()))
+	{
+		return ParentOp->GetCustomTargetIKRig();
+	}
+
+	return nullptr;
+}
+
+int32 UIKRetargeterController::AddRetargetOp(const UScriptStruct* InRetargetOpType, const FName InParentOpName) const
 {
 	check(Asset)
 
@@ -596,16 +724,18 @@ int32 UIKRetargeterController::AddRetargetOp(
 		UE_LOG(LogIKRigEditor, Warning, TEXT("Could not add retarget operations. Invalid op-type specified. Must be child of FIKRetargetOpBase."));
 		return INDEX_NONE;
 	}
-
-	// optionally avoid duplicate ops of the same type
-	if (!bAllowDuplicates)
+	
+	// check if the op is a singleton and the stack already contains an op of that type
+	FInstancedStruct NewOpStruct = FInstancedStruct(InRetargetOpType);
+	if (NewOpStruct.GetPtr<FIKRetargetOpBase>()->IsSingleton())
 	{
 		for (int32 OpIndex=0; OpIndex<Asset->RetargetOps.Num(); ++OpIndex)
 		{
 			const FInstancedStruct& Op = Asset->RetargetOps[OpIndex];
 			if (Op.GetScriptStruct() == InRetargetOpType)
 			{
-				return OpIndex;
+				UE_LOG(LogIKRigEditor, Warning, TEXT("Op not added. It is a singleton and the stack already contains an op of that type."));
+				return INDEX_NONE;
 			}
 		}
 	}
@@ -615,10 +745,10 @@ int32 UIKRetargeterController::AddRetargetOp(
 	Asset->Modify();
 
 	// add the op
-	int32 NewOpIndex = Asset->RetargetOps.Emplace(InRetargetOpType);
+	int32 NewOpIndex = Asset->RetargetOps.Emplace(MoveTemp(NewOpStruct));
 
 	// give a unique default name
-	FIKRetargetOpBase* NewOp = GetOpAtIndex(NewOpIndex);
+	FIKRetargetOpBase* NewOp = GetRetargetOpByIndex(NewOpIndex);
 	const FName NewOpName = GetUniqueOpName(NewOp->GetDefaultName(), NewOpIndex);
 	NewOp->SetName(NewOpName);
 
@@ -628,15 +758,13 @@ int32 UIKRetargeterController::AddRetargetOp(
 	// enforce correct execution order
 	Asset->CleanOpStack();
 
-	// optionally run the initial setup
-	if (bRunInitialSetup)
-	{
-		NewOp->RunInitialSetup(GetAsset());
-	}
-
 	// in the unlikely event that cleaning the op stack reshuffled the indices, get the index again
 	NewOpIndex = GetIndexOfOpByName(NewOpName);
 	
+	// run the initial setup
+	FIKRetargetOpBase* ParentOp = GetRetargetOpByIndex(GetParentOpIndex(NewOpIndex));
+	NewOp->OnAddedToStack(GetAsset(), ParentOp);
+	
 	return NewOpIndex;
 }
 
@@ -658,7 +786,7 @@ void UIKRetargeterController::SetAskedToFixRootHeightForMesh(USkeletalMesh* Mesh
 	}
 }
 
-FName UIKRetargeterController::GetRetargetRootBone(const ERetargetSourceOrTarget SourceOrTarget) const
+FName UIKRetargeterController::GetPelvisBone(const ERetargetSourceOrTarget SourceOrTarget) const
 {
 	FScopeLock Lock(&ControllerLock);
 	const UIKRigDefinition* IKRig = GetIKRig(SourceOrTarget);
@@ -705,105 +833,37 @@ UIKRigStructViewer* UIKRetargeterController::GetStructViewer()
 	return StructViewer;
 }
 
-void UIKRetargeterController::AutoMapChains(const EAutoMapChainType AutoMapType, const bool bForceRemap) const
+void UIKRetargeterController::AutoMapChains(const EAutoMapChainType AutoMapType, const bool bForceRemap, const FName InOpName) const
 {
-	auto GetChainNames = [this](const ERetargetSourceOrTarget SourceOrTarget, TArray<FName>& OutNames)
-	{
-		if (const UIKRigDefinition* IKRig = GetIKRig(SourceOrTarget))
-		{
-			const TArray<FBoneChain>& Chains = IKRig->GetRetargetChains();
-			for (const FBoneChain& Chain : Chains)
-			{
-				OutNames.Add(Chain.ChainName);
-			}
-		}
-	};
-	
 	FScopeLock Lock(&ControllerLock);
 	FScopedTransaction Transaction(LOCTEXT("AutoMapRetargetChains", "Auto-Map Retarget Chains"));
 	FScopedReinitializeIKRetargeter Reinitialize(this);
 	Asset->Modify();
 
-	// get the chain mapping from the asset
-	FRetargetChainMapping& ChainMapping = Asset->GetChainMapping();
-
-	// clean the chain map
-	const UIKRigDefinition* SourceIKRig = GetIKRig(ERetargetSourceOrTarget::Source);
-	const UIKRigDefinition* TargetIKRig = GetIKRig(ERetargetSourceOrTarget::Target);
-	ChainMapping.CleanMapping(SourceIKRig, TargetIKRig);
-
-	// get names of all the potential chains we could map to on the source
-	TArray<FName> SourceChainNames;
-	GetChainNames(ERetargetSourceOrTarget::Source, SourceChainNames);
-
-	// iterate over all the chain pairs and find matching source chain
-	TArray<FRetargetChainPair>& ChainPairs = Asset->GetChainMapping().GetChainPairs();
-	for (FRetargetChainPair& ChainPair : ChainPairs)
+	auto AutoMapChainsInOp = [AutoMapType, bForceRemap](FIKRetargetOpBase* InRetargetOp)
 	{
-		const bool bIsMappedAlready = ChainPair.SourceChainName != NAME_None;
-		if (bIsMappedAlready && !bForceRemap)
+		if (InRetargetOp)
 		{
-			continue; // already set by user
-		}
-
-		// find a source chain to map to
-		int32 SourceChainIndexToMapTo = -1;
-
-		switch (AutoMapType)
-		{
-		case EAutoMapChainType::Fuzzy:
+			if (FRetargetChainMapping* ChainMapping = InRetargetOp->GetChainMapping())
 			{
-				// auto-map chains using a fuzzy string comparison
-				FString TargetNameLowerCase = ChainPair.TargetChainName.ToString().ToLower();
-				float HighestScore = 0.2f;
-				for (int32 ChainIndex=0; ChainIndex<SourceChainNames.Num(); ++ChainIndex)
-				{
-					FString SourceNameLowerCase = SourceChainNames[ChainIndex].ToString().ToLower();
-					float WorstCase = static_cast<float>(TargetNameLowerCase.Len() + SourceNameLowerCase.Len());
-					WorstCase = WorstCase < 1.0f ? 1.0f : WorstCase;
-					const float Score = 1.0f - (static_cast<float>(Algo::LevenshteinDistance(TargetNameLowerCase, SourceNameLowerCase)) / WorstCase);
-					if (Score > HighestScore)
-					{
-						HighestScore = Score;
-						SourceChainIndexToMapTo = ChainIndex;
-					}
-				}
-				break;
-			}
-		case EAutoMapChainType::Exact:
-			{
-				// if no exact match is found, then set to None
-				ChainPair.SourceChainName = NAME_None;
-				
-				// auto-map chains with EXACT same name
-				for (int32 ChainIndex=0; ChainIndex<SourceChainNames.Num(); ++ChainIndex)
-				{
-					if (SourceChainNames[ChainIndex] == ChainPair.TargetChainName)
-					{
-						SourceChainIndexToMapTo = ChainIndex;
-						break;
-					}
-				}
-				break;
+				ChainMapping->AutoMapChains(AutoMapType, bForceRemap);
 			}
-		case EAutoMapChainType::Clear:
-			{
-				ChainPair.SourceChainName = NAME_None;
-				break;
-			}
-		default:
-			checkNoEntry();
 		}
+	};
 
-		// apply source if any decent matches were found
-		if (SourceChainNames.IsValidIndex(SourceChainIndexToMapTo))
-		{
-			ChainPair.SourceChainName = SourceChainNames[SourceChainIndexToMapTo];
-		}
+	// single op
+	if (InOpName != NAME_None)
+	{
+		AutoMapChainsInOp(GetRetargetOpByName(InOpName));
+		return;
 	}
 
-	// sort asset mapping again
-	ChainMapping.SortMapping(TargetIKRig);
+	// all ops
+	const int32 NumOps = GetNumRetargetOps();
+	for (int32 OpIndex = 0; OpIndex < NumOps; ++OpIndex)
+	{
+		AutoMapChainsInOp(GetRetargetOpByIndex(OpIndex));
+	}
 }
 
 void UIKRetargeterController::HandleRetargetChainAdded(UIKRigDefinition* IKRig) const
@@ -814,120 +874,214 @@ void UIKRetargeterController::HandleRetargetChainAdded(UIKRigDefinition* IKRig)
 		// if a source chain is added, it will simply be available as a new option, no need to reinitialize until it's used
 		return;
 	}
-
+	
 	FScopedReinitializeIKRetargeter Reinitialize(this);
 	
 	// clean the chain map (this will add the new chain automatically)
-	const UIKRigDefinition* SourceIKRig = GetIKRig(ERetargetSourceOrTarget::Source);
-	const UIKRigDefinition* TargetIKRig = GetIKRig(ERetargetSourceOrTarget::Target);
-	Asset->GetChainMapping().CleanMapping(SourceIKRig, TargetIKRig);
+	CleanChainMaps();
 }
 
-void UIKRetargeterController::HandleRetargetChainRenamed(UIKRigDefinition* IKRig, FName OldChainName, FName NewChainName) const
+void UIKRetargeterController::HandleRetargetChainRenamed(UIKRigDefinition* InIKRig, FName OldChainName, FName NewChainName) const
 {
-	const bool bIsSourceRig = IKRig == Asset->GetIKRig(ERetargetSourceOrTarget::Source);
-	const bool bIsTargetRig = IKRig == Asset->GetIKRig(ERetargetSourceOrTarget::Target);
-	if (!ensure(bIsSourceRig || bIsTargetRig))
-	{
-		return;
-	}
-	
-	FRetargetChainMapping& ChainMapping = Asset->GetChainMapping();
-	for (FRetargetChainPair& ChainPair : ChainMapping.GetChainPairs())
+	FScopedTransaction Transaction(LOCTEXT("RetargetChainRenamed_Label", "Retarget Chain Renamed"));
+	FScopedReinitializeIKRetargeter Reinitialize(this);
+	Asset->Modify();
+
+	int32 NumOps = GetNumRetargetOps();
+	for (int32 OpIndex = 0; OpIndex < NumOps; ++OpIndex)
 	{
-		FName& ChainNameToUpdate = bIsSourceRig ? ChainPair.SourceChainName : ChainPair.TargetChainName;
-		if (ChainNameToUpdate == OldChainName)
+		FIKRetargetOpBase* Op = GetRetargetOpByIndex(OpIndex);
+		FRetargetChainMapping* ChainMapping = Op->GetChainMapping();
+		if (!ChainMapping)
 		{
-			FScopedTransaction Transaction(LOCTEXT("RetargetChainRenamed_Label", "Retarget Chain Renamed"));
-			FScopedReinitializeIKRetargeter Reinitialize(this);
-			Asset->Modify();
-			ChainNameToUpdate = NewChainName;
-			return;
+			continue;
+		}
+		
+		const bool bIsSourceRig = InIKRig == ChainMapping->GetIKRig(ERetargetSourceOrTarget::Source);
+		const bool bIsTargetRig = InIKRig == GetTargetIKRigForOp(Op->GetName());
+		if (!(bIsSourceRig || bIsTargetRig))
+		{
+			continue;
+		}
+
+		// maintain mappings to old name with new name
+		for (FRetargetChainPair& ChainPair : ChainMapping->GetChainPairs())
+		{
+			FName& ChainNameToUpdate = bIsSourceRig ? ChainPair.SourceChainName : ChainPair.TargetChainName;
+			if (ChainNameToUpdate == OldChainName)
+			{
+				ChainNameToUpdate = NewChainName;
+				break;
+			}
+		}
+
+		// allow op to retain the old chain settings under the new name
+		// NOTE: this is only called for ops that use the target IK rig
+		if (bIsTargetRig)
+		{
+			Op->OnTargetChainRenamed(OldChainName, NewChainName);	
 		}
 	}
 }
 
 void UIKRetargeterController::HandleRetargetChainRemoved(UIKRigDefinition* IKRig, const FName& InChainRemoved) const
 {
-	const bool bIsSourceRig = IKRig == Asset->GetIKRig(ERetargetSourceOrTarget::Source);
-	const bool bIsTargetRig = IKRig == Asset->GetIKRig(ERetargetSourceOrTarget::Target);
-	if (!ensure(bIsSourceRig || bIsTargetRig))
-	{
-		return;
-	}
-	
 	FScopedTransaction Transaction(LOCTEXT("RetargetChainRemoved_Label", "Retarget Chain Removed"));
 	FScopedReinitializeIKRetargeter Reinitialize(this);
 	Asset->Modify();
 
-	FRetargetChainMapping& ChainMapping = Asset->GetChainMapping();
-	TArray<FRetargetChainPair>& ChainPairs = ChainMapping.GetChainPairs();
-	
-	// set source chain name to NONE if it has been deleted 
-	if (bIsSourceRig)
+	int32 NumOps = GetNumRetargetOps();
+	for (int32 OpIndex = 0; OpIndex < NumOps; ++OpIndex)
 	{
-		for (FRetargetChainPair& ChainPair : ChainPairs)
+		FIKRetargetOpBase* Op = GetRetargetOpByIndex(OpIndex);
+		FRetargetChainMapping* ChainMapping = Op->GetChainMapping();
+		if (!ChainMapping)
+		{
+			continue;
+		}
+		
+		const bool bIsSourceRig = IKRig == ChainMapping->GetIKRig(ERetargetSourceOrTarget::Source);
+		const bool bIsTargetRig = IKRig == Op->GetCustomTargetIKRig();
+		if (!(bIsSourceRig || bIsTargetRig))
 		{
-			if (ChainPair.SourceChainName == InChainRemoved)
+			continue;
+		}
+
+		TArray<FRetargetChainPair>& ChainPairs = ChainMapping->GetChainPairs();
+	
+		// set source chain name to NONE if it has been deleted 
+		if (bIsSourceRig)
+		{
+			for (FRetargetChainPair& ChainPair : ChainPairs)
 			{
-				ChainPair.SourceChainName = NAME_None;
-				return;
+				if (ChainPair.SourceChainName == InChainRemoved)
+				{
+					ChainPair.SourceChainName = NAME_None;
+					return;
+				}
 			}
+			return;
 		}
-		return;
-	}
 	
-	// remove target mapping if the target chain has been removed
-	const int32 ChainIndex = ChainPairs.IndexOfByPredicate([&InChainRemoved](const FRetargetChainPair& ChainPair)
-	{
-		return ChainPair.TargetChainName == InChainRemoved;
-	});
+		// remove target mapping if the target chain has been removed
+		const int32 ChainIndex = ChainPairs.IndexOfByPredicate([&InChainRemoved](const FRetargetChainPair& ChainPair)
+		{
+			return ChainPair.TargetChainName == InChainRemoved;
+		});
 	
-	if (ChainIndex != INDEX_NONE)
-	{
-		ChainPairs.RemoveAt(ChainIndex);
+		if (ChainIndex != INDEX_NONE)
+		{
+			ChainPairs.RemoveAt(ChainIndex);
+		}
+
+		// force regeneration of chain settings
+		Op->OnReinitPropertyEdited(nullptr);
 	}
 }
 
-bool UIKRetargeterController::SetSourceChain(FName SourceChainName, FName TargetChainName) const
+bool UIKRetargeterController::SetSourceChain(FName SourceChainName, FName TargetChainName, const FName InOpName) const
 {
 	FScopeLock Lock(&ControllerLock);
-	FRetargetChainMapping& ChainMap = Asset->GetChainMapping();
-	if (!ChainMap.HasChain(TargetChainName, ERetargetSourceOrTarget::Target))
-	{
-		UE_LOG(LogIKRigEditor, Warning, TEXT("Could not map target chain. Target chain not found, %s."), *TargetChainName.ToString());
-		return false;
-	}
-	
 	FScopedTransaction Transaction(LOCTEXT("SetRetargetChainSource", "Set Retarget Chain Source"));
 	FScopedReinitializeIKRetargeter Reinitialize(this);
 	Asset->Modify();
-	ChainMap.SetChainMapping(TargetChainName, SourceChainName);
-	return true;
+
+	bool bModifiedChainMap = false;
+	const int32 NumOps = GetNumRetargetOps();
+	for (int32 OpIndex = 0; OpIndex < NumOps; ++OpIndex)
+	{
+		FIKRetargetOpBase* Op = GetRetargetOpByIndex(OpIndex);
+		if (InOpName != NAME_None && Op->GetName() != InOpName)
+		{
+			continue;
+		}
+		
+		FRetargetChainMapping* ChainMapping = Op->GetChainMapping();
+		if (!ChainMapping)
+		{
+			continue;
+		}
+		
+		if (!ChainMapping->HasChain(TargetChainName, ERetargetSourceOrTarget::Target))
+		{
+			continue;
+		}
+		
+		ChainMapping->SetChainMapping(TargetChainName, SourceChainName);
+		bModifiedChainMap = true;
+	}
+	
+	return bModifiedChainMap;
 }
 
-FName UIKRetargeterController::GetSourceChain(const FName& TargetChainName) const
+FName UIKRetargeterController::GetSourceChain(const FName& TargetChainName, const FName InOpName) const
 {
-	FRetargetChainMapping& ChainMap = Asset->GetChainMapping();
-	if (!ChainMap.HasChain(TargetChainName, ERetargetSourceOrTarget::Target))
+	const int32 NumOps = GetNumRetargetOps();
+	for (int32 OpIndex = 0; OpIndex < NumOps; ++OpIndex)
 	{
-		UE_LOG(LogIKRigEditor, Warning, TEXT("Could not get source chain. Target chain not found, %s."), *TargetChainName.ToString());
-		return NAME_None;
+		FIKRetargetOpBase* Op = GetRetargetOpByIndex(OpIndex);
+		if (InOpName != NAME_None && Op->GetName() != InOpName)
+		{
+			continue;
+		}
+		
+		FRetargetChainMapping* ChainMapping = Op->GetChainMapping();
+		if (!ChainMapping)
+		{
+			continue;
+		}
+		
+		if (!ChainMapping->HasChain(TargetChainName, ERetargetSourceOrTarget::Target))
+		{
+			continue;
+		}
+		
+		return ChainMapping->GetChainMappedTo(TargetChainName, ERetargetSourceOrTarget::Target);
+	}
+
+	return NAME_None;
+}
+
+const FRetargetChainMapping* UIKRetargeterController::GetChainMapping(const FName InOpName) const
+{
+	const int32 NumOps = GetNumRetargetOps();
+	for (int32 OpIndex = 0; OpIndex < NumOps; ++OpIndex)
+	{
+		FIKRetargetOpBase* Op = GetRetargetOpByIndex(OpIndex);
+		if (InOpName != NAME_None && Op->GetName() != InOpName)
+		{
+			continue;
+		}
+		
+		return Op->GetChainMapping();
 	}
 	
-	return ChainMap.GetChainMappedTo(TargetChainName, ERetargetSourceOrTarget::Target);
+	return nullptr;
 }
 
 bool UIKRetargeterController::IsChainGoalConnectedToASolver(const FName& GoalName) const
 {
-	const UIKRigDefinition* TargetIKRig = GetIKRig(ERetargetSourceOrTarget::Target);
-	if (!TargetIKRig)
+	TArray<UIKRigDefinition*> AllTargetIKRigs = GetAllTargetIKRigs();
+	if (AllTargetIKRigs.IsEmpty())
 	{
 		return false;
 	}
 
-	const UIKRigController* RigController = UIKRigController::GetController(TargetIKRig);
-	return RigController->IsGoalConnectedToAnySolver(GoalName);
+	for (const UIKRigDefinition* TargetIKRig : AllTargetIKRigs)
+	{
+		const UIKRigController* RigController = UIKRigController::GetController(TargetIKRig);
+		if (!RigController)
+		{
+			continue;
+		}
+		
+		if (RigController->IsGoalConnectedToAnySolver(GoalName))
+		{
+			return true;
+		}
+	}
+	
+	return false; 
 }
 
 FName UIKRetargeterController::CreateRetargetPose(const FName& NewPoseName, const ERetargetSourceOrTarget SourceOrTarget) const
@@ -1077,7 +1231,7 @@ void UIKRetargeterController::ResetRetargetPose(
 		FScopedTransaction Transaction(LOCTEXT("ResetRetargetBonePose", "Reset Bone Pose"));
 		Asset->Modify();
 		
-		const FName RootBoneName = GetRetargetRootBone(SourceOrTarget);
+		const FName RootBoneName = GetPelvisBone(SourceOrTarget);
 		for (const FName& BoneToReset : BonesToReset)
 		{
 			if (PoseToEdit.BoneRotationOffsets.Contains(BoneToReset))
diff --git a/Engine/Plugins/Animation/IKRig/Source/IKRigEditor/Private/RetargetEditor/IKRetargeterPoseGenerator.cpp b/Engine/Plugins/Animation/IKRig/Source/IKRigEditor/Private/RetargetEditor/IKRetargeterPoseGenerator.cpp
--- a/Engine/Plugins/Animation/IKRig/Source/IKRigEditor/Private/RetargetEditor/IKRetargeterPoseGenerator.cpp
+++ b/Engine/Plugins/Animation/IKRig/Source/IKRigEditor/Private/RetargetEditor/IKRetargeterPoseGenerator.cpp
@@ -112,7 +112,7 @@ void FRetargetAutoPoseGenerator::AlignBone(
 	
 	// find chains on both sides
 	const FName ChainToAlign = Processor->GetChainNameForBone(BoneToAlign, SideToAlign);
-	const FName ChainToMatch = Processor->GetChainMappedToChain(ChainToAlign, SideToAlign);
+	const FName ChainToMatch = Processor->GetFirstChainMappedToChain(ChainToAlign, SideToAlign);
 	checkf(ChainToMatch != NAME_None != 0, TEXT("Bone should never be retargeted and not in a mapped chain."));
 	
 	// get the normalized chain parameter to match
@@ -225,6 +225,14 @@ void FRetargetAutoPoseGenerator::AlignBone(
 			FText::Format(LOCTEXT("AutoAlignMeshNoSource", "Could not determine the direction of bone to match from it's mesh for {0}. Likely because no vertices are skinned to the source bone."),
 				FText::FromName(BoneToAlign)));
 		}
+
+		// user asked to align the mesh, but direction vectors for the target could not be generated
+		if (!bFoundAlignVector || !bFoundMatchVector)
+		{
+			Processor->Log.LogWarning(
+			FText::Format(LOCTEXT("AutoAlignFailed", "Could not determine the direction align the bone: "),
+				FText::FromName(BoneToAlign)));
+		}
 		
 		return;
 	}
@@ -249,7 +257,7 @@ void FRetargetAutoPoseGenerator::AlignRetargetRootBone(const ERetargetSourceOrTa
 	if (!(bFoundSourceSpine && bFoundTargetSpine))
 	{
 		Processor->Log.LogWarning(
-		LOCTEXT("MissingSpine", "Could not auto align the retarget root because a Spine chain was not found. Expects a retarget chain named 'Spine' or similar."));
+		LOCTEXT("MissingSpine", "Could not auto align the pelvis because a Spine chain was not found. Expects a retarget chain named 'Spine' or similar."));
 		return;
 	}
 
@@ -333,8 +341,14 @@ bool FRetargetAutoPoseGenerator::GetSpineChain(const ERetargetSourceOrTarget Sou
 	OutChainName = NAME_None;
 	const FString ChainNameToFind = "spine";
 	float HighestScore = -1.f;
-	const TArray<FResolvedBoneChain>& BoneChains = Processor->GetBoneChains().GetResolvedBoneChains(SourceOrTarget);
-	for (const FResolvedBoneChain& BoneChain : BoneChains)
+	const UIKRigDefinition* IKRig = Processor->GetRetargetAsset()->GetIKRig(SourceOrTarget);
+	const TArray<FResolvedBoneChain>* BoneChains = Processor->GetBoneChains().GetResolvedBoneChains(SourceOrTarget, IKRig);
+	if (!BoneChains)
+	{
+		return false;
+	}
+	
+	for (const FResolvedBoneChain& BoneChain : *BoneChains)
 	{
 		FName ChainNameToTest = BoneChain.ChainName;
 		FString ChainNameStr = ChainNameToTest.ToString().ToLower();
diff --git a/Engine/Plugins/Animation/IKRig/Source/IKRigEditor/Private/RetargetEditor/SIKRetargetChainMapList.cpp b/Engine/Plugins/Animation/IKRig/Source/IKRigEditor/Private/RetargetEditor/SIKRetargetChainMapList.cpp
--- a/Engine/Plugins/Animation/IKRig/Source/IKRigEditor/Private/RetargetEditor/SIKRetargetChainMapList.cpp
+++ b/Engine/Plugins/Animation/IKRig/Source/IKRigEditor/Private/RetargetEditor/SIKRetargetChainMapList.cpp
@@ -5,19 +5,14 @@
 #include "ScopedTransaction.h"
 #include "Framework/Commands/UICommandList.h"
 #include "Framework/MultiBox/MultiBoxBuilder.h"
-#include "IPersonaToolkit.h"
-#include "SKismetInspector.h"
-#include "Dialogs/Dialogs.h"
-#include "Widgets/Text/SInlineEditableTextBlock.h"
-#include "SPositiveActionButton.h"
 #include "Retargeter/IKRetargeter.h"
-#include "RigEditor/IKRigEditorStyle.h"
 #include "Widgets/Input/SComboBox.h"
 #include "SSearchableComboBox.h"
 #include "Misc/TextFilterExpressionEvaluator.h"
 #include "Widgets/Input/SSearchBox.h"
 #include "Retargeter/IKRetargetSettings.h"
 #include "RigEditor/IKRigController.h"
+#include "RigEditor/IKRigStructViewer.h"
 
 #define LOCTEXT_NAMESPACE "SIKRigRetargetChains"
 
@@ -47,7 +42,7 @@ void SIKRetargetChainMapRow::Construct(
 	// NOTE: cannot just use FName because "None" is considered a null entry and removed from ComboBox.
 	SourceChainOptions.Reset();
 	SourceChainOptions.Add(MakeShareable(new FString(TEXT("None"))));
-	const UIKRigDefinition* SourceIKRig = ChainMapList.Pin()->EditorController.Pin()->AssetController->GetIKRig(ERetargetSourceOrTarget::Source);
+	const UIKRigDefinition* SourceIKRig = ChainMapList.Pin()->Config.Controller->GetIKRig(ERetargetSourceOrTarget::Source);
 	if (SourceIKRig)
 	{
 		const TArray<FBoneChain>& Chains = SourceIKRig->GetRetargetChains();
@@ -84,13 +79,16 @@ TSharedRef< SWidget > SIKRetargetChainMapRow::GenerateWidgetForColumn(const FNam
 		TSharedRef<SWidget> NewWidget =
 		SNew(SHorizontalBox)
 		+ SHorizontalBox::Slot()
-		.AutoWidth()
-		.HAlign(HAlign_Left)
+		.HAlign(HAlign_Fill)
 		.VAlign(VAlign_Center)
 		.Padding(3.0f, 1.0f)
 		[
 			SNew(SSearchableComboBox)
 			.OptionsSource(&SourceChainOptions)
+			.IsEnabled_Lambda([this]()
+			{
+				return ChainMapList.Pin()->Config.bEnableChainMapping;
+			})
 			.OnGenerateWidget_Lambda([](TSharedPtr<FString> InItem)
 			{
 				return SNew(STextBlock).Text(FText::FromString(*InItem.Get()));
@@ -157,48 +155,56 @@ FReply SIKRetargetChainMapRow::OnResetToDefaultClicked()
 
 EVisibility SIKRetargetChainMapRow::GetResetToDefaultVisibility() const
 {
-	UIKRetargeterController* RetargeterController = ChainMapList.Pin()->GetRetargetController();
+	UIKRetargeterController* RetargeterController = ChainMapList.Pin()->Config.Controller.Get();
 	if (!RetargeterController)
 	{
 		return EVisibility::Hidden; 
 	}
 
-	bool bIsChainNonDefault = RetargeterController->ChainHasNonDefaultSettings(ChainMapElement.Pin()->TargetChainName);
-	return bIsChainNonDefault ? EVisibility::Visible : EVisibility::Hidden;
+	const FName TargetChainName = ChainMapElement.Pin()->TargetChainName;
+	const FName OpName = ChainMapList.Pin()->Config.OpWithChainSettings;
+	bool bAreSettingsAtDefault = RetargeterController->AreChainSettingsAtDefault(TargetChainName, OpName);
+	return bAreSettingsAtDefault ? EVisibility::Hidden : EVisibility::Visible;
 }
 
 void SIKRetargetChainMapRow::OnSourceChainComboSelectionChanged(TSharedPtr<FString> InName, ESelectInfo::Type SelectInfo)
 {
-	UIKRetargeterController* RetargeterController = ChainMapList.Pin()->GetRetargetController();
+	UIKRetargeterController* RetargeterController = ChainMapList.Pin()->Config.Controller.Get();
 	if (!RetargeterController)
 	{
 		return; 
 	}
 
 	const FName SourceChainName = FName(*InName.Get());
-	RetargeterController->SetSourceChain(SourceChainName, ChainMapElement.Pin()->TargetChainName);
+	const FName TargetChainName = ChainMapElement.Pin()->TargetChainName;
+	const FName OpWithMappingName = ChainMapList.Pin()->Config.OpWithChainMapping;
+	RetargeterController->SetSourceChain(SourceChainName, TargetChainName, OpWithMappingName);
 }
 
 FText SIKRetargetChainMapRow::GetSourceChainName() const
 {
-	UIKRetargeterController* RetargeterController = ChainMapList.Pin()->GetRetargetController();
+	UIKRetargeterController* RetargeterController = ChainMapList.Pin()->Config.Controller.Get();
 	if (!RetargeterController)
 	{
 		return FText::FromName(NAME_None); 
 	}
 
-	return FText::FromName(RetargeterController->GetSourceChain(ChainMapElement.Pin()->TargetChainName));
+	const FName TargetChainName = ChainMapElement.Pin()->TargetChainName;
+	const FName OpName = ChainMapList.Pin()->Config.OpWithChainMapping;
+	const FName SourceChainName = RetargeterController->GetSourceChain(TargetChainName, OpName);
+	return FText::FromName(SourceChainName);
 }
 
 FText SIKRetargetChainMapRow::GetTargetIKGoalName() const
 {
-	UIKRetargeterController* RetargeterController = ChainMapList.Pin()->GetRetargetController();
+	UIKRetargeterController* RetargeterController = ChainMapList.Pin()->Config.Controller.Get();
 	if (!RetargeterController)
 	{
 		return FText(); 
 	}
 
-	const UIKRigDefinition* IKRig = RetargeterController->GetIKRig(ERetargetSourceOrTarget::Target);
+	const FName OpName = ChainMapList.Pin()->Config.OpWithChainMapping;
+	const UIKRigDefinition* IKRig = RetargeterController->GetTargetIKRigForOp(OpName);
 	if (!IKRig)
 	{
 		return FText(); 
@@ -219,81 +225,85 @@ FText SIKRetargetChainMapRow::GetTargetIKGoalName() const
 	return FText::FromName(Chain->IKGoalName);
 }
 
-void SIKRetargetChainMapList::Construct(
-	const FArguments& InArgs,
-	TSharedRef<FIKRetargetEditorController> InEditorController)
+void SIKRetargetChainMapList::Construct(const FArguments& InArgs)
 {
-	EditorController = InEditorController;
-	EditorController.Pin()->SetChainsView(SharedThis(this));
-
+	Config = InArgs._InChainMapListConfig;
+	if (!ensure(Config.IsValid()))
+	{
+		return; // chain map editing must be associated with a particular op
+	}
+	
 	TextFilter = MakeShareable(new FTextFilterExpressionEvaluator(ETextFilterExpressionEvaluatorMode::BasicString));
 	
+	// create property view for the chain properties
+	FDetailsViewArgs DetailsViewArgs;
+	DetailsViewArgs.bAllowSearch = false;
+	DetailsViewArgs.bLockable = false;
+	DetailsViewArgs.bUpdatesFromSelection = false;
+	DetailsViewArgs.bAllowMultipleTopLevelObjects = false;
+	DetailsViewArgs.NameAreaSettings = FDetailsViewArgs::ENameAreaSettings::HideNameArea;
+
+	FPropertyEditorModule& PropertyModule = FModuleManager::LoadModuleChecked<FPropertyEditorModule>("PropertyEditor");
+	DetailsView = PropertyModule.CreateDetailView(DetailsViewArgs);
+
+	auto GenerateHeaderRow = [this]() -> TSharedPtr<SHeaderRow>
+	{
+		TSharedRef<SHeaderRow> HeaderRowWidget = SNew(SHeaderRow);
+		
+		HeaderRowWidget->AddColumn(SHeaderRow::Column(ColumnId_TargetChainLabel)
+			.HAlignHeader(HAlign_Center)
+			.DefaultLabel(LOCTEXT("TargetColumnLabel", "Target Chain"))
+			.DefaultTooltip(LOCTEXT("TargetChainToolTip", "The chain on the target skeleton to copy animation TO.")));
+
+		HeaderRowWidget->AddColumn(SHeaderRow::Column(ColumnId_IKGoalNameLabel)
+			.HAlignHeader(HAlign_Center)
+			.Visibility_Lambda([this]()
+			{
+				return Config.bEnableGoalColumn ? EVisibility::Visible : EVisibility::Collapsed;
+			})
+			.DefaultLabel(LOCTEXT("IKColumnLabel", "Target IK Goal"))
+			.DefaultTooltip(LOCTEXT("IKGoalToolTip",
+    								"The IK Goal assigned to the target chain (if any). Note, this goal should be on the LAST bone in the chain.")));
+
+		HeaderRowWidget->AddColumn(SHeaderRow::Column(ColumnId_SourceChainLabel)
+			.HAlignHeader(HAlign_Center)
+			.DefaultLabel(LOCTEXT("SourceColumnLabel", "Source Chain"))
+			.DefaultTooltip(LOCTEXT("SourceChainToolTip", "The chain on the source skeleton to copy animation FROM.")));
+
+		if (Config.ChainSettingsGetterFunc.IsSet())
+		{
+			HeaderRowWidget->AddColumn(SHeaderRow::Column(ColumnId_ResetLabel)
+			.HAlignHeader(HAlign_Center)
+			.FixedWidth(50.0f)
+			.DefaultLabel(LOCTEXT("ResetColumnLabel", "Reset"))
+			.DefaultTooltip(LOCTEXT("ResetChainToolTip", "Reset the settings for this chain on this op.")));
+		}
+			
+		return HeaderRowWidget;
+	};
+	
 	ChildSlot
     [
 	    SNew(SVerticalBox)
-		+SVerticalBox::Slot()
-		.AutoHeight()
-		[
-			SNew(SHorizontalBox)
-			
-			+ SHorizontalBox::Slot()
-			.AutoWidth()
-			.HAlign(HAlign_Left)
-			.VAlign(VAlign_Center)
-			[
-				SNew(STextBlock)
-				.Text(LOCTEXT("TargetRootLabel", "Target Root: "))
-				.TextStyle(FAppStyle::Get(), "NormalText")
-			]
-
-			+ SHorizontalBox::Slot()
-			.AutoWidth()
-			.Padding(5, 0)
-			.HAlign(HAlign_Left)
-			.VAlign(VAlign_Center)
-			[
-				SNew(STextBlock)
-				.Text(this, &SIKRetargetChainMapList::GetTargetRootBone)
-				.IsEnabled(false)
-			]
-				
-			+ SHorizontalBox::Slot()
-			.AutoWidth()
-			.HAlign(HAlign_Left)
-			.VAlign(VAlign_Center)
-			[
-				SNew(STextBlock)
-				.Text(LOCTEXT("SourceRootLabel", "Source Root: "))
-				.TextStyle(FAppStyle::Get(), "NormalText")
-			]
-
-			+ SHorizontalBox::Slot()
-			.AutoWidth()
-			.Padding(5, 0)
-			.HAlign(HAlign_Left)
-			.VAlign(VAlign_Center)
-			[
-				SNew(STextBlock)
-				.Text(this, &SIKRetargetChainMapList::GetSourceRootBone)
-				.IsEnabled(false)
-			]
-		]
 
 		+SVerticalBox::Slot()
 		.Padding(2.0f)
 		.AutoHeight()
 		[
 			SNew(SHorizontalBox)
-			
+
+			// filter list text field
 			+SHorizontalBox::Slot()
 			.FillWidth(1.0f)
 			[
 				SNew(SSearchBox)
+				.IsEnabled(this, &SIKRetargetChainMapList::IsListEnabled)
 				.SelectAllTextWhenFocused(true)
 				.OnTextChanged( this, &SIKRetargetChainMapList::OnFilterTextChanged )
 				.HintText( LOCTEXT( "SearchBoxHint", "Filter Chain List...") )
 			]
 
+			// filter list options
 			+SHorizontalBox::Slot()
 			.AutoWidth()
 			.Padding(FMargin(6.f, 0.0))
@@ -303,6 +313,7 @@ void SIKRetargetChainMapList::Construct(
 				.ComboButtonStyle(&FAppStyle::Get().GetWidgetStyle<FComboButtonStyle>("SimpleComboButton"))
 				.ForegroundColor(FSlateColor::UseStyle())
 				.ContentPadding(2.0f)
+				.IsEnabled(this, &SIKRetargetChainMapList::IsListEnabled)
 				.OnGetMenuContent( this, &SIKRetargetChainMapList::CreateFilterMenuWidget )
 				.ToolTipText( LOCTEXT("ChainMapFilterToolTip", "Filter list of chain mappings."))
 				.HasDownArrow(true)
@@ -314,6 +325,7 @@ void SIKRetargetChainMapList::Construct(
 				]
 			]
 
+			// chain mapping menu
 			+SHorizontalBox::Slot()
 			.AutoWidth()
 			.VAlign(VAlign_Center)
@@ -323,7 +335,7 @@ void SIKRetargetChainMapList::Construct(
 				.ComboButtonStyle(&FAppStyle::Get().GetWidgetStyle<FComboButtonStyle>("SimpleComboButton"))
 				.ForegroundColor(FSlateColor::UseStyle())
 				.ContentPadding(2.0f)
-				.Visibility(this, &SIKRetargetChainMapList::IsAutoMapButtonVisible)
+				.IsEnabled(this, &SIKRetargetChainMapList::IsChainMappingEnabled)
 				.ToolTipText(LOCTEXT("AutoMapButtonToolTip", "Automatically assign source chains based on matching rule."))
 				.OnGetMenuContent( this, &SIKRetargetChainMapList::CreateChainMapMenuWidget )
 				.HasDownArrow(true)
@@ -333,38 +345,34 @@ void SIKRetargetChainMapList::Construct(
 				]
 			]
 		]
-	    
+
+	    // chain list view
         +SVerticalBox::Slot()
         [
-			SAssignNew(ListView, SRetargetChainMapListViewType )
-			.SelectionMode(ESelectionMode::Multi)
-			.IsEnabled(this, &SIKRetargetChainMapList::IsChainMapEnabled)
-			.ListItemsSource( &ListViewItems )
-			.OnGenerateRow( this, &SIKRetargetChainMapList::MakeListRowWidget )
-			.OnMouseButtonClick_Lambda([this](TSharedPtr<FRetargetChainMapElement> Item)
-			{
-				OnItemClicked(Item);
-			})
-			.HeaderRow
-			(
-				SNew( SHeaderRow )
-				+ SHeaderRow::Column(ColumnId_TargetChainLabel)
-				.DefaultLabel(LOCTEXT("TargetColumnLabel", "Target Chain"))
-				.DefaultTooltip(LOCTEXT("TargetChainToolTip", "The chain on the target skeleton to copy animation TO."))
-				
-				+ SHeaderRow::Column(ColumnId_IKGoalNameLabel)
-				.DefaultLabel(LOCTEXT("IKColumnLabel", "Target IK Goal"))
-				.DefaultTooltip(LOCTEXT("IKGoalToolTip", "The IK Goal assigned to the target chain (if any). Note, this goal should be on the LAST bone in the chain."))
-
-				+ SHeaderRow::Column(ColumnId_SourceChainLabel)
-				.DefaultLabel(LOCTEXT("SourceColumnLabel", "Source Chain"))
-				.DefaultTooltip(LOCTEXT("SourceChainToolTip", "The chain on the source skeleton to copy animation FROM."))
-				
-				+ SHeaderRow::Column(ColumnId_ResetLabel)
-                .DefaultLabel(LOCTEXT("ResetColumnLabel", "Reset"))
-                .DefaultTooltip(LOCTEXT("ResetChainToolTip", "Reset the FK and IK settings for this chain."))
-			)
+        	SNew(SBox)
+			.MaxDesiredHeight(300.0f)
+			[
+				SAssignNew(ListView, SRetargetChainMapListViewType )
+        		.ScrollbarVisibility(EVisibility::Visible)
+				.SelectionMode(ESelectionMode::Multi)
+				.IsEnabled(this, &SIKRetargetChainMapList::IsListEnabled)
+				.ListItemsSource( &ListViewItems )
+				.OnGenerateRow( this, &SIKRetargetChainMapList::MakeListRowWidget )
+				.OnMouseButtonClick_Lambda([this](TSharedPtr<FRetargetChainMapElement> Item)
+				{
+					OnItemClicked(Item);
+				})
+				.HeaderRow
+				(
+					GenerateHeaderRow()
+				)
+			]
         ]
+    	+ SVerticalBox::Slot()
+		.AutoHeight()
+		[
+			DetailsView.ToSharedRef()
+		]
     ];
 
 	RefreshView();
@@ -377,64 +385,54 @@ void SIKRetargetChainMapList::ClearSelection() const
 
 void SIKRetargetChainMapList::ResetChainSettings(const FName InTargetChainName) const
 {
-	EditorController.Pin()->AssetController->ResetChainSettingsInAllOps(InTargetChainName);
-	EditorController.Pin()->RefreshDetailsView();
-}
+	Config.Controller->ResetChainSettingsToDefault(InTargetChainName, Config.OpWithChainSettings);
 
-UIKRetargeterController* SIKRetargetChainMapList::GetRetargetController() const
-{
-	const TSharedPtr<FIKRetargetEditorController> Controller = EditorController.Pin();
-	if (!Controller.IsValid())
+	// update wrappers to show new values
+	for (TObjectPtr<UObject> Wrapper : AllStructWrappers)
 	{
-		return nullptr;
+		if (UIKRigStructWrapperBase* WrapperBase = Cast<UIKRigStructWrapperBase>(Wrapper))
+		{
+			WrapperBase->UpdateWrapperStructWithLatestValues();
+		}
 	}
-
-	return Controller->AssetController;
+	
+	DetailsView->ForceRefresh();
 }
 
-FText SIKRetargetChainMapList::GetSourceRootBone() const
+void SIKRetargetChainMapList::AddReferencedObjects(FReferenceCollector& Collector)
 {
-	UIKRetargeterController* RetargeterController = GetRetargetController();
-	if (!RetargeterController)
-	{
-		return FText::FromName(NAME_None); 
-	}
-	
-	return FText::FromName(RetargeterController->GetRetargetRootBone(ERetargetSourceOrTarget::Source));
+	Collector.AddReferencedObjects(AllStructWrappers);
 }
 
-FText SIKRetargetChainMapList::GetTargetRootBone() const
+bool SIKRetargetChainMapList::IsListEnabled() const
 {
-	UIKRetargeterController* RetargeterController = GetRetargetController();
-	if (!RetargeterController)
-	{
-		return FText::FromName(NAME_None); 
-	}
-	
-	return FText::FromName(RetargeterController->GetRetargetRootBone(ERetargetSourceOrTarget::Target));
+	return Config.IsValid();
 }
 
-bool SIKRetargetChainMapList::IsChainMapEnabled() const
+bool SIKRetargetChainMapList::IsChainMappingEnabled() const
 {
-	UIKRetargeterController* RetargeterController = GetRetargetController();
-	if (!RetargeterController)
+	if (!Config.IsValid())
 	{
 		return false; 
 	}
+
+	if (!Config.bEnableChainMapping)
+	{
+		return false;
+	}
 	
-	if (const UIKRigDefinition* IKRig = RetargeterController->GetAsset()->GetIKRig(ERetargetSourceOrTarget::Target))
+	const FRetargetChainMapping* ChainMapping = Config.Controller->GetChainMapping(Config.OpWithChainMapping);
+	if (!ChainMapping)
 	{
-		const TArray<FBoneChain>& Chains = IKRig->GetRetargetChains();
-		return Chains.Num() > 0;
+		return false;
 	}
 
-	return false;
+	return ChainMapping->IsReady();
 }
 
 void SIKRetargetChainMapList::RefreshView()
 {
-	UIKRetargeterController* AssetController = GetRetargetController();
-	if (!AssetController)
+	if (!Config.Controller.IsValid())
 	{
 		return; 
 	}
@@ -444,9 +442,9 @@ void SIKRetargetChainMapList::RefreshView()
 		return TextFilter->TestTextFilter(FBasicStringFilterExpressionContext(StringToTest));
 	};
 
-	auto DoesChainHaveIK = [AssetController](const FName TargetChainName) ->bool
+	auto DoesChainHaveIK = [this](const FName TargetChainName) ->bool
 	{
-		const UIKRigDefinition* IKRig = AssetController->GetIKRig(ERetargetSourceOrTarget::Target);
+		const UIKRigDefinition* IKRig = Config.Controller->GetTargetIKRigForOp(Config.OpWithChainMapping);
 		if (!IKRig)
 		{
 			return false;
@@ -462,17 +460,26 @@ void SIKRetargetChainMapList::RefreshView()
 		return Chain->IKGoalName != NAME_None;
 	};
 
-	const TArray<FName>& SelectedChains = EditorController.Pin()->GetSelectedChains();
-	const FName LiteralNone = FName("None");
-
 	// refresh items
 	ListViewItems.Reset();
-	const FRetargetChainMapping& ChainMapping = AssetController->GetChainMapping();
-	const TArray<FRetargetChainPair>& ChainMap = ChainMapping.GetChainPairs();
+
+	// get the chain mapping for the op
+	const FRetargetChainMapping* ChainMapping = Config.Controller->GetChainMapping(Config.OpWithChainMapping);
+	if (!ChainMapping->IsReady())
+	{
+		// invalid chain mapping means we don't display anything
+		// user must have valid source and target IK rigs to make a mapping and store settings
+		return;
+	}
+
+	// add list item for each target chain
+	static FName LiteralNone = FName("None");
+	const TArray<FRetargetChainPair>& ChainMap = ChainMapping->GetChainPairs();
 	for (const FRetargetChainPair& ChainPair : ChainMap)
 	{
 		const FName TargetChainName = ChainPair.TargetChainName;
 		const FName SourceChainName = ChainPair.SourceChainName;
+
 		// apply text filter to items
 		if (!(TextFilter->GetFilterText().IsEmpty() ||
 			FilterString(SourceChainName.ToString()) ||
@@ -482,19 +489,20 @@ void SIKRetargetChainMapList::RefreshView()
 		}
 		
 		// apply "only IK" filter
-		if (ChainFilterOptions.bHideChainsWithoutIK && !DoesChainHaveIK(TargetChainName))
+		const bool bFilterNonIKChains = Config.Filter.bHideChainsWithoutIK || Config.Filter.bNeverShowChainsWithoutIK;
+		if (bFilterNonIKChains && !DoesChainHaveIK(TargetChainName))
 		{
 			continue;
 		}
 
 		// apply "hide mapped chains" filter
-		if (ChainFilterOptions.bHideMappedChains && SourceChainName != LiteralNone)
+		if (Config.Filter.bHideMappedChains && SourceChainName != LiteralNone)
 		{
 			continue;
 		}
 		
 		// apply "hide un-mapped chains" filter
-		if (ChainFilterOptions.bHideUnmappedChains && SourceChainName == LiteralNone)
+		if (Config.Filter.bHideUnmappedChains && SourceChainName == LiteralNone)
 		{
 			continue;
 		}
@@ -502,10 +510,6 @@ void SIKRetargetChainMapList::RefreshView()
 		// create an item for this chain
 		TSharedPtr<FRetargetChainMapElement> ChainItem = FRetargetChainMapElement::Make(TargetChainName);
 		ListViewItems.Add(ChainItem);
-
-		// select/deselect it
-		const bool bIsSelected = SelectedChains.Contains(TargetChainName);
-		ListView->SetItemSelection(ChainItem, bIsSelected, ESelectInfo::Direct);
 	}
 	
 	ListView->RequestListRefresh();
@@ -516,41 +520,44 @@ TSharedRef<SWidget> SIKRetargetChainMapList::CreateFilterMenuWidget()
 	const FUIAction FilterHideMappedAction = FUIAction(
 		FExecuteAction::CreateLambda([this]
 		{
-			ChainFilterOptions.bHideMappedChains = !ChainFilterOptions.bHideMappedChains;
+			Config.Filter.bHideMappedChains = !Config.Filter.bHideMappedChains;
 			RefreshView();
 		}),
 		FCanExecuteAction(),
 		FIsActionChecked::CreateLambda([this]()
 		{
-			return ChainFilterOptions.bHideMappedChains;
+			return Config.Filter.bHideMappedChains;
 		}));
 
 	const FUIAction FilterOnlyUnMappedAction = FUIAction(
 		FExecuteAction::CreateLambda([this]
 		{
-			ChainFilterOptions.bHideUnmappedChains = !ChainFilterOptions.bHideUnmappedChains;
+			Config.Filter.bHideUnmappedChains = !Config.Filter.bHideUnmappedChains;
 			RefreshView();
 		}),
 		FCanExecuteAction(),
 		FIsActionChecked::CreateLambda([this]()
 		{
-			return ChainFilterOptions.bHideUnmappedChains;
+			return Config.Filter.bHideUnmappedChains;
 		}));
 
 	const FUIAction FilterIKChainAction = FUIAction(
 		FExecuteAction::CreateLambda([this]
 		{
-			ChainFilterOptions.bHideChainsWithoutIK = !ChainFilterOptions.bHideChainsWithoutIK;
+			Config.Filter.bHideChainsWithoutIK = !Config.Filter.bHideChainsWithoutIK;
 			RefreshView();
 		}),
-		FCanExecuteAction(),
+		FCanExecuteAction::CreateLambda([this]()
+		{
+			return !Config.Filter.bNeverShowChainsWithoutIK;
+		}),
 		FIsActionChecked::CreateLambda([this]()
 		{
-			return ChainFilterOptions.bHideChainsWithoutIK;
+			return Config.Filter.bHideChainsWithoutIK;
 		}));
 	
 	static constexpr bool CloseAfterSelection = true;
-	FMenuBuilder MenuBuilder(CloseAfterSelection, CommandList);
+	FMenuBuilder MenuBuilder(CloseAfterSelection, MenuCommandList);
 
 	MenuBuilder.BeginSection("Chain Map Filters", LOCTEXT("ChainMapFiltersSection", "Filter Chain Mappings"));
 
@@ -588,7 +595,7 @@ TSharedRef<SWidget> SIKRetargetChainMapList::CreateFilterMenuWidget()
 		FSlateIcon(),
 		FUIAction(FExecuteAction::CreateLambda([this]
 		{
-			ChainFilterOptions = FChainMapFilterOptions();
+			Config.Filter.Reset();
 			RefreshView();
 		})));
 
@@ -613,7 +620,7 @@ TSharedRef<ITableRow> SIKRetargetChainMapList::MakeListRowWidget(
 		SharedThis(this));
 }
 
-void SIKRetargetChainMapList::OnItemClicked(TSharedPtr<FRetargetChainMapElement> InItem) const 
+void SIKRetargetChainMapList::OnItemClicked(TSharedPtr<FRetargetChainMapElement> InItem) 
 {
 	// get list of selected chains
 	TArray<FName> SelectedChains;
@@ -622,10 +629,22 @@ void SIKRetargetChainMapList::OnItemClicked(TSharedPtr<FRetargetChainMapElement>
 	{
 		SelectedChains.Add(Item.Get()->TargetChainName);
 	}
-	
-	// replace the chain selection
-	constexpr bool bEditFromChainsView = true;
-	EditorController.Pin()->EditChainSelection(SelectedChains, ESelectionEdit::Replace, bEditFromChainsView);
+
+	// get list of settings to edit for the selected chains
+	AllStructWrappers.Reset();
+	if (Config.ChainSettingsGetterFunc.IsSet())
+	{
+		for (const FName SelectedChainName : SelectedChains)
+		{
+			UObject* StructWrapper = Config.ChainSettingsGetterFunc(SelectedChainName);
+			if (IsValid(StructWrapper))
+			{
+				AllStructWrappers.Add(StructWrapper);
+			}
+		}
+	}
+
+	DetailsView->SetObjects(AllStructWrappers);
 }
 
 TSharedRef<SWidget> SIKRetargetChainMapList::CreateChainMapMenuWidget()
@@ -671,7 +690,7 @@ TSharedRef<SWidget> SIKRetargetChainMapList::CreateChainMapMenuWidget()
 		FCanExecuteAction(), FIsActionChecked());
 	
 	static constexpr bool CloseAfterSelection = true;
-	FMenuBuilder MenuBuilder(CloseAfterSelection, CommandList);
+	FMenuBuilder MenuBuilder(CloseAfterSelection, MenuCommandList);
 
 	MenuBuilder.BeginSection("Auto-Map Chains Fuzzy", LOCTEXT("FuzzyNameSection", "Fuzzy Name Matching"));
 
@@ -726,27 +745,10 @@ TSharedRef<SWidget> SIKRetargetChainMapList::CreateChainMapMenuWidget()
 	return MenuBuilder.MakeWidget();
 }
 
-EVisibility SIKRetargetChainMapList::IsAutoMapButtonVisible() const
-{
-	return IsChainMapEnabled() ? EVisibility::Visible : EVisibility::Hidden;
-}
-
 void SIKRetargetChainMapList::AutoMapChains(const EAutoMapChainType AutoMapType, const bool bForceRemap)
 {
-	const TSharedPtr<FIKRetargetEditorController> Controller = EditorController.Pin();
-	if (!Controller.IsValid())
-	{
-		return;
-	}
-	
-	UIKRetargeterController* RetargeterController = GetRetargetController();
-	if (!RetargeterController)
-	{
-		return;
-	}
-
-	Controller->ClearOutputLog();
-	RetargeterController->AutoMapChains(AutoMapType, bForceRemap);
+	// TODO EditorController->ClearOutputLog();
+	Config.Controller->AutoMapChains(AutoMapType, bForceRemap);
 }
 
 #undef LOCTEXT_NAMESPACE
diff --git a/Engine/Plugins/Animation/IKRig/Source/IKRigEditor/Private/RetargetEditor/SRetargetAnimAssetsWindow.cpp b/Engine/Plugins/Animation/IKRig/Source/IKRigEditor/Private/RetargetEditor/SRetargetAnimAssetsWindow.cpp
--- a/Engine/Plugins/Animation/IKRig/Source/IKRigEditor/Private/RetargetEditor/SRetargetAnimAssetsWindow.cpp
+++ b/Engine/Plugins/Animation/IKRig/Source/IKRigEditor/Private/RetargetEditor/SRetargetAnimAssetsWindow.cpp
@@ -691,8 +691,7 @@ void FProceduralRetargetAssets::AutoGenerateIKRetargetAsset()
 	RetargetController->RemoveAllOps();
 
 	// add all the default ops for retargeting
-	constexpr bool bRunInitialSetup = true;
-	RetargetController->AddDefaultOps(bRunInitialSetup);
+	RetargetController->AddDefaultOps();
 
 	// disable IK pass because auto-IK is not yet robust enough, still useful for manual editing though
 	FIKRetargetRunIKRigOp* RunIKOp = RetargetController->GetFirstRetargetOpOfType<FIKRetargetRunIKRigOp>();
diff --git a/Engine/Plugins/Animation/IKRig/Source/IKRigEditor/Private/RetargetEditor/SRetargetOpStack.cpp b/Engine/Plugins/Animation/IKRig/Source/IKRigEditor/Private/RetargetEditor/SRetargetOpStack.cpp
--- a/Engine/Plugins/Animation/IKRig/Source/IKRigEditor/Private/RetargetEditor/SRetargetOpStack.cpp
+++ b/Engine/Plugins/Animation/IKRig/Source/IKRigEditor/Private/RetargetEditor/SRetargetOpStack.cpp
@@ -230,7 +230,7 @@ FIKRetargetOpBase* SRetargetOpItem::GetRetargetOp() const
 		return nullptr;
 	}
 	const int32 OpIndex = Element.Pin()->GetIndexInStack();
-	return EditorController->AssetController->GetOpAtIndex(OpIndex);
+	return EditorController->AssetController->GetRetargetOpByIndex(OpIndex);
 }
 
 void SRetargetOpList::Construct(const FArguments& InArgs)
@@ -435,18 +435,21 @@ FReply SRetargetOpList::OnFocusReceived(const FGeometry& MyGeometry, const FFocu
 
 void SRetargetOpList::OnSelectionChanged(TSharedPtr<FRetargetOpStackElement> InItem, ESelectInfo::Type SelectInfo) const
 {
-	if (!EditorController.IsValid())
+	// adds support for keyboard navigation of op stack
+	if (SelectInfo == ESelectInfo::OnNavigation || SelectInfo == ESelectInfo::Direct)
 	{
+		if (EditorController.IsValid() && InItem.IsValid())
+		{
+			EditorController.Pin()->SetOpSelected(InItem->GetIndexInStack());
+		}
+		
 		return;
 	}
-	
-	if (SelectInfo != ESelectInfo::Direct && InItem.IsValid())
+
+	if (!InItem.IsValid())
 	{
-		EditorController.Pin()->SetOpSelected(InItem->GetIndexInStack());
-		return;
+		EditorController.Pin()->ClearSelection();
 	}
-
-	EditorController.Pin()->SetDetailsObject(EditorController.Pin()->AssetController->GetAsset());
 }
 
 void SRetargetOpList::OnItemClicked(TSharedPtr<FRetargetOpStackElement> InItem)
@@ -467,6 +470,15 @@ void SRetargetOpList::OnItemClicked(TSharedPtr<FRetargetOpStackElement> InItem)
 
 	LastClickCycles = CurrentCycles;
 	LastSelectedElement = InItem;
+
+	if (EditorController.IsValid() && InItem.IsValid())
+	{
+		EditorController.Pin()->SetOpSelected(InItem->GetIndexInStack());
+	}
+	else
+	{
+		EditorController.Pin()->ClearSelection();
+	}
 }
 
 bool SRetargetOpList::IsEnabled() const
@@ -495,7 +507,7 @@ void SRetargetOpList::RestoreSelectionAndRefreshList()
 			return TSharedPtr<FRetargetOpStackElement>();
 		}
 		
-		const FName LastSelectedOpName = EditorController.Pin()->GetLastSelectedOpName();
+		const FName LastSelectedOpName = EditorController.Pin()->GetSelectedOpName();
 		if (LastSelectedOpName == NAME_None)
 		{
 			return TSharedPtr<FRetargetOpStackElement>();
@@ -574,10 +586,8 @@ void SRetargetOpList::AddNewRetargetOp(const UScriptStruct* ScriptStruct)
 	}
 
 	// add a new op to the stack
-	constexpr bool bAllowDuplicates = true;
-	constexpr bool bRunInitialSetup = true;
 	const FName ParentOpName = ParentElement.IsValid() ? ParentElement.Pin()->GetName() : NAME_None;
-	int32 NewOpIndex = AssetController->AddRetargetOp(ScriptStruct, ParentOpName, bAllowDuplicates, bRunInitialSetup);
+	int32 NewOpIndex = AssetController->AddRetargetOp(ScriptStruct, ParentOpName);
 
 	RestoreSelectionAndRefreshList();
 }
@@ -611,11 +621,20 @@ void SParentRetargetOpItem::Construct(const FArguments& InArgs, const TSharedRef
 		.OnDragDetected(OpListWidget.Pin().Get(), &SRetargetOpList::OnDragDetected)
 		.OnCanAcceptDrop(OpListWidget.Pin().Get(), &SRetargetOpList::OnCanAcceptDrop)
 		.OnAcceptDrop(OpListWidget.Pin().Get(), &SRetargetOpList::OnAcceptDrop)
+		.ShowSelection(false)
 		.Padding(FMargin(6.f, 3.f))
 		.Content()
 		[
 			SNew(SBorder)
-			.BorderImage(FIKRetargetEditorStyle::Get().GetBrush("IKRetarget.OpBorder"))
+			.BorderImage_Lambda([this]()
+			{
+				const FName SelectedOpName = OpListWidget.Pin()->EditorController.Pin()->GetSelectedOpName();
+				if (ParentStackElement.Pin()->GetName() == SelectedOpName)
+				{
+					return FIKRetargetEditorStyle::Get().GetBrush("IKRetarget.OpBorderSelected");
+				}
+				return FIKRetargetEditorStyle::Get().GetBrush("IKRetarget.OpBorder");
+			})
 			.Padding(FMargin(SRetargetOpItem::OpHorizontalPadding, SRetargetOpItem::OpVerticalPadding))
 			[
 				SNew(SVerticalBox)
@@ -689,11 +708,20 @@ void SRetargetOpSingleItem::Construct(const FArguments& InArgs, const TSharedRef
         .OnDragDetected(OpListWidget.Pin().Get(), &SRetargetOpList::OnDragDetected)
         .OnCanAcceptDrop(OpListWidget.Pin().Get(), &SRetargetOpList::OnCanAcceptDrop)
         .OnAcceptDrop(OpListWidget.Pin().Get(), &SRetargetOpList::OnAcceptDrop)
+        .ShowSelection(false)
         .Padding(FMargin(6.f, 3.f))
         .Content()
         [
         	SNew(SBorder)
-			.BorderImage(FIKRetargetEditorStyle::Get().GetBrush("IKRetarget.OpBorder"))
+        	.BorderImage_Lambda([this]()
+			{
+				const FName SelectedOpName = OpListWidget.Pin()->EditorController.Pin()->GetSelectedOpName();
+				if (StackElement.Pin()->GetName() == SelectedOpName)
+				{
+					return FIKRetargetEditorStyle::Get().GetBrush("IKRetarget.OpBorderSelected");
+				}
+				return FIKRetargetEditorStyle::Get().GetBrush("IKRetarget.OpBorder");
+			})
 			.Padding(FMargin(SRetargetOpItem::OpHorizontalPadding, SRetargetOpItem::OpVerticalPadding))
 			[
 				SNew(SRetargetOpItem)
@@ -745,7 +773,7 @@ FIKRetargetOpBase* SRetargetOpSingleItem::GetRetargetOp() const
 		return nullptr;
 	}
 	const int32 OpIndex = StackElement.Pin()->GetIndexInStack();
-	return EditorController->AssetController->GetOpAtIndex(OpIndex);
+	return EditorController->AssetController->GetRetargetOpByIndex(OpIndex);
 }
 
 void SRetargetOpStack::Construct(
diff --git a/Engine/Plugins/Animation/IKRig/Source/IKRigEditor/Private/RigEditor/IKRigController.cpp b/Engine/Plugins/Animation/IKRig/Source/IKRigEditor/Private/RigEditor/IKRigController.cpp
--- a/Engine/Plugins/Animation/IKRig/Source/IKRigEditor/Private/RigEditor/IKRigController.cpp
+++ b/Engine/Plugins/Animation/IKRig/Source/IKRigEditor/Private/RigEditor/IKRigController.cpp
@@ -11,7 +11,7 @@
 #include "ScopedTransaction.h"
 #include "RigEditor/IKRigAutoCharacterizer.h"
 #include "RigEditor/IKRigAutoFBIK.h"
-#include "RigEditor/IKRigDetailCustomizations.h"
+#include "RigEditor/IKRigStructViewer.h"
 
 #include UE_INLINE_GENERATED_CPP_BY_NAME(IKRigController)
 
diff --git a/Engine/Plugins/Animation/IKRig/Source/IKRigEditor/Private/RigEditor/IKRigDetailCustomizations.cpp b/Engine/Plugins/Animation/IKRig/Source/IKRigEditor/Private/RigEditor/IKRigDetailCustomizations.cpp
--- a/Engine/Plugins/Animation/IKRig/Source/IKRigEditor/Private/RigEditor/IKRigDetailCustomizations.cpp
+++ b/Engine/Plugins/Animation/IKRig/Source/IKRigEditor/Private/RigEditor/IKRigDetailCustomizations.cpp
@@ -21,116 +21,6 @@ namespace IKRigDetailCustomizationsConstants
 	static const float ItemWidth = 125.0f;
 }
 
-void UIKRigStructViewer::TriggerReinitIfNeeded(const FPropertyChangedEvent& InEvent)
-{
-	// never reinitialize while interacting with a slider
-	if (InEvent.ChangeType == EPropertyChangeType::Interactive || !InEvent.Property)
-	{
-		return;
-	}
-
-	// always reinitialize if property is marked for it
-	if (InEvent.Property->HasMetaData(IKRigReinitOnEditMetaLabel))
-	{
-		OnStructPropertyEdited().Broadcast();
-		return;
-	}
-
-	// if it is a name property, we want to reinitialize to validate it
-	if (CastField<FNameProperty>(InEvent.Property) != nullptr)
-	{
-		OnStructPropertyEdited().Broadcast();
-		return;
-	}
-
-	// always reinitialize if the property is a bone name
-	// NOTE: this is needed because FBoneReference's child BoneName property
-	// cannot be marked with the meta-tag (IKRig plugin doesn't own it and it's only relevant here)
-	if (InEvent.Property->GetFName() == FName(TEXT("BoneName")))
-	{
-		OnStructPropertyEdited().Broadcast();
-		return;
-	}
-}
-
-USkeleton* UIKRigStructViewer::GetSkeleton(bool& bInvalidSkeletonIsError, const IPropertyHandle* PropertyHandle)
-{
-	if (!StructType)
-	{
-		return nullptr;
-	}
-
-#if WITH_EDITOR
-	// NOTE: it's not ideal that we are hardcoding supported types here, but because UStruct's do not support multiple
-	// inheritance we cannot use an interface to identify skeleton providers as we normally would.
-	if (StructType->IsChildOf(FIKRetargetOpSettingsBase::StaticStruct()))
-	{
-		FName PropertyName = PropertyHandle->GetProperty()->GetFName();
-		uint8* StructMemory = StructMemoryProvider();
-		FIKRetargetOpSettingsBase* SkeletonProvider = reinterpret_cast<FIKRetargetOpSettingsBase*>(StructMemory);
-		return SkeletonProvider->GetSkeleton(PropertyName);
-	}
-#endif
-	
-	return nullptr;
-}
-
-TSharedRef<IDetailCustomization> FStructViewerCustomization::MakeInstance()
-{
-	return MakeShareable(new FStructViewerCustomization);
-}
-
-void FStructViewerCustomization::CustomizeDetails(IDetailLayoutBuilder& DetailBuilder)
-{
-	TArray<TWeakObjectPtr<UObject>> ObjectsBeingCustomized;
-	DetailBuilder.GetObjectsBeingCustomized(ObjectsBeingCustomized);
-	UIKRigStructViewer* StructViewer = Cast<UIKRigStructViewer>(ObjectsBeingCustomized[0]);
-	if (!ensure(StructViewer))
-	{
-		return;
-	}
-	if (!StructViewer->IsValid())
-	{
-		return;
-	}
-
-	// show the struct in the details panel
-	const TSharedPtr<FStructOnScope>& StructOnScope = StructViewer->GetStructOnScope();
-	FName StructTitle = StructViewer->GetStructName();
-	IDetailCategoryBuilder& CategoryBuilder = DetailBuilder.EditCategory(StructTitle);
-	TArray<IDetailPropertyRow*> OutRows;
-	TArray<TSharedPtr<IPropertyHandle>> Handles = CategoryBuilder.AddAllExternalStructureProperties(StructOnScope.ToSharedRef(), EPropertyLocation::Default, &OutRows);
-
-	// setup callbacks to begin/end an Undo transaction when property is edited
-	UObject* ObjectToTransact = StructViewer->GetStructOwner();
-	for (IDetailPropertyRow* Row : OutRows)
-	{
-		Row->GetPropertyHandle()->SetOnPropertyValuePreChange(FSimpleDelegate::CreateLambda([ObjectToTransact]()
-		{
-			GEditor->BeginTransaction(LOCTEXT("IKRigStructViewer", "Edited property."));
-			ObjectToTransact->Modify();
-		}));
-		
-		Row->GetPropertyHandle()->SetOnPropertyValueChangedWithData(TDelegate<void(const FPropertyChangedEvent&)>::CreateLambda([StructViewer](const FPropertyChangedEvent& InEvent)
-		{
-			GEditor->EndTransaction();
-			StructViewer->TriggerReinitIfNeeded(InEvent);
-		}));
-
-		Row->GetPropertyHandle()->SetOnChildPropertyValuePreChange(FSimpleDelegate::CreateLambda([ObjectToTransact]()
-		{
-			GEditor->BeginTransaction(LOCTEXT("IKRigStructViewerChild", "Edited child property."));
-			ObjectToTransact->Modify();
-		}));
-
-		Row->GetPropertyHandle()->SetOnChildPropertyValueChangedWithData(TDelegate<void(const FPropertyChangedEvent&)>::CreateLambda([StructViewer](const FPropertyChangedEvent& InEvent)
-		{
-			GEditor->EndTransaction();
-			StructViewer->TriggerReinitIfNeeded(InEvent);
-		}));
-	}
-}
-
 const TArray<FText>& FIKRigGenericDetailCustomization::GetButtonLabels()
 {
 	static const TArray<FText> ButtonLabels
diff --git a/Engine/Plugins/Animation/IKRig/Source/IKRigEditor/Private/RigEditor/IKRigEditorController.cpp b/Engine/Plugins/Animation/IKRig/Source/IKRigEditor/Private/RigEditor/IKRigEditorController.cpp
--- a/Engine/Plugins/Animation/IKRig/Source/IKRigEditor/Private/RigEditor/IKRigEditorController.cpp
+++ b/Engine/Plugins/Animation/IKRig/Source/IKRigEditor/Private/RigEditor/IKRigEditorController.cpp
@@ -22,6 +22,7 @@
 #include "Framework/Notifications/NotificationManager.h"
 #include "RigEditor/IKRigAutoFBIK.h"
 #include "RigEditor/IKRigDetailCustomizations.h"
+#include "RigEditor/IKRigStructViewer.h"
 #include "Widgets/Notifications/SNotificationList.h"
 
 #include UE_INLINE_GENERATED_CPP_BY_NAME(IKRigEditorController)
@@ -971,10 +972,16 @@ void FIKRigEditorController::ShowDetailsForBoneSettings(const FName& BoneName, i
 	};
 
 	FInstancedStruct* InstancedStruct = AssetController->GetSolverStructAtIndex(SolverIndex);
-	const UScriptStruct* TypeToEdit = InstancedStruct->GetPtr<FIKRigSolverBase>()->GetBoneSettingsType();
+	const FIKRigSolverBase* Solver = InstancedStruct->GetPtr<FIKRigSolverBase>();
+
+	FIKRigStructToView StructToView;
+	StructToView.Type = Solver->GetBoneSettingsType();
+	StructToView.MemoryProvider = MemoryProvider;
+	StructToView.Owner = AssetController->GetAsset();
+	StructToView.UniqueName = BoneName;
 	
 	UIKRigStructViewer* StructViewer = AssetController->GetStructViewer();
-	StructViewer->SetStructToView(TypeToEdit, MemoryProvider, AssetController->GetAsset());
+	StructViewer->SetStructToView(StructToView);
 	
 	ShowStructViewerInDetailsView(StructViewer);
 }
@@ -1000,10 +1007,15 @@ void FIKRigEditorController::ShowDetailsForGoalSettings(const FName GoalName, co
 	};
 	
 	FInstancedStruct* InstancedStruct = AssetController->GetSolverStructAtIndex(SolverIndex);
-	const UScriptStruct* TypeToEdit = InstancedStruct->GetPtr<FIKRigSolverBase>()->GetGoalSettingsType();
+
+	FIKRigStructToView StructToView;
+	StructToView.Type = InstancedStruct->GetPtr<FIKRigSolverBase>()->GetGoalSettingsType();
+	StructToView.MemoryProvider = MemoryProvider;
+	StructToView.Owner = AssetController->GetAsset();
+	StructToView.UniqueName = GoalName;
 	
 	UIKRigStructViewer* StructViewer = AssetController->GetStructViewer();
-	StructViewer->SetStructToView(TypeToEdit, MemoryProvider, AssetController->GetAsset());
+	StructViewer->SetStructToView(StructToView);
 	
 	ShowStructViewerInDetailsView(StructViewer);
 }
@@ -1024,10 +1036,16 @@ void FIKRigEditorController::ShowDetailsForSolver(const int32 SolverIndex) const
 	};
 	
 	FInstancedStruct* InstancedStruct = AssetController->GetSolverStructAtIndex(SolverIndex);
-	const UScriptStruct* TypeToEdit = InstancedStruct->GetPtr<FIKRigSolverBase>()->GetSolverSettingsType();
+	const FIKRigSolverBase* Solver = InstancedStruct->GetPtr<FIKRigSolverBase>();
+	
+	FIKRigStructToView StructToView;
+	StructToView.Type = Solver->GetSolverSettingsType();
+	StructToView.MemoryProvider = MemoryProvider;
+	StructToView.Owner = AssetController->GetAsset();
+	StructToView.UniqueName = FName(Solver->GetNiceName().ToString());
 	
 	UIKRigStructViewer* StructViewer = AssetController->GetStructViewer();
-	StructViewer->SetStructToView(TypeToEdit, MemoryProvider, AssetController->GetAsset());
+	StructViewer->SetStructToView(StructToView);
 	
 	ShowStructViewerInDetailsView(StructViewer);
 }
diff --git a/Engine/Plugins/Animation/IKRig/Source/IKRigEditor/Public/RetargetEditor/IKRetargetEditorController.h b/Engine/Plugins/Animation/IKRig/Source/IKRigEditor/Public/RetargetEditor/IKRetargetEditorController.h
--- a/Engine/Plugins/Animation/IKRig/Source/IKRigEditor/Public/RetargetEditor/IKRetargetEditorController.h
+++ b/Engine/Plugins/Animation/IKRig/Source/IKRigEditor/Public/RetargetEditor/IKRetargetEditorController.h
@@ -135,7 +135,6 @@ public:
 	// store pointers to various tabs of UI,
 	// have to manage access to these because they can be null if the tabs are closed
 	void SetDetailsView(const TSharedPtr<IDetailsView>& InDetailsView) { DetailsView = InDetailsView; };
-	void SetChainsView(const TSharedPtr<SIKRetargetChainMapList>& InChainsView) { ChainsView = InChainsView; };
 	void SetAssetBrowserView(const TSharedPtr<SIKRetargetAssetBrowser>& InAssetBrowserView) { AssetBrowserView = InAssetBrowserView; };
 	void SetOutputLogView(const TSharedPtr<SIKRigOutputLog>& InOutputLogView) { OutputLogView = InOutputLogView; };
 	void SetHierarchyView(const TSharedPtr<SIKRetargetHierarchy>& InHierarchyView) { HierarchyView = InHierarchyView; };
@@ -145,7 +144,6 @@ public:
 	// force refresh all views in the editor
 	void RefreshAllViews() const;
 	void RefreshDetailsView() const;
-	void RefreshChainsView() const;
 	void RefreshAssetBrowserView() const;
 	void RefreshHierarchyView() const;
 	void RefreshOpStackView() const;
@@ -234,7 +232,7 @@ public:
 
 	// SELECTION - OP STACK
 	void SetOpSelected(const int32 InOpIndex);
-	FName GetLastSelectedOpName() const;
+	FName GetSelectedOpName() const;
 	FIKRetargetOpBase* GetSelectedOp() const;
 	int32 GetSelectedOpIndex() const;
 	// END op selection
@@ -317,8 +315,6 @@ private:
 
 	// asset properties tab
 	TSharedPtr<IDetailsView> DetailsView;
-	// chain list view
-	TSharedPtr<SIKRetargetChainMapList> ChainsView;
 	// asset browser view
 	TSharedPtr<SIKRetargetAssetBrowser> AssetBrowserView;
 	// output log view
diff --git a/Engine/Plugins/Animation/IKRig/Source/IKRigEditor/Public/RetargetEditor/IKRetargetEditorStyle.h b/Engine/Plugins/Animation/IKRig/Source/IKRigEditor/Public/RetargetEditor/IKRetargetEditorStyle.h
--- a/Engine/Plugins/Animation/IKRig/Source/IKRigEditor/Public/RetargetEditor/IKRetargetEditorStyle.h
+++ b/Engine/Plugins/Animation/IKRig/Source/IKRigEditor/Public/RetargetEditor/IKRetargetEditorStyle.h
@@ -34,9 +34,13 @@ public:
 		Set("IKRetarget.AutoAlign", new IMAGE_BRUSH_SVG("Slate/AutoRetargetPose", Icon16x16));
 
 		FSlateColor OpOutlineColor = FSlateColor(FLinearColor(0.1843f, 0.1843f, 0.1843f, 0.5f));
+		FSlateColor OpOutlineColorSelected = FSlateColor(FLinearColor(0.1843f, 0.1843f, 0.1843f, 1.0f));
 		Set("IKRetarget.OpBorder", new FSlateRoundedBoxBrush(
 			FStyleColors::Header, CoreStyleConstants::InputFocusRadius,
 			OpOutlineColor, CoreStyleConstants::InputFocusThickness));
+		Set("IKRetarget.OpBorderSelected", new FSlateRoundedBoxBrush(
+			FStyleColors::Select, CoreStyleConstants::InputFocusRadius,
+			OpOutlineColorSelected, CoreStyleConstants::InputFocusThickness));
 
 		Set("IKRetarget.OpGroupBorder", new FSlateRoundedBoxBrush(FStyleColors::Recessed, CoreStyleConstants::InputFocusRadius));
 
diff --git a/Engine/Plugins/Animation/IKRig/Source/IKRigEditor/Public/RetargetEditor/IKRetargeterController.h b/Engine/Plugins/Animation/IKRig/Source/IKRigEditor/Public/RetargetEditor/IKRetargeterController.h
--- a/Engine/Plugins/Animation/IKRig/Source/IKRigEditor/Public/RetargetEditor/IKRetargeterController.h
+++ b/Engine/Plugins/Animation/IKRig/Source/IKRigEditor/Public/RetargetEditor/IKRetargeterController.h
@@ -23,14 +23,6 @@ class UIKRetargeter;
 class USkeletalMesh;
 class UIKRigStructViewer;
 
-UENUM()
-enum class EAutoMapChainType : uint8
-{
-	Exact, // map chains that have exactly the same name (case insensitive)
-	Fuzzy, // map chains to the chain with the closest name (levenshtein distance)
-	Clear, // clear all mappings, set them to None
-};
-
 // A stateless singleton (1-per-asset) class used to make modifications to a UIKRetargeter asset.
 // Use UIKRetargeter.GetController() to get the controller for the asset you want to modify.  
 UCLASS(BlueprintType, hidecategories = UObject)
@@ -77,6 +69,11 @@ public:
 	 * @return the IK Rig asset associated with either the Source or Target of the IK Retargeter */ 
 	UFUNCTION(BlueprintCallable, Category="IK Retarget", meta = (BlueprintThreadSafe, DisplayName = "Get IK Rig"))
 	const UIKRigDefinition* GetIKRig(const ERetargetSourceOrTarget SourceOrTarget) const;
+	
+	/** Get all target IK Rigs referenced by all ops
+	 * @return an array of all target IK Rig assets associated with all the ops in the stack */ 
+	UFUNCTION(BlueprintCallable, Category="IK Retarget", meta = (BlueprintThreadSafe))
+	TArray<UIKRigDefinition*> GetAllTargetIKRigs() const;
 
 	/** Set the preview skeletal mesh for either source or target
 	* @param SourceOrTarget an enum specifying either "Source" or "Target"
@@ -95,14 +92,9 @@ public:
 	//
 
 	/** Add a new retarget op of the given type to the bottom of the stack. Returns the stack index.
-	 * @param InIKRetargetOpType: the full package path of the UStruct type, ie /Script/IKRig.IKRetargetPinBoneOp
-	 * @param bAllowDuplicates: if false, will not add a new op if one of that type already exists (returns index of existing op instead)
-	 * @param bRunInitialSetup: if true, will run the Op's setup routine after adding it to the stack */
+	 * @param InIKRetargetOpType: the full package path of the UStruct type, ie /Script/IKRig.IKRetargetPinBoneOp */
 	UFUNCTION(BlueprintCallable, BlueprintPure=false, Category="IK Retarget")
-	int32 AddRetargetOp(
-		const FString InIKRetargetOpType, 
-		const bool bAllowDuplicates = true,
-		const bool bRunInitialSetup = true) const;
+	int32 AddRetargetOp(const FString InIKRetargetOpType) const;
 
 	/** Remove the retarget op at the given stack index
 	 * NOTE: if this op is a parent, all it's children will be removed as well
@@ -154,7 +146,7 @@ public:
 	 * If any of these ops are already present, they will not be re-added.
 	 * @param bRunInitialSetup: if true, will run the initial setup routine on each op after they are all added to the stack. */
 	UFUNCTION(BlueprintCallable, Category="IK Retarget", meta = (BlueprintThreadSafe))
-	void AddDefaultOps(const bool bRunInitialSetup=true) const;
+	void AddDefaultOps() const;
 
 	/** Force the op to run the initial setup. This is normally run when an Op is added through the editor UI.
 	 * NOTE: ops may not have custom initial setup routines; the exact behavior is op-dependent
@@ -203,11 +195,7 @@ public:
 	//
 	
 	// add a new retarget op to the stack
-	int32 AddRetargetOp(
-		const UScriptStruct* InRetargetOpType,
-		const FName InParentOpName = NAME_None,
-		const bool bAllowDuplicates = true,
-		const bool bRunInitialSetup = true) const;
+	int32 AddRetargetOp(const UScriptStruct* InRetargetOpType, const FName InParentOpName = NAME_None) const;
 
 	// Gather up all the ops that are children of this op
 	TArray<int32> GetChildOpIndices(const int32 InOpIndex) const;
@@ -229,7 +217,7 @@ public:
 	}
 
 	// Get access to the given retarget operation by index.
-	FIKRetargetOpBase* GetOpAtIndex(const int32 InOpIndex) const;
+	FIKRetargetOpBase* GetRetargetOpByIndex(const int32 InOpIndex) const;
 
 	// Get access to the given retarget operation by name.
 	FIKRetargetOpBase* GetRetargetOpByName(const FName InOpName) const;
@@ -240,6 +228,9 @@ public:
 	// Get the index of a given retarget operation. 
 	int32 GetIndexOfRetargetOp(FIKRetargetOpBase* RetargetOp) const;
 
+	// A callback whenever the property of an op is modified
+	void OnOpPropertyChanged(const FName& InOpName, const FPropertyChangedEvent& InPropertyChangedEvent) const;
+
 	//
 	// GENERAL C++ ONLY API
 	//
@@ -259,7 +250,7 @@ public:
 	void SetAskedToFixRootHeightForMesh(USkeletalMesh* Mesh, bool InAsked) const;
 	
 	// Get name of the Root bone used for retargeting. 
-	FName GetRetargetRootBone(const ERetargetSourceOrTarget SourceOrTarget) const;
+	FName GetPelvisBone(const ERetargetSourceOrTarget SourceOrTarget) const;
 
 	//
 	// RETARGET CHAIN MAPPING PUBLIC/SCRIPTING API
@@ -267,31 +258,53 @@ public:
 	
 	/** Use string comparision to find "best" Source chain to map to each Target chain or clear the mappings
 	 * @param AutoMapType an enum specifying the type of mapping to perform (ie, Exact, Fuzzy or Clear)
-	 * @param bForceRemap if false, will only remap those mappings that unset/None */
+	 * @param bForceRemap if false, will only remap those mappings that unset/None
+	 * @param InOpName if specified, will auto-map only the chain mapping in the specified op, otherwise applies to all ops*/
 	UFUNCTION(BlueprintCallable, BlueprintPure=false, Category="IK Retarget", meta = (BlueprintThreadSafe))
-	void AutoMapChains(const EAutoMapChainType AutoMapType, const bool bForceRemap) const;
+	void AutoMapChains(const EAutoMapChainType AutoMapType, const bool bForceRemap, const FName InOpName=NAME_None) const;
 
 	/** Assign a source chain to the given target chain. Animation will be copied from the source to the target.
 	 * @param InSourceChainName the name of the source retarget chain to assign
-	 * @param InTargetChainName the name of the target chain to map */
+	 * @param InTargetChainName the name of the target chain to map
+	 * @param InOpName if specified, will set the source chain only for that op (otherwise applies to all ops with a chain mapping)*/
 	UFUNCTION(BlueprintCallable, BlueprintPure=false, Category="IK Retarget", meta = (BlueprintThreadSafe))
-	bool SetSourceChain(FName InSourceChainName, FName InTargetChainName) const;
+	bool SetSourceChain(FName InSourceChainName, FName InTargetChainName, const FName InOpName=NAME_None) const;
 
 	/** Get the name of the source chain mapped to a given target chain (the chain animation is copied FROM).
 	 * @param InTargetChainName the name of the target retarget chain to get a mapping for
+	 * @param InOpName if specified, will get the source chain for that op (otherwise gets the first chain mapping in the op stack)
 	 * @return the name of the source chain mapped to the given target chain (or None if not mapped) */
 	UFUNCTION(BlueprintCallable, Category="IK Retarget", meta = (BlueprintThreadSafe))
-	FName GetSourceChain(const FName& InTargetChainName) const;
+	FName GetSourceChain(const FName& InTargetChainName, const FName InOpName=NAME_None) const;
+
+	/** Reset the settings for the given chain in the given op back to the defaults.
+	 * @param InTargetChainName the name of the chain
+	 * @param InOpName the name of the op containing the chain settings */
+	UFUNCTION(BlueprintCallable, Category="IK Retarget", meta = (BlueprintThreadSafe))
+	void ResetChainSettingsToDefault(const FName InTargetChainName, const FName InOpName) const;
+
+	/** Get the target IK Rig associated with the given Op. May be different than the default IK Rig.
+	 * NOTE: all ops use the global SOURCE IK rig, but each op may use its own custom TARGET IK Rig.
+	 * NOTE: not all ops maintain their own target IK Rig in which case this function returns nullptr
+	 * NOTE: some ops may refer to the target IK Rig used by their parent op. In that case, this function returns the parent Op's IK Rig.
+	 * @param InOpName the name of the retarget op
+	 * @return the IK Rig assset associated with this op */
+	UFUNCTION(BlueprintCallable, Category="IK Retarget", meta = (BlueprintThreadSafe))
+	const UIKRigDefinition* GetTargetIKRigForOp(const FName InOpName) const;
 
 	//
 	// RETARGET CHAIN MAPPING C++ ONLY API
 	//
 
-	// get access to the mapping between source/target chains
-	const FRetargetChainMapping& GetChainMapping() const;
+	/** Get access to the mapping between source/target chains for the particular op (may be null)
+	 * If no op name is specified, it returns the first chain mapping it finds*/
+	const FRetargetChainMapping* GetChainMapping(const FName InOpName = NAME_None) const;
+
+	// Clean all chain mappings in all ops, or if InOpName is specified, just that op
+	void CleanChainMaps(const FName InOpName=NAME_None) const;
 	
-	// Ask if any op has settings for the given chain that are not at the default
-	bool ChainHasNonDefaultSettings(const FName InTargetChainName) const;
+	// Ask if the given op has settings for the given chain that are not at the default
+	bool AreChainSettingsAtDefault(const FName InTargetChainName, const FName InOpName) const;
 	
 	// Get whether the given chain's IK goal is connected to a solver 
 	bool IsChainGoalConnectedToASolver(const FName& GoalName) const;
@@ -300,7 +313,7 @@ public:
 	void HandleRetargetChainAdded(UIKRigDefinition* IKRig) const;
 	
 	// Call this when IK Rig chain is renamed. Retains existing mappings using the new name 
-	void HandleRetargetChainRenamed(UIKRigDefinition* IKRig, FName OldChainName, FName NewChainName) const;
+	void HandleRetargetChainRenamed(UIKRigDefinition* InIKRig, FName OldChainName, FName NewChainName) const;
 	
 	// Call this when IK Rig chain is removed. 
 	void HandleRetargetChainRemoved(UIKRigDefinition* IKRig, const FName& InChainRemoved) const;
@@ -436,6 +449,10 @@ private:
 	PRAGMA_DISABLE_DEPRECATION_WARNINGS
 public:
 
+	// get access to the mapping between source/target chains
+	UE_DEPRECATED(5.6, "Chain mappings are stored on individual ops. Use the version that takes an op name. ")
+	const FRetargetChainMapping& GetChainMapping() const;
+	
 	// Get a copy of the retarget root settings for this asset.
 	UE_DEPRECATED(5.6, "Root settings are now accessed through the Pelvis Motion Op controller.")
 	UFUNCTION(BlueprintCallable, BlueprintPure=false, Category="IK Retarget", meta = (BlueprintThreadSafe))
diff --git a/Engine/Plugins/Animation/IKRig/Source/IKRigEditor/Public/RetargetEditor/SIKRetargetChainMapList.h b/Engine/Plugins/Animation/IKRig/Source/IKRigEditor/Public/RetargetEditor/SIKRetargetChainMapList.h
--- a/Engine/Plugins/Animation/IKRig/Source/IKRigEditor/Public/RetargetEditor/SIKRetargetChainMapList.h
+++ b/Engine/Plugins/Animation/IKRig/Source/IKRigEditor/Public/RetargetEditor/SIKRetargetChainMapList.h
@@ -11,6 +11,7 @@
 #include "Widgets/Views/STableRow.h"
 #include "Widgets/Views/SListView.h"
 #include "Framework/Commands/UICommandList.h"
+#include "RigEditor/IKRigDetailCustomizations.h"
 
 class FTextFilterExpressionEvaluator;
 class FIKRigEditorController;
@@ -83,49 +84,88 @@ struct FChainMapFilterOptions
 	bool bHideUnmappedChains = false;
 	bool bHideMappedChains = false;
 	bool bHideChainsWithoutIK = false;
+	
+	bool bNeverShowChainsWithoutIK = false;
+
+	void Reset()
+	{
+		bool bNeverShowNonIK = bNeverShowChainsWithoutIK;
+		*this = FChainMapFilterOptions();
+		bNeverShowChainsWithoutIK = bNeverShowNonIK;
+	}
+};
+
+struct FChainMapListConfig
+{
+	// the name of the op that owns the chain settings
+	FName OpWithChainSettings = NAME_None;
+	// the name of the op that owns the chain mapping
+	FName OpWithChainMapping = NAME_None;
+	// the asset controller that owns the op
+	TWeakObjectPtr<UIKRetargeterController> Controller;
+	// the default filter options for the list
+	FChainMapFilterOptions Filter;
+	// whether to show the column for IK goals or not
+	bool bEnableGoalColumn = true;
+	// whether to allow remapping target chains to source chains
+	bool bEnableChainMapping = true;
+	// a callback to get the chain settings for a given chain
+	TFunction<UObject*(FName)> ChainSettingsGetterFunc;
+
+	bool IsValid() const
+	{
+		// chain mapping is essential, but chain settings are optional
+		return OpWithChainMapping != NAME_None && Controller != nullptr;
+	}
 };
 
 typedef SListView< TSharedPtr<FRetargetChainMapElement> > SRetargetChainMapListViewType;
 
-class SIKRetargetChainMapList : public SCompoundWidget, public FEditorUndoClient
+class SIKRetargetChainMapList : public SCompoundWidget, FEditorUndoClient, FGCObject
 {
 	
 public:
 	
 	SLATE_BEGIN_ARGS(SIKRetargetChainMapList) {}
+		SLATE_ARGUMENT(FChainMapListConfig, InChainMapListConfig)
 	SLATE_END_ARGS()
 
-	void Construct(const FArguments& InArgs, TSharedRef<FIKRetargetEditorController> InEditorController);
+	void Construct(const FArguments& InArgs);
 
 	void ClearSelection() const;
 
 	void ResetChainSettings(const FName InTargetChainName) const;
+
+	// FGCObject interface
+	virtual void AddReferencedObjects( FReferenceCollector& Collector ) override;
+	virtual FString GetReferencerName() const override
+	{
+		return TEXT("SIKRetargetChainMapList");
+	}
+	// End of FGCObject interface
 	
 private:
-	
-	/** menu for adding new solver commands */
-	TSharedPtr<FUICommandList> CommandList;
-	
-	/** editor controller */
-	TWeakPtr<FIKRetargetEditorController> EditorController;
+
+	/** the options for this chain map list */
+	FChainMapListConfig Config;
 
 	/** list view */
 	TSharedPtr<SRetargetChainMapListViewType> ListView;
 	TArray< TSharedPtr<FRetargetChainMapElement> > ListViewItems;
 	/** END list view */
 
-	UIKRetargeterController* GetRetargetController() const;
-
+	/** the details view for editing selected chain properties */
+	TSharedPtr<IDetailsView> DetailsView;
+	
 	/** callbacks */
-	FText GetSourceRootBone() const;
-	FText GetTargetRootBone() const;
-	bool IsChainMapEnabled() const;
+	bool IsListEnabled() const;
+	bool IsChainMappingEnabled() const;
 	/** when a chain is clicked on in the table view */
-	void OnItemClicked(TSharedPtr<FRetargetChainMapElement> InItem) const;
+	void OnItemClicked(TSharedPtr<FRetargetChainMapElement> InItem);
+	TArray<TObjectPtr<UObject>> AllStructWrappers;
 
 	/** auto-map chain button*/
 	TSharedRef<SWidget> CreateChainMapMenuWidget();
-	EVisibility IsAutoMapButtonVisible() const;
 	void AutoMapChains(const EAutoMapChainType AutoMapType, const bool bSetUnmatchedToNone);
 	/** END auto-map chain button*/
 
@@ -133,7 +173,8 @@ private:
 	TSharedRef<SWidget> CreateFilterMenuWidget();
 	void OnFilterTextChanged(const FText& SearchText);
 	TSharedPtr<FTextFilterExpressionEvaluator> TextFilter;
-	FChainMapFilterOptions ChainFilterOptions;
+	/** menu for adding new solver commands */
+	TSharedPtr<FUICommandList> MenuCommandList;
 
 	/** list view generate row callback */
 	TSharedRef<ITableRow> MakeListRowWidget(TSharedPtr<FRetargetChainMapElement> InElement, const TSharedRef<STableViewBase>& OwnerTable);
diff --git a/Engine/Plugins/Animation/IKRig/Source/IKRigEditor/Public/RigEditor/IKRigDetailCustomizations.h b/Engine/Plugins/Animation/IKRig/Source/IKRigEditor/Public/RigEditor/IKRigDetailCustomizations.h
--- a/Engine/Plugins/Animation/IKRig/Source/IKRigEditor/Public/RigEditor/IKRigDetailCustomizations.h
+++ b/Engine/Plugins/Animation/IKRig/Source/IKRigEditor/Public/RigEditor/IKRigDetailCustomizations.h
@@ -8,113 +8,6 @@
 #include "Interfaces/Interface_BoneReferenceSkeletonProvider.h"
 #include "RigEditor/IKRigEditorController.h"
 
-#include "IKRigDetailCustomizations.generated.h"
-
-
-// a thin wrapper around UStruct data to display in a details panel
-UCLASS(Blueprintable)
-class UIKRigStructViewer : public UObject, public IBoneReferenceSkeletonProvider
-{
-	GENERATED_BODY()
-
-public:
-	
-	/**
-	 * Configures an instanced struct to display a child struct property in the details panel with undo/redo support.
-	 * NOTE: we can't just pass in raw pointers to the struct memory because these can be destroyed after a transaction.
-	 * So instead we pass in TFunction callbacks that get the latest memory locations when the details panel is refreshed.
-	 *
-	 * @param InTypeToEdit The type that corresponds to the struct memory returned by the StructMemoryProvider
-	 * @param InStructMemoryProvider A callback function that retrieves the memory location for the struct to be edited
-	 * @param InStructOwner The UObject that owns the struct. This is what will be transacted when the property is edited.
-	 */
-	void SetStructToView(
-		const UScriptStruct* InTypeToEdit,
-		TFunction<uint8*()> InStructMemoryProvider,
-		UObject* InStructOwner)
-	{
-		StructType = InTypeToEdit;
-		StructMemoryProvider = MoveTemp(InStructMemoryProvider);
-		StructOwner = InStructOwner;
-	}
-
-	bool IsValid() const
-	{
-		uint8* Memory = StructMemoryProvider();
-		return Memory != nullptr && StructType && StructOwner.IsValid();
-	}
-
-	void Reset()
-	{
-		StructType = nullptr;
-		StructMemoryProvider.Reset();
-		StructOwner = nullptr;
-	}
-
-	TSharedPtr<FStructOnScope>& GetStructOnScope()
-	{
-		uint8* Memory = StructMemoryProvider();
-		StructOnScope = MakeShared<FStructOnScope>(StructType, Memory);
-		return StructOnScope;
-	}
-
-	FName GetStructName() const
-	{
-		// try to get the "nice name" from metadata
-		FString TypeName = StructType->GetMetaData(TEXT("DisplayName"));
-
-		// if no "DisplayName" metadata is found, fall back to the struct name
-		if (TypeName.IsEmpty())
-		{
-			TypeName = StructType->GetName();
-		}
-		
-		return FName(TypeName);
-	}
-
-	UObject* GetStructOwner() const
-	{
-		return StructOwner.Get();
-	};
-
-	void TriggerReinitIfNeeded(const FPropertyChangedEvent& InEvent);
-
-	// IBoneReferenceSkeletonProvider overrides.
-	virtual USkeleton* GetSkeleton(bool& bInvalidSkeletonIsError, const IPropertyHandle* PropertyHandle) override;
-	// ~END IBoneReferenceSkeletonProvider overrides.
-
-private:
-
-	DECLARE_MULTICAST_DELEGATE(FOnStructPropertyEdited);
-	FOnStructPropertyEdited StructPropertyEditedDelegate;
-	
-	// the type of struct we are editing
-	const UScriptStruct* StructType;
-	// provides the memory address of the struct to edit (refreshed after undo/redo)
-	TFunction<uint8*()> StructMemoryProvider;
-	
-	// a wrapper of the struct sent to the details panel
-	TSharedPtr<FStructOnScope> StructOnScope;
-
-	// a UObject that owns the struct for Undo/Redo transactions (ie, the asset we are editing)
-	TWeakObjectPtr<UObject> StructOwner = nullptr;
-
-public:
-	
-	// Attach a delegate to be notified whenever a property (or child property) in the currently displayed UStruct is edited
-	FOnStructPropertyEdited& OnStructPropertyEdited(){ return StructPropertyEditedDelegate; };
-};
-
-class FStructViewerCustomization : public IDetailCustomization
-{
-public:
-	static TSharedRef<IDetailCustomization> MakeInstance();
-
-	// IDetailCustomization interface
-	virtual void CustomizeDetails(IDetailLayoutBuilder& DetailBuilder) override;
-	// END IDetailCustomization interface
-};
-
 
 class FIKRigGenericDetailCustomization : public IDetailCustomization
 {
diff --git a/Engine/Source/Editor/PropertyEditor/Private/CustomChildBuilder.cpp b/Engine/Source/Editor/PropertyEditor/Private/CustomChildBuilder.cpp
--- a/Engine/Source/Editor/PropertyEditor/Private/CustomChildBuilder.cpp
+++ b/Engine/Source/Editor/PropertyEditor/Private/CustomChildBuilder.cpp
@@ -17,10 +17,10 @@ IDetailChildrenBuilder& FCustomChildrenBuilder::AddCustomBuilder( TSharedRef<cla
 	return *this;
 }
 
-IDetailGroup& FCustomChildrenBuilder::AddGroup( FName GroupName, const FText& LocalizedDisplayName )
+IDetailGroup& FCustomChildrenBuilder::AddGroup( FName GroupName, const FText& LocalizedDisplayName, const bool bStartExpanded)
 {
 	FDetailLayoutCustomization NewCustomization;
-	NewCustomization.DetailGroup = MakeShareable( new FDetailGroup( GroupName, ParentCategory.Pin().ToSharedRef(), LocalizedDisplayName ) );
+	NewCustomization.DetailGroup = MakeShareable( new FDetailGroup( GroupName, ParentCategory.Pin().ToSharedRef(), LocalizedDisplayName, bStartExpanded ) );
 
 	ChildCustomizations.Add( NewCustomization );
 
diff --git a/Engine/Source/Editor/PropertyEditor/Private/CustomChildBuilder.h b/Engine/Source/Editor/PropertyEditor/Private/CustomChildBuilder.h
--- a/Engine/Source/Editor/PropertyEditor/Private/CustomChildBuilder.h
+++ b/Engine/Source/Editor/PropertyEditor/Private/CustomChildBuilder.h
@@ -21,7 +21,7 @@ public:
 	{}
 
 	virtual IDetailChildrenBuilder& AddCustomBuilder(TSharedRef<class IDetailCustomNodeBuilder> InCustomBuilder) override;
-	virtual IDetailGroup& AddGroup(FName GroupName, const FText& LocalizedDisplayName) override;
+	virtual IDetailGroup& AddGroup(FName GroupName, const FText& LocalizedDisplayName, const bool bStartExpanded = false) override;
 	virtual FDetailWidgetRow& AddCustomRow(const FText& SearchString) override;
 	virtual IDetailPropertyRow& AddProperty(TSharedRef<IPropertyHandle> PropertyHandle) override;
 	UE_DEPRECATED(5.1, "Please use the overload that takes an FAddPropertyParams instead.")
diff --git a/Engine/Source/Editor/PropertyEditor/Public/IDetailChildrenBuilder.h b/Engine/Source/Editor/PropertyEditor/Public/IDetailChildrenBuilder.h
--- a/Engine/Source/Editor/PropertyEditor/Public/IDetailChildrenBuilder.h
+++ b/Engine/Source/Editor/PropertyEditor/Public/IDetailChildrenBuilder.h
@@ -36,9 +36,10 @@ public:
 	 *
 	 * @param GroupName	The name of the group
 	 * @param LocalizedDisplayName	The display name of the group
+	 * @param bStartExpanded Set to true to force the category to expand by default
 	 * @param true if the group should appear in the advanced section of the category
 	 */
-	virtual IDetailGroup& AddGroup(FName GroupName, const FText& LocalizedDisplayName) = 0;
+	virtual IDetailGroup& AddGroup(FName GroupName, const FText& LocalizedDisplayName, const bool bStartExpanded=false) = 0;
 
 	/**
 	 * Adds new custom content as a child to the struct
