[IK Rig] Fix issue, IK Rig anim node triggers ensure if no rig is assigned.

diff --git a/Engine/Plugins/Animation/IKRig/Source/IKRig/Private/AnimNodes/AnimNode_IKRig.cpp b/Engine/Plugins/Animation/IKRig/Source/IKRig/Private/AnimNodes/AnimNode_IKRig.cpp
--- a/Engine/Plugins/Animation/IKRig/Source/IKRig/Private/AnimNodes/AnimNode_IKRig.cpp
+++ b/Engine/Plugins/Animation/IKRig/Source/IKRig/Private/AnimNodes/AnimNode_IKRig.cpp
@@ -51,6 +51,12 @@ void FAnimNode_IKRig::Evaluate_AnyThread(FPoseContext& Output)
 	{
 		IKRigProcessor.Initialize(RigDefinitionAsset, Output.AnimInstanceProxy->GetSkelMeshComponent()->GetSkeletalMeshAsset(),FIKRigGoalContainer());
 	}
+
+	// bail out if uninitialized
+	if (!IKRigProcessor.IsInitialized())
+	{
+		return;
+	}
 	
 	// copy input pose to solver stack
 	CopyInputPoseToSolver(Output.Pose);
diff --git a/Engine/Plugins/Animation/IKRig/Source/IKRig/Private/Rig/IKRigProcessor.cpp b/Engine/Plugins/Animation/IKRig/Source/IKRig/Private/Rig/IKRigProcessor.cpp
--- a/Engine/Plugins/Animation/IKRig/Source/IKRig/Private/Rig/IKRigProcessor.cpp
+++ b/Engine/Plugins/Animation/IKRig/Source/IKRig/Private/Rig/IKRigProcessor.cpp
@@ -19,7 +19,7 @@ void FIKRigProcessor::Initialize(
 	bInitialized = false;
 	
 	// can't initialize without a rig definition
-	if (!ensure(InRigAsset))
+	if (!InRigAsset)
 	{
 		return;
 	}
