[IK Retarget] Adding tooltop to target IK Rig and new display name to make clear it is now a "default" that may be overridden.

diff --git a/Engine/Plugins/Animation/IKRig/Source/IKRig/Public/Retargeter/IKRetargeter.h b/Engine/Plugins/Animation/IKRig/Source/IKRig/Public/Retargeter/IKRetargeter.h
--- a/Engine/Plugins/Animation/IKRig/Source/IKRig/Public/Retargeter/IKRetargeter.h
+++ b/Engine/Plugins/Animation/IKRig/Source/IKRig/Public/Retargeter/IKRetargeter.h
@@ -327,8 +327,8 @@ private:
 	TSoftObjectPtr<USkeletalMesh> SourcePreviewMesh = nullptr;
 #endif
 	
-	// The rig to copy animation TO.
-	UPROPERTY(EditAnywhere, Category = Target)
+	/** The rig to copy animation TO. Note that this is only the default target IK Rig and ops can be setup to use other IK Rigs as desired. */  
+	UPROPERTY(EditAnywhere, Category = Target, DisplayName="Default Target IK Rig")
 	TObjectPtr<UIKRigDefinition> TargetIKRigAsset = nullptr;
 
 #if WITH_EDITORONLY_DATA
