[IK Retarget] Fixes crash when loading old ops with chain settings.

diff --git a/Engine/Plugins/Animation/IKRig/Source/IKRig/Private/Retargeter/IKRetargeter.cpp b/Engine/Plugins/Animation/IKRig/Source/IKRig/Private/Retargeter/IKRetargeter.cpp
--- a/Engine/Plugins/Animation/IKRig/Source/IKRig/Private/Retargeter/IKRetargeter.cpp
+++ b/Engine/Plugins/Animation/IKRig/Source/IKRig/Private/Retargeter/IKRetargeter.cpp
@@ -454,36 +454,14 @@ void UIKRetargeter::PostLoadPutChainMappingInOps()
 	
 	PRAGMA_DISABLE_DEPRECATION_WARNINGS
 
-	// NOTE: we go through all op types that have been upgraded to support a custom IK Rig asset in order to:
-	// 1. copy the global chain mapping into the op's local chain mapping
-	// 2. call the "OnAddedToStack" virtual to regenerate internal settings structs
+	// NOTE: we go through all op types that have been upgraded to support a custom IK Rig asset and
+	// copy the global chain mapping into the op's local chain mapping.
+	// Later in PostLoad(), the op stack will be cleaned which gives each op a callback to regenerate it's per-chain settings
 	
 	TArray<FIKRetargetRunIKRigOp*> RunIKOps = GetAllRetargetOpsOfType<FIKRetargetRunIKRigOp>();
 	for (FIKRetargetRunIKRigOp* Op : RunIKOps)
 	{
 		Op->ChainMapping = ChainMap_DEPRECATED;
-		Op->OnAddedToStack(this, nullptr);
-	}
-
-	TArray<FIKRetargetIKChainsOp*> IKChainsOp = GetAllRetargetOpsOfType<FIKRetargetIKChainsOp>();
-	for (FIKRetargetIKChainsOp* Op : IKChainsOp)
-	{
-		const FIKRetargetOpBase* ParentOp = GetRetargetOpByName(Op->GetParentOpName());
-		Op->OnAddedToStack(this, ParentOp);
-	}
-
-	TArray<FIKRetargetStrideWarpingOp*> WarpOp = GetAllRetargetOpsOfType<FIKRetargetStrideWarpingOp>();
-	for (FIKRetargetStrideWarpingOp* Op : WarpOp)
-	{
-		const FIKRetargetOpBase* ParentOp = GetRetargetOpByName(Op->GetParentOpName());
-		Op->OnAddedToStack(this, ParentOp);
-	}
-
-	TArray<FIKRetargetSpeedPlantingOp*> PlantingOp = GetAllRetargetOpsOfType<FIKRetargetSpeedPlantingOp>();
-	for (FIKRetargetSpeedPlantingOp* Op : PlantingOp)
-	{
-		const FIKRetargetOpBase* ParentOp = GetRetargetOpByName(Op->GetParentOpName());
-		Op->OnAddedToStack(this, ParentOp);
 	}
 	
 	TArray<FIKRetargetFKChainsOp*> FKChainOps = GetAllRetargetOpsOfType<FIKRetargetFKChainsOp>();
@@ -491,7 +469,6 @@ void UIKRetargeter::PostLoadPutChainMappingInOps()
 	{
 		Op->Settings.IKRigAsset = TargetIKRigAsset;
 		Op->ChainMapping = ChainMap_DEPRECATED;
-		Op->OnAddedToStack(this, nullptr);
 	}
 
 	TArray<FIKRetargetAlignPoleVectorOp*> PoleVectorOps = GetAllRetargetOpsOfType<FIKRetargetAlignPoleVectorOp>();
@@ -499,8 +476,8 @@ void UIKRetargeter::PostLoadPutChainMappingInOps()
 	{
 		Op->Settings.IKRigAsset = TargetIKRigAsset;
 		Op->ChainMapping = ChainMap_DEPRECATED;
-		Op->OnAddedToStack(this, nullptr);
 	}
+	
 	PRAGMA_ENABLE_DEPRECATION_WARNINGS
 }
 #endif
