[IK Retarget] Fixes crash when deleting selected op in the op stack.

diff --git a/Engine/Plugins/Animation/IKRig/Source/IKRigEditor/Private/RetargetEditor/IKRetargetEditorController.cpp b/Engine/Plugins/Animation/IKRig/Source/IKRigEditor/Private/RetargetEditor/IKRetargetEditorController.cpp
--- a/Engine/Plugins/Animation/IKRig/Source/IKRigEditor/Private/RetargetEditor/IKRetargetEditorController.cpp
+++ b/Engine/Plugins/Animation/IKRig/Source/IKRigEditor/Private/RetargetEditor/IKRetargetEditorController.cpp
@@ -660,10 +660,10 @@ bool FIKRetargetEditorController::IsObjectInDetailsView(const UObject* Object)
 void FIKRetargetEditorController::RefreshAllViews() const
 {
 	Editor.Pin()->RegenerateMenusAndToolbars();
+	RefreshOpStackView(); // must be before details panel in case details object was deleted
 	RefreshDetailsView();
 	RefreshAssetBrowserView();
 	RefreshHierarchyView();
-	RefreshOpStackView();
 }
 
 void FIKRetargetEditorController::RefreshDetailsView() const
diff --git a/Engine/Plugins/Animation/IKRig/Source/IKRigEditor/Private/RetargetEditor/SRetargetOpStack.cpp b/Engine/Plugins/Animation/IKRig/Source/IKRigEditor/Private/RetargetEditor/SRetargetOpStack.cpp
--- a/Engine/Plugins/Animation/IKRig/Source/IKRigEditor/Private/RetargetEditor/SRetargetOpStack.cpp
+++ b/Engine/Plugins/Animation/IKRig/Source/IKRigEditor/Private/RetargetEditor/SRetargetOpStack.cpp
@@ -421,7 +421,7 @@ FReply SRetargetOpList::OnAcceptDrop(
 	const bool bWasReparented = AssetController->MoveRetargetOpInStack(DraggedElement->GetIndexInStack(),IndexToMoveTo);
 	if (bWasReparented)
 	{
-		RestoreSelectionAndRefreshList();
+		RequestListRefresh();
 	}
 	
 	return FReply::Handled();
@@ -496,53 +496,6 @@ bool SRetargetOpList::IsEnabled() const
 	return false;
 }
 
-void SRetargetOpList::RestoreSelectionAndRefreshList()
-{
-	auto GetElementOfLastSelectedOp = [this]() -> TSharedPtr<FRetargetOpStackElement>
-	{
-		TSharedPtr<FRetargetOpStackElement> ElementToSelect;
-
-		if (!EditorController.IsValid())
-		{
-			return TSharedPtr<FRetargetOpStackElement>();
-		}
-		
-		const FName LastSelectedOpName = EditorController.Pin()->GetSelectedOpName();
-		if (LastSelectedOpName == NAME_None)
-		{
-			return TSharedPtr<FRetargetOpStackElement>();
-		}
-
-		for (TSharedPtr<FRetargetOpStackElement>& Element : Elements)
-		{
-			if (Element->GetName() == LastSelectedOpName)
-			{
-				return Element;
-			}
-		}
-
-		return TSharedPtr<FRetargetOpStackElement>();
-	};
-
-	// restore selection to the last selected element
-	TSharedPtr<FRetargetOpStackElement> ElementToRestore = GetElementOfLastSelectedOp();
-	if (ElementToRestore.IsValid())
-	{
-		// restore selection
-		SetSelection(ElementToRestore);
-		// 
-		EditorController.Pin()->SetOpSelected(ElementToRestore->GetIndexInStack());
-	}
-	else
-	{
-		// clear selection otherwise
-		ClearSelection();
-	}
-		
-	// refresh the list
-	RequestListRefresh();
-}
-
 void SRetargetOpList::RequestRenameSelectedOp() const
 {
 	if (!LastSelectedElement.IsValid())
@@ -588,8 +541,6 @@ void SRetargetOpList::AddNewRetargetOp(const UScriptStruct* ScriptStruct)
 	// add a new op to the stack
 	const FName ParentOpName = ParentElement.IsValid() ? ParentElement.Pin()->GetName() : NAME_None;
 	int32 NewOpIndex = AssetController->AddRetargetOp(ScriptStruct, ParentOpName);
-
-	RestoreSelectionAndRefreshList();
 }
 
 void SRetargetOpList::DeleteRetargetOp(TSharedPtr<FRetargetOpStackElement> OpToDelete)
@@ -600,8 +551,64 @@ void SRetargetOpList::DeleteRetargetOp(TSharedPtr<FRetargetOpStackElement> OpToD
 	}
 	
 	EditorController.Pin()->AssetController->RemoveRetargetOp(OpToDelete->GetIndexInStack());
+}
+
+void SRetargetOpList::RefreshAndRestore()
+{
+	FIKRetargetEditorController* Controller = EditorController.Pin().Get();
+	if (!ensure(Controller))
+	{
+		return;
+	}
+
+	RequestListRefresh();
+	
+	auto GetElementOfLastSelectedOp = [this, Controller]() -> TSharedPtr<FRetargetOpStackElement>
+	{
+		TSharedPtr<FRetargetOpStackElement> ElementToSelect;
+
+		if (Controller->GetSelectionState().LastSelectedType != ERetargetSelectionType::OP)
+		{
+			return TSharedPtr<FRetargetOpStackElement>();
+		}
+	
+		const FName LastSelectedOpName = Controller->GetSelectedOpName();
+		if (LastSelectedOpName == NAME_None)
+		{
+			return TSharedPtr<FRetargetOpStackElement>();
+		}
+
+		for (TSharedPtr<FRetargetOpStackElement> Element : Elements)
+		{
+			if (Element->GetName() == LastSelectedOpName)
+			{
+				return Element;
+			}
+		}
+
+		return TSharedPtr<FRetargetOpStackElement>();
+	};
 	
-	RestoreSelectionAndRefreshList();
+	// restore selection to the last selected element
+	TSharedPtr<FRetargetOpStackElement> ElementToRestore = GetElementOfLastSelectedOp();
+	if (ElementToRestore.IsValid())
+	{
+		// restore selection
+		ElementToRestore.Get()->GetOpList().Pin()->SetSelection(ElementToRestore);
+		Controller->SetOpSelected(ElementToRestore->GetIndexInStack());
+	}
+	else
+	{
+		// if an op was selected but it's no longer valid, then clear the selection
+		if (Controller->GetSelectionState().LastSelectedType == ERetargetSelectionType::OP)
+		{
+			FIKRetargetOpBase* SelectedOp = Controller->AssetController->GetRetargetOpByName(Controller->GetSelectedOpName());
+			if (!SelectedOp)
+			{
+				Controller->ClearSelection();	
+			}
+		}
+	}
 }
 
 void SParentRetargetOpItem::Construct(const FArguments& InArgs, const TSharedRef<STableViewBase>& OwnerTable)
@@ -695,7 +702,7 @@ void SParentRetargetOpItem::RefreshListView() const
 		ParentStackElement.Pin()->AddChild(ChildElement.Get());
 	}
 
-	ChildrenListView->RestoreSelectionAndRefreshList();
+	ChildrenListView->RefreshAndRestore();
 }
 
 void SRetargetOpSingleItem::Construct(const FArguments& InArgs, const TSharedRef<STableViewBase>& OwnerTable)
@@ -864,7 +871,8 @@ void SRetargetOpStack::RefreshStackView() const
 		ListView->Elements.Add(StackElement);
 	}
 
-	ListView->RestoreSelectionAndRefreshList();
+	// refresh the list and restore the selection
+	ListView->RefreshAndRestore();
 }
 
 TSharedRef<FRetargetOpStackDragDropOp> FRetargetOpStackDragDropOp::New(TWeakPtr<FRetargetOpStackElement> InElement)
diff --git a/Engine/Plugins/Animation/IKRig/Source/IKRigEditor/Public/RetargetEditor/IKRetargetEditorController.h b/Engine/Plugins/Animation/IKRig/Source/IKRigEditor/Public/RetargetEditor/IKRetargetEditorController.h
--- a/Engine/Plugins/Animation/IKRig/Source/IKRigEditor/Public/RetargetEditor/IKRetargetEditorController.h
+++ b/Engine/Plugins/Animation/IKRig/Source/IKRigEditor/Public/RetargetEditor/IKRetargetEditorController.h
@@ -207,8 +207,14 @@ public:
 
 	// ------------------------- SELECTION -----------------------------
 
+	// GENERAL SELECTION
+	//
 	// access to all selection state for the editor
 	const FIKRetargetDebugDrawState& GetSelectionState() const { return Selection; };
+	void CleanSelection(ERetargetSourceOrTarget SourceOrTarget);
+	void ClearSelection(const bool bKeepBoneSelection=false);
+	// to frame selection when pressing "f" in viewport
+	bool GetCameraTargetForSelection(FSphere& OutTarget) const;
 
 	// SELECTION - BONES (viewport or hierarchy view)
 	void EditBoneSelection(
@@ -237,12 +243,6 @@ public:
 	int32 GetSelectedOpIndex() const;
 	// END op selection
 
-	// GENERAL SELECTION
-	void CleanSelection(ERetargetSourceOrTarget SourceOrTarget);
-	void ClearSelection(const bool bKeepBoneSelection=false);
-	// to frame selection when pressing "f" in viewport
-	bool GetCameraTargetForSelection(FSphere& OutTarget) const;
-
 	// ------------------------- END SELECTION -----------------------------
 
 	// determine if bone in the specified skeleton is part of the retarget (in a mapped chain)
diff --git a/Engine/Plugins/Animation/IKRig/Source/IKRigEditor/Public/RetargetEditor/SRetargetOpStack.h b/Engine/Plugins/Animation/IKRig/Source/IKRigEditor/Public/RetargetEditor/SRetargetOpStack.h
--- a/Engine/Plugins/Animation/IKRig/Source/IKRigEditor/Public/RetargetEditor/SRetargetOpStack.h
+++ b/Engine/Plugins/Animation/IKRig/Source/IKRigEditor/Public/RetargetEditor/SRetargetOpStack.h
@@ -52,19 +52,21 @@ public:
 	
 	void SetParent(FRetargetOpStackElement* InParentElement) { ParentElement = InParentElement; };
 
+	TWeakPtr<SRetargetOpList> GetOpList() const { return OpListWidget; };
+
 	DECLARE_DELEGATE(FOnRenameRequested);
 	FOnRenameRequested OnRenameRequested;
 	
 private:
 
-	int32 IndexInStack;
-	bool bCanHaveChildren;
+	int32 IndexInStack = INDEX_NONE;
+	bool bCanHaveChildren = false;
 	TArray<FRetargetOpStackElement*> ChildElements;
 	FRetargetOpStackElement* ParentElement = nullptr;
 	TWeakPtr<SRetargetOpList> OpListWidget;
 	
 	// hidden constructors, always use Make above
-	FRetargetOpStackElement() {}
+	FRetargetOpStackElement() = default;
 	FRetargetOpStackElement(
 		const int32 InIndexInStack,
 		const bool InHasChildren,
@@ -94,8 +96,8 @@ class SRetargetOpList : public SListView<TSharedPtr<FRetargetOpStackElement>>
 public:
 	
 	SLATE_BEGIN_ARGS(SRetargetOpList) {}
-	SLATE_ARGUMENT(TWeakPtr<FIKRetargetEditorController>, InEditorController)
-	SLATE_ARGUMENT(TWeakPtr<FRetargetOpStackElement>, InParentElement)
+		SLATE_ARGUMENT(TWeakPtr<FIKRetargetEditorController>, InEditorController)
+		SLATE_ARGUMENT(TWeakPtr<FRetargetOpStackElement>, InParentElement)
 	SLATE_END_ARGS()
 
 	virtual ~SRetargetOpList() {}
@@ -105,8 +107,6 @@ public:
 	TSharedRef<ITableRow> MakeListRowWidget(TSharedPtr<FRetargetOpStackElement> InElement, const TSharedRef<STableViewBase>& OwnerTable);
 	
 	bool IsEnabled() const;
-
-	void RestoreSelectionAndRefreshList();
 	
 	// drag and drop operations
 	FReply OnDragDetected(
@@ -139,6 +139,9 @@ public:
 	virtual FReply OnKeyDown(const FGeometry& MyGeometry, const FKeyEvent& InKeyEvent);
 	void DeleteRetargetOp(TSharedPtr<FRetargetOpStackElement> OpToDelete);
 
+	// must be called after refresh
+	void RefreshAndRestore();
+
 	// the elements contained in this list
 	TArray< TSharedPtr<FRetargetOpStackElement> > Elements;
 
