[IK Retarget] Pole vector op settings reset to default doesn't show correctly.

diff --git a/Engine/Plugins/Animation/IKRig/Source/IKRig/Private/Retargeter/RetargetOps/AlignPoleVectorOp.cpp b/Engine/Plugins/Animation/IKRig/Source/IKRig/Private/Retargeter/RetargetOps/AlignPoleVectorOp.cpp
--- a/Engine/Plugins/Animation/IKRig/Source/IKRig/Private/Retargeter/RetargetOps/AlignPoleVectorOp.cpp
+++ b/Engine/Plugins/Animation/IKRig/Source/IKRig/Private/Retargeter/RetargetOps/AlignPoleVectorOp.cpp
@@ -183,7 +183,8 @@ FVector FPoleVectorMatcher::GetChainAxisNormalized(
 
 bool FRetargetPoleVectorSettings::operator==(const FRetargetPoleVectorSettings& Other) const
 {
-	return FMath::IsNearlyEqualByULP(AlignAlpha,Other.AlignAlpha)
+	return bEnabled == Other.bEnabled
+		&& FMath::IsNearlyEqualByULP(AlignAlpha,Other.AlignAlpha)
 		&& FMath::IsNearlyEqualByULP(StaticAngularOffset, Other.StaticAngularOffset)
 		&& MaintainOffset == Other.MaintainOffset;
 }
diff --git a/Engine/Plugins/Animation/IKRig/Source/IKRigEditor/Private/RetargetEditor/IKRetargetOpDetails.cpp b/Engine/Plugins/Animation/IKRig/Source/IKRigEditor/Private/RetargetEditor/IKRetargetOpDetails.cpp
--- a/Engine/Plugins/Animation/IKRig/Source/IKRigEditor/Private/RetargetEditor/IKRetargetOpDetails.cpp
+++ b/Engine/Plugins/Animation/IKRig/Source/IKRigEditor/Private/RetargetEditor/IKRetargetOpDetails.cpp
@@ -484,6 +484,10 @@ void FPoleVectorOpCustomization::CustomizeChildren(
 		return;
 	}
 
+	// add the IK Rig asset input field
+	TSharedPtr<IPropertyHandle> IKRigHandle = StructPropertyHandle->GetChildHandle(GET_MEMBER_NAME_CHECKED(FIKRetargetAlignPoleVectorOpSettings, IKRigAsset));
+	ChildBuilder.AddProperty(IKRigHandle.ToSharedRef());
+
 	// configure chain map list for the IK op
 	FChainMapListConfig ChainListConfig;
 	ChainListConfig.OpWithChainSettings = OpName;
