[IK Retarget] Curve remap op now has separate toggle for Copying and Remapping.

diff --git a/Engine/Plugins/Animation/IKRig/Source/IKRig/Private/Retargeter/RetargetOps/CurveRemapOp.cpp b/Engine/Plugins/Animation/IKRig/Source/IKRig/Private/Retargeter/RetargetOps/CurveRemapOp.cpp
--- a/Engine/Plugins/Animation/IKRig/Source/IKRig/Private/Retargeter/RetargetOps/CurveRemapOp.cpp
+++ b/Engine/Plugins/Animation/IKRig/Source/IKRig/Private/Retargeter/RetargetOps/CurveRemapOp.cpp
@@ -48,17 +48,21 @@ void FIKRetargetCurveRemapOp::AnimGraphEvaluateAnyThread(FPoseContext& Output)
 	}
 
 	// copy curves over with different names (remap)
-	FBlendedCurve RemapedCurves;
-	for (const FCurveRemapPair& CurveToRemap : Settings.CurvesToRemap)
+	if (Settings.bRemapCurves)
 	{
-		bool bOutIsValid = false;
-		const float SourceValue = SourceCurves.Get(CurveToRemap.SourceCurve, bOutIsValid);
-		if (bOutIsValid)
+		FBlendedCurve RemapedCurves;
+		for (const FCurveRemapPair& CurveToRemap : Settings.CurvesToRemap)
 		{
-			RemapedCurves.Add(CurveToRemap.TargetCurve, SourceValue);
+			bool bOutIsValid = false;
+			const float SourceValue = SourceCurves.Get(CurveToRemap.SourceCurve, bOutIsValid);
+			if (bOutIsValid)
+			{
+				RemapedCurves.Add(CurveToRemap.TargetCurve, SourceValue);
+			}
 		}
+		
+		OutputCurves.Combine(RemapedCurves);
 	}
-	OutputCurves.Combine(RemapedCurves);
 }
 
 const UClass* FIKRetargetCurveRemapOp::GetControllerType() const
diff --git a/Engine/Plugins/Animation/IKRig/Source/IKRig/Public/Retargeter/RetargetOps/CurveRemapOp.h b/Engine/Plugins/Animation/IKRig/Source/IKRig/Public/Retargeter/RetargetOps/CurveRemapOp.h
--- a/Engine/Plugins/Animation/IKRig/Source/IKRig/Public/Retargeter/RetargetOps/CurveRemapOp.h
+++ b/Engine/Plugins/Animation/IKRig/Source/IKRig/Public/Retargeter/RetargetOps/CurveRemapOp.h
@@ -32,14 +32,18 @@ struct IKRIG_API FIKRetargetCurveRemapOpSettings : public FIKRetargetOpSettingsB
 	// NOTE: This setting also applies when exporting retargeted animations.
 	// True: all source curves are copied to the target animation instance/asset
 	// False: only remapped curves are copied on the target animation instance/asset
-	UPROPERTY(EditAnywhere, BlueprintReadWrite, Category="OpSettings")
+	UPROPERTY(EditAnywhere, BlueprintReadWrite, Category="Copy Curves")
 	bool bCopyAllSourceCurves = true;
+
+	// Toggle curve remapping on/off
+	UPROPERTY(EditAnywhere, BlueprintReadWrite, Category="Remap Curves")
+	bool bRemapCurves = true;
 	
 	// Add pairs of Source/Target curve names to remap. While retargeting, the animation from the source curves
 	// will be redirected to the curves on the target skeletal meshes. Can be used to drive, blendshapes or other downstream systems.
 	// NOTE: By default the IK Retargeter will automatically copy all equivalently named curves from the source to the
 	// target. Remapping is only necessary when the target curve name(s) are different.
-	UPROPERTY(EditAnywhere, BlueprintReadWrite, Category="OpSettings")
+	UPROPERTY(EditAnywhere, BlueprintReadWrite, Category="Remap Curves")
 	TArray<FCurveRemapPair> CurvesToRemap;
 };
 
