[IK Retarget] Fixing op script API to support runtime editing.

diff --git a/Engine/Plugins/Animation/IKRig/Source/IKRig/Private/Retargeter/IKRetargetOps.cpp b/Engine/Plugins/Animation/IKRig/Source/IKRig/Private/Retargeter/IKRetargetOps.cpp
--- a/Engine/Plugins/Animation/IKRig/Source/IKRig/Private/Retargeter/IKRetargetOps.cpp
+++ b/Engine/Plugins/Animation/IKRig/Source/IKRig/Private/Retargeter/IKRetargetOps.cpp
@@ -4,15 +4,33 @@
 
 #define LOCTEXT_NAMESPACE "RetargetOpBase"
 
-void FIKRetargetOpBase::SetSettings(const FIKRetargetOpSettingsBase* InSettings)
+const UClass* FIKRetargetOpSettingsBase::GetControllerType() const
 {
-	// copy everything by default
-	static TArray<FName> PropertiesToIgnore;
-	CopySettingsRaw(InSettings, PropertiesToIgnore);
+	return UIKRetargetOpControllerBase::StaticClass();
+}
 
-#if WITH_EDITOR
-	GetSettings()->OwningOpName = GetName();
-#endif
+UIKRetargetOpControllerBase* FIKRetargetOpSettingsBase::GetController(UObject* Outer)
+{
+	return CreateControllerIfNeeded(Outer);
+}
+
+UIKRetargetOpControllerBase* FIKRetargetOpSettingsBase::CreateControllerIfNeeded(UObject* Outer)
+{
+	if (!Controller.IsValid())
+	{
+		const UClass* ClassType = GetControllerType();
+		if (ensure(ClassType && ClassType->IsChildOf(UIKRetargetOpControllerBase::StaticClass())))
+		{
+			Controller = TStrongObjectPtr(NewObject<UIKRetargetOpControllerBase>(Outer, ClassType));
+			Controller->OpSettingsToControl = this;
+		}
+	}
+	return Controller.Get();
+}
+
+void FIKRetargetOpBase::SetSettings(const FIKRetargetOpSettingsBase* InSettings)
+{
+	GetSettings()->CopySettingsAtRuntime(InSettings);
 }
 
 #if WITH_EDITOR
@@ -48,16 +66,6 @@ FText FIKRetargetOpBase::GetWarningMessage() const
 }
 #endif
 
-UIKRetargetOpControllerBase* FIKRetargetOpBase::GetController(UObject* Outer)
-{
-	return CreateControllerIfNeeded(Outer);
-}
-
-const UClass* FIKRetargetOpBase::GetControllerType() const
-{
-	return UIKRetargetOpControllerBase::StaticClass();
-}
-
 FName FIKRetargetOpBase::GetName() const
 {
 	return Name;
@@ -124,18 +132,4 @@ void FIKRetargetOpBase::CopyStructProperties(
 	}
 }
 
-UIKRetargetOpControllerBase* FIKRetargetOpBase::CreateControllerIfNeeded(UObject* Outer)
-{
-	if (!Controller.IsValid())
-	{
-		const UClass* ClassType = GetControllerType();
-		if (ensure(ClassType && ClassType->IsChildOf(UIKRetargetOpControllerBase::StaticClass())))
-		{
-			Controller = TStrongObjectPtr(NewObject<UIKRetargetOpControllerBase>(Outer, ClassType));
-			Controller->OpToControl = this;
-		}
-	}
-	return Controller.Get();
-}
-
 #undef LOCTEXT_NAMESPACE
diff --git a/Engine/Plugins/Animation/IKRig/Source/IKRig/Private/Retargeter/IKRetargetProfile.cpp b/Engine/Plugins/Animation/IKRig/Source/IKRig/Private/Retargeter/IKRetargetProfile.cpp
--- a/Engine/Plugins/Animation/IKRig/Source/IKRig/Private/Retargeter/IKRetargetProfile.cpp
+++ b/Engine/Plugins/Animation/IKRig/Source/IKRig/Private/Retargeter/IKRetargetProfile.cpp
@@ -59,8 +59,34 @@ bool FRetargetOpProfile::CopySettingsToOp(FInstancedStruct& InOutOpStruct, ECopy
 	return true;
 }
 
+UIKRetargetOpControllerBase* FRetargetOpProfile::CreateControllerIfNeeded(UObject* Outer)
+{
+	if (!ensure(SettingsToApply.IsValid()))
+	{
+		return nullptr;
+	}
+	
+	if (!Controller.IsValid())
+	{
+		FIKRetargetOpSettingsBase* OpSettings = SettingsToApply.GetMutablePtr<FIKRetargetOpSettingsBase>();
+		const UClass* ClassType = OpSettings->GetControllerType();
+		if (ensure(ClassType && ClassType->IsChildOf(UIKRetargetOpControllerBase::StaticClass())))
+		{
+			Controller = TStrongObjectPtr(NewObject<UIKRetargetOpControllerBase>(Outer, ClassType));
+			Controller->OpSettingsToControl = OpSettings;
+		}
+	}
+	
+	return Controller.Get();
+}
+
 void FRetargetProfile::FillProfileWithAssetSettings(const UIKRetargeter* InAsset)
 {
+	if (!InAsset)
+	{
+		return;
+	}
+	
 	// profile can apply retarget poses
 	bApplyTargetRetargetPose = true;
 	TargetRetargetPoseName = InAsset->GetCurrentRetargetPoseName(ERetargetSourceOrTarget::Target);
@@ -288,3 +314,35 @@ bool FRetargetProfile::ApplyOpProfilesToOpStruct(
 	}
 	return bOpSettingsApplied;
 }
+
+FRetargetOpProfile* FRetargetProfile::GetOpProfileByName(const FName InOpName)
+{
+	for (FRetargetOpProfile& OpProfile : RetargetOpProfiles)
+	{
+		if (OpProfile.OpToApplySettingsTo == InOpName)
+		{
+			return &OpProfile;
+		}
+	}
+	return nullptr;
+}
+
+FRetargetProfile URetargetProfileLibrary::CopyRetargetProfileFromRetargetAsset(const UIKRetargeter* InRetargetAsset)
+{
+	FRetargetProfile RetargetProfile;
+	RetargetProfile.FillProfileWithAssetSettings(InRetargetAsset);
+	return MoveTemp(RetargetProfile);
+}
+
+UIKRetargetOpControllerBase* URetargetProfileLibrary::GetOpControllerFromRetargetProfile(
+	FRetargetProfile& InRetargetProfile,
+	const FName InRetargetOpName)
+{
+	FRetargetOpProfile* OpProfile = InRetargetProfile.GetOpProfileByName(InRetargetOpName);
+	if (!OpProfile)
+	{
+		return nullptr;
+	}
+	
+	return OpProfile->CreateControllerIfNeeded((UObject*)GetTransientPackage());
+}
diff --git a/Engine/Plugins/Animation/IKRig/Source/IKRig/Private/Retargeter/RetargetOps/AlignPoleVectorOp.cpp b/Engine/Plugins/Animation/IKRig/Source/IKRig/Private/Retargeter/RetargetOps/AlignPoleVectorOp.cpp
--- a/Engine/Plugins/Animation/IKRig/Source/IKRig/Private/Retargeter/RetargetOps/AlignPoleVectorOp.cpp
+++ b/Engine/Plugins/Animation/IKRig/Source/IKRig/Private/Retargeter/RetargetOps/AlignPoleVectorOp.cpp
@@ -202,6 +202,27 @@ void FIKRetargetAlignPoleVectorOpSettings::MergePoleVectorSettings(const FRetarg
 	ChainsToAlign.Add(InSettingsToMerge);
 }
 
+const UClass* FIKRetargetAlignPoleVectorOpSettings::GetControllerType() const
+{
+	return UIKRetargetAlignPoleVectorController::StaticClass();
+}
+
+void FIKRetargetAlignPoleVectorOpSettings::CopySettingsAtRuntime(const FIKRetargetOpSettingsBase* InSettingsToCopyFrom)
+{
+	// copy settings only for chains that the op has initialized
+	const FIKRetargetAlignPoleVectorOpSettings* NewSettings = reinterpret_cast<const FIKRetargetAlignPoleVectorOpSettings*>(InSettingsToCopyFrom);
+	for (const FRetargetPoleVectorSettings& NewChainSettings : NewSettings->ChainsToAlign)
+	{
+		for (FRetargetPoleVectorSettings& ChainSettings : ChainsToAlign)
+		{
+			if (ChainSettings.TargetChainName == NewChainSettings.TargetChainName)
+			{
+				ChainSettings = NewChainSettings;
+			}
+		}
+	}
+}
+
 bool FIKRetargetAlignPoleVectorOp::Initialize(
 	const FIKRetargetProcessor& InProcessor,
 	const FRetargetSkeleton& InSourceSkeleton,
@@ -299,23 +320,6 @@ FIKRetargetOpSettingsBase* FIKRetargetAlignPoleVectorOp::GetSettings()
 	return &Settings;
 }
 
-void FIKRetargetAlignPoleVectorOp::SetSettings(const FIKRetargetOpSettingsBase* InSettings)
-{
-	const FIKRetargetAlignPoleVectorOpSettings* NewSettings = reinterpret_cast<const FIKRetargetAlignPoleVectorOpSettings*>(InSettings);
-	
-	// copy settings only for chains that the op has initialized
-	for (const FRetargetPoleVectorSettings& NewChainSettings : NewSettings->ChainsToAlign)
-	{
-		for (FRetargetPoleVectorSettings& ChainSettings : Settings.ChainsToAlign)
-		{
-			if (ChainSettings.TargetChainName == NewChainSettings.TargetChainName)
-			{
-				ChainSettings = NewChainSettings;
-			}
-		}
-	}
-}
-
 const UScriptStruct* FIKRetargetAlignPoleVectorOp::GetSettingsType() const
 {
 	return FIKRetargetAlignPoleVectorOpSettings::StaticStruct();
@@ -326,11 +330,6 @@ const UScriptStruct* FIKRetargetAlignPoleVectorOp::GetType() const
 	return FIKRetargetAlignPoleVectorOp::StaticStruct();
 }
 
-const UClass* FIKRetargetAlignPoleVectorOp::GetControllerType() const
-{
-	return UIKRetargetAlignPoleVectorController::StaticClass();
-}
-
 const UIKRigDefinition* FIKRetargetAlignPoleVectorOp::GetCustomTargetIKRig() const
 {
 	return Settings.IKRigAsset;
@@ -446,12 +445,12 @@ void FIKRetargetAlignPoleVectorOp::ApplyIKRigs(const UIKRigDefinition* InSourceI
 
 FIKRetargetAlignPoleVectorOpSettings UIKRetargetAlignPoleVectorController::GetSettings()
 {
-	return *reinterpret_cast<FIKRetargetAlignPoleVectorOpSettings*>(OpToControl->GetSettings());
+	return *reinterpret_cast<FIKRetargetAlignPoleVectorOpSettings*>(OpSettingsToControl);
 }
 
 void UIKRetargetAlignPoleVectorController::SetSettings(FIKRetargetAlignPoleVectorOpSettings InSettings)
 {
-	OpToControl->SetSettings(&InSettings);
+	OpSettingsToControl->CopySettingsAtRuntime(&InSettings);
 }
 
 #undef LOCTEXT_NAMESPACE
diff --git a/Engine/Plugins/Animation/IKRig/Source/IKRig/Private/Retargeter/RetargetOps/CopyBasePoseOp.cpp b/Engine/Plugins/Animation/IKRig/Source/IKRig/Private/Retargeter/RetargetOps/CopyBasePoseOp.cpp
--- a/Engine/Plugins/Animation/IKRig/Source/IKRig/Private/Retargeter/RetargetOps/CopyBasePoseOp.cpp
+++ b/Engine/Plugins/Animation/IKRig/Source/IKRig/Private/Retargeter/RetargetOps/CopyBasePoseOp.cpp
@@ -2,6 +2,22 @@
 
 #include "Retargeter/RetargetOps/CopyBasePoseOp.h"
 
+const UClass* FIKRetargetCopyBasePoseOpSettings::GetControllerType() const
+{
+	return UIKRetargetCopyBasePoseController::StaticClass();
+}
+
+void FIKRetargetCopyBasePoseOpSettings::CopySettingsAtRuntime(const FIKRetargetOpSettingsBase* InSettingsToCopyFrom)
+{
+	// copies all properties
+	const TArray<FName> PropertiesToIgnore = {};
+	FIKRetargetOpBase::CopyStructProperties(
+		FIKRetargetCopyBasePoseOpSettings::StaticStruct(),
+		InSettingsToCopyFrom,
+		this,
+		PropertiesToIgnore);
+}
+
 bool FIKRetargetCopyBasePoseOp::Initialize(
 	const FIKRetargetProcessor& InProcessor,
 	const FRetargetSkeleton& InSourceSkeleton,
@@ -18,11 +34,6 @@ FIKRetargetOpSettingsBase* FIKRetargetCopyBasePoseOp::GetSettings()
 	return &Settings;
 }
 
-void FIKRetargetCopyBasePoseOp::SetSettings(const FIKRetargetOpSettingsBase* InSettings)
-{
-	Settings = *reinterpret_cast<const FIKRetargetCopyBasePoseOpSettings*>(InSettings);
-}
-
 const UScriptStruct* FIKRetargetCopyBasePoseOp::GetSettingsType() const
 {
 	return FIKRetargetCopyBasePoseOpSettings::StaticStruct();
@@ -33,17 +44,12 @@ const UScriptStruct* FIKRetargetCopyBasePoseOp::GetType() const
 	return FIKRetargetCopyBasePoseOp::StaticStruct();
 }
 
-const UClass* FIKRetargetCopyBasePoseOp::GetControllerType() const
-{
-	return UIKRetargetCopyBasePoseController::StaticClass();
-}
-
 FIKRetargetCopyBasePoseOpSettings UIKRetargetCopyBasePoseController::GetSettings()
 {
-	return *reinterpret_cast<FIKRetargetCopyBasePoseOpSettings*>(OpToControl->GetSettings());
+	return *reinterpret_cast<FIKRetargetCopyBasePoseOpSettings*>(OpSettingsToControl);
 }
 
 void UIKRetargetCopyBasePoseController::SetSettings(FIKRetargetCopyBasePoseOpSettings InSettings)
 {
-	OpToControl->SetSettings(&InSettings);
+	OpSettingsToControl->CopySettingsAtRuntime(&InSettings);
 }
diff --git a/Engine/Plugins/Animation/IKRig/Source/IKRig/Private/Retargeter/RetargetOps/CurveRemapOp.cpp b/Engine/Plugins/Animation/IKRig/Source/IKRig/Private/Retargeter/RetargetOps/CurveRemapOp.cpp
--- a/Engine/Plugins/Animation/IKRig/Source/IKRig/Private/Retargeter/RetargetOps/CurveRemapOp.cpp
+++ b/Engine/Plugins/Animation/IKRig/Source/IKRig/Private/Retargeter/RetargetOps/CurveRemapOp.cpp
@@ -8,6 +8,22 @@
 #include "Components/SkeletalMeshComponent.h"
 #include "Animation/AnimNodeBase.h"
 
+const UClass* FIKRetargetCurveRemapOpSettings::GetControllerType() const
+{
+	return UIKRetargetCurveRemapController::StaticClass();
+}
+
+void FIKRetargetCurveRemapOpSettings::CopySettingsAtRuntime(const FIKRetargetOpSettingsBase* InSettingsToCopyFrom)
+{
+	// copies all properties
+	const TArray<FName> PropertiesToIgnore = {};
+	FIKRetargetOpBase::CopyStructProperties(
+		FIKRetargetCurveRemapOpSettings::StaticStruct(),
+		InSettingsToCopyFrom,
+		this,
+		PropertiesToIgnore);
+}
+
 void FIKRetargetCurveRemapOp::AnimGraphPreUpdateMainThread(
 	USkeletalMeshComponent& SourceMeshComponent,
     USkeletalMeshComponent& TargetMeshComponent)
@@ -65,17 +81,12 @@ void FIKRetargetCurveRemapOp::AnimGraphEvaluateAnyThread(FPoseContext& Output)
 	}
 }
 
-const UClass* FIKRetargetCurveRemapOp::GetControllerType() const
-{
-	return UIKRetargetCurveRemapController::StaticClass();
-}
-
 FIKRetargetCurveRemapOpSettings UIKRetargetCurveRemapController::GetSettings()
 {
-	return *reinterpret_cast<FIKRetargetCurveRemapOpSettings*>(OpToControl->GetSettings());
+	return *reinterpret_cast<FIKRetargetCurveRemapOpSettings*>(OpSettingsToControl);
 }
 
 void UIKRetargetCurveRemapController::SetSettings(FIKRetargetCurveRemapOpSettings InSettings)
 {
-	OpToControl->SetSettings(&InSettings);
+	OpSettingsToControl->CopySettingsAtRuntime(&InSettings);
 }
diff --git a/Engine/Plugins/Animation/IKRig/Source/IKRig/Private/Retargeter/RetargetOps/FKChainsOp.cpp b/Engine/Plugins/Animation/IKRig/Source/IKRig/Private/Retargeter/RetargetOps/FKChainsOp.cpp
--- a/Engine/Plugins/Animation/IKRig/Source/IKRig/Private/Retargeter/RetargetOps/FKChainsOp.cpp
+++ b/Engine/Plugins/Animation/IKRig/Source/IKRig/Private/Retargeter/RetargetOps/FKChainsOp.cpp
@@ -545,6 +545,35 @@ bool FRetargetFKChainSettings::operator==(const FRetargetFKChainSettings& Other)
 		&& FMath::IsNearlyEqualByULP(TranslationAlpha, Other.TranslationAlpha);
 }
 
+const UClass* FIKRetargetFKChainsOpSettings::GetControllerType() const
+{
+	return UIKRetargetFKChainsController::StaticClass();
+}
+
+void FIKRetargetFKChainsOpSettings::CopySettingsAtRuntime(const FIKRetargetOpSettingsBase* InSettingsToCopyFrom)
+{
+	// copies everything except the ChainsToRetarget array (those are copied below, only for already existing chains)
+	const TArray<FName> PropertiesToIgnore = {"ChainsToRetarget"};
+	FIKRetargetOpBase::CopyStructProperties(
+		FIKRetargetFKChainsOpSettings::StaticStruct(),
+		InSettingsToCopyFrom,
+		this,
+		PropertiesToIgnore);
+	
+	// copy settings only for chains that the op has initialized
+	const FIKRetargetFKChainsOpSettings* NewSettings = reinterpret_cast<const FIKRetargetFKChainsOpSettings*>(InSettingsToCopyFrom);
+	for (const FRetargetFKChainSettings& NewChainSettings : NewSettings->ChainsToRetarget)
+	{
+		for (FRetargetFKChainSettings& ChainSettings : ChainsToRetarget)
+		{
+			if (ChainSettings.TargetChainName == NewChainSettings.TargetChainName)
+			{
+				ChainSettings = NewChainSettings;
+			}
+		}
+	}
+}
+
 bool FIKRetargetFKChainsOp::Initialize(
 	const FIKRetargetProcessor& InProcessor,
 	const FRetargetSkeleton& InSourceSkeleton,
@@ -729,27 +758,6 @@ FIKRetargetOpSettingsBase* FIKRetargetFKChainsOp::GetSettings()
 	return &Settings;
 }
 
-void FIKRetargetFKChainsOp::SetSettings(const FIKRetargetOpSettingsBase* InSettings)
-{
-	const FIKRetargetFKChainsOpSettings* NewSettings = reinterpret_cast<const FIKRetargetFKChainsOpSettings*>(InSettings);
-
-	// copies everything except the ChainsToRetarget array (those are copied below, only for already existing chains)
-	const TArray<FName> PropertiesToIgnore = {"ChainsToRetarget"};
-	CopySettingsRaw(InSettings, PropertiesToIgnore);
-	
-	// copy settings only for chains that the op has initialized
-	for (const FRetargetFKChainSettings& NewChainSettings : NewSettings->ChainsToRetarget)
-	{
-		for (FRetargetFKChainSettings& ChainSettings : Settings.ChainsToRetarget)
-		{
-			if (ChainSettings.TargetChainName == NewChainSettings.TargetChainName)
-			{
-				ChainSettings = NewChainSettings;
-			}
-		}
-	}
-}
-
 const UScriptStruct* FIKRetargetFKChainsOp::GetSettingsType() const
 {
 	return FIKRetargetFKChainsOpSettings::StaticStruct();
@@ -760,11 +768,6 @@ const UScriptStruct* FIKRetargetFKChainsOp::GetType() const
 	return FIKRetargetFKChainsOp::StaticStruct();
 }
 
-const UClass* FIKRetargetFKChainsOp::GetControllerType() const
-{
-	return UIKRetargetFKChainsController::StaticClass();
-}
-
 #if WITH_EDITOR
 
 FCriticalSection FIKRetargetFKChainsOp::DebugDataMutex;
@@ -915,12 +918,12 @@ void FIKRetargetFKChainsOp::ApplyIKRigs(const UIKRigDefinition* InSourceIKRig, c
 
 FIKRetargetFKChainsOpSettings UIKRetargetFKChainsController::GetSettings()
 {
-	return *reinterpret_cast<FIKRetargetFKChainsOpSettings*>(OpToControl->GetSettings());
+	return *reinterpret_cast<FIKRetargetFKChainsOpSettings*>(OpSettingsToControl);
 }
 
 void UIKRetargetFKChainsController::SetSettings(FIKRetargetFKChainsOpSettings InSettings)
 {
-	OpToControl->SetSettings(&InSettings);
+	OpSettingsToControl->CopySettingsAtRuntime(&InSettings);
 }
 
 #undef LOCTEXT_NAMESPACE
diff --git a/Engine/Plugins/Animation/IKRig/Source/IKRig/Private/Retargeter/RetargetOps/IKChainsOp.cpp b/Engine/Plugins/Animation/IKRig/Source/IKRig/Private/Retargeter/RetargetOps/IKChainsOp.cpp
--- a/Engine/Plugins/Animation/IKRig/Source/IKRig/Private/Retargeter/RetargetOps/IKChainsOp.cpp
+++ b/Engine/Plugins/Animation/IKRig/Source/IKRig/Private/Retargeter/RetargetOps/IKChainsOp.cpp
@@ -187,6 +187,35 @@ bool FRetargetIKChainSettings::operator==(const FRetargetIKChainSettings& Other)
 		&& FMath::IsNearlyEqualByULP(Extension, Other.Extension);
 }
 
+const UClass* FIKRetargetIKChainsOpSettings::GetControllerType() const
+{
+	return UIKRetargetIKChainsController::StaticClass();
+}
+
+void FIKRetargetIKChainsOpSettings::CopySettingsAtRuntime(const FIKRetargetOpSettingsBase* InSettingsToCopyFrom)
+{
+	// copies everything except the ChainsToRetarget array (those are copied below, only for already existing chains)
+	const TArray<FName> PropertiesToIgnore = {"ChainsToRetarget"};
+	FIKRetargetOpBase::CopyStructProperties(
+		FIKRetargetIKChainsOpSettings::StaticStruct(),
+		InSettingsToCopyFrom,
+		this,
+		PropertiesToIgnore);
+	
+	// copy ChainsToRetarget only for chains are in-common
+	const FIKRetargetIKChainsOpSettings* NewSettings = reinterpret_cast<const FIKRetargetIKChainsOpSettings*>(InSettingsToCopyFrom);
+	for (const FRetargetIKChainSettings& NewChainSettings : NewSettings->ChainsToRetarget)
+	{
+		for (FRetargetIKChainSettings& ChainSettings : ChainsToRetarget)
+		{
+			if (ChainSettings.TargetChainName == NewChainSettings.TargetChainName)
+			{
+				ChainSettings = NewChainSettings;
+			}
+		}
+	}
+}
+
 bool FIKRetargetIKChainsOp::Initialize(
 	const FIKRetargetProcessor& InProcessor,
 	const FRetargetSkeleton& InSourceSkeleton,
@@ -377,11 +406,6 @@ const UScriptStruct* FIKRetargetIKChainsOp::GetType() const
 	return FIKRetargetIKChainsOp::StaticStruct();
 }
 
-const UClass* FIKRetargetIKChainsOp::GetControllerType() const
-{
-	return UIKRetargetIKChainsController::StaticClass();
-}
-
 const UScriptStruct* FIKRetargetIKChainsOp::GetParentOpType() const
 {
 	return FIKRetargetRunIKRigOp::StaticStruct();
@@ -584,12 +608,12 @@ void FIKRetargetIKChainsOp::RegenerateChainSettings(const FIKRetargetOpBase* InP
 
 FIKRetargetIKChainsOpSettings UIKRetargetIKChainsController::GetSettings()
 {
-	return *reinterpret_cast<FIKRetargetIKChainsOpSettings*>(OpToControl->GetSettings());
+	return *reinterpret_cast<FIKRetargetIKChainsOpSettings*>(OpSettingsToControl);
 }
 
 void UIKRetargetIKChainsController::SetSettings(FIKRetargetIKChainsOpSettings InSettings)
 {
-	OpToControl->SetSettings(&InSettings);
+	OpSettingsToControl->CopySettingsAtRuntime(&InSettings);
 }
 
 #undef LOCTEXT_NAMESPACE
diff --git a/Engine/Plugins/Animation/IKRig/Source/IKRig/Private/Retargeter/RetargetOps/PelvisMotionOp.cpp b/Engine/Plugins/Animation/IKRig/Source/IKRig/Private/Retargeter/RetargetOps/PelvisMotionOp.cpp
--- a/Engine/Plugins/Animation/IKRig/Source/IKRig/Private/Retargeter/RetargetOps/PelvisMotionOp.cpp
+++ b/Engine/Plugins/Animation/IKRig/Source/IKRig/Private/Retargeter/RetargetOps/PelvisMotionOp.cpp
@@ -13,6 +13,22 @@
 #define LOCTEXT_NAMESPACE "PelvisMotionOp"
 
 
+const UClass* FIKRetargetPelvisMotionOpSettings::GetControllerType() const
+{
+	return UIKRetargetPelvisMotionController::StaticClass();
+}
+
+void FIKRetargetPelvisMotionOpSettings::CopySettingsAtRuntime(const FIKRetargetOpSettingsBase* InSettingsToCopyFrom)
+{
+	// copies everything except the bones we are operating on (those require reinit)
+	const TArray<FName> PropertiesToIgnore = {"SourcePelvisBone", "TargetPelvisBone"};
+	FIKRetargetOpBase::CopyStructProperties(
+		FIKRetargetPelvisMotionOpSettings::StaticStruct(),
+		InSettingsToCopyFrom,
+		this,
+		PropertiesToIgnore);
+}
+
 #if WITH_EDITOR
 USkeleton* FIKRetargetPelvisMotionOpSettings::GetSkeleton(const FName InPropertyName)
 {
@@ -174,13 +190,6 @@ FIKRetargetOpSettingsBase* FIKRetargetPelvisMotionOp::GetSettings()
 	return &Settings;
 }
 
-void FIKRetargetPelvisMotionOp::SetSettings(const FIKRetargetOpSettingsBase* InSettings)
-{
-	// copies everything except the bones we are operating on (those require reinit)
-	const TArray<FName> PropertiesToIgnore = {"SourcePelvisBone", "TargetPelvisBone"};
-	CopySettingsRaw(InSettings, PropertiesToIgnore);
-}
-
 const UScriptStruct* FIKRetargetPelvisMotionOp::GetSettingsType() const
 {
 	return FIKRetargetPelvisMotionOpSettings::StaticStruct();
@@ -191,11 +200,6 @@ const UScriptStruct* FIKRetargetPelvisMotionOp::GetType() const
 	return FIKRetargetPelvisMotionOp::StaticStruct();
 }
 
-const UClass* FIKRetargetPelvisMotionOp::GetControllerType() const
-{
-	return UIKRetargetPelvisMotionController::StaticClass();
-}
-
 void FIKRetargetPelvisMotionOp::CollectRetargetedBones(TSet<int32>& OutRetargetedBones) const
 {
 	// the pelvis bone is retargeted
@@ -345,36 +349,37 @@ void FIKRetargetPelvisMotionOp::DebugDraw(
 
 FIKRetargetPelvisMotionOpSettings UIKRetargetPelvisMotionController::GetSettings()
 {
-	return *reinterpret_cast<FIKRetargetPelvisMotionOpSettings*>(OpToControl->GetSettings());
+	return GetPelvisOpSettings();
 }
 
 void UIKRetargetPelvisMotionController::SetSettings(FIKRetargetPelvisMotionOpSettings InSettings)
 {
-	OpToControl->SetSettings(&InSettings);
+	OpSettingsToControl->CopySettingsAtRuntime(&InSettings);
 }
 
 void UIKRetargetPelvisMotionController::SetSourcePelvisBone(const FName InSourcePelvisBone)
 {
-	FIKRetargetPelvisMotionOp* Op = reinterpret_cast<FIKRetargetPelvisMotionOp*>(OpToControl);
-	Op->Settings.SourcePelvisBone.BoneName = InSourcePelvisBone;
+	GetPelvisOpSettings().SourcePelvisBone.BoneName = InSourcePelvisBone;
 }
 
 FName UIKRetargetPelvisMotionController::GetSourcePelvisBone()
 {
-	FIKRetargetPelvisMotionOp* Op = reinterpret_cast<FIKRetargetPelvisMotionOp*>(OpToControl);
-	return Op->Settings.SourcePelvisBone.BoneName;
+	return GetPelvisOpSettings().SourcePelvisBone.BoneName;
 }
 
 void UIKRetargetPelvisMotionController::SetTargetPelvisBone(const FName InTargetPelvisBone)
 {
-	FIKRetargetPelvisMotionOp* Op = reinterpret_cast<FIKRetargetPelvisMotionOp*>(OpToControl);
-	Op->Settings.TargetPelvisBone.BoneName = InTargetPelvisBone;
+	GetPelvisOpSettings().TargetPelvisBone.BoneName = InTargetPelvisBone;
 }
 
 FName UIKRetargetPelvisMotionController::GetTargetPelvisBone()
 {
-	FIKRetargetPelvisMotionOp* Op = reinterpret_cast<FIKRetargetPelvisMotionOp*>(OpToControl);
-	return Op->Settings.TargetPelvisBone.BoneName;
+	return GetPelvisOpSettings().TargetPelvisBone.BoneName;
+}
+
+FIKRetargetPelvisMotionOpSettings& UIKRetargetPelvisMotionController::GetPelvisOpSettings() const
+{
+	return *reinterpret_cast<FIKRetargetPelvisMotionOpSettings*>(OpSettingsToControl);
 }
 
 #undef LOCTEXT_NAMESPACE
diff --git a/Engine/Plugins/Animation/IKRig/Source/IKRig/Private/Retargeter/RetargetOps/PinBoneOp.cpp b/Engine/Plugins/Animation/IKRig/Source/IKRig/Private/Retargeter/RetargetOps/PinBoneOp.cpp
--- a/Engine/Plugins/Animation/IKRig/Source/IKRig/Private/Retargeter/RetargetOps/PinBoneOp.cpp
+++ b/Engine/Plugins/Animation/IKRig/Source/IKRig/Private/Retargeter/RetargetOps/PinBoneOp.cpp
@@ -60,8 +60,23 @@ void FPinBoneData::CachePinData(ERetargetSourceOrTarget InSkeletonToCopyFrom, co
     LocalRefPoseBoneToCopyTo = BoneToCopyToParentRefPoseGlobal.GetRelativeTransform(BoneToCopyToRefPoseGlobal);
 }
 
-#if WITH_EDITOR
+const UClass* FIKRetargetPinBoneOpSettings::GetControllerType() const
+{
+	return UIKRetargetPinBoneController::StaticClass();
+}
+
+void FIKRetargetPinBoneOpSettings::CopySettingsAtRuntime(const FIKRetargetOpSettingsBase* InSettingsToCopyFrom)
+{
+	// copies everything except the bones we are operating on (those require reinit)
+	const TArray<FName> PropertiesToIgnore = {"BonesToPin"};
+	FIKRetargetOpBase::CopyStructProperties(
+		FIKRetargetPinBoneOpSettings::StaticStruct(),
+		InSettingsToCopyFrom,
+		this,
+		PropertiesToIgnore);
+}
 
+#if WITH_EDITOR
 USkeleton* FIKRetargetPinBoneOpSettings::GetSkeleton(const FName InPropertyName)
 {
 	if (!SourceSkeletonAsset || !TargetSkeletonAsset)
@@ -198,13 +213,6 @@ FIKRetargetOpSettingsBase* FIKRetargetPinBoneOp::GetSettings()
 	return &Settings;
 }
 
-void FIKRetargetPinBoneOp::SetSettings(const FIKRetargetOpSettingsBase* InSettings)
-{
-	// copies everything except the bones we are operating on (those require reinit)
-	const TArray<FName> PropertiesToIgnore = {"BonesToPin"};
-	CopySettingsRaw(InSettings, PropertiesToIgnore);
-}
-
 const UScriptStruct* FIKRetargetPinBoneOp::GetSettingsType() const
 {
 	return FIKRetargetPinBoneOpSettings::StaticStruct();
@@ -215,11 +223,6 @@ const UScriptStruct* FIKRetargetPinBoneOp::GetType() const
 	return FIKRetargetPinBoneOp::StaticStruct();
 }
 
-const UClass* FIKRetargetPinBoneOp::GetControllerType() const
-{
-	return UIKRetargetPinBoneController::StaticClass();
-}
-
 #if WITH_EDITOR
 FText FIKRetargetPinBoneOp::GetWarningMessage() const
 {
@@ -413,26 +416,26 @@ PRAGMA_ENABLE_DEPRECATION_WARNINGS
 
 FIKRetargetPinBoneOpSettings UIKRetargetPinBoneController::GetSettings()
 {
-	return *reinterpret_cast<FIKRetargetPinBoneOpSettings*>(OpToControl->GetSettings());
+	return *reinterpret_cast<FIKRetargetPinBoneOpSettings*>(OpSettingsToControl);
 }
 
 void UIKRetargetPinBoneController::SetSettings(FIKRetargetPinBoneOpSettings InSettings)
 {
-	OpToControl->SetSettings(&InSettings);
+	OpSettingsToControl->CopySettingsAtRuntime(&InSettings);
 }
 
 void UIKRetargetPinBoneController::ClearAllBonePairs()
 {
-	FIKRetargetPinBoneOp* Op = reinterpret_cast<FIKRetargetPinBoneOp*>(OpToControl);
-	Op->Settings.BonesToPin.Reset();
+	FIKRetargetPinBoneOpSettings* Settings = reinterpret_cast<FIKRetargetPinBoneOpSettings*>(OpSettingsToControl);
+	Settings->BonesToPin.Reset();
 }
 
 void UIKRetargetPinBoneController::SetBonePair(const FName InBoneToCopyFrom, const FName InBoneToCopyTo)
 {
-	FIKRetargetPinBoneOp* Op = reinterpret_cast<FIKRetargetPinBoneOp*>(OpToControl);
+	FIKRetargetPinBoneOpSettings* Settings = reinterpret_cast<FIKRetargetPinBoneOpSettings*>(OpSettingsToControl);
 
 	// update existing pair with new bone to copy from (if there is one)
-	for (FPinBoneData& BonePair : Op->Settings.BonesToPin)
+	for (FPinBoneData& BonePair : Settings->BonesToPin)
 	{
 		if (BonePair.BoneToCopyTo == InBoneToCopyTo)
 		{
@@ -445,15 +448,15 @@ void UIKRetargetPinBoneController::SetBonePair(const FName InBoneToCopyFrom, con
 	FPinBoneData NewBonePair;
 	NewBonePair.BoneToCopyFrom.BoneName = InBoneToCopyFrom;
 	NewBonePair.BoneToCopyTo.BoneName = InBoneToCopyTo;
-	Op->Settings.BonesToPin.Add(NewBonePair);
+	Settings->BonesToPin.Add(NewBonePair);
 }
 
 TMap<FName, FName> UIKRetargetPinBoneController::GetAllBonePairs()
 {
 	TMap<FName, FName> AllBonePairs;
 	
-	FIKRetargetPinBoneOp* Op = reinterpret_cast<FIKRetargetPinBoneOp*>(OpToControl);
-	for (FPinBoneData& BonePair : Op->Settings.BonesToPin)
+	FIKRetargetPinBoneOpSettings* Settings = reinterpret_cast<FIKRetargetPinBoneOpSettings*>(OpSettingsToControl);
+	for (FPinBoneData& BonePair : Settings->BonesToPin)
 	{
 		AllBonePairs.Add(BonePair.BoneToCopyTo.BoneName, BonePair.BoneToCopyFrom.BoneName);
 	}
diff --git a/Engine/Plugins/Animation/IKRig/Source/IKRig/Private/Retargeter/RetargetOps/RootMotionGeneratorOp.cpp b/Engine/Plugins/Animation/IKRig/Source/IKRig/Private/Retargeter/RetargetOps/RootMotionGeneratorOp.cpp
--- a/Engine/Plugins/Animation/IKRig/Source/IKRig/Private/Retargeter/RetargetOps/RootMotionGeneratorOp.cpp
+++ b/Engine/Plugins/Animation/IKRig/Source/IKRig/Private/Retargeter/RetargetOps/RootMotionGeneratorOp.cpp
@@ -15,6 +15,22 @@
 IMPLEMENT_HIT_PROXY(HIKRetargetEditorRootProxy, HHitProxy);
 #endif
 
+const UClass* FIKRetargetRootMotionOpSettings::GetControllerType() const
+{
+	return UIKRetargetRootMotionController::StaticClass();
+}
+
+void FIKRetargetRootMotionOpSettings::CopySettingsAtRuntime(const FIKRetargetOpSettingsBase* InSettingsToCopyFrom)
+{
+	// copies everything except the bones we are operating on (those require reinit)
+	const TArray<FName> PropertiesToIgnore = {"SourceRootBone", "TargetRootBone", "TargetPelvisBone"};
+	FIKRetargetOpBase::CopyStructProperties(
+		FIKRetargetRootMotionOpSettings::StaticStruct(),
+		InSettingsToCopyFrom,
+		this,
+		PropertiesToIgnore);
+}
+
 #if WITH_EDITOR
 USkeleton* FIKRetargetRootMotionOpSettings::GetSkeleton(const FName InPropertyName)
 {
@@ -199,13 +215,6 @@ FIKRetargetOpSettingsBase* FIKRetargetRootMotionOp::GetSettings()
 	return &Settings;
 }
 
-void FIKRetargetRootMotionOp::SetSettings(const FIKRetargetOpSettingsBase* InSettings)
-{
-	// copies everything except the bones we are operating on (those require reinit)
-	const TArray<FName> PropertiesToIgnore = {"SourceRootBone", "TargetRootBone", "TargetPelvisBone"};
-	CopySettingsRaw(InSettings, PropertiesToIgnore);
-}
-
 const UScriptStruct* FIKRetargetRootMotionOp::GetSettingsType() const
 {
 	return FIKRetargetRootMotionOpSettings::StaticStruct();
@@ -216,11 +225,6 @@ const UScriptStruct* FIKRetargetRootMotionOp::GetType() const
 	return FIKRetargetRootMotionOp::StaticStruct();
 }
 
-const UClass* FIKRetargetRootMotionOp::GetControllerType() const
-{
-	return UIKRetargetRootMotionController::StaticClass();
-}
-
 void FIKRetargetRootMotionOp::Reset()
 {
 	SourceRootIndex = INDEX_NONE;
@@ -338,36 +342,47 @@ void URootMotionGeneratorOp::ConvertToInstancedStruct(FInstancedStruct& OutInsta
 
 FIKRetargetRootMotionOpSettings UIKRetargetRootMotionController::GetSettings()
 {
-	return *reinterpret_cast<FIKRetargetRootMotionOpSettings*>(OpToControl->GetSettings());
+	return *GetSettingsPtr();
 }
 
 void UIKRetargetRootMotionController::SetSettings(FIKRetargetRootMotionOpSettings InSettings)
 {
-	OpToControl->SetSettings(&InSettings);
+	OpSettingsToControl->CopySettingsAtRuntime(&InSettings);
 }
 
 void UIKRetargetRootMotionController::SetSourceRootBone(const FName InSourceRootBone)
 {
-	FIKRetargetRootMotionOp* Op = reinterpret_cast<FIKRetargetRootMotionOp*>(OpToControl);
-	Op->Settings.SourceRoot.BoneName = InSourceRootBone;
+	GetSettingsPtr()->SourceRoot.BoneName = InSourceRootBone;
 }
 
 FName UIKRetargetRootMotionController::GetSourceRootBone()
 {
-	FIKRetargetRootMotionOp* Op = reinterpret_cast<FIKRetargetRootMotionOp*>(OpToControl);
-	return Op->Settings.SourceRoot.BoneName;
+	return GetSettingsPtr()->SourceRoot.BoneName;
 }
 
 void UIKRetargetRootMotionController::SetTargetRootBone(const FName InTargetRootBone)
 {
-	FIKRetargetRootMotionOp* Op = reinterpret_cast<FIKRetargetRootMotionOp*>(OpToControl);
-	Op->Settings.TargetRoot.BoneName = InTargetRootBone;
+	GetSettingsPtr()->TargetRoot.BoneName = InTargetRootBone;
 }
 
 FName UIKRetargetRootMotionController::GetTargetRootBone()
 {
-	FIKRetargetRootMotionOp* Op = reinterpret_cast<FIKRetargetRootMotionOp*>(OpToControl);
-	return Op->Settings.TargetRoot.BoneName;
+	return GetSettingsPtr()->TargetRoot.BoneName;
+}
+
+void UIKRetargetRootMotionController::SetTargetPelvisBone(const FName InTargetPelvisBone)
+{
+	GetSettingsPtr()->TargetPelvis.BoneName = InTargetPelvisBone;
+}
+
+FName UIKRetargetRootMotionController::GetTargetPelvisBone()
+{
+	return GetSettingsPtr()->TargetPelvis.BoneName;
+}
+
+FIKRetargetRootMotionOpSettings* UIKRetargetRootMotionController::GetSettingsPtr() const
+{
+	return reinterpret_cast<FIKRetargetRootMotionOpSettings*>(OpSettingsToControl);
 }
 
 #undef LOCTEXT_NAMESPACE
diff --git a/Engine/Plugins/Animation/IKRig/Source/IKRig/Private/Retargeter/RetargetOps/RunIKRigOp.cpp b/Engine/Plugins/Animation/IKRig/Source/IKRig/Private/Retargeter/RetargetOps/RunIKRigOp.cpp
--- a/Engine/Plugins/Animation/IKRig/Source/IKRig/Private/Retargeter/RetargetOps/RunIKRigOp.cpp
+++ b/Engine/Plugins/Animation/IKRig/Source/IKRig/Private/Retargeter/RetargetOps/RunIKRigOp.cpp
@@ -8,6 +8,22 @@
 
 #define LOCTEXT_NAMESPACE "RunIKRigSolversOp"
 
+const UClass* FIKRetargetRunIKRigOpSettings::GetControllerType() const
+{
+	return UIKRetargetRunIKRigController::StaticClass();
+}
+
+void FIKRetargetRunIKRigOpSettings::CopySettingsAtRuntime(const FIKRetargetOpSettingsBase* InSettingsToCopyFrom)
+{
+	// copies everything except "IKRigAsset"
+	const TArray<FName> PropertiesToIgnore = {"IKRigAsset"};
+	FIKRetargetOpBase::CopyStructProperties(
+		FIKRetargetRunIKRigOpSettings::StaticStruct(),
+		InSettingsToCopyFrom,
+		this,
+		PropertiesToIgnore);
+}
+
 bool FIKRetargetRunIKRigOp::Initialize(
 	const FIKRetargetProcessor& InProcessor,
 	const FRetargetSkeleton& InSourceSkeleton,
@@ -104,13 +120,6 @@ FIKRetargetOpSettingsBase* FIKRetargetRunIKRigOp::GetSettings()
 	return &Settings;
 }
 
-void FIKRetargetRunIKRigOp::SetSettings(const FIKRetargetOpSettingsBase* InSettings)
-{
-	// copies everything except "IKRigAsset"
-	const TArray<FName> PropertiesToIgnore = {"IKRigAsset"};
-	CopySettingsRaw(InSettings, PropertiesToIgnore);
-}
-
 const UScriptStruct* FIKRetargetRunIKRigOp::GetSettingsType() const
 {
 	return FIKRetargetRunIKRigOpSettings::StaticStruct();
@@ -121,11 +130,6 @@ const UScriptStruct* FIKRetargetRunIKRigOp::GetType() const
 	return FIKRetargetRunIKRigOp::StaticStruct();
 }
 
-const UClass* FIKRetargetRunIKRigOp::GetControllerType() const
-{
-	return UIKRetargetRunIKRigController::StaticClass();
-}
-
 const UIKRigDefinition* FIKRetargetRunIKRigOp::GetCustomTargetIKRig() const
 {
 	return Settings.IKRigAsset;
@@ -279,12 +283,12 @@ TArray<FName> FIKRetargetRunIKRigOp::GetRequiredTargetChains() const
 
 FIKRetargetRunIKRigOpSettings UIKRetargetRunIKRigController::GetSettings()
 {
-	return *reinterpret_cast<FIKRetargetRunIKRigOpSettings*>(OpToControl->GetSettings());
+	return *reinterpret_cast<FIKRetargetRunIKRigOpSettings*>(OpSettingsToControl);
 }
 
 void UIKRetargetRunIKRigController::SetSettings(FIKRetargetRunIKRigOpSettings InSettings)
 {
-	OpToControl->SetSettings(&InSettings);
+	OpSettingsToControl->CopySettingsAtRuntime(&InSettings);
 }
 
 #undef LOCTEXT_NAMESPACE
diff --git a/Engine/Plugins/Animation/IKRig/Source/IKRig/Private/Retargeter/RetargetOps/ScaleSourceOp.cpp b/Engine/Plugins/Animation/IKRig/Source/IKRig/Private/Retargeter/RetargetOps/ScaleSourceOp.cpp
--- a/Engine/Plugins/Animation/IKRig/Source/IKRig/Private/Retargeter/RetargetOps/ScaleSourceOp.cpp
+++ b/Engine/Plugins/Animation/IKRig/Source/IKRig/Private/Retargeter/RetargetOps/ScaleSourceOp.cpp
@@ -2,6 +2,22 @@
 
 #include "Retargeter/RetargetOps/ScaleSourceOp.h"
 
+const UClass* FIKRetargetScaleSourceOpSettings::GetControllerType() const
+{
+	return UIKRetargetScaleSourceController::StaticClass();
+}
+
+void FIKRetargetScaleSourceOpSettings::CopySettingsAtRuntime(const FIKRetargetOpSettingsBase* InSettingsToCopyFrom)
+{
+	// copy ALL properties
+	static TArray<FName> PropertiesToIgnore;
+	FIKRetargetOpBase::CopyStructProperties(
+		FIKRetargetScaleSourceOpSettings::StaticStruct(),
+		InSettingsToCopyFrom,
+		this,
+		PropertiesToIgnore);
+}
+
 bool FIKRetargetScaleSourceOp::Initialize(
 	const FIKRetargetProcessor& InProcessor,
 	const FRetargetSkeleton& InSourceSkeleton,
@@ -18,11 +34,6 @@ FIKRetargetOpSettingsBase* FIKRetargetScaleSourceOp::GetSettings()
 	return &Settings;
 }
 
-void FIKRetargetScaleSourceOp::SetSettings(const FIKRetargetOpSettingsBase* InSettings)
-{
-	Settings = *reinterpret_cast<const FIKRetargetScaleSourceOpSettings*>(InSettings);
-}
-
 const UScriptStruct* FIKRetargetScaleSourceOp::GetSettingsType() const
 {
 	return FIKRetargetScaleSourceOpSettings::StaticStruct();
@@ -33,17 +44,12 @@ const UScriptStruct* FIKRetargetScaleSourceOp::GetType() const
 	return FIKRetargetScaleSourceOp::StaticStruct();
 }
 
-const UClass* FIKRetargetScaleSourceOp::GetControllerType() const
-{
-	return UIKRetargetScaleSourceController::StaticClass();
-}
-
 FIKRetargetScaleSourceOpSettings UIKRetargetScaleSourceController::GetSettings()
 {
-	return *reinterpret_cast<FIKRetargetScaleSourceOpSettings*>(OpToControl->GetSettings());
+	return *reinterpret_cast<FIKRetargetScaleSourceOpSettings*>(OpSettingsToControl);
 }
 
 void UIKRetargetScaleSourceController::SetSettings(FIKRetargetScaleSourceOpSettings InSettings)
 {
-	OpToControl->SetSettings(&InSettings);
+	OpSettingsToControl->CopySettingsAtRuntime(&InSettings);
 }
diff --git a/Engine/Plugins/Animation/IKRig/Source/IKRig/Private/Retargeter/RetargetOps/SpeedPlantingOp.cpp b/Engine/Plugins/Animation/IKRig/Source/IKRig/Private/Retargeter/RetargetOps/SpeedPlantingOp.cpp
--- a/Engine/Plugins/Animation/IKRig/Source/IKRig/Private/Retargeter/RetargetOps/SpeedPlantingOp.cpp
+++ b/Engine/Plugins/Animation/IKRig/Source/IKRig/Private/Retargeter/RetargetOps/SpeedPlantingOp.cpp
@@ -10,6 +10,35 @@
 
 #define LOCTEXT_NAMESPACE "SpeedPlantingOp"
 
+const UClass* FIKRetargetSpeedPlantingOpSettings::GetControllerType() const
+{
+	return UIKRetargetSpeedPlantingController::StaticClass();
+}
+
+void FIKRetargetSpeedPlantingOpSettings::CopySettingsAtRuntime(const FIKRetargetOpSettingsBase* InSettingsToCopyFrom)
+{
+	// copies everything except "ChainsToSpeedPlant"
+	static TArray<FName> PropertiesToIgnore = {"ChainsToSpeedPlant"};
+	FIKRetargetOpBase::CopyStructProperties(
+		FIKRetargetSpeedPlantingOpSettings::StaticStruct(),
+		InSettingsToCopyFrom,
+		this,
+		PropertiesToIgnore);
+	
+	// copy settings only for chains that the op has initialized
+	const FIKRetargetSpeedPlantingOpSettings* NewSettings = reinterpret_cast<const FIKRetargetSpeedPlantingOpSettings*>(InSettingsToCopyFrom);
+	for (const FRetargetSpeedPlantingSettings& NewChainSettings : NewSettings->ChainsToSpeedPlant)
+	{
+		for (FRetargetSpeedPlantingSettings& ChainSettings : ChainsToSpeedPlant)
+		{
+			if (ChainSettings.TargetChainName == NewChainSettings.TargetChainName)
+			{
+				ChainSettings = NewChainSettings;
+			}
+		}
+	}
+}
+
 bool FIKRetargetSpeedPlantingOp::Initialize(
 	const FIKRetargetProcessor& InProcessor,
 	const FRetargetSkeleton& InSourceSkeleton,
@@ -128,27 +157,6 @@ FIKRetargetOpSettingsBase* FIKRetargetSpeedPlantingOp::GetSettings()
 	return &Settings;
 }
 
-void FIKRetargetSpeedPlantingOp::SetSettings(const FIKRetargetOpSettingsBase* InSettings)
-{
-	const FIKRetargetSpeedPlantingOpSettings* NewSettings = reinterpret_cast<const FIKRetargetSpeedPlantingOpSettings*>(InSettings);
-	
-	// copies everything except "ChainsToSpeedPlant"
-	const TArray<FName> PropertiesToIgnore = {"ChainsToSpeedPlant"};
-	CopySettingsRaw(InSettings, PropertiesToIgnore);
-	
-	// copy settings only for chains that the op has initialized
-	for (const FRetargetSpeedPlantingSettings& NewChainSettings : NewSettings->ChainsToSpeedPlant)
-	{
-		for (FRetargetSpeedPlantingSettings& ChainSettings : Settings.ChainsToSpeedPlant)
-		{
-			if (ChainSettings.TargetChainName == NewChainSettings.TargetChainName)
-			{
-				ChainSettings = NewChainSettings;
-			}
-		}
-	}
-}
-
 const UScriptStruct* FIKRetargetSpeedPlantingOp::GetSettingsType() const
 {
 	return FIKRetargetSpeedPlantingOpSettings::StaticStruct();
@@ -159,11 +167,6 @@ const UScriptStruct* FIKRetargetSpeedPlantingOp::GetType() const
 	return FIKRetargetSpeedPlantingOp::StaticStruct();
 }
 
-const UClass* FIKRetargetSpeedPlantingOp::GetControllerType() const
-{
-	return UIKRetargetSpeedPlantingController::StaticClass();
-}
-
 const UScriptStruct* FIKRetargetSpeedPlantingOp::GetParentOpType() const
 {
 	return FIKRetargetRunIKRigOp::StaticStruct();
@@ -323,12 +326,12 @@ void FIKRetargetSpeedPlantingOp::RegenerateChainSettings(const FIKRetargetOpBase
 
 FIKRetargetSpeedPlantingOpSettings UIKRetargetSpeedPlantingController::GetSettings()
 {
-	return *reinterpret_cast<FIKRetargetSpeedPlantingOpSettings*>(OpToControl->GetSettings());
+	return *reinterpret_cast<FIKRetargetSpeedPlantingOpSettings*>(OpSettingsToControl);
 }
 
 void UIKRetargetSpeedPlantingController::SetSettings(FIKRetargetSpeedPlantingOpSettings InSettings)
 {
-	OpToControl->SetSettings(&InSettings);
+	OpSettingsToControl->CopySettingsAtRuntime(&InSettings);
 }
 
 #undef LOCTEXT_NAMESPACE
\ No newline at end of file
diff --git a/Engine/Plugins/Animation/IKRig/Source/IKRig/Private/Retargeter/RetargetOps/StrideWarpingOp.cpp b/Engine/Plugins/Animation/IKRig/Source/IKRig/Private/Retargeter/RetargetOps/StrideWarpingOp.cpp
--- a/Engine/Plugins/Animation/IKRig/Source/IKRig/Private/Retargeter/RetargetOps/StrideWarpingOp.cpp
+++ b/Engine/Plugins/Animation/IKRig/Source/IKRig/Private/Retargeter/RetargetOps/StrideWarpingOp.cpp
@@ -20,6 +20,22 @@ bool FRetargetStrideWarpChainSettings::operator==(const FRetargetStrideWarpChain
 	return EnableStrideWarping == Other.EnableStrideWarping;
 }
 
+const UClass* FIKRetargetStrideWarpingOpSettings::GetControllerType() const
+{
+	return UIKRetargetStrideWarpingController::StaticClass();
+}
+
+void FIKRetargetStrideWarpingOpSettings::CopySettingsAtRuntime(const FIKRetargetOpSettingsBase* InSettingsToCopyFrom)
+{
+	// copies everything except "ChainsToStrideWarp"
+	static TArray<FName> PropertiesToIgnore = {"ChainsToStrideWarp"};
+	FIKRetargetOpBase::CopyStructProperties(
+		FIKRetargetStrideWarpingOpSettings::StaticStruct(),
+		InSettingsToCopyFrom,
+		this,
+		PropertiesToIgnore);
+}
+
 FVector FIKRetargetStrideWarpingOpSettings::GetAxisVector(const EBasicAxis& Axis)
 {
 	switch (Axis)
@@ -332,13 +348,6 @@ FIKRetargetOpSettingsBase* FIKRetargetStrideWarpingOp::GetSettings()
 	return &Settings;
 }
 
-void FIKRetargetStrideWarpingOp::SetSettings(const FIKRetargetOpSettingsBase* InSettings)
-{
-	// copies everything except "ChainsToStrideWarp"
-	const TArray<FName> PropertiesToIgnore = {"ChainsToStrideWarp"};
-	CopySettingsRaw(InSettings, PropertiesToIgnore);
-}
-
 const UScriptStruct* FIKRetargetStrideWarpingOp::GetSettingsType() const
 {
 	return FIKRetargetStrideWarpingOpSettings::StaticStruct();
@@ -349,11 +358,6 @@ const UScriptStruct* FIKRetargetStrideWarpingOp::GetType() const
 	return FIKRetargetStrideWarpingOp::StaticStruct();
 }
 
-const UClass* FIKRetargetStrideWarpingOp::GetControllerType() const
-{
-	return UIKRetargetStrideWarpingController::StaticClass();
-}
-
 const UScriptStruct* FIKRetargetStrideWarpingOp::GetParentOpType() const
 {
 	return FIKRetargetRunIKRigOp::StaticStruct();
@@ -481,12 +485,12 @@ void FIKRetargetStrideWarpingOp::RegenerateChainSettings(const FIKRetargetOpBase
 
 FIKRetargetStrideWarpingOpSettings UIKRetargetStrideWarpingController::GetSettings()
 {
-	return *reinterpret_cast<FIKRetargetStrideWarpingOpSettings*>(OpToControl->GetSettings());
+	return *reinterpret_cast<FIKRetargetStrideWarpingOpSettings*>(OpSettingsToControl);
 }
 
 void UIKRetargetStrideWarpingController::SetSettings(FIKRetargetStrideWarpingOpSettings InSettings)
 {
-	OpToControl->SetSettings(&InSettings);
+	OpSettingsToControl->CopySettingsAtRuntime(&InSettings);
 }
 
 #undef LOCTEXT_NAMESPACE
\ No newline at end of file
diff --git a/Engine/Plugins/Animation/IKRig/Source/IKRig/Public/Retargeter/IKRetargetOps.h b/Engine/Plugins/Animation/IKRig/Source/IKRig/Public/Retargeter/IKRetargetOps.h
--- a/Engine/Plugins/Animation/IKRig/Source/IKRig/Public/Retargeter/IKRetargetOps.h
+++ b/Engine/Plugins/Animation/IKRig/Source/IKRig/Public/Retargeter/IKRetargetOps.h
@@ -41,7 +41,27 @@ struct IKRIG_API FIKRetargetOpSettingsBase
 	
 	virtual ~FIKRetargetOpSettingsBase() {};
 
+	/** (required) override to specify how settings should be applied in a way that will not require reinitialization (ie runtime compatible)*/
+	virtual void CopySettingsAtRuntime(const FIKRetargetOpSettingsBase* InSettingsToCopyFrom) { checkNoEntry(); };
+	
+	/** (optional, but recommended) provide a custom controller type (deriving from UIKRetargetOpControllerBase) as an API for editing the op
+	 * NOTE: this type will automatically be lazy instantiated when/if the user calls GetController() on the op. */
+	virtual const UClass* GetControllerType() const;
+	
+	/** provide a scripting object to edit your custom settings type in the editor via BP/Python
+	* NOTE: this returns an instance of whatever type is returned by GetControllerType(). This can be custom per op. */
+	UIKRetargetOpControllerBase* GetController(UObject* Outer);
+
+private:
+
+	// helper function for settings structs to instantiate their own controller
+	UIKRetargetOpControllerBase* CreateControllerIfNeeded(UObject* Outer);
+	
+	// the controller used to edit this op by script/blueprint (lazy instantiated when needed)
+	TStrongObjectPtr<UIKRetargetOpControllerBase> Controller = nullptr;
+
 #if WITH_EDITORONLY_DATA
+public:
 	/** the op these settings belong to */
 	UPROPERTY(transient)
 	FName OwningOpName;
@@ -109,11 +129,11 @@ public:
 	/** (required) override and return a pointer to the settings struct used by this operation */ 
 	virtual FIKRetargetOpSettingsBase* GetSettings() { checkNoEntry(); return nullptr; };
 	
-	/** (optional) override to customize which settings are copied into the op at runtime
+	/** (optional) override to react when settings are applied at runtime
 	 * NOTE:
 	 * This is called while the Op is running AFTER Initialize() so it up to the Op author
 	 * to copy only those settings which are safe to be updated while the op is running.
-	 * The default implementation copies ALL settings all the time.
+	 * The default implementation uses CopySettingsAtRuntime() which must be implemented on the settings type.
 	 * @param InSettings - pointer to the settings to copy, can be safely cast to the type returned by GetSettingsType() */ 
 	virtual void SetSettings(const FIKRetargetOpSettingsBase* InSettings);
 	
@@ -137,14 +157,6 @@ public:
 	/** (optional) override and add any bones that your op modifies to the output TSet of bone indices
 	 * Add indices of any bone that this op modifies. Any bone not registered here will be FK parented by other operations*/
 	virtual void CollectRetargetedBones(TSet<int32>& OutRetargetedBones) const {};
-	
-	/** provide a scripting object to edit your custom solver type in the editor via BP/Python
-	* NOTE: this returns an instance of whatever type is returned by GetControllerType(). This can be custom per op. */
-	UIKRetargetOpControllerBase* GetController(UObject* Outer);
-
-	/** (optional, but recommended) provide a custom controller type (deriving from UIKRetargetOpControllerBase) as an API for editing the op
-	 * NOTE: this type will automatically be lazy instantiated when/if the user calls GetController() on the op. */
-	virtual const UClass* GetControllerType() const;
 
 	/** (optional) ops can optionally behave as a 'parent' where child ops must be executed first. */
 	virtual bool CanHaveChildOps() const { return false; };
@@ -219,6 +231,13 @@ public:
 	 * NOTE: be careful calling this on an initialized op as it may invalidate it's runtime state (use SetSettings instead) */
 	void CopySettingsRaw(const FIKRetargetOpSettingsBase* InSettings, const TArray<FName>& InPropertiesToIgnore);
 
+	/** a convenience function to copy all properties from one struct to another while ignoring some */
+	static void CopyStructProperties(
+		const UStruct* InStructType,
+		const void* InSrcStruct,
+		void* InOutDestStruct,
+		const TArray<FName>& InPropertiesToIgnore);
+
 private:
 	
 	/** when false, execution of this op is skipped */
@@ -234,19 +253,6 @@ private:
 	FName ParentOpName;
 
 protected:
-
-	/** a convenience function to copy all properties from one struct to another while ignoring some */
-	static void CopyStructProperties(
-		const UStruct* InStructType,
-		const void* InSrcStruct,
-		void* InOutDestStruct,
-		const TArray<FName>& InPropertiesToIgnore);
-
-	// helper function for ops to instantiate their own controller
-	UIKRetargetOpControllerBase* CreateControllerIfNeeded(UObject* Outer);
-
-	// the  controller used to edit this solver by script/blueprint (lazy instantiated when needed)
-	TStrongObjectPtr<UIKRetargetOpControllerBase> Controller = nullptr;
 	
 	/** set to true when the op is ready to run */
 	bool bIsInitialized = false;
@@ -269,19 +275,9 @@ class IKRIG_API UIKRetargetOpControllerBase : public UObject
 	GENERATED_BODY()
 	
 public:
-
-	/* Turn the op on/off in the retargeting stack.
-	 * @param bIsEnabled if true, the op will execute, false otherwise. */
-	UFUNCTION(BlueprintCallable, Category = RetargetOp)
-	void SetEnabled(bool bIsEnabled) { OpToControl->SetEnabled(bIsEnabled); };
-
-	/* Get if the op is on or off in the retargeting stack.
-	 * @param bIsEnabled returns true if the op is enabled, false otherwise. */
-	UFUNCTION(BlueprintCallable, Category = RetargetOp)
-	bool GetEnabled() { return OpToControl->IsEnabled(); };
 	
-	// the op this controller controls
-	FIKRetargetOpBase* OpToControl = nullptr;
+	// the op settings this controller controls
+	FIKRetargetOpSettingsBase* OpSettingsToControl = nullptr;
 };
 
 
diff --git a/Engine/Plugins/Animation/IKRig/Source/IKRig/Public/Retargeter/IKRetargetProfile.h b/Engine/Plugins/Animation/IKRig/Source/IKRig/Public/Retargeter/IKRetargetProfile.h
--- a/Engine/Plugins/Animation/IKRig/Source/IKRig/Public/Retargeter/IKRetargetProfile.h
+++ b/Engine/Plugins/Animation/IKRig/Source/IKRig/Public/Retargeter/IKRetargetProfile.h
@@ -3,6 +3,7 @@
 #pragma once
 
 #include "CoreMinimal.h"
+#include "Kismet/BlueprintFunctionLibrary.h"
 #include "IKRetargetOps.h"
 #include "Retargeter/IKRetargetDeprecated.h"
 #include "StructUtils/InstancedStruct.h"
@@ -39,6 +40,9 @@ struct FRetargetOpProfile
 	 * @param InApplyContext - in 'runtime' mode, will let op apply the settings itself to avoid changes that require reinitialization
 	 * @return true if the supplied op was compatible with the settings, false otherwise */
 	bool CopySettingsToOp(FInstancedStruct& InOutOpStruct, ECopyOpSettingsContext InApplyContext) const;
+
+	/** helper function for ops to instantiate their own controller */
+	UIKRetargetOpControllerBase* CreateControllerIfNeeded(UObject* Outer);
 	
 	/** The name of the op these settings should be applied to.
 	 * If left at None, will apply to all ops that use the type of settings stored in 'SettingsToApply' */
@@ -51,13 +55,8 @@ struct FRetargetOpProfile
 	UPROPERTY(EditAnywhere, BlueprintReadWrite, Category = "IK Retarget Op Profile", meta = (ExcludeBaseStruct, BaseStruct = "/Script/IKRig.IKRetargetOpSettingsBase"))
 	FInstancedStruct SettingsToApply;
 
-	// TODO store sparse set of property overrides
-	// Is it possible we could do this with FInstancedPropertyBag?
-	// Perhaps we should store a bool mask with the same shape as the settings struct
-	// (ie with one bool per-property indicating if the profile should override it)
-	// Regardless, it will take some major UI work to expose this in a useful way.
-	// UPROPERTY(EditDefaultsOnly, Category = "IK Retarget Op Profile")
-	// FInstancedPropertyBag TestPropertyBag;
+	// the controller used to edit these op settings by script/blueprint (lazy instantiated when needed)
+	TStrongObjectPtr<UIKRetargetOpControllerBase> Controller = nullptr;
 };
 
 USTRUCT(BlueprintType)
@@ -109,6 +108,8 @@ public:
 	FRetargetOpProfile* FindMatchingOpProfile(const FRetargetOpProfile& OtherOpProfile);
 
 	bool ApplyOpProfilesToOpStruct(FInstancedStruct& InOutOpStruct, const ECopyOpSettingsContext InCopyContext) const;
+
+	FRetargetOpProfile* GetOpProfileByName(const FName InOpName);
 	
 	// search all the op profiles and get an array of all the settings of a particular type
 	template <typename T>
@@ -166,3 +167,17 @@ public:
 	// END DEPRECATED API
 	//
 };
+
+UCLASS()
+class URetargetProfileLibrary : public UBlueprintFunctionLibrary
+{
+	GENERATED_BODY()
+
+public:
+	
+	UFUNCTION(BlueprintCallable, Category = "Retarget Profile")
+	static FRetargetProfile CopyRetargetProfileFromRetargetAsset(const UIKRetargeter* InRetargetAsset);
+	
+	UFUNCTION(BlueprintPure, Category = "Retarget Profile")
+	static UIKRetargetOpControllerBase* GetOpControllerFromRetargetProfile(UPARAM(ref) FRetargetProfile& InRetargetProfile, const FName InRetargetOpName);
+};
\ No newline at end of file
diff --git a/Engine/Plugins/Animation/IKRig/Source/IKRig/Public/Retargeter/RetargetOps/AlignPoleVectorOp.h b/Engine/Plugins/Animation/IKRig/Source/IKRig/Public/Retargeter/RetargetOps/AlignPoleVectorOp.h
--- a/Engine/Plugins/Animation/IKRig/Source/IKRig/Public/Retargeter/RetargetOps/AlignPoleVectorOp.h
+++ b/Engine/Plugins/Animation/IKRig/Source/IKRig/Public/Retargeter/RetargetOps/AlignPoleVectorOp.h
@@ -100,6 +100,10 @@ struct IKRIG_API FIKRetargetAlignPoleVectorOpSettings : public FIKRetargetOpSett
 	TArray<FRetargetPoleVectorSettings> ChainsToAlign;
 
 	void MergePoleVectorSettings(const FRetargetPoleVectorSettings& InSettingsToMerge);
+
+	virtual const UClass* GetControllerType() const override;
+	
+	virtual void CopySettingsAtRuntime(const FIKRetargetOpSettingsBase* InSettingsToCopyFrom) override;
 };
 
 USTRUCT(BlueprintType, meta = (DisplayName = "Pole Vector Alignment"))
@@ -123,14 +127,10 @@ struct IKRIG_API FIKRetargetAlignPoleVectorOp : public FIKRetargetOpBase
 	virtual void OnAddedToStack(const UIKRetargeter* InRetargetAsset, const FIKRetargetOpBase* InParentOp) override;
 
 	virtual FIKRetargetOpSettingsBase* GetSettings() override;
-
-	virtual void SetSettings(const FIKRetargetOpSettingsBase* InSettings) override;
 	
 	virtual const UScriptStruct* GetSettingsType() const override;
 
 	virtual const UScriptStruct* GetType() const override;
-
-	virtual const UClass* GetControllerType() const override;
 	
 	virtual const UIKRigDefinition* GetCustomTargetIKRig() const override;
 
diff --git a/Engine/Plugins/Animation/IKRig/Source/IKRig/Public/Retargeter/RetargetOps/CopyBasePoseOp.h b/Engine/Plugins/Animation/IKRig/Source/IKRig/Public/Retargeter/RetargetOps/CopyBasePoseOp.h
--- a/Engine/Plugins/Animation/IKRig/Source/IKRig/Public/Retargeter/RetargetOps/CopyBasePoseOp.h
+++ b/Engine/Plugins/Animation/IKRig/Source/IKRig/Public/Retargeter/RetargetOps/CopyBasePoseOp.h
@@ -13,14 +13,17 @@ struct IKRIG_API FIKRetargetCopyBasePoseOpSettings : public FIKRetargetOpSetting
 	GENERATED_BODY()
 	
 	/** When true, will copy all the source bone transforms with matching names to use as a base pose. This can be useful for partial retargeting.
-	 * NOTE: no retargeting is applied to the bone transforms, they are assumed to be fully compatible between source/target skeletons.
-	 */
+	 * NOTE: no retargeting is applied to the bone transforms, they are assumed to be fully compatible between source/target skeletons. */
 	UPROPERTY(EditAnywhere, BlueprintReadWrite, Category = RetargetPhases)
 	bool bCopyBasePose = false;
 
 	/** Filters the bones to copy when using "Copy Base Pose". If specified, will only copy all children of the specified bone (inclusive).*/
 	UPROPERTY(EditAnywhere, BlueprintReadWrite, Category = RetargetPhases, meta = (EditCondition = "bCopyBasePose"))
 	FName CopyBasePoseRoot = NAME_None;
+
+	virtual const UClass* GetControllerType() const override;
+
+	virtual void CopySettingsAtRuntime(const FIKRetargetOpSettingsBase* InSettingsToCopyFrom) override;
 };
 
 USTRUCT(BlueprintType, meta = (DisplayName = "Copy Base Pose"))
@@ -45,16 +48,12 @@ struct IKRIG_API FIKRetargetCopyBasePoseOp : public FIKRetargetOpBase
 		TArray<FTransform>& OutTargetGlobalPose) override {};
 
 	virtual FIKRetargetOpSettingsBase* GetSettings() override;
-
-	virtual void SetSettings(const FIKRetargetOpSettingsBase* InSettings) override;
 	
 	virtual const UScriptStruct* GetSettingsType() const override;
 	
 	virtual const UScriptStruct* GetType() const override;
 
 	virtual bool IsSingleton() const override { return true; };
-	
-	virtual const UClass* GetControllerType() const override;
 
 	UPROPERTY()
 	FIKRetargetCopyBasePoseOpSettings Settings;
diff --git a/Engine/Plugins/Animation/IKRig/Source/IKRig/Public/Retargeter/RetargetOps/CurveRemapOp.h b/Engine/Plugins/Animation/IKRig/Source/IKRig/Public/Retargeter/RetargetOps/CurveRemapOp.h
--- a/Engine/Plugins/Animation/IKRig/Source/IKRig/Public/Retargeter/RetargetOps/CurveRemapOp.h
+++ b/Engine/Plugins/Animation/IKRig/Source/IKRig/Public/Retargeter/RetargetOps/CurveRemapOp.h
@@ -45,6 +45,10 @@ struct IKRIG_API FIKRetargetCurveRemapOpSettings : public FIKRetargetOpSettingsB
 	// target. Remapping is only necessary when the target curve name(s) are different.
 	UPROPERTY(EditAnywhere, BlueprintReadWrite, Category="Remap Curves")
 	TArray<FCurveRemapPair> CurvesToRemap;
+
+	virtual const UClass* GetControllerType() const override;
+	
+	virtual void CopySettingsAtRuntime(const FIKRetargetOpSettingsBase* InSettingsToCopyFrom) override;
 };
 
 USTRUCT(BlueprintType, meta = (DisplayName = "Remap Curves"))
@@ -99,8 +103,6 @@ struct IKRIG_API FIKRetargetCurveRemapOp : public FIKRetargetOpBase
 
 	virtual void AnimGraphEvaluateAnyThread(FPoseContext& Output) override;
 
-	virtual const UClass* GetControllerType() const override;
-
 	UPROPERTY()
 	FIKRetargetCurveRemapOpSettings Settings;
 	
diff --git a/Engine/Plugins/Animation/IKRig/Source/IKRig/Public/Retargeter/RetargetOps/FKChainsOp.h b/Engine/Plugins/Animation/IKRig/Source/IKRig/Public/Retargeter/RetargetOps/FKChainsOp.h
--- a/Engine/Plugins/Animation/IKRig/Source/IKRig/Public/Retargeter/RetargetOps/FKChainsOp.h
+++ b/Engine/Plugins/Animation/IKRig/Source/IKRig/Public/Retargeter/RetargetOps/FKChainsOp.h
@@ -190,6 +190,10 @@ struct IKRIG_API FIKRetargetFKChainsOpSettings : public FIKRetargetOpSettingsBas
 	/** Debug draw size used for single-bone chains. */
 	UPROPERTY(EditAnywhere, Category = Debug)
 	double ChainDrawSize = 1.0f;
+
+	virtual const UClass* GetControllerType() const override;
+
+	virtual void CopySettingsAtRuntime(const FIKRetargetOpSettingsBase* InSettingsToCopyFrom) override;
 };
 
 #if WITH_EDITOR
@@ -230,15 +234,11 @@ struct IKRIG_API FIKRetargetFKChainsOp : public FIKRetargetOpBase
 	virtual void CollectRetargetedBones(TSet<int32>& OutRetargetedBones) const override;
 
 	virtual FIKRetargetOpSettingsBase* GetSettings() override;
-
-	virtual void SetSettings(const FIKRetargetOpSettingsBase* InSettings) override;
 	
 	virtual const UScriptStruct* GetSettingsType() const override;
 
 	virtual const UScriptStruct* GetType() const override;
 
-	virtual const UClass* GetControllerType() const override;
-
 	virtual const UIKRigDefinition* GetCustomTargetIKRig() const override;
 
 	virtual FRetargetChainMapping* GetChainMapping() override;
diff --git a/Engine/Plugins/Animation/IKRig/Source/IKRig/Public/Retargeter/RetargetOps/IKChainsOp.h b/Engine/Plugins/Animation/IKRig/Source/IKRig/Public/Retargeter/RetargetOps/IKChainsOp.h
--- a/Engine/Plugins/Animation/IKRig/Source/IKRig/Public/Retargeter/RetargetOps/IKChainsOp.h
+++ b/Engine/Plugins/Animation/IKRig/Source/IKRig/Public/Retargeter/RetargetOps/IKChainsOp.h
@@ -175,6 +175,10 @@ struct IKRIG_API FIKRetargetIKChainsOpSettings : public FIKRetargetOpSettingsBas
 	// Adjust thickness of goal debug drawing in viewport
 	UPROPERTY(EditAnywhere, Category = Debug)
 	double GoalDrawThickness = 1.0;
+
+	virtual const UClass* GetControllerType() const override;
+
+	virtual void CopySettingsAtRuntime(const FIKRetargetOpSettingsBase* InSettingsToCopyFrom) override;
 };
 
 #if WITH_EDITOR
@@ -216,8 +220,6 @@ struct IKRIG_API FIKRetargetIKChainsOp : public FIKRetargetOpBase
 
 	virtual const UScriptStruct* GetType() const override;
 
-	virtual const UClass* GetControllerType() const override;
-
 	virtual const UScriptStruct* GetParentOpType() const override;
 
 	virtual void OnTargetChainRenamed(const FName InOldChainName, const FName InNewChainName) override;
diff --git a/Engine/Plugins/Animation/IKRig/Source/IKRig/Public/Retargeter/RetargetOps/PelvisMotionOp.h b/Engine/Plugins/Animation/IKRig/Source/IKRig/Public/Retargeter/RetargetOps/PelvisMotionOp.h
--- a/Engine/Plugins/Animation/IKRig/Source/IKRig/Public/Retargeter/RetargetOps/PelvisMotionOp.h
+++ b/Engine/Plugins/Animation/IKRig/Source/IKRig/Public/Retargeter/RetargetOps/PelvisMotionOp.h
@@ -87,6 +87,10 @@ struct IKRIG_API FIKRetargetPelvisMotionOpSettings : public FIKRetargetOpSetting
 	UPROPERTY(EditAnywhere, Category=Debug)
 	double DebugDrawThickness = 5.0f;
 
+	virtual const UClass* GetControllerType() const override;
+	
+	virtual void CopySettingsAtRuntime(const FIKRetargetOpSettingsBase* InSettingsToCopyFrom) override;
+
 	/** FBoneReference needs a skeleton for the UI widget. */
 #if WITH_EDITOR
 	virtual USkeleton* GetSkeleton(const FName InPropertyName) override;
@@ -138,14 +142,10 @@ struct IKRIG_API FIKRetargetPelvisMotionOp : public FIKRetargetOpBase
 	virtual void OnAddedToStack(const UIKRetargeter* InRetargetAsset, const FIKRetargetOpBase* InParentOp) override;
 
 	virtual FIKRetargetOpSettingsBase* GetSettings() override;
-
-	virtual void SetSettings(const FIKRetargetOpSettingsBase* InSettings) override;
 	
 	virtual const UScriptStruct* GetSettingsType() const override;
 
 	virtual const UScriptStruct* GetType() const override;
-
-	virtual const UClass* GetControllerType() const override;
 	
 	virtual void CollectRetargetedBones(TSet<int32>& OutRetargetedBones) const override;
 
@@ -232,6 +232,10 @@ public:
 	 * @return the name of the pelvis bone on the target skeleton. */
 	UFUNCTION(BlueprintCallable, Category = Settings)
 	FName GetTargetPelvisBone();
+
+private:
+
+	FIKRetargetPelvisMotionOpSettings& GetPelvisOpSettings() const;
 };
 
 #if WITH_EDITOR
diff --git a/Engine/Plugins/Animation/IKRig/Source/IKRig/Public/Retargeter/RetargetOps/PinBoneOp.h b/Engine/Plugins/Animation/IKRig/Source/IKRig/Public/Retargeter/RetargetOps/PinBoneOp.h
--- a/Engine/Plugins/Animation/IKRig/Source/IKRig/Public/Retargeter/RetargetOps/PinBoneOp.h
+++ b/Engine/Plugins/Animation/IKRig/Source/IKRig/Public/Retargeter/RetargetOps/PinBoneOp.h
@@ -130,6 +130,10 @@ struct IKRIG_API FIKRetargetPinBoneOpSettings : public FIKRetargetOpSettingsBase
 	UPROPERTY(EditAnywhere, BlueprintReadWrite, Category=Settings)
 	FTransform LocalOffset;
 
+	virtual const UClass* GetControllerType() const override;
+	
+	virtual void CopySettingsAtRuntime(const FIKRetargetOpSettingsBase* InSettingsToCopyFrom) override;
+
 #if WITH_EDITOR
 	virtual USkeleton* GetSkeleton(const FName InPropertyName) override;
 #endif
@@ -154,15 +158,11 @@ struct IKRIG_API FIKRetargetPinBoneOp : public FIKRetargetOpBase
 		TArray<FTransform>& OutTargetGlobalPose) override;
 
 	virtual FIKRetargetOpSettingsBase* GetSettings() override;
-
-	virtual void SetSettings(const FIKRetargetOpSettingsBase* InSettings) override;
 	
 	virtual const UScriptStruct* GetSettingsType() const override;
 	
 	virtual const UScriptStruct* GetType() const override;
 
-	virtual const UClass* GetControllerType() const override;
-
 	UPROPERTY()
 	FIKRetargetPinBoneOpSettings Settings;
 	
diff --git a/Engine/Plugins/Animation/IKRig/Source/IKRig/Public/Retargeter/RetargetOps/RootMotionGeneratorOp.h b/Engine/Plugins/Animation/IKRig/Source/IKRig/Public/Retargeter/RetargetOps/RootMotionGeneratorOp.h
--- a/Engine/Plugins/Animation/IKRig/Source/IKRig/Public/Retargeter/RetargetOps/RootMotionGeneratorOp.h
+++ b/Engine/Plugins/Animation/IKRig/Source/IKRig/Public/Retargeter/RetargetOps/RootMotionGeneratorOp.h
@@ -74,6 +74,10 @@ struct IKRIG_API FIKRetargetRootMotionOpSettings : public FIKRetargetOpSettingsB
 	UPROPERTY(EditAnywhere, BlueprintReadWrite, Category="Update Children")
 	bool bPropagateToNonRetargetedChildren = true;
 
+	virtual const UClass* GetControllerType() const override;
+	
+	virtual void CopySettingsAtRuntime(const FIKRetargetOpSettingsBase* InSettingsToCopyFrom) override;
+
 	/** FBoneReference needs a skeleton for the UI widget. */
 #if WITH_EDITOR
 	virtual USkeleton* GetSkeleton(const FName InPropertyName) override;
@@ -129,15 +133,11 @@ struct IKRIG_API FIKRetargetRootMotionOp : public FIKRetargetOpBase
 	virtual void OnAddedToStack(const UIKRetargeter* InRetargetAsset, const FIKRetargetOpBase* InParentOp) override;
 	
 	virtual FIKRetargetOpSettingsBase* GetSettings() override;
-
-	virtual void SetSettings(const FIKRetargetOpSettingsBase* InSettings) override;
 	
 	virtual const UScriptStruct* GetSettingsType() const override;
 	
 	virtual const UScriptStruct* GetType() const override;
 
-	virtual const UClass* GetControllerType() const override;
-
 	UPROPERTY()
 	FIKRetargetRootMotionOpSettings Settings;
 
@@ -207,6 +207,20 @@ public:
 	 * @return the name of the root bone on the target skeleton. */
 	UFUNCTION(BlueprintCallable, Category = Settings)
 	FName GetTargetRootBone();
+
+	/* Set the pelvis bone for the target.
+	 * @param InTargetPelvisBone the name of the pelvis bone on the target skeleton. */
+	UFUNCTION(BlueprintCallable, Category = Settings)
+	void SetTargetPelvisBone(const FName InTargetPelvisBone);
+
+	/* Get the pelvis bone for the target.
+	 * @return the name of the pelvis bone on the target skeleton. */
+	UFUNCTION(BlueprintCallable, Category = Settings)
+	FName GetTargetPelvisBone();
+
+private:
+
+	FIKRetargetRootMotionOpSettings* GetSettingsPtr() const;
 };
 
 //
diff --git a/Engine/Plugins/Animation/IKRig/Source/IKRig/Public/Retargeter/RetargetOps/RunIKRigOp.h b/Engine/Plugins/Animation/IKRig/Source/IKRig/Public/Retargeter/RetargetOps/RunIKRigOp.h
--- a/Engine/Plugins/Animation/IKRig/Source/IKRig/Public/Retargeter/RetargetOps/RunIKRigOp.h
+++ b/Engine/Plugins/Animation/IKRig/Source/IKRig/Public/Retargeter/RetargetOps/RunIKRigOp.h
@@ -37,6 +37,10 @@ struct IKRIG_API FIKRetargetRunIKRigOpSettings : public FIKRetargetOpSettingsBas
 	// Adjust thickness of goal debug drawing in viewport
 	UPROPERTY(EditAnywhere, Category = Debug)
 	double GoalDrawThickness = 1.0;
+
+	virtual const UClass* GetControllerType() const override;
+
+	virtual void CopySettingsAtRuntime(const FIKRetargetOpSettingsBase* InSettingsToCopyFrom) override;
 };
 
 #if WITH_EDITOR
@@ -70,13 +74,9 @@ struct IKRIG_API FIKRetargetRunIKRigOp : public FIKRetargetOpBase
 
 	virtual FIKRetargetOpSettingsBase* GetSettings() override;
 	
-	virtual void SetSettings(const FIKRetargetOpSettingsBase* InSettings) override;
-	
 	virtual const UScriptStruct* GetSettingsType() const override;
 
 	virtual const UScriptStruct* GetType() const override;
-
-	virtual const UClass* GetControllerType() const override;
 	
 	virtual bool CanHaveChildOps() const override { return true; };
 
diff --git a/Engine/Plugins/Animation/IKRig/Source/IKRig/Public/Retargeter/RetargetOps/ScaleSourceOp.h b/Engine/Plugins/Animation/IKRig/Source/IKRig/Public/Retargeter/RetargetOps/ScaleSourceOp.h
--- a/Engine/Plugins/Animation/IKRig/Source/IKRig/Public/Retargeter/RetargetOps/ScaleSourceOp.h
+++ b/Engine/Plugins/Animation/IKRig/Source/IKRig/Public/Retargeter/RetargetOps/ScaleSourceOp.h
@@ -15,6 +15,10 @@ struct IKRIG_API FIKRetargetScaleSourceOpSettings : public FIKRetargetOpSettings
 	/** Range 0.01 to +inf. Default 1. Scales the incoming source pose. Affects entire skeleton and all IK goals.*/
 	UPROPERTY(EditAnywhere, BlueprintReadWrite, Category = Scale, meta = (ReinitializeOnEdit, ClampMin = "0.01", UIMin = "0.01", UIMax = "10.0"))
 	double SourceScaleFactor = 1.0f;
+
+	virtual const UClass* GetControllerType() const override;
+
+	virtual void CopySettingsAtRuntime(const FIKRetargetOpSettingsBase* InSettingsToCopyFrom) override;
 };
 
 USTRUCT(BlueprintType, meta = (DisplayName = "Scale Source"))
@@ -39,15 +43,11 @@ struct IKRIG_API FIKRetargetScaleSourceOp : public FIKRetargetOpBase
 		TArray<FTransform>& OutTargetGlobalPose) override {};
 
 	virtual FIKRetargetOpSettingsBase* GetSettings() override;
-
-	virtual void SetSettings(const FIKRetargetOpSettingsBase* InSettings) override;
 	
 	virtual const UScriptStruct* GetSettingsType() const override;
 	
 	virtual const UScriptStruct* GetType() const override;
 
-	virtual const UClass* GetControllerType() const override;
-
 	UPROPERTY()
 	FIKRetargetScaleSourceOpSettings Settings;
 };
diff --git a/Engine/Plugins/Animation/IKRig/Source/IKRig/Public/Retargeter/RetargetOps/SpeedPlantingOp.h b/Engine/Plugins/Animation/IKRig/Source/IKRig/Public/Retargeter/RetargetOps/SpeedPlantingOp.h
--- a/Engine/Plugins/Animation/IKRig/Source/IKRig/Public/Retargeter/RetargetOps/SpeedPlantingOp.h
+++ b/Engine/Plugins/Animation/IKRig/Source/IKRig/Public/Retargeter/RetargetOps/SpeedPlantingOp.h
@@ -46,6 +46,10 @@ struct IKRIG_API FIKRetargetSpeedPlantingOpSettings : public FIKRetargetOpSettin
 	// How much damping to apply to the spring (0 means no damping, 1 means critically damped which means no oscillation)
 	UPROPERTY(EditAnywhere, BlueprintReadWrite, Category = "Speed Planting", meta = (ClampMin = "0.0", ClampMax = "1.0", UIMin = "0.0", UIMax = "1.0"))
 	double CriticalDamping = 1.0f;
+
+	virtual const UClass* GetControllerType() const override;
+
+	virtual void CopySettingsAtRuntime(const FIKRetargetOpSettingsBase* InSettingsToCopyFrom) override;
 };
 
 struct FSpeedPlantingGoalState
@@ -91,8 +95,6 @@ struct IKRIG_API FIKRetargetSpeedPlantingOp : public FIKRetargetOpBase
 	virtual void OnAddedToStack(const UIKRetargeter* InRetargetAsset, const FIKRetargetOpBase* InParentOp) override;
 
 	virtual FIKRetargetOpSettingsBase* GetSettings() override;
-
-	virtual void SetSettings(const FIKRetargetOpSettingsBase* InSettings) override;
 	
 	virtual const UScriptStruct* GetSettingsType() const override;
 	
@@ -103,8 +105,6 @@ struct IKRIG_API FIKRetargetSpeedPlantingOp : public FIKRetargetOpBase
 	virtual void AnimGraphPreUpdateMainThread(USkeletalMeshComponent& SourceMeshComponent, USkeletalMeshComponent& TargetMeshComponent) override;
 
 	virtual void AnimGraphEvaluateAnyThread(FPoseContext& Output) override;
-	
-	virtual const UClass* GetControllerType() const override;
 
 	virtual const UScriptStruct* GetParentOpType() const override;
 
diff --git a/Engine/Plugins/Animation/IKRig/Source/IKRig/Public/Retargeter/RetargetOps/StrideWarpingOp.h b/Engine/Plugins/Animation/IKRig/Source/IKRig/Public/Retargeter/RetargetOps/StrideWarpingOp.h
--- a/Engine/Plugins/Animation/IKRig/Source/IKRig/Public/Retargeter/RetargetOps/StrideWarpingOp.h
+++ b/Engine/Plugins/Animation/IKRig/Source/IKRig/Public/Retargeter/RetargetOps/StrideWarpingOp.h
@@ -93,6 +93,10 @@ struct IKRIG_API FIKRetargetStrideWarpingOpSettings : public FIKRetargetOpSettin
 	UPROPERTY(EditAnywhere, BlueprintReadWrite, Category = Debug)
 	double DebugDrawThickness = 5.0;
 
+	virtual const UClass* GetControllerType() const override;
+
+	virtual void CopySettingsAtRuntime(const FIKRetargetOpSettingsBase* InSettingsToCopyFrom) override;
+
 	static FVector GetAxisVector(const EBasicAxis& Axis);
 
 	bool operator==(const FIKRetargetStrideWarpingOpSettings& Other) const;
@@ -129,14 +133,10 @@ struct IKRIG_API FIKRetargetStrideWarpingOp : public FIKRetargetOpBase
 	virtual void OnAddedToStack(const UIKRetargeter* InRetargetAsset, const FIKRetargetOpBase* InParentOp) override;
 
 	virtual FIKRetargetOpSettingsBase* GetSettings() override;
-
-	virtual void SetSettings(const FIKRetargetOpSettingsBase* InSettings) override;
 	
 	virtual const UScriptStruct* GetSettingsType() const override;
 	
 	virtual const UScriptStruct* GetType() const override;
-
-	virtual const UClass* GetControllerType() const override;
 	
 	virtual const UScriptStruct* GetParentOpType() const override;
 	
diff --git a/Engine/Plugins/Animation/IKRig/Source/IKRigEditor/Private/RetargetEditor/IKRetargeterController.cpp b/Engine/Plugins/Animation/IKRig/Source/IKRigEditor/Private/RetargetEditor/IKRetargeterController.cpp
--- a/Engine/Plugins/Animation/IKRig/Source/IKRigEditor/Private/RetargetEditor/IKRetargeterController.cpp
+++ b/Engine/Plugins/Animation/IKRig/Source/IKRigEditor/Private/RetargetEditor/IKRetargeterController.cpp
@@ -598,8 +598,9 @@ UIKRetargetOpControllerBase* UIKRetargeterController::GetOpController(int32 OpIn
 {
 	if (FIKRetargetOpBase* OpAtIndex = GetRetargetOpByIndex(OpIndex))
 	{
-		return OpAtIndex->GetController(this);
+		return OpAtIndex->GetSettings()->GetController(this);
 	}
+	
 	return nullptr;
 }
 
