[IK Retarget] details panel resets when editing some op settings.

diff --git a/Engine/Plugins/Animation/IKRig/Source/IKRigEditor/Private/RetargetEditor/IKRetargetEditor.cpp b/Engine/Plugins/Animation/IKRig/Source/IKRigEditor/Private/RetargetEditor/IKRetargetEditor.cpp
--- a/Engine/Plugins/Animation/IKRig/Source/IKRigEditor/Private/RetargetEditor/IKRetargetEditor.cpp
+++ b/Engine/Plugins/Animation/IKRig/Source/IKRigEditor/Private/RetargetEditor/IKRetargetEditor.cpp
@@ -371,20 +371,6 @@ TStatId FIKRetargetEditor::GetStatId() const
 	RETURN_QUICK_DECLARE_CYCLE_STAT(FIKRetargetEditor, STATGROUP_Tickables);
 }
 
-void FIKRetargetEditor::PostUndo(bool bSuccess)
-{
-	EditorController->AssetController->CleanAsset();
-	EditorController->HandlePreviewMeshReplaced(ERetargetSourceOrTarget::Source);
-	EditorController->HandleRetargeterNeedsInitialized();
-}
-
-void FIKRetargetEditor::PostRedo(bool bSuccess)
-{
-	EditorController->AssetController->CleanAsset();
-	EditorController->HandlePreviewMeshReplaced(ERetargetSourceOrTarget::Source);
-	EditorController->HandleRetargeterNeedsInitialized();
-}
-
 void FIKRetargetEditor::HandleViewportCreated(const TSharedRef<class IPersonaViewport>& InViewport)
 {
 	// register callbacks to allow the asset to store the Bone Size viewport setting
diff --git a/Engine/Plugins/Animation/IKRig/Source/IKRigEditor/Private/RetargetEditor/IKRetargetEditorController.cpp b/Engine/Plugins/Animation/IKRig/Source/IKRigEditor/Private/RetargetEditor/IKRetargetEditorController.cpp
--- a/Engine/Plugins/Animation/IKRig/Source/IKRigEditor/Private/RetargetEditor/IKRetargetEditorController.cpp
+++ b/Engine/Plugins/Animation/IKRig/Source/IKRigEditor/Private/RetargetEditor/IKRetargetEditorController.cpp
@@ -131,8 +131,8 @@ void FRetargetPlaybackManager::ResumePlayback() const
 	if (AnimThatWasPlaying)
 	{
 		AnimInstance->SetAnimationAsset(AnimThatWasPlaying);
+		AnimInstance->SetPosition(TimeWhenPaused);
 		AnimInstance->SetPlaying(true);
-		AnimInstance->SetPosition(TimeWhenPaused);	
 	}
 }
 
@@ -191,6 +191,22 @@ void FIKRetargetEditorController::Close()
 	}
 }
 
+void FIKRetargetEditorController::PostUndo(bool bSuccess)
+{
+	AssetController->CleanAsset();
+	HandlePreviewMeshReplaced(ERetargetSourceOrTarget::Source);
+	HandleRetargeterNeedsInitialized();
+	RefreshAllViews();
+}
+
+void FIKRetargetEditorController::PostRedo(bool bSuccess)
+{
+	AssetController->CleanAsset();
+	HandlePreviewMeshReplaced(ERetargetSourceOrTarget::Source);
+	HandleRetargeterNeedsInitialized();
+	RefreshAllViews();
+}
+
 void FIKRetargetEditorController::BindToIKRigAssets()
 {
 	const UIKRetargeter* Asset = AssetController->GetAsset();
@@ -252,9 +268,6 @@ void FIKRetargetEditorController::HandleRetargeterNeedsInitialized() const
 	FixZeroHeightRetargetRoot(ERetargetSourceOrTarget::Target);
 	
 	ReinitializeRetargeterNoUIRefresh();
-	
-	// refresh all the UI views
-	RefreshAllViews();	
 }
 
 void FIKRetargetEditorController::ReinitializeRetargeterNoUIRefresh() const
diff --git a/Engine/Plugins/Animation/IKRig/Source/IKRigEditor/Private/RetargetEditor/IKRetargeterController.cpp b/Engine/Plugins/Animation/IKRig/Source/IKRigEditor/Private/RetargetEditor/IKRetargeterController.cpp
--- a/Engine/Plugins/Animation/IKRig/Source/IKRigEditor/Private/RetargetEditor/IKRetargeterController.cpp
+++ b/Engine/Plugins/Animation/IKRig/Source/IKRigEditor/Private/RetargetEditor/IKRetargeterController.cpp
@@ -9,6 +9,7 @@
 #include "IKRigEditor.h"
 #include "RetargetEditor/IKRetargeterPoseGenerator.h"
 #include "Retargeter/IKRetargetOps.h"
+#include "Retargeter/RetargetOps/CurveRemapOp.h"
 #include "Retargeter/RetargetOps/FKChainsOp.h"
 #include "Retargeter/RetargetOps/IKChainsOp.h"
 #include "Retargeter/RetargetOps/PelvisMotionOp.h"
@@ -347,6 +348,7 @@ void UIKRetargeterController::AddDefaultOps() const
 	const FName RunIKOpName = GetOpName(RunIKIndex);
 	AddRetargetOp(FIKRetargetIKChainsOp::StaticStruct(), RunIKOpName);
 	AddRetargetOp(FIKRetargetRootMotionOp::StaticStruct());
+	AddRetargetOp(FIKRetargetCurveRemapOp::StaticStruct());
 }
 
 void UIKRetargeterController::RunOpInitialSetup(const int32 InOpIndex) const
diff --git a/Engine/Plugins/Animation/IKRig/Source/IKRigEditor/Private/RetargetEditor/SRetargetOpStack.cpp b/Engine/Plugins/Animation/IKRig/Source/IKRigEditor/Private/RetargetEditor/SRetargetOpStack.cpp
--- a/Engine/Plugins/Animation/IKRig/Source/IKRigEditor/Private/RetargetEditor/SRetargetOpStack.cpp
+++ b/Engine/Plugins/Animation/IKRig/Source/IKRigEditor/Private/RetargetEditor/SRetargetOpStack.cpp
@@ -541,6 +541,9 @@ void SRetargetOpList::AddNewRetargetOp(const UScriptStruct* ScriptStruct)
 	// add a new op to the stack
 	const FName ParentOpName = ParentElement.IsValid() ? ParentElement.Pin()->GetName() : NAME_None;
 	int32 NewOpIndex = AssetController->AddRetargetOp(ScriptStruct, ParentOpName);
+
+	// refresh the op stack
+	EditorController.Pin()->RefreshOpStackView();
 }
 
 void SRetargetOpList::DeleteRetargetOp(TSharedPtr<FRetargetOpStackElement> OpToDelete)
@@ -551,6 +554,9 @@ void SRetargetOpList::DeleteRetargetOp(TSharedPtr<FRetargetOpStackElement> OpToD
 	}
 	
 	EditorController.Pin()->AssetController->RemoveRetargetOp(OpToDelete->GetIndexInStack());
+	
+	// refresh the op stack
+	EditorController.Pin()->RefreshOpStackView();
 }
 
 void SRetargetOpList::RefreshAndRestore()
diff --git a/Engine/Plugins/Animation/IKRig/Source/IKRigEditor/Public/RetargetEditor/IKRetargetEditor.h b/Engine/Plugins/Animation/IKRig/Source/IKRigEditor/Public/RetargetEditor/IKRetargetEditor.h
--- a/Engine/Plugins/Animation/IKRig/Source/IKRigEditor/Public/RetargetEditor/IKRetargetEditor.h
+++ b/Engine/Plugins/Animation/IKRig/Source/IKRigEditor/Public/RetargetEditor/IKRetargetEditor.h
@@ -9,7 +9,6 @@
 #include "IKRetargetApplicationMode.h"
 #include "IPersonaPreviewScene.h"
 #include "PersonaAssetEditorToolkit.h"
-#include "EditorUndoClient.h"
 
 class IAnimationSequenceBrowser;
 class UIKRetargetSkeletalMeshComponent;
@@ -30,7 +29,6 @@ namespace IKRetargetApplicationModes
 class FIKRetargetEditor :
 	public FPersonaAssetEditorToolkit,
 	public IHasPersonaToolkit,
-	public FSelfRegisteringEditorUndoClient,
 	public FTickableEditorObject
 {
 public:
@@ -64,11 +62,6 @@ public:
 	virtual TSharedRef<IPersonaToolkit> GetPersonaToolkit() const override { return PersonaToolkit.ToSharedRef(); }
 	/** END IHasPersonaToolkit interface */
 
-	/** FSelfRegisteringEditorUndoClient interface */
-	virtual void PostUndo( bool bSuccess );
-	virtual void PostRedo( bool bSuccess );
-	/** END FSelfRegisteringEditorUndoClient interface */
-
 	TSharedRef<FIKRetargetEditorController> GetController() const {return EditorController;};
 	
 	void SetupAnimInstance();
diff --git a/Engine/Plugins/Animation/IKRig/Source/IKRigEditor/Public/RetargetEditor/IKRetargetEditorController.h b/Engine/Plugins/Animation/IKRig/Source/IKRigEditor/Public/RetargetEditor/IKRetargetEditorController.h
--- a/Engine/Plugins/Animation/IKRig/Source/IKRigEditor/Public/RetargetEditor/IKRetargetEditorController.h
+++ b/Engine/Plugins/Animation/IKRig/Source/IKRigEditor/Public/RetargetEditor/IKRetargetEditorController.h
@@ -10,6 +10,7 @@
 #include "Input/Reply.h"
 #include "Widgets/Input/SEditableTextBox.h"
 #include "Widgets/SWindow.h"
+#include "EditorUndoClient.h"
 
 class SRetargetOpStack;
 enum class ERetargetSourceOrTarget : uint8;
@@ -71,7 +72,10 @@ private:
 };
 
 // a home for cross-widget communication to synchronize state across all tabs and viewport
-class FIKRetargetEditorController : public TSharedFromThis<FIKRetargetEditorController>, FGCObject
+class FIKRetargetEditorController :
+	public TSharedFromThis<FIKRetargetEditorController>,
+	public FSelfRegisteringEditorUndoClient,
+	FGCObject 
 {
 public:
 
@@ -81,6 +85,11 @@ public:
 	void Initialize(TSharedPtr<FIKRetargetEditor> InEditor, UIKRetargeter* InAsset);
 	// Close the editor
 	void Close();
+
+	/** FSelfRegisteringEditorUndoClient interface */
+	virtual void PostUndo( bool bSuccess );
+	virtual void PostRedo( bool bSuccess );
+	/** END FSelfRegisteringEditorUndoClient interface */
 	
 	// Bind callbacks to this IK Rig
 	void BindToIKRigAssets();
