[IK Retarget] Fixes UI bug when reordering op stack.

diff --git a/Engine/Plugins/Animation/IKRig/Source/IKRigEditor/Private/RetargetEditor/SRetargetOpStack.cpp b/Engine/Plugins/Animation/IKRig/Source/IKRigEditor/Private/RetargetEditor/SRetargetOpStack.cpp
--- a/Engine/Plugins/Animation/IKRig/Source/IKRigEditor/Private/RetargetEditor/SRetargetOpStack.cpp
+++ b/Engine/Plugins/Animation/IKRig/Source/IKRigEditor/Private/RetargetEditor/SRetargetOpStack.cpp
@@ -421,7 +421,8 @@ FReply SRetargetOpList::OnAcceptDrop(
 	const bool bWasReparented = AssetController->MoveRetargetOpInStack(DraggedElement->GetIndexInStack(),IndexToMoveTo);
 	if (bWasReparented)
 	{
-		RequestListRefresh();
+		// refresh the op stack
+		EditorController.Pin()->RefreshOpStackView();
 	}
 	
 	return FReply::Handled();
