[IK Retarget] Adding Retarget Pose Op

diff --git a/Engine/Plugins/Animation/IKRig/Source/IKRig/Private/Retargeter/IKRetargetProcessor.cpp b/Engine/Plugins/Animation/IKRig/Source/IKRig/Private/Retargeter/IKRetargetProcessor.cpp
--- a/Engine/Plugins/Animation/IKRig/Source/IKRig/Private/Retargeter/IKRetargetProcessor.cpp
+++ b/Engine/Plugins/Animation/IKRig/Source/IKRig/Private/Retargeter/IKRetargetProcessor.cpp
@@ -21,83 +21,48 @@
 // This is the default end of branch index value, meaning we haven't cached it yet
 #define RETARGETSKELETON_INVALID_BRANCH_INDEX -2
 
-void FRetargetSkeleton::Initialize(
-	const USkeletalMesh* InSkeletalMesh,
-	const FName InRetargetPoseName,
-	const FIKRetargetPose* RetargetPose,
-	const FName RetargetRootBone,
-	const double ScaleFactor)
-{
-	// reset all skeleton data
-	Reset();
-	
-	// record which skeletal mesh this is running on
-	SkeletalMesh = InSkeletalMesh;
-	
-	// copy names and parent indices into local storage
-	const FReferenceSkeleton& RefSkeleton = SkeletalMesh->GetRefSkeleton();
-	for (int32 BoneIndex=0; BoneIndex<RefSkeleton.GetNum(); ++BoneIndex)
-	{
-		BoneNames.Add(RefSkeleton.GetBoneName(BoneIndex));
-		ParentIndices.Add(RefSkeleton.GetParentIndex(BoneIndex));	
-	}
-
-	// initialize branch caching
-	CachedEndOfBranchIndices.Init(RETARGETSKELETON_INVALID_BRANCH_INDEX, ParentIndices.Num());
-
-	// update retarget pose to reflect custom offsets (applies stored offsets)
-	// NOTE: this must be done AFTER generating IsBoneInAnyTargetChain array above
-	GenerateRetargetPose(InRetargetPoseName, RetargetPose, RetargetRootBone, ScaleFactor);
-}
-
-void FRetargetSkeleton::Reset()
-{
-	BoneNames.Reset();
-	ParentIndices.Reset();
-	RetargetLocalPose.Reset();
-	RetargetGlobalPose.Reset();
-	SkeletalMesh = nullptr;
-}
-
-void FRetargetSkeleton::GenerateRetargetPose(
+FResolvedRetargetPose& FResolvedRetargetPoseSet::AddOrUpdateRetargetPose(
+	const FRetargetSkeleton& InSkeleton,
 	const FName InRetargetPoseName,
 	const FIKRetargetPose* InRetargetPose,
 	const FName PelvisBoneName,
 	const double ScaleFactor)
 {
+	int32 NewPoseIndex = RetargetPoses.Emplace(InRetargetPoseName);
+	FResolvedRetargetPose& NewRetargetPose = RetargetPoses[NewPoseIndex];
+	
 	// record the name of the retarget pose (prevents re-initialization if profile swaps it)
-	RetargetPoseName = InRetargetPoseName;
-	RetargetPoseVersion = InRetargetPose->GetVersion();
-	RetargetPoseScale = ScaleFactor;
+	NewRetargetPose.Version = InRetargetPose->GetVersion();
+	NewRetargetPose.Scale = ScaleFactor;
 	
 	// initialize retarget pose to the skeletal mesh reference pose
-	RetargetLocalPose = SkeletalMesh->GetRefSkeleton().GetRefBonePose();
+	NewRetargetPose.LocalPose = InSkeleton.SkeletalMesh->GetRefSkeleton().GetRefBonePose();
 	// copy local pose to global
-	RetargetGlobalPose = RetargetLocalPose;
+	NewRetargetPose.GlobalPose = NewRetargetPose.LocalPose;
 	// convert to global space
-	UpdateGlobalTransformsBelowBone(INDEX_NONE, RetargetLocalPose, RetargetGlobalPose);
+	InSkeleton.UpdateGlobalTransformsBelowBone(INDEX_NONE, NewRetargetPose.LocalPose, NewRetargetPose.GlobalPose);
 
 	// strip scale (done AFTER generating global pose so that scales are baked into translation)
-	for (int32 BoneIndex=0; BoneIndex<BoneNames.Num(); ++BoneIndex)
+	for (int32 BoneIndex=0; BoneIndex<InSkeleton.BoneNames.Num(); ++BoneIndex)
 	{
-		RetargetLocalPose[BoneIndex].SetScale3D(FVector::OneVector);
-		RetargetGlobalPose[BoneIndex].SetScale3D(FVector::OneVector);
+		NewRetargetPose.LocalPose[BoneIndex].SetScale3D(FVector::OneVector);
+		NewRetargetPose.GlobalPose[BoneIndex].SetScale3D(FVector::OneVector);
 	}
 
 	// apply pelvis translation offset
-	const int32 PelvisBoneIndex = FindBoneIndexByName(PelvisBoneName);
+	const int32 PelvisBoneIndex = InSkeleton.FindBoneIndexByName(PelvisBoneName);
 	if (PelvisBoneIndex != INDEX_NONE)
 	{
-		FTransform& PelvisTransform = RetargetGlobalPose[PelvisBoneIndex];
+		FTransform& PelvisTransform = NewRetargetPose.GlobalPose[PelvisBoneIndex];
 		PelvisTransform.AddToTranslation(InRetargetPose->GetRootTranslationDelta());
-		UpdateLocalTransformOfSingleBone(PelvisBoneIndex, RetargetLocalPose, RetargetGlobalPose);
+		InSkeleton.UpdateLocalTransformOfSingleBone(PelvisBoneIndex, NewRetargetPose.LocalPose, NewRetargetPose.GlobalPose);
 	}
 
 	// apply retarget pose offsets (retarget pose is stored as offset relative to reference pose)
-	const TArray<FTransform>& RefPoseLocal = SkeletalMesh->GetRefSkeleton().GetRefBonePose();
+	const TArray<FTransform>& RefPoseLocal = InSkeleton.SkeletalMesh->GetRefSkeleton().GetRefBonePose();
 	for (const TTuple<FName, FQuat>& BoneDelta : InRetargetPose->GetAllDeltaRotations())
 	{
-		const int32 BoneIndex = FindBoneIndexByName(BoneDelta.Key);
+		const int32 BoneIndex = InSkeleton.FindBoneIndexByName(BoneDelta.Key);
 		if (BoneIndex == INDEX_NONE)
 		{
 			// this can happen if a retarget pose recorded a bone offset for a bone that is not present in the
@@ -106,24 +71,130 @@ void FRetargetSkeleton::GenerateRetargetPose(
 		}
 
 		const FQuat LocalBoneRotation = RefPoseLocal[BoneIndex].GetRotation() * BoneDelta.Value;
-		RetargetLocalPose[BoneIndex].SetRotation(LocalBoneRotation);
+		NewRetargetPose.LocalPose[BoneIndex].SetRotation(LocalBoneRotation);
 	}
 
 	// update global transforms based on local pose modified by the retarget pose offsets
-	UpdateGlobalTransformsBelowBone(INDEX_NONE, RetargetLocalPose, RetargetGlobalPose);
+	InSkeleton.UpdateGlobalTransformsBelowBone(INDEX_NONE, NewRetargetPose.LocalPose, NewRetargetPose.GlobalPose);
 
 	// scale the retarget pose
 	if (!FMath::IsNearlyEqual(ScaleFactor, 1.0f))
 	{
 		// scale translations of global transforms
-		for (FTransform& BoneTransform : RetargetGlobalPose)
+		for (FTransform& BoneTransform : NewRetargetPose.GlobalPose)
 		{
 			BoneTransform.SetTranslation(BoneTransform.GetTranslation() * ScaleFactor);
 		}
 
 		// update the local transforms
-		UpdateLocalTransformsBelowBone(INDEX_NONE, RetargetLocalPose, RetargetGlobalPose);
+		InSkeleton.UpdateLocalTransformsBelowBone(INDEX_NONE, NewRetargetPose.LocalPose, NewRetargetPose.GlobalPose);
+	}
+
+	return NewRetargetPose;
+}
+
+const FResolvedRetargetPose* FResolvedRetargetPoseSet::FindRetargetPoseByName(const FName InRetargetPoseName) const
+{
+	for (const FResolvedRetargetPose& Pose : RetargetPoses)
+	{
+		if (Pose.Name == CurrentRetargetPoseName)
+		{
+			return &Pose;
+		}
+	}
+
+	ensureMsgf(false, TEXT("Specified retarget pose was not found. Calling code should check existence of pose."));
+	return nullptr;
+}
+
+const TArray<FTransform>& FResolvedRetargetPoseSet::GetLocalRetargetPose() const
+{
+	if (const FResolvedRetargetPose* ResolvedRetargetPose = FindRetargetPoseByName(CurrentRetargetPoseName))
+	{
+		return ResolvedRetargetPose->LocalPose;
 	}
+	
+	return RetargetPoses[0].LocalPose;
+}
+
+const TArray<FTransform>& FResolvedRetargetPoseSet::GetGlobalRetargetPose() const
+{
+	if (const FResolvedRetargetPose* ResolvedRetargetPose = FindRetargetPoseByName(CurrentRetargetPoseName))
+	{
+		return ResolvedRetargetPose->GlobalPose;
+	}
+	
+	return RetargetPoses[0].GlobalPose;
+}
+
+FTransform FResolvedRetargetPoseSet::GetGlobalRetargetPoseOfSingleBone(
+	const FRetargetSkeleton& InSkeleton,
+	const int32 BoneIndex,
+	const TArray<FTransform>& InGlobalPose) const
+{
+	const TArray<FTransform>& LocalRetargetPose = GetLocalRetargetPose();
+	const int32 ParentIndex = InSkeleton.ParentIndices[BoneIndex];
+	if (ParentIndex == INDEX_NONE)
+	{
+		return LocalRetargetPose[BoneIndex]; // root always in global space
+	}
+	const FTransform& ChildLocalTransform = LocalRetargetPose[BoneIndex];
+	const FTransform& ParentGlobalTransform = InGlobalPose[ParentIndex];
+	return ChildLocalTransform * ParentGlobalTransform;
+}
+
+void FResolvedRetargetPoseSet::Reset()
+{
+	RetargetPoses.Reset();
+}
+
+void FRetargetSkeleton::Initialize(
+	const USkeletalMesh* InSkeletalMesh,
+	const ERetargetSourceOrTarget InSourceOrTarget,
+	const UIKRetargeter* InRetargetAsset,
+	const FName RetargetRootBone,
+	const double ScaleFactor)
+{
+	// reset all skeleton data
+	Reset();
+	
+	// record which skeletal mesh this is running on
+	SkeletalMesh = InSkeletalMesh;
+	
+	// copy names and parent indices into local storage
+	const FReferenceSkeleton& RefSkeleton = SkeletalMesh->GetRefSkeleton();
+	for (int32 BoneIndex=0; BoneIndex<RefSkeleton.GetNum(); ++BoneIndex)
+	{
+		BoneNames.Add(RefSkeleton.GetBoneName(BoneIndex));
+		ParentIndices.Add(RefSkeleton.GetParentIndex(BoneIndex));	
+	}
+
+	// initialize branch caching
+	CachedEndOfBranchIndices.Init(RETARGETSKELETON_INVALID_BRANCH_INDEX, ParentIndices.Num());
+
+	// add default retarget pose for this skeleton
+	static FName DefaultRetargetPoseName = UIKRetargeter::GetDefaultPoseName();
+	const FIKRetargetPose* DefaultRetargetPose = InRetargetAsset->GetRetargetPoseByName(InSourceOrTarget, UIKRetargeter::GetDefaultPoseName());
+	RetargetPoses.AddOrUpdateRetargetPose(*this,DefaultRetargetPoseName, DefaultRetargetPose, RetargetRootBone, ScaleFactor);
+	RetargetPoses.CurrentRetargetPoseName = DefaultRetargetPoseName;
+
+	// add current retarget pose
+	const FName CurrentRetargetPoseName = InRetargetAsset->GetCurrentRetargetPoseName(InSourceOrTarget);
+	const FIKRetargetPose* CurrentRetargetPose = InRetargetAsset->GetRetargetPoseByName(InSourceOrTarget, CurrentRetargetPoseName);
+	if (CurrentRetargetPose)
+	{
+		RetargetPoses.AddOrUpdateRetargetPose(*this,CurrentRetargetPoseName, CurrentRetargetPose, RetargetRootBone, ScaleFactor);
+		// set this as the current pose to use (may be overridden by ops)
+		RetargetPoses.CurrentRetargetPoseName = CurrentRetargetPoseName;
+	}
+}
+
+void FRetargetSkeleton::Reset()
+{
+	BoneNames.Reset();
+	ParentIndices.Reset();
+	RetargetPoses.Reset();
+	SkeletalMesh = nullptr;
 }
 
 int32 FRetargetSkeleton::FindBoneIndexByName(const FName InName) const
@@ -215,20 +286,6 @@ void FRetargetSkeleton::UpdateLocalTransformOfSingleBone(
 	OutLocalPose[BoneIndex] = ChildGlobalTransform.GetRelativeTransform(ParentGlobalTransform);
 }
 
-FTransform FRetargetSkeleton::GetGlobalRetargetPoseOfSingleBone(
-	const int32 BoneIndex,
-	const TArray<FTransform>& InGlobalPose) const
-{
-	const int32 ParentIndex = ParentIndices[BoneIndex];
-	if (ParentIndex == INDEX_NONE)
-	{
-		return RetargetLocalPose[BoneIndex]; // root always in global space
-	}
-	const FTransform& ChildLocalTransform = RetargetLocalPose[BoneIndex];
-	const FTransform& ParentGlobalTransform = InGlobalPose[ParentIndex];
-	return ChildLocalTransform * ParentGlobalTransform;
-}
-
 int32 FRetargetSkeleton::GetCachedEndOfBranchIndex(const int32 InBoneIndex) const
 {
 	if (!CachedEndOfBranchIndices.IsValidIndex(InBoneIndex))
@@ -336,15 +393,15 @@ int32 FRetargetSkeleton::GetParentIndex(const int32 BoneIndex) const
 }
 
 void FTargetSkeleton::Initialize(
-	const USkeletalMesh* InSkeletalMesh,
-	const FName InRetargetPoseName,
-	const FIKRetargetPose* RetargetPose,
+const USkeletalMesh* InSkeletalMesh,
+	const ERetargetSourceOrTarget InSourceOrTarget,
+	const UIKRetargeter* InRetargetAsset,
 	const FName RetargetRootBone)
 {
 	Reset();
 
 	constexpr double TargetSkeletonScaleFactor = 1.0;
-	FRetargetSkeleton::Initialize(InSkeletalMesh, InRetargetPoseName, RetargetPose, RetargetRootBone, TargetSkeletonScaleFactor);
+	FRetargetSkeleton::Initialize(InSkeletalMesh, InSourceOrTarget, InRetargetAsset, RetargetRootBone, TargetSkeletonScaleFactor);
 
 	// make storage for per-bone "Is Retargeted" flag (used for hierarchy updates)
 	// these are bones that are in a target chain that is mapped to a source chain (ie, will actually be retargeted)
@@ -352,11 +409,11 @@ void FTargetSkeleton::Initialize(
 	bIsMaskInitialized = false;
 	IsBoneRetargeted.Init(false, BoneNames.Num());
 
-	// initialize pose for base local pose (may be overridden by pose copier, but uses retarget pose by default)
-	InputLocalPose = RetargetLocalPose;
+	// initialize base local pose (may be overridden by pose copier, but uses retarget pose by default)
+	InputLocalPose = RetargetPoses.GetLocalRetargetPose();
 	
 	// initialize storage for output pose (the result of the retargeting)
-	OutputGlobalPose = RetargetGlobalPose;
+	OutputGlobalPose = RetargetPoses.GetGlobalRetargetPose();
 }
 
 void FTargetSkeleton::Reset()
@@ -468,20 +525,22 @@ bool FResolvedBoneChain::IsValid() const
 
 void FResolvedBoneChain::UpdatePoseFromSkeleton(const FRetargetSkeleton& InSkeleton)
 {
+	const TArray<FTransform>& GlobalRetargetPose = InSkeleton.RetargetPoses.GetGlobalRetargetPose();
+	
 	// store all the initial bone transforms in the bone chain
 	RefPoseGlobalTransforms.Reset();
 	for (int32 Index=0; Index < BoneIndices.Num(); ++Index)
 	{
 		const int32 BoneIndex = BoneIndices[Index];
-		if (ensure(InSkeleton.RetargetGlobalPose.IsValidIndex(BoneIndex)))
+		if (ensure(GlobalRetargetPose.IsValidIndex(BoneIndex)))
 		{
-			RefPoseGlobalTransforms.Emplace(InSkeleton.RetargetGlobalPose[BoneIndex]);
+			RefPoseGlobalTransforms.Emplace(GlobalRetargetPose[BoneIndex]);
 		}
 	}
 
 	// get the local space of the chain in retarget pose
 	RefPoseLocalTransforms.SetNum(RefPoseGlobalTransforms.Num());
-	FillTransformsWithLocalSpaceOfChain(InSkeleton, InSkeleton.RetargetGlobalPose, BoneIndices, RefPoseLocalTransforms);
+	FillTransformsWithLocalSpaceOfChain(InSkeleton, GlobalRetargetPose, BoneIndices, RefPoseLocalTransforms);
 
 	// initialize storage for current local transforms
 	CurrentLocalTransforms = RefPoseLocalTransforms;
@@ -491,7 +550,7 @@ void FResolvedBoneChain::UpdatePoseFromSkeleton(const FRetargetSkeleton& InSkele
 	ChainParentInitialGlobalTransform = FTransform::Identity;
 	if (ChainParentBoneIndex != INDEX_NONE)
 	{
-		ChainParentInitialGlobalTransform = InSkeleton.RetargetGlobalPose[ChainParentBoneIndex];
+		ChainParentInitialGlobalTransform = GlobalRetargetPose[ChainParentBoneIndex];
 	}
 }
 
@@ -667,8 +726,8 @@ void FResolvedBoneChain::FillTransformsWithLocalSpaceOfChain(
 	}
 }
 
-void FResolvedBoneChain::FillTransformsWithGlobalRefPoseOfChain(
-	const FRetargetSkeleton& Skeleton,
+void FResolvedBoneChain::FillTransformsWithGlobalRetargetPoseOfChain(
+	const FRetargetSkeleton& InSkeleton,
 	const TArray<FTransform>& InGlobalPose,
 	const TArray<int32>& InBoneIndices,
 	TArray<FTransform>& OutGlobalTransforms)
@@ -676,20 +735,21 @@ void FResolvedBoneChain::FillTransformsWithGlobalRefPoseOfChain(
 	check(InBoneIndices.Num() == OutGlobalTransforms.Num())
 	
 	// update chain current transforms to the retarget pose in global space
+	TArray<FTransform> LocalRetargetPose = InSkeleton.RetargetPoses.GetLocalRetargetPose();
 	for (int32 ChainIndex=0; ChainIndex<InBoneIndices.Num(); ++ChainIndex)
 	{
 		// update first bone in chain based on the incoming parent
 		if (ChainIndex == 0)
 		{
 			const int32 BoneIndex = InBoneIndices[ChainIndex];
-			OutGlobalTransforms[ChainIndex] = Skeleton.GetGlobalRetargetPoseOfSingleBone(BoneIndex, InGlobalPose);
+			OutGlobalTransforms[ChainIndex] = InSkeleton.RetargetPoses.GetGlobalRetargetPoseOfSingleBone(InSkeleton, BoneIndex, InGlobalPose);
 		}
 		else
 		{
 			// all subsequent bones in chain are based on previous parent
 			const int32 BoneIndex = InBoneIndices[ChainIndex];
 			const FTransform& ParentGlobalTransform = OutGlobalTransforms[ChainIndex-1];
-			const FTransform& ChildLocalTransform = Skeleton.RetargetLocalPose[BoneIndex];
+			const FTransform& ChildLocalTransform = LocalRetargetPose[BoneIndex];
 			OutGlobalTransforms[ChainIndex] = ChildLocalTransform * ParentGlobalTransform;
 		}
 	}
@@ -1087,16 +1147,16 @@ void FIKRetargetProcessor::Initialize(
 	// initialize source skeleton
 	SourceSkeleton.Initialize(
 		SourceSkeletalMesh,
-		RetargeterAsset->GetCurrentRetargetPoseName(ERetargetSourceOrTarget::Source),
-		RetargeterAsset->GetCurrentRetargetPose(ERetargetSourceOrTarget::Source),
+		ERetargetSourceOrTarget::Source,
+		RetargeterAsset,
 		SourceIKRig->GetRetargetRoot(),
 		GetSourceScaleFactor());
 	
 	// initialize target skeleton
 	TargetSkeleton.Initialize(
 		TargetSkeletalMesh,
-		RetargeterAsset->GetCurrentRetargetPoseName(ERetargetSourceOrTarget::Target),
-		RetargeterAsset->GetCurrentRetargetPose(ERetargetSourceOrTarget::Target),
+		ERetargetSourceOrTarget::Target,
+		RetargeterAsset,
 		TargetIKRig->GetRetargetRoot());
 
 	// setup the goal container
@@ -1297,8 +1357,8 @@ void FIKRetargetProcessor::ScaleSourcePose(TArray<FTransform>& InOutSourceGlobal
 void FIKRetargetProcessor::GenerateBasePoses(TArray<FTransform>& InSourceGlobalPose)
 {
 	// start from the retarget pose
-	TargetSkeleton.OutputGlobalPose = TargetSkeleton.RetargetGlobalPose;
-	TargetSkeleton.InputLocalPose = TargetSkeleton.RetargetLocalPose;
+	TargetSkeleton.OutputGlobalPose = TargetSkeleton.RetargetPoses.GetGlobalRetargetPose();
+	TargetSkeleton.InputLocalPose = TargetSkeleton.RetargetPoses.GetLocalRetargetPose();
 
 	// partially override with input pose
 	FName CopyBasePoseRootBoneName;
@@ -1560,7 +1620,9 @@ FTransform FIKRetargetProcessor::GetRetargetPoseBoneTransform(
 	}
 
 	const FRetargetSkeleton& Skeleton = GetSkeleton(SourceOrTarget);
-	return BoneSpace == ERetargetBoneSpace::Global ? Skeleton.RetargetGlobalPose[BoneIndex] : Skeleton.RetargetLocalPose[BoneIndex];
+	const bool bGlobal = BoneSpace == ERetargetBoneSpace::Global;
+	const TArray<FTransform>& Pose = bGlobal ? Skeleton.RetargetPoses.GetGlobalRetargetPose() : Skeleton.RetargetPoses.GetLocalRetargetPose();
+	return Pose[BoneIndex];
 }
 
 float FIKRetargetProcessor::GetParamOfBoneInChain(
@@ -1732,27 +1794,55 @@ void FIKRetargetProcessor::UpdateRetargetPoseAtRuntime(
 	{
 		return; // retarget pose not found
 	}
+	
+	auto AddOrUpdateNewRetargetPose = [this, NewRetargetPoseName, NewRetargetPose, SourceOrTarget]()
+	{
+		const FName PelvisBoneName = GetPelvisBone(SourceOrTarget);
+		const double ScaleFactor = SourceOrTarget == ERetargetSourceOrTarget::Source ? GetSourceScaleFactor() : 1.0;
+		const bool bIsSource = SourceOrTarget == ERetargetSourceOrTarget::Source;
+		FRetargetSkeleton& RetargetSkeleton = bIsSource ? SourceSkeleton : TargetSkeleton;
+		return RetargetSkeleton.RetargetPoses.AddOrUpdateRetargetPose(
+			RetargetSkeleton,
+			NewRetargetPoseName,
+			NewRetargetPose,
+			PelvisBoneName,
+			ScaleFactor);
+	};
 
 	const bool bIsSource = SourceOrTarget == ERetargetSourceOrTarget::Source;
 	FRetargetSkeleton& RetargetSkeleton = bIsSource ? SourceSkeleton : TargetSkeleton;
-	const bool bSameRetargetPose = RetargetSkeleton.RetargetPoseName == NewRetargetPoseName;
-	const bool bSameVersion = RetargetSkeleton.RetargetPoseVersion == NewRetargetPose->GetVersion();
+	const FResolvedRetargetPose* NewResolvedPose = RetargetSkeleton.RetargetPoses.FindRetargetPoseByName(NewRetargetPoseName);
+	if (!NewResolvedPose)
+	{
+		// trying to switch to a retarget pose that hasn't been loaded yet, so load it
+		AddOrUpdateNewRetargetPose();
+		NewResolvedPose = RetargetSkeleton.RetargetPoses.FindRetargetPoseByName(NewRetargetPoseName);
+		if (!NewResolvedPose)
+		{
+			// could not load the new pose
+			return;
+		}
+	}
+
+	// verify if pose needs rebuilding
+	const FName CurrentRetargetPoseName = RetargetSkeleton.RetargetPoses.CurrentRetargetPoseName;
+	const FResolvedRetargetPose* CurrentResolvedPose = RetargetSkeleton.RetargetPoses.FindRetargetPoseByName(CurrentRetargetPoseName);
+	const bool bSameRetargetPose = NewResolvedPose->Name == CurrentResolvedPose->Name;
+	const bool bSameVersion = CurrentResolvedPose->Version == NewRetargetPose->GetVersion();
 	const double SourceScaleFactor = GetSourceScaleFactor();
-	const bool bSameGlobalScale = bIsSource ? FMath::IsNearlyEqual(RetargetSkeleton.RetargetPoseScale, SourceScaleFactor) : true;
+	const bool bSameGlobalScale = bIsSource ? FMath::IsNearlyEqual(CurrentResolvedPose->Scale, SourceScaleFactor) : true;
 	if (bSameRetargetPose && bSameVersion && bSameGlobalScale)
 	{
 		return; // retarget pose has not changed since it was initialized
 	}
 	
 	// re-generate the retarget pose
-	const double ScaleFactor = SourceOrTarget == ERetargetSourceOrTarget::Source ? SourceScaleFactor : 1.0;
-	const FName PelvisBoneName = GetPelvisBone(SourceOrTarget);
-	RetargetSkeleton.GenerateRetargetPose(NewRetargetPoseName, NewRetargetPose, PelvisBoneName, ScaleFactor);
+	AddOrUpdateNewRetargetPose();
 
 	// re-load the updated pose into the bone chains
 	AllBoneChains.UpdatePoseFromSkeleton(RetargetSkeleton, SourceOrTarget);
 
-	// re-initialize the op stack (ops sometimes cache retarget poses
+	// re-initialize the op stack (ops sometimes cache retarget poses)
 	for (FInstancedStruct& OpStruct : OpStack)
 	{
 		FIKRetargetOpBase& Op = OpStruct.GetMutable<FIKRetargetOpBase>();
diff --git a/Engine/Plugins/Animation/IKRig/Source/IKRig/Private/Retargeter/RetargetOps/AlignPoleVectorOp.cpp b/Engine/Plugins/Animation/IKRig/Source/IKRig/Private/Retargeter/RetargetOps/AlignPoleVectorOp.cpp
--- a/Engine/Plugins/Animation/IKRig/Source/IKRig/Private/Retargeter/RetargetOps/AlignPoleVectorOp.cpp
+++ b/Engine/Plugins/Animation/IKRig/Source/IKRig/Private/Retargeter/RetargetOps/AlignPoleVectorOp.cpp
@@ -18,12 +18,15 @@ bool FPoleVectorMatcher::Initialize(
 	Settings = InSettings;
 	SourceBoneChain = &InSourceBoneChain;
 	TargetBoneChain = &InTargetBoneChain;
+
+	const TArray<FTransform>& SourceRetargetPose = InSourceSkeleton.RetargetPoses.GetGlobalRetargetPose();
+	const TArray<FTransform>& TargetRetargetPose = InTargetSkeleton.RetargetPoses.GetGlobalRetargetPose();
 	
-	SourcePoleAxis = CalculateBestPoleAxisForChain(SourceBoneChain->BoneIndices, InSourceSkeleton.RetargetGlobalPose);
-	TargetPoleAxis = CalculateBestPoleAxisForChain(TargetBoneChain->BoneIndices, InTargetSkeleton.RetargetGlobalPose);
+	SourcePoleAxis = CalculateBestPoleAxisForChain(SourceBoneChain->BoneIndices, SourceRetargetPose);
+	TargetPoleAxis = CalculateBestPoleAxisForChain(TargetBoneChain->BoneIndices, TargetRetargetPose);
 
-	const FVector SourcePoleVector = CalculatePoleVector(SourcePoleAxis, SourceBoneChain->BoneIndices, InSourceSkeleton.RetargetGlobalPose);
-	const FVector TargetPoleVector = CalculatePoleVector(TargetPoleAxis, TargetBoneChain->BoneIndices, InTargetSkeleton.RetargetGlobalPose);
+	const FVector SourcePoleVector = CalculatePoleVector(SourcePoleAxis, SourceBoneChain->BoneIndices, SourceRetargetPose);
+	const FVector TargetPoleVector = CalculatePoleVector(TargetPoleAxis, TargetBoneChain->BoneIndices, TargetRetargetPose);
 	
 	TargetToSourceAngularOffsetAtRefPose = static_cast<float>(FMath::Acos(FVector::DotProduct(TargetPoleVector, SourcePoleVector)));
 
diff --git a/Engine/Plugins/Animation/IKRig/Source/IKRig/Private/Retargeter/RetargetOps/FKChainsOp.cpp b/Engine/Plugins/Animation/IKRig/Source/IKRig/Private/Retargeter/RetargetOps/FKChainsOp.cpp
--- a/Engine/Plugins/Animation/IKRig/Source/IKRig/Private/Retargeter/RetargetOps/FKChainsOp.cpp
+++ b/Engine/Plugins/Animation/IKRig/Source/IKRig/Private/Retargeter/RetargetOps/FKChainsOp.cpp
@@ -127,7 +127,7 @@ void FChainDecoderFK::DecodePose(
 	if (!Settings.EnableFK)
 	{
 		// put the chain in the global ref pose (globally rotated by parent bone in it's currently retargeted state)
-		BoneChain->FillTransformsWithGlobalRefPoseOfChain(TargetSkeleton, InOutGlobalPose, TargetBoneIndices, CurrentGlobalTransforms);
+		BoneChain->FillTransformsWithGlobalRetargetPoseOfChain(TargetSkeleton, InOutGlobalPose, TargetBoneIndices, CurrentGlobalTransforms);
 		
 		for (int32 ChainIndex=0; ChainIndex<TargetBoneIndices.Num(); ++ChainIndex)
 		{
@@ -212,6 +212,7 @@ void FChainDecoderFK::DecodePose(
 		const FQuat RotationDelta = SourceCurrentRotation * SourceInitialRotation.Inverse();
 		const FQuat TargetInitialRotation = TargetInitialTransform.GetRotation();
 		const FQuat OutRotation = RotationDelta * TargetInitialRotation;
+		
 
 		// calculate output POSITION based on translation mode setting
 		FTransform ParentGlobalTransform = FTransform::Identity;
@@ -225,7 +226,7 @@ void FChainDecoderFK::DecodePose(
 		{
 		case EFKChainTranslationMode::None:
 			{
-				const FVector InitialLocalOffset = TargetSkeleton.RetargetLocalPose[BoneIndex].GetTranslation();
+				const FVector InitialLocalOffset = TargetSkeleton.RetargetPoses.GetLocalRetargetPose()[BoneIndex].GetTranslation();
 				OutPosition = ParentGlobalTransform.TransformPosition(InitialLocalOffset);
 				break;
 			}
@@ -244,7 +245,7 @@ void FChainDecoderFK::DecodePose(
 			{
 				if (ChainIndex==0)
 				{
-					const FVector InitialLocalOffset = TargetSkeleton.RetargetLocalPose[BoneIndex].GetTranslation();
+					const FVector InitialLocalOffset = TargetSkeleton.RetargetPoses.GetLocalRetargetPose()[BoneIndex].GetTranslation();
 					OutPosition = ParentGlobalTransform.TransformPosition(InitialLocalOffset);
 					break;
 				}
@@ -253,7 +254,7 @@ void FChainDecoderFK::DecodePose(
 				double SourceChainLengthCurrent = BoneChain->GetChainLength(SourceChain.CurrentGlobalTransforms);
 				double StretchRatio = SourceChainLengthInitial < UE_KINDA_SMALL_NUMBER ? 1.0f : SourceChainLengthCurrent / SourceChainLengthInitial;
 				// stretch local translation
-				const FVector InitialLocalOffset = TargetSkeleton.RetargetLocalPose[BoneIndex].GetTranslation();
+				const FVector InitialLocalOffset = TargetSkeleton.RetargetPoses.GetLocalRetargetPose()[BoneIndex].GetTranslation();
 				OutPosition = ParentGlobalTransform.TransformPosition(InitialLocalOffset * StretchRatio);
 				break;
 			}
@@ -262,13 +263,13 @@ void FChainDecoderFK::DecodePose(
 				const double Param = BoneChain->Params[ChainIndex];
 				double StretchRatio = SourceChain.BoneChain->GetStretchAtParam( SourceChain.BoneChain->RefPoseGlobalTransforms, SourceChain.CurrentGlobalTransforms, Param);
 				// stretch local translation
-				const FVector InitialLocalOffset = TargetSkeleton.RetargetLocalPose[BoneIndex].GetTranslation();
+				const FVector InitialLocalOffset = TargetSkeleton.RetargetPoses.GetLocalRetargetPose()[BoneIndex].GetTranslation();
 				OutPosition = ParentGlobalTransform.TransformPosition(InitialLocalOffset * StretchRatio);
 				break;
 			}
 			default:
 			{
-				const FVector InitialLocalOffset = TargetSkeleton.RetargetLocalPose[BoneIndex].GetTranslation();
+				const FVector InitialLocalOffset = TargetSkeleton.RetargetPoses.GetLocalRetargetPose()[BoneIndex].GetTranslation();
 				OutPosition = ParentGlobalTransform.TransformPosition(InitialLocalOffset);
 				break;
 			}
diff --git a/Engine/Plugins/Animation/IKRig/Source/IKRig/Private/Retargeter/RetargetOps/PelvisMotionOp.cpp b/Engine/Plugins/Animation/IKRig/Source/IKRig/Private/Retargeter/RetargetOps/PelvisMotionOp.cpp
--- a/Engine/Plugins/Animation/IKRig/Source/IKRig/Private/Retargeter/RetargetOps/PelvisMotionOp.cpp
+++ b/Engine/Plugins/Animation/IKRig/Source/IKRig/Private/Retargeter/RetargetOps/PelvisMotionOp.cpp
@@ -139,7 +139,7 @@ bool FIKRetargetPelvisMotionOp::InitializeSource(
 	}
 	
 	// record initial root data
-	const FTransform InitialTransform = SourceSkeleton.RetargetGlobalPose[Source.BoneIndex]; 
+	const FTransform InitialTransform = SourceSkeleton.RetargetPoses.GetGlobalRetargetPose()[Source.BoneIndex]; 
 	float InitialHeight = static_cast<float>(InitialTransform.GetTranslation().Z);
 	Source.InitialRotation = InitialTransform.GetRotation();
 
@@ -173,7 +173,7 @@ bool FIKRetargetPelvisMotionOp::InitializeTarget(
 		return false;
 	}
 
-	const FTransform TargetInitialTransform = TargetSkeleton.RetargetGlobalPose[Target.BoneIndex];
+	const FTransform TargetInitialTransform = TargetSkeleton.RetargetPoses.GetGlobalRetargetPose()[Target.BoneIndex];
 	Target.InitialHeight = static_cast<float>(TargetInitialTransform.GetTranslation().Z);
 	Target.InitialRotation = TargetInitialTransform.GetRotation();
 	Target.InitialPosition = TargetInitialTransform.GetTranslation();
diff --git a/Engine/Plugins/Animation/IKRig/Source/IKRig/Private/Retargeter/RetargetOps/PinBoneOp.cpp b/Engine/Plugins/Animation/IKRig/Source/IKRig/Private/Retargeter/RetargetOps/PinBoneOp.cpp
--- a/Engine/Plugins/Animation/IKRig/Source/IKRig/Private/Retargeter/RetargetOps/PinBoneOp.cpp
+++ b/Engine/Plugins/Animation/IKRig/Source/IKRig/Private/Retargeter/RetargetOps/PinBoneOp.cpp
@@ -46,11 +46,15 @@ void FPinBoneData::CachePinData(ERetargetSourceOrTarget InSkeletonToCopyFrom, co
     const FRetargetSkeleton& SkeletonToCopyFrom = Processor.GetSkeleton(InSkeletonToCopyFrom);
     const FRetargetSkeleton& SkeletonToCopyTo = Processor.GetSkeleton(ERetargetSourceOrTarget::Target);
 
+	// get retarget pose of source and target
+	const TArray<FTransform>& FromRetargetPose = SkeletonToCopyFrom.RetargetPoses.GetGlobalRetargetPose();
+	const TArray<FTransform>& ToRetargetPose = SkeletonToCopyTo.RetargetPoses.GetGlobalRetargetPose();
+
 	// get relevant transforms from the retarget poses
-	FTransform BoneToCopyFromParentRefPoseGlobal = GetParentTransform(BoneToCopyFrom.BoneIndex, SkeletonToCopyFrom, SkeletonToCopyFrom.RetargetGlobalPose);
-	FTransform BoneToCopyToParentRefPoseGlobal = GetParentTransform(BoneToCopyTo.BoneIndex, SkeletonToCopyTo, SkeletonToCopyTo.RetargetGlobalPose);
-	FTransform BoneToCopyFromRefPoseGlobal = SkeletonToCopyFrom.RetargetGlobalPose[BoneToCopyFrom.BoneIndex];
-	FTransform BoneToCopyToRefPoseGlobal = SkeletonToCopyTo.RetargetGlobalPose[BoneToCopyTo.BoneIndex];
+	FTransform BoneToCopyFromParentRefPoseGlobal = GetParentTransform(BoneToCopyFrom.BoneIndex, SkeletonToCopyFrom, FromRetargetPose);
+	FTransform BoneToCopyToParentRefPoseGlobal = GetParentTransform(BoneToCopyTo.BoneIndex, SkeletonToCopyTo, ToRetargetPose);
+	FTransform BoneToCopyFromRefPoseGlobal = FromRetargetPose[BoneToCopyFrom.BoneIndex];
+	FTransform BoneToCopyToRefPoseGlobal = ToRetargetPose[BoneToCopyTo.BoneIndex];
 
 	// cache offset from BoneToCopyFrom to BoneToCopyTo in ref pose
     OffsetFromBoneToCopyFromInRefPose = BoneToCopyFromRefPoseGlobal.GetRelativeTransform(BoneToCopyToRefPoseGlobal);
@@ -335,8 +339,11 @@ FTransform FIKRetargetPinBoneOp::GetNewBoneTransform(
 		}
 	case EPinBoneRotationMode::MaintainOffsetFromBoneToCopyFrom:
 		{
-			const FTransform& BoneToCopyFromRefPoseGlobal = CopyFrom.RetargetGlobalPose[PinData.BoneToCopyFrom.BoneIndex];
-			const FTransform& BoneToCopyToRefPoseGlobal = CopyTo.RetargetGlobalPose[PinData.BoneToCopyTo.BoneIndex];
+			// get retarget pose of source and target
+			const TArray<FTransform>& FromRetargetPose = CopyFrom.RetargetPoses.GetGlobalRetargetPose();
+			const TArray<FTransform>& ToRetargetPose = CopyTo.RetargetPoses.GetGlobalRetargetPose();
+			const FTransform& BoneToCopyFromRefPoseGlobal = FromRetargetPose[PinData.BoneToCopyFrom.BoneIndex];
+			const FTransform& BoneToCopyToRefPoseGlobal = ToRetargetPose[PinData.BoneToCopyTo.BoneIndex];
 			const FTransform Offset = BoneToCopyFromRefPoseGlobal.GetRelativeTransform(BoneToCopyToRefPoseGlobal);
 			const FTransform BoneToCopyFromTransform = PoseToCopyFrom[PinData.BoneToCopyFrom.BoneIndex];
 			Result.SetRotation((Offset * BoneToCopyFromTransform).GetRotation());
diff --git a/Engine/Plugins/Animation/IKRig/Source/IKRig/Private/Retargeter/RetargetOps/RetargetPoseOp.cpp b/Engine/Plugins/Animation/IKRig/Source/IKRig/Private/Retargeter/RetargetOps/RetargetPoseOp.cpp
new file mode 100644
--- /dev/null
+++ b/Engine/Plugins/Animation/IKRig/Source/IKRig/Private/Retargeter/RetargetOps/RetargetPoseOp.cpp
@@ -0,0 +1,72 @@
+﻿// Copyright Epic Games, Inc. All Rights Reserved.
+
+#include "Retargeter/RetargetOps/RetargetPoseOp.h"
+
+const UClass* FIKRetargetPoseOpSettings::GetControllerType() const
+{
+	return UIKRetargetPoseController::StaticClass();
+}
+
+void FIKRetargetPoseOpSettings::CopySettingsAtRuntime(const FIKRetargetOpSettingsBase* InSettingsToCopyFrom)
+{
+	// copy all properties
+	static TArray<FName> PropertiesToIgnore = {};
+	FIKRetargetOpBase::CopyStructProperties(
+		FIKRetargetPoseOpSettings::StaticStruct(),
+		InSettingsToCopyFrom,
+		this,
+		PropertiesToIgnore);
+}
+
+bool FIKRetargetPoseOp::Initialize(
+	const FIKRetargetProcessor& InProcessor,
+	const FRetargetSkeleton& InSourceSkeleton,
+	const FTargetSkeleton& InTargetSkeleton,
+	const FIKRetargetOpBase* InParentOp,
+	FIKRigLogger& Log)
+{
+	// TODO cache the retarget pose in the processor 
+	bIsInitialized = true;
+	return true;
+}
+
+void FIKRetargetPoseOp::Run(
+	FIKRetargetProcessor& InProcessor,
+	const double InDeltaTime,
+	const TArray<FTransform>& InSourceGlobalPose,
+	TArray<FTransform>& OutTargetGlobalPose)
+{
+	// TODO update/switch to current retarget pose
+	// TODO apply pose to output pose
+}
+
+FIKRetargetOpSettingsBase* FIKRetargetPoseOp::GetSettings()
+{
+	return &Settings;
+}
+
+void FIKRetargetPoseOp::SetSettings(const FIKRetargetOpSettingsBase* InSettings)
+{
+	Settings.CopySettingsAtRuntime(InSettings);
+}
+
+const UScriptStruct* FIKRetargetPoseOp::GetSettingsType() const
+{
+	return FIKRetargetPoseOpSettings::StaticStruct();
+}
+
+const UScriptStruct* FIKRetargetPoseOp::GetType() const
+{
+	return FIKRetargetPoseOp::StaticStruct();
+}
+
+FIKRetargetPoseOpSettings UIKRetargetPoseController::GetSettings()
+{
+	return *reinterpret_cast<FIKRetargetPoseOpSettings*>(OpSettingsToControl);
+}
+
+void UIKRetargetPoseController::SetSettings(FIKRetargetPoseOpSettings InSettings)
+{
+	reinterpret_cast<FIKRetargetPoseOpSettings*>(OpSettingsToControl)->CopySettingsAtRuntime(&InSettings);
+}
+
diff --git a/Engine/Plugins/Animation/IKRig/Source/IKRig/Private/Retargeter/RetargetOps/RootMotionGeneratorOp.cpp b/Engine/Plugins/Animation/IKRig/Source/IKRig/Private/Retargeter/RetargetOps/RootMotionGeneratorOp.cpp
--- a/Engine/Plugins/Animation/IKRig/Source/IKRig/Private/Retargeter/RetargetOps/RootMotionGeneratorOp.cpp
+++ b/Engine/Plugins/Animation/IKRig/Source/IKRig/Private/Retargeter/RetargetOps/RootMotionGeneratorOp.cpp
@@ -114,9 +114,13 @@ bool FIKRetargetRootMotionOp::Initialize(
 		return false;
 	}
 
-	TargetPelvisInRefPose = TargetSkeleton.RetargetGlobalPose[TargetPelvisIndex];
-	SourceRootInRefPose = SourceSkeleton.RetargetGlobalPose[SourceRootIndex];
-	TargetRootInRefPose = TargetSkeleton.RetargetGlobalPose[TargetRootIndex];
+	// get retarget pose of source and target
+	const TArray<FTransform>& TargetRetargetPose = TargetSkeleton.RetargetPoses.GetGlobalRetargetPose();
+	const TArray<FTransform>& SourceRetargetPose = SourceSkeleton.RetargetPoses.GetGlobalRetargetPose();
+
+	TargetPelvisInRefPose = TargetRetargetPose[TargetPelvisIndex];
+	SourceRootInRefPose = SourceRetargetPose[SourceRootIndex];
+	TargetRootInRefPose = TargetRetargetPose[TargetRootIndex];
 	TargetPelvisRelativeToTargetRootRefPose = TargetRootInRefPose.GetRelativeTransform(TargetPelvisInRefPose);
 
 	bIsInitialized = true;
diff --git a/Engine/Plugins/Animation/IKRig/Source/IKRig/Private/Retargeter/RetargetOps/StrideWarpingOp.cpp b/Engine/Plugins/Animation/IKRig/Source/IKRig/Private/Retargeter/RetargetOps/StrideWarpingOp.cpp
--- a/Engine/Plugins/Animation/IKRig/Source/IKRig/Private/Retargeter/RetargetOps/StrideWarpingOp.cpp
+++ b/Engine/Plugins/Animation/IKRig/Source/IKRig/Private/Retargeter/RetargetOps/StrideWarpingOp.cpp
@@ -98,6 +98,7 @@ bool FIKRetargetStrideWarpingOp::Initialize(
 	
 	const FRetargeterBoneChains& BoneChains = InProcessor.GetBoneChains();
 	const FIKRigGoalContainer& GoalContainer = InProcessor.GetIKRigGoalContainer();
+	const TArray<FTransform>& TargetGlobalRetargetPose = InTargetSkeleton.RetargetPoses.GetGlobalRetargetPose();
 	for (const FRetargetStrideWarpChainSettings& ChainSettings : Settings.ChainSettings)
 	{
 		if (!ChainSettings.EnableStrideWarping)
@@ -144,7 +145,7 @@ bool FIKRetargetStrideWarpingOp::Initialize(
 		}
 		
 		// store goal to warp
-		GoalsToWarp.Emplace(Goal->Name, InTargetSkeleton.RetargetGlobalPose[GoalBoneIndex]);
+		GoalsToWarp.Emplace(Goal->Name, TargetGlobalRetargetPose[GoalBoneIndex]);
 	}
 
 	bIsInitialized = !GoalsToWarp.IsEmpty();
@@ -270,12 +271,13 @@ void FIKRetargetStrideWarpingOp::Run(
 	case EWarpingDirectionSource::RootBone:
 		{
 			const FTargetSkeleton& TargetSkeleton = InProcessor.GetTargetSkeleton();
+			const TArray<FTransform>& TargetRetargetPose = TargetSkeleton.RetargetPoses.GetGlobalRetargetPose();
 			// use the root bone
-			InitialBodyPosition = TargetSkeleton.RetargetGlobalPose[0].GetTranslation();
+			InitialBodyPosition = TargetRetargetPose[0].GetTranslation();
 			CurrentBodyPosition = TargetSkeleton.OutputGlobalPose[0].GetTranslation();
 
 			// get the forward vectors of the root bone
-			const FQuat DeltaRotation = TargetSkeleton.RetargetGlobalPose[0].GetRotation() * TargetSkeleton.OutputGlobalPose[0].GetRotation().Inverse();
+			const FQuat DeltaRotation = TargetRetargetPose[0].GetRotation() * TargetSkeleton.OutputGlobalPose[0].GetRotation().Inverse();
 			FVector ForwardOrig = Settings.GetAxisVector(Settings.ForwardDirection);
 			FVector ForwardCurrent = DeltaRotation.RotateVector(ForwardOrig);
 
diff --git a/Engine/Plugins/Animation/IKRig/Source/IKRig/Public/Retargeter/IKRetargetProcessor.h b/Engine/Plugins/Animation/IKRig/Source/IKRig/Public/Retargeter/IKRetargetProcessor.h
--- a/Engine/Plugins/Animation/IKRig/Source/IKRig/Public/Retargeter/IKRetargetProcessor.h
+++ b/Engine/Plugins/Animation/IKRig/Source/IKRig/Public/Retargeter/IKRetargetProcessor.h
@@ -30,6 +30,47 @@ struct FIKRetargetPelvisMotionOp;
 struct FResolvedBoneChain;
 struct FIKRetargetDebugDrawState;
 
+
+// cached retarget pose transforms (both local and global) after it has been applied/resolved on a specific skeletal mesh
+struct FResolvedRetargetPose
+{
+	FResolvedRetargetPose(const FName InName) : Name(InName) {}
+	
+	FName Name;						// the name of the retarget pose this was initialized with
+	int32 Version;					// the version of the retarget pose this was initialized with (transient)
+	double Scale;					// the scale of the retarget pose this was initialized with (transient)
+	
+	TArray<FTransform> LocalPose;	// local space retarget pose
+	TArray<FTransform> GlobalPose;	// global space retarget pose
+};
+
+// the set of all resolved/cached retarget poses in use for a source or target skeleton
+struct FResolvedRetargetPoseSet
+{
+	FName CurrentRetargetPoseName;
+	TArray<FResolvedRetargetPose> RetargetPoses;	// all the retarget poses resolved on the target skeleton
+
+	UE_API FResolvedRetargetPose& AddOrUpdateRetargetPose(
+		const FRetargetSkeleton& InSkeleton,
+		const FName InRetargetPoseName,
+		const FIKRetargetPose* InRetargetPose,
+		const FName PelvisBoneName,
+		const double ScaleFactor);
+
+	UE_API const FResolvedRetargetPose* FindRetargetPoseByName(const FName InRetargetPoseName) const;
+
+	UE_API const TArray<FTransform>& GetLocalRetargetPose() const;
+	
+	UE_API const TArray<FTransform>& GetGlobalRetargetPose() const;
+
+	UE_API FTransform GetGlobalRetargetPoseOfSingleBone(
+		const FRetargetSkeleton& InSkeleton,
+		const int32 BoneIndex,
+		const TArray<FTransform>& InGlobalPose) const;
+
+	UE_API void Reset();
+};
+
 /**
  * A retarget skeleton contains:
  * 1. hierarchy data (bones, parents)
@@ -42,19 +83,13 @@ struct FRetargetSkeleton
 {
 	TArray<FName> BoneNames;				// list of all bone names in ref skeleton order
 	TArray<int32> ParentIndices;			// per-bone indices of parent bones (the hierarchy)
-	TArray<FTransform> RetargetLocalPose;	// local space retarget pose
-	TArray<FTransform> RetargetGlobalPose;	// global space retarget pose
-	
-	FName RetargetPoseName;					// the name of the retarget pose this was initialized with
-	int32 RetargetPoseVersion;				// the version of the retarget pose this was initialized with (transient)
-	double RetargetPoseScale;				// the scale of the retarget pose this was initialized with (transient)
-	
+	FResolvedRetargetPoseSet RetargetPoses;	// the set of cached retarget poses
 	const USkeletalMesh* SkeletalMesh;		// the skeletal mesh this was initialized with
 
 	UE_API void Initialize(
 		const USkeletalMesh* InSkeletalMesh,
-		const FName InRetargetPoseName,
-		const FIKRetargetPose* RetargetPose,
+		const ERetargetSourceOrTarget InSourceOrTarget,
+		const UIKRetargeter* InRetargetAsset,
 		const FName RetargetRootBone,
 		const double ScaleFactor);
 	
@@ -95,10 +130,6 @@ struct FRetargetSkeleton
 		TArray<FTransform>& OutLocalPose,
 		const TArray<FTransform>& InGlobalPose) const;
 
-	UE_API FTransform GetGlobalRetargetPoseOfSingleBone(
-		const int32 BoneIndex,
-		const TArray<FTransform>& InGlobalPose) const;
-
 	UE_API int32 GetCachedEndOfBranchIndex(const int32 InBoneIndex) const;
 
 	UE_API void GetChildrenIndices(const int32 BoneIndex, TArray<int32>& OutChildren) const;
@@ -125,8 +156,8 @@ struct FTargetSkeleton : public FRetargetSkeleton
 
 	void Initialize(
 		const USkeletalMesh* InSkeletalMesh,
-		const FName InRetargetPoseName,
-		const FIKRetargetPose* RetargetPose,
+		const ERetargetSourceOrTarget InSourceOrTarget,
+		const UIKRetargeter* InRetargetAsset,
 		const FName RetargetRootBone);
 
 	void Reset();
@@ -188,8 +219,8 @@ struct FResolvedBoneChain
 		const TArray<FTransform>& InGlobalPose,
 		const TArray<int32>& BoneIndices,
 		TArray<FTransform>& OutLocalTransforms);
-
-	static UE_API void FillTransformsWithGlobalRefPoseOfChain(
+	
+	static UE_API void FillTransformsWithGlobalRetargetPoseOfChain(
 		const FRetargetSkeleton& Skeleton,
 		const TArray<FTransform>& InGlobalPose,
 		const TArray<int32>& BoneIndices,
diff --git a/Engine/Plugins/Animation/IKRig/Source/IKRig/Public/Retargeter/RetargetOps/RetargetPoseOp.h b/Engine/Plugins/Animation/IKRig/Source/IKRig/Public/Retargeter/RetargetOps/RetargetPoseOp.h
new file mode 100644
--- /dev/null
+++ b/Engine/Plugins/Animation/IKRig/Source/IKRig/Public/Retargeter/RetargetOps/RetargetPoseOp.h
@@ -0,0 +1,85 @@
+﻿// Copyright Epic Games, Inc. All Rights Reserved.
+#pragma once
+
+#include "Retargeter/IKRetargetOps.h"
+#include "Retargeter/IKRetargetSettings.h"
+
+#include "RetargetPoseOp.generated.h"
+
+#define LOCTEXT_NAMESPACE "RetargetPoseOp"
+
+UENUM()
+enum class ERetargetPoseMode : uint8
+{
+	RetargetPose,
+	AdditivePose,
+};
+
+USTRUCT(BlueprintType, meta = (DisplayName = "Retarget Pose Op Settings"))
+struct IKRIG_API FIKRetargetPoseOpSettings : public FIKRetargetOpSettingsBase
+{
+	GENERATED_BODY()
+
+	UPROPERTY(EditAnywhere, BlueprintReadWrite, Category="Pose")
+	ERetargetSourceOrTarget SourceOrTarget;
+
+	UPROPERTY(EditAnywhere, BlueprintReadWrite, Category="Pose")
+	ERetargetPoseMode Mode;
+	
+	UPROPERTY(EditAnywhere, BlueprintReadWrite, Category="Pose")
+	FName CurrentPose;
+	
+	virtual const UClass* GetControllerType() const override;
+	
+	virtual void CopySettingsAtRuntime(const FIKRetargetOpSettingsBase* InSettingsToCopyFrom) override;
+};
+
+USTRUCT(BlueprintType, meta = (DisplayName = "Retarget Pose"))
+struct IKRIG_API FIKRetargetPoseOp : public FIKRetargetOpBase
+{
+	GENERATED_BODY()
+	
+	virtual bool Initialize(
+		const FIKRetargetProcessor& InProcessor,
+		const FRetargetSkeleton& InSourceSkeleton,
+		const FTargetSkeleton& InTargetSkeleton,
+		const FIKRetargetOpBase* InParentOp,
+		FIKRigLogger& Log) override;
+	
+	virtual void Run(
+		FIKRetargetProcessor& InProcessor,
+		const double InDeltaTime,
+		const TArray<FTransform>& InSourceGlobalPose,
+		TArray<FTransform>& OutTargetGlobalPose) override;
+
+	virtual FIKRetargetOpSettingsBase* GetSettings() override;
+
+	virtual void SetSettings(const FIKRetargetOpSettingsBase* InSettings) override;
+	
+	virtual const UScriptStruct* GetSettingsType() const override;
+	
+	virtual const UScriptStruct* GetType() const override;
+
+	UPROPERTY()
+	FIKRetargetPoseOpSettings Settings;
+};
+
+/* The blueprint/python API for editing a Retarget Pose Op */
+UCLASS(BlueprintType)
+class IKRIG_API UIKRetargetPoseController : public UIKRetargetOpControllerBase
+{
+	GENERATED_BODY()
+	
+public:
+	/* Get the current op settings as a struct.
+	 * @return FIKRetargetPoseOpSettings struct with the current settings used by the op. */
+	UFUNCTION(BlueprintCallable, Category = Settings)
+	FIKRetargetPoseOpSettings GetSettings();
+
+	/* Set the op settings. Input is a custom struct type for this op.
+	 * @param InSettings a FIKRetargetPoseOpSettings struct containing all the settings to apply to this op */
+	UFUNCTION(BlueprintCallable, Category = Settings)
+	void SetSettings(FIKRetargetPoseOpSettings InSettings);
+};
+
+#undef LOCTEXT_NAMESPACE
diff --git a/Engine/Plugins/Animation/IKRig/Source/IKRigEditor/Private/IKRigEditor.cpp b/Engine/Plugins/Animation/IKRig/Source/IKRigEditor/Private/IKRigEditor.cpp
--- a/Engine/Plugins/Animation/IKRig/Source/IKRigEditor/Private/IKRigEditor.cpp
+++ b/Engine/Plugins/Animation/IKRig/Source/IKRigEditor/Private/IKRigEditor.cpp
@@ -101,6 +101,9 @@ void FIKRigEditor::StartupModule()
 	PropertyEditorModule.RegisterCustomPropertyTypeLayout(
 		"IKRetargetAlignPoleVectorOpSettings",
 		FOnGetPropertyTypeCustomizationInstance::CreateStatic(&FPoleVectorOpCustomization::MakeInstance));
+	PropertyEditorModule.RegisterCustomPropertyTypeLayout(
+			"IKRetargetPoseOpSettings",
+			FOnGetPropertyTypeCustomizationInstance::CreateStatic(&FRetargetPoseOpCustomization::MakeInstance));
 
 	// register a thumbnail renderer for the assets
 	UThumbnailManager::Get().RegisterCustomRenderer(UIKRigDefinition::StaticClass(), UIKRigThumbnailRenderer::StaticClass());
diff --git a/Engine/Plugins/Animation/IKRig/Source/IKRigEditor/Private/RetargetEditor/IKRetargetEditorController.cpp b/Engine/Plugins/Animation/IKRig/Source/IKRigEditor/Private/RetargetEditor/IKRetargetEditorController.cpp
--- a/Engine/Plugins/Animation/IKRig/Source/IKRigEditor/Private/RetargetEditor/IKRetargetEditorController.cpp
+++ b/Engine/Plugins/Animation/IKRig/Source/IKRigEditor/Private/RetargetEditor/IKRetargetEditorController.cpp
@@ -603,9 +603,10 @@ void FIKRetargetEditorController::GetGlobalRetargetPoseOfImmediateChildren(
 	RetargetSkeleton.GetChildrenIndices(BoneIndex, OutChildrenIndices);
 
 	// get the positions of the immediate children
+	const TArray<FTransform>& RetargetPose = RetargetSkeleton.RetargetPoses.GetGlobalRetargetPose();
 	for (const int32& ChildIndex : OutChildrenIndices)
 	{
-		OutChildrenPositions.Emplace(RetargetSkeleton.RetargetGlobalPose[ChildIndex].GetTranslation());
+		OutChildrenPositions.Emplace(RetargetPose[ChildIndex].GetTranslation());
 	}
 
 	// apply scale and offset to positions
@@ -1620,7 +1621,6 @@ void FIKRetargetEditorController::FixZeroHeightRetargetRoot(ERetargetSourceOrTar
 		return;
 	}
 
-	const FName CurrentRetargetPoseName = AssetController->GetCurrentRetargetPoseName(SourceOrTarget);
 	FIKRetargetPose& CurrentRetargetPose = AssetController->GetCurrentRetargetPose(SourceOrTarget);
 	const FName RetargetRootBoneName = AssetController->GetPelvisBone(SourceOrTarget);
 	if (RetargetRootBoneName == NAME_None)
@@ -1633,8 +1633,8 @@ void FIKRetargetEditorController::FixZeroHeightRetargetRoot(ERetargetSourceOrTar
 	FRetargetSkeleton DummySkeleton;
 	DummySkeleton.Initialize(
 		SkeletalMesh,
-		CurrentRetargetPoseName,
-		&CurrentRetargetPose,
+		SourceOrTarget,
+		AssetController->GetAsset(),
 		RetargetRootBoneName,
 		ScaleFactor);
 
@@ -1644,7 +1644,7 @@ void FIKRetargetEditorController::FixZeroHeightRetargetRoot(ERetargetSourceOrTar
 		return;
 	}
 
-	const FTransform& RootTransform = DummySkeleton.RetargetGlobalPose[RootBoneIndex];
+	const FTransform& RootTransform = DummySkeleton.RetargetPoses.GetGlobalRetargetPose()[RootBoneIndex];
 	if (RootTransform.GetLocation().Z < 1.0f)
 	{
 		if (PromptToFixRootHeight(SourceOrTarget))
diff --git a/Engine/Plugins/Animation/IKRig/Source/IKRigEditor/Private/RetargetEditor/IKRetargetOpDetails.cpp b/Engine/Plugins/Animation/IKRig/Source/IKRigEditor/Private/RetargetEditor/IKRetargetOpDetails.cpp
--- a/Engine/Plugins/Animation/IKRig/Source/IKRigEditor/Private/RetargetEditor/IKRetargetOpDetails.cpp
+++ b/Engine/Plugins/Animation/IKRig/Source/IKRigEditor/Private/RetargetEditor/IKRetargetOpDetails.cpp
@@ -4,6 +4,7 @@
 
 #include "IDetailChildrenBuilder.h"
 #include "IPropertyUtilities.h"
+#include "Widgets/Input/SComboBox.h"
 #include "RetargetEditor/IKRetargeterController.h"
 #include "RetargetEditor/SIKRetargetChainMapList.h"
 #include "Retargeter/IKRetargeter.h"
@@ -536,5 +537,102 @@ void FPoleVectorOpCustomization::CustomizeChildren(
 	];
 }
 
+void FRetargetPoseOpCustomization::CustomizeChildren(
+	TSharedRef<IPropertyHandle> StructPropertyHandle,
+	IDetailChildrenBuilder& ChildBuilder,
+	IPropertyTypeCustomizationUtils& StructCustomizationUtils)
+{
+	if (!LoadAndValidateStructToCustomize(StructPropertyHandle, StructCustomizationUtils))
+	{
+		return;
+	}
+
+	// store property handles for callbacks
+	SourceOrTargetProperty = StructPropertyHandle->GetChildHandle(GET_MEMBER_NAME_CHECKED(FIKRetargetPoseOpSettings, SourceOrTarget));
+	CurrentPoseProperty = StructPropertyHandle->GetChildHandle(GET_MEMBER_NAME_CHECKED(FIKRetargetPoseOpSettings, CurrentPose));
+
+	// add the source/target property
+	ChildBuilder.AddProperty(SourceOrTargetProperty.ToSharedRef());
+
+	// add the mode property
+	TSharedPtr<IPropertyHandle> ModeProperty = StructPropertyHandle->GetChildHandle(GET_MEMBER_NAME_CHECKED(FIKRetargetPoseOpSettings, Mode));
+	ChildBuilder.AddProperty(ModeProperty.ToSharedRef());
+
+	// add a callback to update the pose name list when switching source/target
+	SourceOrTargetProperty->SetOnPropertyValueChanged((FSimpleDelegate::CreateLambda([this]()
+	{
+		UpdatePoseNameOptions();
+	})));
+	
+	UpdatePoseNameOptions();
+
+	// add dropdown menu to select retarget pose
+	ChildBuilder.AddCustomRow(LOCTEXT("CurrentPoseLabel", "Current Pose"))
+	.NameContent()
+	[
+		CurrentPoseProperty->CreatePropertyNameWidget()
+	]
+	.ValueContent()
+	.MinDesiredWidth(200)
+	[
+		SNew(SComboBox<TSharedPtr<FName>>)
+		.OptionsSource(&PoseNameOptions)
+		.InitiallySelectedItem(CurrentPoseOption)
+		.OnGenerateWidget_Lambda([](TSharedPtr<FName> InItem)
+		{
+			return SNew(STextBlock).Text(FText::FromName(*InItem));
+		})
+		.OnSelectionChanged_Lambda([this](TSharedPtr<FName> NewSelection, ESelectInfo::Type)
+		{
+			if (NewSelection.IsValid())
+			{
+				CurrentPoseProperty->SetValue(*NewSelection);
+			}
+		})
+		[
+			SNew(STextBlock)
+			.Text_Lambda([this]() {
+				FName Value;
+				CurrentPoseProperty->GetValue(Value);
+				return FText::FromName(Value);
+			})
+		]
+	];
+}
+
+void FRetargetPoseOpCustomization::UpdatePoseNameOptions()
+{
+	// determine currently selected pose
+	FName CurrentPoseName;
+	CurrentPoseProperty->GetValue(CurrentPoseName);
+	
+	// get all the retarget poses
+	int64 EnumValueInt;
+	SourceOrTargetProperty->GetValue(EnumValueInt);
+	ERetargetSourceOrTarget SourceOrTarget = static_cast<ERetargetSourceOrTarget>(EnumValueInt);
+	const TMap<FName, FIKRetargetPose>& RetargetPoses = AssetController->GetRetargetPoses(SourceOrTarget);
+	TArray<FName> PoseNames;
+	RetargetPoses.GetKeys(PoseNames);
+
+	// reset list of names
+	PoseNameOptions.Reset(PoseNames.Num());
+
+	// add all the other poses
+	for (const FName& PoseName : PoseNames)
+	{
+		TSharedPtr<FName> PostNameOption = MakeShared<FName>(PoseName);
+		if (PoseName == CurrentPoseName)
+		{
+			CurrentPoseOption = PostNameOption;
+		}
+		PoseNameOptions.Emplace(PostNameOption);
+	}
+
+	// default to first pose if stored pose no longer available
+	if (!CurrentPoseOption.IsValid())
+	{
+		CurrentPoseOption = PoseNameOptions[0];
+	}
+}
 
 #undef LOCTEXT_NAMESPACE
diff --git a/Engine/Plugins/Animation/IKRig/Source/IKRigEditor/Private/RetargetEditor/IKRetargeterPoseGenerator.cpp b/Engine/Plugins/Animation/IKRig/Source/IKRigEditor/Private/RetargetEditor/IKRetargeterPoseGenerator.cpp
--- a/Engine/Plugins/Animation/IKRig/Source/IKRigEditor/Private/RetargetEditor/IKRetargeterPoseGenerator.cpp
+++ b/Engine/Plugins/Animation/IKRig/Source/IKRigEditor/Private/RetargetEditor/IKRetargeterPoseGenerator.cpp
@@ -304,7 +304,7 @@ void FRetargetAutoPoseGenerator::SnapToGround(const FName BoneToPutOnGround, con
 			{
 				continue;
 			}
-			const FTransform& BoneTransform = Skeleton.RetargetGlobalPose[BoneIndex];
+			const FTransform& BoneTransform = Skeleton.RetargetPoses.GetGlobalRetargetPose()[BoneIndex];
 			const float BoneHeight = static_cast<float>(BoneTransform.GetLocation().Z);
 			if (BoneHeight < HeightOfBoneInRetargetPose)
 			{
@@ -317,7 +317,7 @@ void FRetargetAutoPoseGenerator::SnapToGround(const FName BoneToPutOnGround, con
 	{
 		// use the user specified bone
 		BoneToSnapIndex = Skeleton.FindBoneIndexByName(BoneToPutOnGround);
-		HeightOfBoneInRetargetPose = static_cast<float>(Skeleton.RetargetGlobalPose[BoneToSnapIndex].GetLocation().Z);
+		HeightOfBoneInRetargetPose = static_cast<float>(Skeleton.RetargetPoses.GetGlobalRetargetPose()[BoneToSnapIndex].GetLocation().Z);
 	}
 
 	if (BoneToSnapIndex == INDEX_NONE)
@@ -430,7 +430,7 @@ bool FRetargetAutoPoseGenerator::GetChainTangentFromMesh(
 	// return vector from bone to centroid of mesh
 	const FRetargetSkeleton& Skeleton = Processor->GetSkeleton(SourceOrTarget);
 	const int32 BoneIndex = Skeleton.FindBoneIndexByName(BoneName);
-	const FVector BonePosition = Skeleton.RetargetGlobalPose[BoneIndex].GetLocation();
+	const FVector BonePosition = Skeleton.RetargetPoses.GetGlobalRetargetPose()[BoneIndex].GetLocation();
 	OutDirection = (Centroid - BonePosition).GetSafeNormal();
 	return true;
 }
@@ -484,7 +484,7 @@ void FRetargetAutoPoseGenerator::GetPointsWeightedToBones(
 		VerticesToConsider,
 		VertexRefPose,
 		ComponentSpaceRefPoseBoneTransforms,
-		Skeleton.RetargetGlobalPose,
+		Skeleton.RetargetPoses.GetGlobalRetargetPose(),
 		SkinWeights,
 		SourceOrTarget,
 		OutPoints);
@@ -652,7 +652,7 @@ void FRetargetAutoPoseGenerator::ApplyWorldRotationToRetargetPose(
 	const int32 ParentIndex = RefSkeleton.GetParentIndex(BoneIndex);
 	if (ParentIndex != INDEX_NONE)
 	{
-		ParentGlobalTransform = Skeleton.RetargetGlobalPose[ParentIndex];
+		ParentGlobalTransform = Skeleton.RetargetPoses.GetGlobalRetargetPose()[ParentIndex];
 	}
 
 	// get the current global transform with no delta's at all
diff --git a/Engine/Plugins/Animation/IKRig/Source/IKRigEditor/Public/RetargetEditor/IKRetargetOpDetails.h b/Engine/Plugins/Animation/IKRig/Source/IKRigEditor/Public/RetargetEditor/IKRetargetOpDetails.h
--- a/Engine/Plugins/Animation/IKRig/Source/IKRigEditor/Public/RetargetEditor/IKRetargetOpDetails.h
+++ b/Engine/Plugins/Animation/IKRig/Source/IKRigEditor/Public/RetargetEditor/IKRetargetOpDetails.h
@@ -5,6 +5,7 @@
 #include "Retargeter/RetargetOps/AlignPoleVectorOp.h"
 #include "Retargeter/RetargetOps/FKChainsOp.h"
 #include "Retargeter/RetargetOps/IKChainsOp.h"
+#include "Retargeter/RetargetOps/RetargetPoseOp.h"
 #include "Retargeter/RetargetOps/SpeedPlantingOp.h"
 #include "RigEditor/IKRigStructViewer.h"
 #include "Templates/SharedPointer.h"
@@ -219,3 +220,42 @@ public:
 	UPROPERTY(EditAnywhere, Category = "Pole Vector Settings")
 	FRetargetPoleVectorSettings Settings;
 };
+
+//
+// Retarget Pose Op Customization
+//
+
+class FRetargetPoseOpCustomization : public FIKRetargetOpBaseSettingsCustomization
+{
+public:
+	static TSharedRef<IPropertyTypeCustomization> MakeInstance()
+	{
+		return MakeShared<FRetargetPoseOpCustomization>();
+	}
+
+	// IPropertyTypeCustomization
+	virtual void CustomizeChildren(
+		TSharedRef<IPropertyHandle> StructPropertyHandle,
+		IDetailChildrenBuilder& ChildBuilder,
+		IPropertyTypeCustomizationUtils& StructCustomizationUtils) override;
+	// End IPropertyTypeCustomization
+
+private:
+	
+	void UpdatePoseNameOptions();
+
+	TSharedPtr<IPropertyHandle> SourceOrTargetProperty;
+	TSharedPtr<IPropertyHandle> CurrentPoseProperty;
+	
+	TSharedPtr<FName> CurrentPoseOption;
+	TArray<TSharedPtr<FName>> PoseNameOptions;
+};
+
+UCLASS(BlueprintType)
+class URetargetPoseOpSettingsWrapper : public UIKRigStructWrapperBase
+{
+	GENERATED_BODY()
+public:
+	UPROPERTY(EditAnywhere, Category = "Retarget Pose Settings")
+	FIKRetargetPoseOpSettings Settings;
+};
