[IK Retarget] Fixes crash when adding default ops with no IK Rig

diff --git a/Engine/Plugins/Animation/IKRig/Source/IKRig/Private/Retargeter/IKRetargetChainMapping.cpp b/Engine/Plugins/Animation/IKRig/Source/IKRig/Private/Retargeter/IKRetargetChainMapping.cpp
--- a/Engine/Plugins/Animation/IKRig/Source/IKRig/Private/Retargeter/IKRetargetChainMapping.cpp
+++ b/Engine/Plugins/Animation/IKRig/Source/IKRig/Private/Retargeter/IKRetargetChainMapping.cpp
@@ -123,7 +123,13 @@ const TArray<FRetargetChainPair>& FRetargetChainMapping::GetChainPairs() const
 
 TArray<FName> FRetargetChainMapping::GetChainNames(ERetargetSourceOrTarget InSourceOrTarget) const
 {
-	TArray<FName> ChainNames = InSourceOrTarget == ERetargetSourceOrTarget::Source ? SourceIKRig->GetRetargetChainNames() : TargetIKRig->GetRetargetChainNames();
+	const UIKRigDefinition* IKRig = InSourceOrTarget == ERetargetSourceOrTarget::Source ?  SourceIKRig : TargetIKRig;
+	if (!IKRig)
+	{
+		return TArray<FName>();
+	}
+	
+	TArray<FName> ChainNames = IKRig->GetRetargetChainNames();
 	return MoveTemp(ChainNames);
 }
 
