[IK Retarget] Fix bug, op debug draws only when op is selected

diff --git a/Engine/Plugins/Animation/IKRig/Source/IKRig/Private/Retargeter/IKRetargetProcessor.cpp b/Engine/Plugins/Animation/IKRig/Source/IKRig/Private/Retargeter/IKRetargetProcessor.cpp
--- a/Engine/Plugins/Animation/IKRig/Source/IKRig/Private/Retargeter/IKRetargetProcessor.cpp
+++ b/Engine/Plugins/Animation/IKRig/Source/IKRig/Private/Retargeter/IKRetargetProcessor.cpp
@@ -1703,35 +1703,27 @@ FName FIKRetargetProcessor::GetFirstChainMappedToChain(
 void FIKRetargetProcessor::DebugDrawAllOps(
 	FPrimitiveDrawInterface* InPDI,
 	const FTransform& InComponentTransform,
-	const FIKRetargetDebugDrawState& EditorState) const
+	const FIKRetargetDebugDrawState& EditorState,
+	const int32 InSelectedOpIndex) const
 {
 	const float ComponentScale = static_cast<float>(InComponentTransform.GetScale3D().GetMax());
 	
-	for (const FInstancedStruct& OpStruct : OpStack)
+	for (int32 OpIndex=0; OpIndex<OpStack.Num(); ++OpIndex)
 	{
-		const FIKRetargetOpBase* Op = OpStruct.GetPtr<FIKRetargetOpBase>();
-		if (Op->IsEnabled() && Op->IsInitialized())
+		// filter unselected ops if any op is selected
+		if (InSelectedOpIndex != INDEX_NONE && OpIndex != InSelectedOpIndex)
 		{
-			Op->DebugDraw(InPDI, InComponentTransform, ComponentScale, EditorState);
+			continue;
 		}
-	}
-}
 
-void FIKRetargetProcessor::DebugDrawOpAtIndex(
-	const int32 InOpIndex,
-	FPrimitiveDrawInterface* InPDI,
-	const FTransform& InComponentTransform,
-	const FIKRetargetDebugDrawState& EditorState) const
-{
-	if (!OpStack.IsValidIndex(InOpIndex))
-	{
-		return;
-	}
+		const FIKRetargetOpBase* Op = OpStack[OpIndex].GetPtr<FIKRetargetOpBase>();
+
+		// filter disabled or uninitialized ops
+		if (!Op->IsEnabled() || !Op->IsInitialized())
+		{
+			continue;
+		}
 
-	const FIKRetargetOpBase* Op = OpStack[InOpIndex].GetPtr<FIKRetargetOpBase>();
-	if (Op->IsEnabled() && Op->IsInitialized())
-	{
-		const float ComponentScale = static_cast<float>(InComponentTransform.GetScale3D().GetMax());
 		Op->DebugDraw(InPDI, InComponentTransform, ComponentScale, EditorState);
 	}
 }
diff --git a/Engine/Plugins/Animation/IKRig/Source/IKRig/Public/Retargeter/IKRetargetProcessor.h b/Engine/Plugins/Animation/IKRig/Source/IKRig/Public/Retargeter/IKRetargetProcessor.h
--- a/Engine/Plugins/Animation/IKRig/Source/IKRig/Public/Retargeter/IKRetargetProcessor.h
+++ b/Engine/Plugins/Animation/IKRig/Source/IKRig/Public/Retargeter/IKRetargetProcessor.h
@@ -526,17 +526,12 @@ public:
 	UE_API FName GetFirstChainMappedToChain(const FName InChainName, const ERetargetSourceOrTarget InSourceOrTarget) const;
 	/** Attach a delegate to be notified whenever this processor is initialized. */
 	FOnRetargeterInitialized& OnRetargeterInitialized(){ return RetargeterInitialized; };
-	/** Run all debug drawing on all ops in the stack */
+	/** Run all debug drawing on all ops in the stack (or just selected op if one is provided) */
 	UE_API void DebugDrawAllOps(
 		FPrimitiveDrawInterface* InPDI,
 		const FTransform& InComponentTransform,
-		const FIKRetargetDebugDrawState& EditorState) const;
-	/** Run debug drawing on the specified op */
-	UE_API void DebugDrawOpAtIndex(
-		const int32 InOpIndex,
-		FPrimitiveDrawInterface* InPDI,
-		const FTransform& InComponentTransform,
-		const FIKRetargetDebugDrawState& EditorState) const;
+		const FIKRetargetDebugDrawState& EditorState,
+		const int32 InSelectedOpIndex) const;
 	
 #endif
 
diff --git a/Engine/Plugins/Animation/IKRig/Source/IKRigEditor/Private/RetargetEditor/IKRetargetDefaultMode.cpp b/Engine/Plugins/Animation/IKRig/Source/IKRigEditor/Private/RetargetEditor/IKRetargetDefaultMode.cpp
--- a/Engine/Plugins/Animation/IKRig/Source/IKRigEditor/Private/RetargetEditor/IKRetargetDefaultMode.cpp
+++ b/Engine/Plugins/Animation/IKRig/Source/IKRigEditor/Private/RetargetEditor/IKRetargetDefaultMode.cpp
@@ -82,7 +82,7 @@ void FIKRetargetDefaultMode::RenderDebugProxies(FPrimitiveDrawInterface* PDI, co
 	int32 SelectedOpIndex = Controller->GetSelectedOpIndex();
 	UDebugSkelMeshComponent* TargetSkelMesh = Controller->GetSkeletalMeshComponent(ERetargetSourceOrTarget::Target);
 	const FIKRetargetDebugDrawState& EditorState = Controller->GetSelectionState();
-	RetargetProcessor->DebugDrawOpAtIndex(SelectedOpIndex, PDI, TargetSkelMesh->GetComponentTransform(), EditorState);
+	RetargetProcessor->DebugDrawAllOps(PDI, TargetSkelMesh->GetComponentTransform(), EditorState, SelectedOpIndex);
 }
 
 bool FIKRetargetDefaultMode::HandleClick(FEditorViewportClient* InViewportClient, HHitProxy* HitProxy, const FViewportClick& Click)
diff --git a/Engine/Plugins/Animation/IKRig/Source/IKRigEditor/Private/RetargetEditor/IKRetargetEditorController.cpp b/Engine/Plugins/Animation/IKRig/Source/IKRigEditor/Private/RetargetEditor/IKRetargetEditorController.cpp
--- a/Engine/Plugins/Animation/IKRig/Source/IKRigEditor/Private/RetargetEditor/IKRetargetEditorController.cpp
+++ b/Engine/Plugins/Animation/IKRig/Source/IKRigEditor/Private/RetargetEditor/IKRetargetEditorController.cpp
@@ -1030,6 +1030,11 @@ int32 FIKRetargetEditorController::GetSelectedOpIndex() const
 		return INDEX_NONE;
 	}
 
+	if (Selection.LastSelectedType != ERetargetSelectionType::OP)
+	{
+		return INDEX_NONE;
+	}
+	
 	return AssetController->GetIndexOfOpByName(Selection.LastSelectedOpName);
 }
 
