[IK Retarget] Fixes crash when swapping source mesh in retargeter

diff --git a/Engine/Plugins/Animation/IKRig/Source/IKRig/Private/AnimNodes/AnimNode_RetargetPoseFromMesh.cpp b/Engine/Plugins/Animation/IKRig/Source/IKRig/Private/AnimNodes/AnimNode_RetargetPoseFromMesh.cpp
--- a/Engine/Plugins/Animation/IKRig/Source/IKRig/Private/AnimNodes/AnimNode_RetargetPoseFromMesh.cpp
+++ b/Engine/Plugins/Animation/IKRig/Source/IKRig/Private/AnimNodes/AnimNode_RetargetPoseFromMesh.cpp
@@ -105,6 +105,12 @@ void FAnimNode_RetargetPoseFromMesh::Evaluate_AnyThread(FPoseContext& Output)
 			{
 				return false; // if copying from another mesh, pose must be copied already from the PreUpdate()
 			}
+
+			const int32 NumBonesInSource = Processor.GetSkeleton(ERetargetSourceOrTarget::Source).BoneNames.Num();
+			if (NumBonesInSource != PoseToRetargetFromComponentSpace.Num())
+			{
+				return false; // the source pose has not been copied yet (can happen if Evaluate() is called before PreUpdate() with a different mesh)	
+			}
 		}
 
 		if (!bCopyingFromOtherComponent && !Source.GetLinkNode())
