[IK Retarget] Prevent adding multiple scale source ops. Menu disables option to add singleton ops when already present in the stack.

diff --git a/Engine/Plugins/Animation/IKRig/Source/IKRig/Public/Retargeter/RetargetOps/ScaleSourceOp.h b/Engine/Plugins/Animation/IKRig/Source/IKRig/Public/Retargeter/RetargetOps/ScaleSourceOp.h
--- a/Engine/Plugins/Animation/IKRig/Source/IKRig/Public/Retargeter/RetargetOps/ScaleSourceOp.h
+++ b/Engine/Plugins/Animation/IKRig/Source/IKRig/Public/Retargeter/RetargetOps/ScaleSourceOp.h
@@ -50,6 +50,8 @@ struct FIKRetargetScaleSourceOp : public FIKRetargetOpBase
 	
 	UE_API virtual const UScriptStruct* GetType() const override;
 
+	virtual bool IsSingleton() const override { return true; };
+
 	UPROPERTY()
 	FIKRetargetScaleSourceOpSettings Settings;
 };
diff --git a/Engine/Plugins/Animation/IKRig/Source/IKRigEditor/Private/RetargetEditor/SRetargetOpStack.cpp b/Engine/Plugins/Animation/IKRig/Source/IKRigEditor/Private/RetargetEditor/SRetargetOpStack.cpp
--- a/Engine/Plugins/Animation/IKRig/Source/IKRigEditor/Private/RetargetEditor/SRetargetOpStack.cpp
+++ b/Engine/Plugins/Animation/IKRig/Source/IKRigEditor/Private/RetargetEditor/SRetargetOpStack.cpp
@@ -301,6 +301,7 @@ void SRetargetOpList::CacheOpTypeMetaData()
 		OpMetaData.NiceName = TempOp->GetDefaultName();
 		OpMetaData.OpType = TempOpStruct.GetScriptStruct();
 		OpMetaData.ParentType = TempOp->GetParentOpType();
+		OpMetaData.bIsSingleton = TempOp->IsSingleton();
 		AllOpsMetaData.Add(OpMetaData);
 	}
 }
@@ -515,9 +516,32 @@ TSharedRef<SWidget> SRetargetOpList::CreateAddNewOpMenu()
 	MenuBuilder.BeginSection("AddNewRetargetOp", LOCTEXT("AddOperations", "Add New Retarget Op"));
 
 	// add menu option to create each retarget op type
+	const TObjectPtr<UIKRetargeterController> AssetController = EditorController.Pin()->AssetController;
 	for (FIKRetargetOpMetaData& OpMetaData : AllOpsMetaData)
 	{
-		FUIAction Action = FUIAction(FExecuteAction::CreateSP(this, &SRetargetOpList::AddNewRetargetOp, OpMetaData.OpType));
+		FUIAction Action = FUIAction(
+			FExecuteAction::CreateSP(this, &SRetargetOpList::AddNewRetargetOp, OpMetaData.OpType),
+			FCanExecuteAction::CreateLambda([AssetController, OpMetaData]()
+			{
+				if (!OpMetaData.bIsSingleton)
+				{
+					return true;
+				}
+
+				// check if another instance of this type already exists in the op stack
+				const int32 NumOps = AssetController->GetNumRetargetOps();
+				for (int32 OpIndex = 0; OpIndex < NumOps; ++OpIndex)
+				{
+					FIKRetargetOpBase* Op = AssetController->GetRetargetOpByIndex(OpIndex);
+					if (Op->GetType() == OpMetaData.OpType)
+					{
+						return false; // can only have one instance of this op type
+					}
+				}
+				
+				return true;
+			}));
+			
 		MenuBuilder.AddMenuEntry(FText::FromName(OpMetaData.NiceName), FText::GetEmpty(), FSlateIcon(), Action);
 	}
 	
diff --git a/Engine/Plugins/Animation/IKRig/Source/IKRigEditor/Public/RetargetEditor/SRetargetOpStack.h b/Engine/Plugins/Animation/IKRig/Source/IKRigEditor/Public/RetargetEditor/SRetargetOpStack.h
--- a/Engine/Plugins/Animation/IKRig/Source/IKRigEditor/Public/RetargetEditor/SRetargetOpStack.h
+++ b/Engine/Plugins/Animation/IKRig/Source/IKRigEditor/Public/RetargetEditor/SRetargetOpStack.h
@@ -83,6 +83,7 @@ struct FIKRetargetOpMetaData
 	FName NiceName;
 	const UScriptStruct* OpType = nullptr;
 	const UScriptStruct* ParentType = nullptr;
+	bool bIsSingleton = false;
 };
 
 // an SListView customized for retarget ops
