[IK Rig] Fixed issues, New Solver menu now shows nice names of solvers.

diff --git a/Engine/Plugins/Animation/IKRig/Source/IKRigEditor/Private/RigEditor/SIKRigSolverStack.cpp b/Engine/Plugins/Animation/IKRig/Source/IKRigEditor/Private/RigEditor/SIKRigSolverStack.cpp
--- a/Engine/Plugins/Animation/IKRig/Source/IKRigEditor/Private/RigEditor/SIKRigSolverStack.cpp
+++ b/Engine/Plugins/Animation/IKRig/Source/IKRigEditor/Private/RigEditor/SIKRigSolverStack.cpp
@@ -240,6 +240,8 @@ void SIKRigSolverStack::Construct(const FArguments& InArgs, TSharedRef<FIKRigEdi
 	EditorController = InEditorController;
 	EditorController.Pin()->SetSolverStackView(SharedThis(this));
 	
+	CacheSolverMetaData();
+	
 	CommandList = MakeShared<FUICommandList>();
 
 	ChildSlot
@@ -307,25 +309,12 @@ TSharedRef<SWidget> SIKRigSolverStack::CreateAddNewMenuWidget()
 void SIKRigSolverStack::BuildAddNewMenu(FMenuBuilder& MenuBuilder)
 {
 	MenuBuilder.BeginSection("AddNewSolver", LOCTEXT("AddOperations", "Add New Solver"));
-	
+
 	// add menu option to create each solver type
-	TArray<UStruct*> SolverTypes;
-	const UStruct* BaseStruct = FIKRigSolverBase::StaticStruct();
-	for(TObjectIterator<UStruct> StructIt; StructIt; ++StructIt)
+	for (const FIKRigSolverMetaData& SolverMetaData : AllSolversMetaData)
 	{
-		UStruct* Struct = *StructIt;
-
-		if(!Struct->IsChildOf(BaseStruct) || Struct == BaseStruct)
-		{
-			continue;
-		}
-
-		if (UScriptStruct* ScriptStruct = Cast<UScriptStruct>(Struct))
-		{
-			FString TypeName = ScriptStruct->GetName();
-			FUIAction Action = FUIAction( FExecuteAction::CreateSP(this, &SIKRigSolverStack::AddNewSolver, ScriptStruct));
-			MenuBuilder.AddMenuEntry(FText::FromString(TypeName), FText::GetEmpty(), FSlateIcon(), Action);
-		}
+		FUIAction Action = FUIAction( FExecuteAction::CreateSP(this, &SIKRigSolverStack::AddNewSolver, SolverMetaData.ScriptStruct));
+		MenuBuilder.AddMenuEntry(SolverMetaData.NiceName, FText::GetEmpty(), FSlateIcon(), Action);
 	}
 	
 	MenuBuilder.EndSection();
@@ -350,6 +339,37 @@ bool SIKRigSolverStack::IsAddSolverEnabled() const
 	return false;
 }
 
+void SIKRigSolverStack::CacheSolverMetaData()
+{
+	AllSolversMetaData.Reset();
+
+	// instantiate all types of solvers to cache their meta data
+	const UStruct* BaseStruct = FIKRigSolverBase::StaticStruct();
+	for (TObjectIterator<UStruct> StructIt; StructIt; ++StructIt)
+	{
+		UStruct* Struct = *StructIt;
+		if (!Struct->IsChildOf(BaseStruct) || Struct == BaseStruct)
+		{
+			continue;
+		}
+
+		UScriptStruct* ScriptStruct = Cast<UScriptStruct>(Struct);
+		if (!ScriptStruct)
+		{
+			continue;
+		}
+
+		FInstancedStruct TempInstancedStruct = FInstancedStruct(ScriptStruct);
+		const FIKRigSolverBase& TempSolver = TempInstancedStruct.Get<FIKRigSolverBase>();
+
+		// store meta data for this type
+		FIKRigSolverMetaData SolverMetaData;
+		SolverMetaData.NiceName = TempSolver.GetNiceName();
+		SolverMetaData.ScriptStruct = ScriptStruct;
+		AllSolversMetaData.Add(SolverMetaData);
+	}
+}
+
 void SIKRigSolverStack::AddNewSolver(UScriptStruct* SolverType)
 {
 	const TSharedPtr<FIKRigEditorController> Controller = EditorController.Pin();
diff --git a/Engine/Plugins/Animation/IKRig/Source/IKRigEditor/Public/RigEditor/SIKRigSolverStack.h b/Engine/Plugins/Animation/IKRig/Source/IKRigEditor/Public/RigEditor/SIKRigSolverStack.h
--- a/Engine/Plugins/Animation/IKRig/Source/IKRigEditor/Public/RigEditor/SIKRigSolverStack.h
+++ b/Engine/Plugins/Animation/IKRig/Source/IKRigEditor/Public/RigEditor/SIKRigSolverStack.h
@@ -74,9 +74,13 @@ public:
 	TWeakPtr<FSolverStackElement> Element;
 };
 
+struct FIKRigSolverMetaData
+{
+	UScriptStruct* ScriptStruct;
+	FText NiceName;
+};
 
 typedef SListView< TSharedPtr<FSolverStackElement> > SSolverStackListViewType;
-
 class SIKRigSolverStack : public SCompoundWidget, public FEditorUndoClient
 {
 public:
@@ -103,6 +107,10 @@ private:
 	void BuildAddNewMenu(FMenuBuilder& MenuBuilder);
 	bool IsAddSolverEnabled() const;
 	/** END solver stack menu */
+
+	/** all type info about all available solver types (cached at startup) */
+	TArray<FIKRigSolverMetaData> AllSolversMetaData;
+	void CacheSolverMetaData();
 	
 	/** menu command callback for adding a new solver */
 	void AddNewSolver(UScriptStruct* SolverType);
