[IK Retarget] Fixes issue where creating retargeter from right-clicking on IK Rig doesn't import the retarget chains.

diff --git a/Engine/Plugins/Animation/IKRig/Source/IKRigEditor/Private/RetargetEditor/AssetTypeActions_IKRetargeter.cpp b/Engine/Plugins/Animation/IKRig/Source/IKRigEditor/Private/RetargetEditor/AssetTypeActions_IKRetargeter.cpp
--- a/Engine/Plugins/Animation/IKRig/Source/IKRigEditor/Private/RetargetEditor/AssetTypeActions_IKRetargeter.cpp
+++ b/Engine/Plugins/Animation/IKRig/Source/IKRigEditor/Private/RetargetEditor/AssetTypeActions_IKRetargeter.cpp
@@ -61,7 +61,7 @@ void FAssetTypeActions_IKRetargeter::ExtendIKRigMenuToMakeRetargeter()
 				InSection.AddMenuEntry(
 					"CreateRetargeter",
 					LOCTEXT("CreateIKRetargeter", "Create IK Retargeter"),
-					LOCTEXT("CreateIKRetargeter_ToolTip", "Creates an IK Retargeter using this IK Rig as the source."),
+					LOCTEXT("CreateIKRetargeter_ToolTip", "Creates an IK Retargeter using this IK Rig as the source and target."),
 					FSlateIcon(FIKRigEditorStyle::Get().GetStyleSetName(), "IKRig", "ClassIcon.IKRigDefinition"),
 					FExecuteAction::CreateLambda([InSelectedAssets = Context->SelectedAssets]()
 					{
@@ -131,6 +131,10 @@ void FAssetTypeActions_IKRetargeter::CreateNewIKRetargeterFromIKRig(UObject* InS
 	// assign ik rig to the retargeter
 	const UIKRetargeterController* Controller = UIKRetargeterController::GetController(RetargetAsset);
 	Controller->SetIKRig(ERetargetSourceOrTarget::Source, IKRigAsset);
+	Controller->SetIKRig(ERetargetSourceOrTarget::Target, IKRigAsset);
+
+	// add default ops
+	Controller->AddDefaultOps();
 }
 
 void FAssetTypeActions_IKRetargeter::ExtendAnimAssetMenusForBatchRetargeting()
