[IK Retarget] Fixes crash when undo'ing bone rotation editing.

diff --git a/Engine/Plugins/Animation/IKRig/Source/IKRigEditor/Private/RetargetEditor/IKRetargetDetails.cpp b/Engine/Plugins/Animation/IKRig/Source/IKRigEditor/Private/RetargetEditor/IKRetargetDetails.cpp
--- a/Engine/Plugins/Animation/IKRig/Source/IKRigEditor/Private/RetargetEditor/IKRetargetDetails.cpp
+++ b/Engine/Plugins/Animation/IKRig/Source/IKRigEditor/Private/RetargetEditor/IKRetargetDetails.cpp
@@ -30,11 +30,13 @@
 FEulerTransform UIKRetargetBoneDetails::GetTransform(EIKRetargetTransformType TransformType, const bool bLocalSpace) const
 {
 	// editor setup?
-	const FIKRetargetEditorController* Controller = EditorController.Pin().Get();
-	if (!Controller)
+	if (!EditorController.IsValid())
 	{
 		return FEulerTransform::Identity;
 	}
+	
+	// get controller
+	const FIKRetargetEditorController* Controller = EditorController.Pin().Get();
 
 	// ensure we have a valid skeletal mesh
 	const bool bEditingSource = Controller->GetSourceOrTarget() == ERetargetSourceOrTarget::Source;
@@ -881,6 +883,11 @@ void FIKRetargetBoneDetailCustomization::CustomizeDetails(IDetailLayoutBuilder&
 		ESlateRotationRepresentation::Type Representation,
 		ESlateTransformSubComponent::Type SubComponent) -> TOptional<FVector::FReal>
 		{
+			if (BonesView.IsEmpty() || !BonesView[0]->IsValidLowLevel())
+			{
+				return TOptional<FVector::FReal>();
+			}
+			
 			TOptional<FVector::FReal> FirstValue = BonesView[0]->GetNumericValue(TransformType, Component, Representation, SubComponent);
 			
 			if (FirstValue)
