[IK Retarget] Reorganized PinBoneOp, added categories

diff --git a/Engine/Plugins/Animation/IKRig/Source/IKRig/Public/Retargeter/RetargetOps/PinBoneOp.h b/Engine/Plugins/Animation/IKRig/Source/IKRig/Public/Retargeter/RetargetOps/PinBoneOp.h
--- a/Engine/Plugins/Animation/IKRig/Source/IKRig/Public/Retargeter/RetargetOps/PinBoneOp.h
+++ b/Engine/Plugins/Animation/IKRig/Source/IKRig/Public/Retargeter/RetargetOps/PinBoneOp.h
@@ -97,39 +97,39 @@ struct FIKRetargetPinBoneOpSettings : public FIKRetargetOpSettingsBase
 	TArray<FPinBoneData> BonesToPin;
 
 	/** Which skeleton to copy animation from. */
-	UPROPERTY(EditAnywhere, BlueprintReadWrite, Category=Settings)
+	UPROPERTY(EditAnywhere, BlueprintReadWrite, Category=Setup)
 	ERetargetSourceOrTarget SkeletonToCopyFrom = ERetargetSourceOrTarget::Target;
 
 	/** Copy the translation of the bone. */
-	UPROPERTY(EditAnywhere, BlueprintReadWrite, Category=Settings)
+	UPROPERTY(EditAnywhere, BlueprintReadWrite, Category=Translation)
 	bool bCopyTranslation = true;
+	
+	/** The method used to calculate the translation of the bone. */
+	UPROPERTY(EditAnywhere, BlueprintReadWrite, Category=Translation)
+	EPinBoneTranslationMode TranslationMode = EPinBoneTranslationMode::CopyGlobalPosition;
 
 	/** Copy the rotation of the bone. */
-	UPROPERTY(EditAnywhere, BlueprintReadWrite, Category=Settings)
+	UPROPERTY(EditAnywhere, BlueprintReadWrite, Category=Rotation)
 	bool bCopyRotation = true;
+
+	/** The method used to calculate the rotation of the bone. */
+	UPROPERTY(EditAnywhere, BlueprintReadWrite, Category=Rotation)
+	EPinBoneRotationMode RotationMode = EPinBoneRotationMode::CopyGlobalRotation;
 	
 	/** Copy the scale of the bone. */
-	UPROPERTY(EditAnywhere, BlueprintReadWrite, Category=Settings)
+	UPROPERTY(EditAnywhere, BlueprintReadWrite, Category=Scale)
 	bool bCopyScale = true;
 
 	/** Update children bone transforms (recursively). */
-	UPROPERTY(EditAnywhere, Category=Settings)
+	UPROPERTY(EditAnywhere, Category=PropagateTransform)
 	bool bPropagateToChildren = false;
 
-	/** The method used to calculate the translation of the bone. */
-	UPROPERTY(EditAnywhere, BlueprintReadWrite, Category=Settings)
-	EPinBoneTranslationMode TranslationMode = EPinBoneTranslationMode::CopyGlobalPosition;
-
-	/** The method used to calculate the rotation of the bone. */
-	UPROPERTY(EditAnywhere, BlueprintReadWrite, Category=Settings)
-	EPinBoneRotationMode RotationMode = EPinBoneRotationMode::CopyGlobalRotation;
-
 	/** A manual offset to apply in global space */
-	UPROPERTY(EditAnywhere, BlueprintReadWrite, Category=Settings)
+	UPROPERTY(EditAnywhere, BlueprintReadWrite, Category=Offset)
 	FTransform GlobalOffset;
 
 	/** A manual offset to apply in local space */
-	UPROPERTY(EditAnywhere, BlueprintReadWrite, Category=Settings)
+	UPROPERTY(EditAnywhere, BlueprintReadWrite, Category=Offset)
 	FTransform LocalOffset;
 
 	UE_API virtual const UClass* GetControllerType() const override;
