[IK Retarget] Fixes bug, loading old retarget assets with disabled ops no longer disabled.

diff --git a/Engine/Plugins/Animation/IKRig/Source/IKRig/Private/Retargeter/RetargetOps/RootMotionGeneratorOp.cpp b/Engine/Plugins/Animation/IKRig/Source/IKRig/Private/Retargeter/RetargetOps/RootMotionGeneratorOp.cpp
--- a/Engine/Plugins/Animation/IKRig/Source/IKRig/Private/Retargeter/RetargetOps/RootMotionGeneratorOp.cpp
+++ b/Engine/Plugins/Animation/IKRig/Source/IKRig/Private/Retargeter/RetargetOps/RootMotionGeneratorOp.cpp
@@ -333,6 +333,7 @@ void URootMotionGeneratorOp::ConvertToInstancedStruct(FInstancedStruct& OutInsta
 {
 	OutInstancedStruct.InitializeAs(FIKRetargetRootMotionOp::StaticStruct());
 	FIKRetargetRootMotionOp& NewOp = OutInstancedStruct.GetMutable<FIKRetargetRootMotionOp>();
+	NewOp.SetEnabled(bIsEnabled);
 	NewOp.Settings.SourceRoot.BoneName = SourceRootBone;
 	NewOp.Settings.TargetRoot.BoneName = TargetRootBone;
 	NewOp.Settings.TargetPelvis.BoneName = TargetPelvisBone;
diff --git a/Engine/Plugins/Animation/IKRig/Source/IKRig/Public/Retargeter/RetargetOps/CurveRemapOp.h b/Engine/Plugins/Animation/IKRig/Source/IKRig/Public/Retargeter/RetargetOps/CurveRemapOp.h
--- a/Engine/Plugins/Animation/IKRig/Source/IKRig/Public/Retargeter/RetargetOps/CurveRemapOp.h
+++ b/Engine/Plugins/Animation/IKRig/Source/IKRig/Public/Retargeter/RetargetOps/CurveRemapOp.h
@@ -144,6 +144,7 @@ class UCurveRemapOp : public URetargetOpBase
 	{
 		OutInstancedStruct.InitializeAs(FIKRetargetCurveRemapOp::StaticStruct());
 		FIKRetargetCurveRemapOp& NewOp = OutInstancedStruct.GetMutable<FIKRetargetCurveRemapOp>();
+		NewOp.SetEnabled(bIsEnabled);
 		NewOp.Settings.CurvesToRemap = CurvesToRemap;
 		NewOp.Settings.bCopyAllSourceCurves = bCopyAllSourceCurves;
 	};
diff --git a/Engine/Plugins/Animation/IKRig/Source/IKRig/Public/Retargeter/RetargetOps/PinBoneOp.h b/Engine/Plugins/Animation/IKRig/Source/IKRig/Public/Retargeter/RetargetOps/PinBoneOp.h
--- a/Engine/Plugins/Animation/IKRig/Source/IKRig/Public/Retargeter/RetargetOps/PinBoneOp.h
+++ b/Engine/Plugins/Animation/IKRig/Source/IKRig/Public/Retargeter/RetargetOps/PinBoneOp.h
@@ -233,6 +233,7 @@ public:
 	{
 		OutInstancedStruct.InitializeAs(FIKRetargetPinBoneOp::StaticStruct());
 		FIKRetargetPinBoneOp& NewOp = OutInstancedStruct.GetMutable<FIKRetargetPinBoneOp>();
+		NewOp.SetEnabled(bIsEnabled);
 		NewOp.Settings.BonesToPin = BonesToPin;
 		NewOp.Settings.SkeletonToCopyFrom = PinTo;
 		NewOp.Settings.bCopyTranslation = bCopyTranslation;
