[IK Retarget] Fixes issue with op stack not refreshing when adding default ops or modifying ops stack from script.

diff --git a/Engine/Plugins/Animation/IKRig/Source/IKRigEditor/Private/RetargetEditor/IKRetargetEditorController.cpp b/Engine/Plugins/Animation/IKRig/Source/IKRigEditor/Private/RetargetEditor/IKRetargetEditorController.cpp
--- a/Engine/Plugins/Animation/IKRig/Source/IKRigEditor/Private/RetargetEditor/IKRetargetEditorController.cpp
+++ b/Engine/Plugins/Animation/IKRig/Source/IKRigEditor/Private/RetargetEditor/IKRetargetEditorController.cpp
@@ -172,6 +172,11 @@ void FIKRetargetEditorController::Initialize(TSharedPtr<FIKRetargetEditor> InEdi
 
 	// bind callback when retargeter needs reinitialized
 	RetargeterReInitDelegateHandle = AssetController->OnRetargeterNeedsInitialized().AddSP(this, &FIKRetargetEditorController::HandleRetargeterNeedsInitialized);
+	// bind callback when retargeter op stack is modified
+	OpStackModifiedDelegateHandle = AssetController->OnOpStackModified().AddLambda([this]()
+	{
+		RefreshOpStackView();
+	});
 	// bind callback when IK Rig asset is replaced with a different asset
 	IKRigReplacedDelegateHandle = AssetController->OnIKRigReplaced().AddSP(this, &FIKRetargetEditorController::HandleIKRigReplaced);
 	// bind callback when Preview Mesh asset is replaced with a different asset
@@ -181,6 +186,7 @@ void FIKRetargetEditorController::Initialize(TSharedPtr<FIKRetargetEditor> InEdi
 void FIKRetargetEditorController::Close()
 {
 	AssetController->OnRetargeterNeedsInitialized().Remove(RetargeterReInitDelegateHandle);
+	AssetController->OnOpStackModified().Remove(OpStackModifiedDelegateHandle);
 	AssetController->OnIKRigReplaced().Remove(IKRigReplacedDelegateHandle);
 	AssetController->OnPreviewMeshReplaced().Remove(PreviewMeshReplacedDelegateHandle);
 	GetRetargetProcessor()->OnRetargeterInitialized().Remove(RetargeterInitializedDelegateHandle);
@@ -284,6 +290,7 @@ void FIKRetargetEditorController::ReinitializeRetargeterNoUIRefresh() const
 void FIKRetargetEditorController::HandleIKRigReplaced(ERetargetSourceOrTarget SourceOrTarget)
 {
 	BindToIKRigAssets();
+	RefreshAllViews();
 }
 
 void FIKRetargetEditorController::HandlePreviewMeshReplaced(ERetargetSourceOrTarget SourceOrTarget)
diff --git a/Engine/Plugins/Animation/IKRig/Source/IKRigEditor/Private/RetargetEditor/IKRetargeterController.cpp b/Engine/Plugins/Animation/IKRig/Source/IKRigEditor/Private/RetargetEditor/IKRetargeterController.cpp
--- a/Engine/Plugins/Animation/IKRig/Source/IKRigEditor/Private/RetargetEditor/IKRetargeterController.cpp
+++ b/Engine/Plugins/Animation/IKRig/Source/IKRigEditor/Private/RetargetEditor/IKRetargeterController.cpp
@@ -108,8 +108,8 @@ void UIKRetargeterController::SetIKRig(const ERetargetSourceOrTarget SourceOrTar
 	}
 
 	// update any editors attached to this asset
-	BroadcastIKRigReplaced(SourceOrTarget);
-	BroadcastPreviewMeshReplaced(SourceOrTarget);
+	IKRigReplaced.Broadcast(SourceOrTarget);
+	PreviewMeshReplaced.Broadcast(SourceOrTarget);
 }
 
 const UIKRigDefinition* UIKRetargeterController::GetIKRig(const ERetargetSourceOrTarget SourceOrTarget) const
@@ -163,7 +163,7 @@ void UIKRetargeterController::SetPreviewMesh(
 	SetAskedToFixRootHeightForMesh(PreviewMesh, false);
 	
 	// update any editors attached to this asset
-	BroadcastPreviewMeshReplaced(SourceOrTarget);
+	PreviewMeshReplaced.Broadcast(SourceOrTarget);
 }
 
 USkeletalMesh* UIKRetargeterController::GetPreviewMesh(const ERetargetSourceOrTarget SourceOrTarget) const
@@ -211,7 +211,7 @@ bool UIKRetargeterController::RemoveRetargetOp(const int32 OpIndex) const
 	}
 
 	FScopedTransaction Transaction(LOCTEXT("RemoveRetargetOp_Label", "Remove Retarget Op"));
-	FScopedReinitializeIKRetargeter Reinitialize(this);
+	FScopedReinitializeIKRetargeter Reinitialize(this, ERetargetRefreshMode::ProcessorAndOpStack);
 	Asset->Modify();
 
 	TArray<int32> IndicesToRemove = GetChildOpIndices(OpIndex);
@@ -234,7 +234,7 @@ bool UIKRetargeterController::RemoveAllOps() const
 	}
 
 	FScopedTransaction Transaction(LOCTEXT("RemoveAllRetargetOps_Label", "Remove All Retarget Ops"));
-	FScopedReinitializeIKRetargeter Reinitialize(this);
+	FScopedReinitializeIKRetargeter Reinitialize(this, ERetargetRefreshMode::ProcessorAndOpStack);
 	Asset->Modify();
 	Asset->RetargetOps.Empty();
 	return true;
@@ -249,7 +249,7 @@ FName UIKRetargeterController::SetOpName(const FName InName, const int32 InOpInd
 	}
 	
 	FScopedTransaction Transaction(LOCTEXT("SetOpName_Label", "Rename Retarget Op"));
-	FScopedReinitializeIKRetargeter Reinitialize(this);
+	FScopedReinitializeIKRetargeter Reinitialize(this, ERetargetRefreshMode::ProcessorAndOpStack);
 	Asset->Modify();
 
 	const FName OldOpName = GetOpName(InOpIndex);
@@ -305,7 +305,7 @@ bool UIKRetargeterController::SetParentOpByName(
 	}
 
 	FScopedTransaction Transaction(LOCTEXT("SetOpParent_Label", "Set Op Parent"));
-	FScopedReinitializeIKRetargeter Reinitialize(this);
+	FScopedReinitializeIKRetargeter Reinitialize(this, ERetargetRefreshMode::ProcessorAndOpStack);
 	Asset->Modify();
 	
 	ChildOp->SetParentOpName(InParentOpName);
@@ -338,7 +338,7 @@ int32 UIKRetargeterController::GetIndexOfOpByName(const FName InOpName) const
 void UIKRetargeterController::AddDefaultOps() const
 {
 	FScopedTransaction Transaction(LOCTEXT("AddDefaultOps_Label", "Add Default Ops"));
-	FScopedReinitializeIKRetargeter Reinitialize(this);
+	FScopedReinitializeIKRetargeter Reinitialize(this, ERetargetRefreshMode::ProcessorAndOpStack);
 	Asset->Modify();
 
 	// add set of default ops for basic retargeting
@@ -361,7 +361,7 @@ void UIKRetargeterController::RunOpInitialSetup(const int32 InOpIndex) const
 	}
 	
 	FScopedTransaction Transaction(LOCTEXT("RunOpInitialSetup_Label", "Run Op Initial Setup"));
-	FScopedReinitializeIKRetargeter Reinitialize(this);
+	FScopedReinitializeIKRetargeter Reinitialize(this, ERetargetRefreshMode::ProcessorAndOpStack);
 	Asset->Modify();
 
 	FIKRetargetOpBase* ParentOp = GetRetargetOpByIndex(GetParentOpIndex(InOpIndex));
@@ -510,7 +510,7 @@ void UIKRetargeterController::OnOpPropertyChanged(const FName& InOpName, const F
 	}
 
 	// reinitialize
-	BroadcastNeedsReinitialized();
+	RetargeterNeedsInitialized.Broadcast();
 }
 
 FInstancedStruct* UIKRetargeterController::GetRetargetOpStructAtIndex(int32 Index) const
@@ -554,7 +554,7 @@ bool UIKRetargeterController::MoveRetargetOpInStack(int32 OpToMoveIndex, int32 T
 	}
 
 	FScopedTransaction Transaction(LOCTEXT("ReorderRetargetOps_Label", "Reorder Retarget Ops"));
-	FScopedReinitializeIKRetargeter Reinitialize(this);
+	FScopedReinitializeIKRetargeter Reinitialize(this, ERetargetRefreshMode::ProcessorAndOpStack);
 	Asset->Modify();
 
 	// extract the element to move
@@ -744,7 +744,7 @@ int32 UIKRetargeterController::AddRetargetOp(const UScriptStruct* InRetargetOpTy
 	}
 
 	FScopedTransaction Transaction(LOCTEXT("AddRetargetOp_Label", "Add Retarget Op"));
-	FScopedReinitializeIKRetargeter Reinitialize(this);
+	FScopedReinitializeIKRetargeter Reinitialize(this, ERetargetRefreshMode::ProcessorAndOpStack);
 	Asset->Modify();
 
 	// add the op
diff --git a/Engine/Plugins/Animation/IKRig/Source/IKRigEditor/Private/RetargetEditor/SRetargetAnimAssetsWindow.cpp b/Engine/Plugins/Animation/IKRig/Source/IKRigEditor/Private/RetargetEditor/SRetargetAnimAssetsWindow.cpp
--- a/Engine/Plugins/Animation/IKRig/Source/IKRigEditor/Private/RetargetEditor/SRetargetAnimAssetsWindow.cpp
+++ b/Engine/Plugins/Animation/IKRig/Source/IKRigEditor/Private/RetargetEditor/SRetargetAnimAssetsWindow.cpp
@@ -664,7 +664,7 @@ void FProceduralRetargetAssets::AutoGenerateIKRetargetAsset()
 
 	const UIKRetargeterController* RetargetController = UIKRetargeterController::GetController(Retargeter);
 	
-	FScopedReinitializeIKRetargeter Reinitialize(RetargetController);
+	FScopedReinitializeIKRetargeter Reinitialize(RetargetController, ERetargetRefreshMode::ProcessorAndOpStack);
 
 	// re-assign both IK Rigs
 	RetargetController->SetIKRig(ERetargetSourceOrTarget::Source, SourceIKRig);
diff --git a/Engine/Plugins/Animation/IKRig/Source/IKRigEditor/Private/RetargetEditor/SRetargetOpStack.cpp b/Engine/Plugins/Animation/IKRig/Source/IKRigEditor/Private/RetargetEditor/SRetargetOpStack.cpp
--- a/Engine/Plugins/Animation/IKRig/Source/IKRigEditor/Private/RetargetEditor/SRetargetOpStack.cpp
+++ b/Engine/Plugins/Animation/IKRig/Source/IKRigEditor/Private/RetargetEditor/SRetargetOpStack.cpp
@@ -420,11 +420,6 @@ FReply SRetargetOpList::OnAcceptDrop(
 	
 	const UIKRetargeterController* AssetController = Controller->AssetController;
 	const bool bWasReparented = AssetController->MoveRetargetOpInStack(DraggedElement->GetIndexInStack(),IndexToMoveTo);
-	if (bWasReparented)
-	{
-		// refresh the op stack
-		EditorController.Pin()->RefreshOpStackView();
-	}
 	
 	return FReply::Handled();
 }
@@ -566,9 +561,6 @@ void SRetargetOpList::AddNewRetargetOp(const UScriptStruct* ScriptStruct)
 	// add a new op to the stack
 	const FName ParentOpName = ParentElement.IsValid() ? ParentElement.Pin()->GetName() : NAME_None;
 	int32 NewOpIndex = AssetController->AddRetargetOp(ScriptStruct, ParentOpName);
-
-	// refresh the op stack
-	EditorController.Pin()->RefreshOpStackView();
 }
 
 void SRetargetOpList::DeleteRetargetOp(TSharedPtr<FRetargetOpStackElement> OpToDelete)
@@ -579,9 +571,6 @@ void SRetargetOpList::DeleteRetargetOp(TSharedPtr<FRetargetOpStackElement> OpToD
 	}
 	
 	EditorController.Pin()->AssetController->RemoveRetargetOp(OpToDelete->GetIndexInStack());
-	
-	// refresh the op stack
-	EditorController.Pin()->RefreshOpStackView();
 }
 
 void SRetargetOpList::RefreshAndRestore()
diff --git a/Engine/Plugins/Animation/IKRig/Source/IKRigEditor/Public/RetargetEditor/IKRetargetEditorController.h b/Engine/Plugins/Animation/IKRig/Source/IKRigEditor/Public/RetargetEditor/IKRetargetEditorController.h
--- a/Engine/Plugins/Animation/IKRig/Source/IKRigEditor/Public/RetargetEditor/IKRetargetEditorController.h
+++ b/Engine/Plugins/Animation/IKRig/Source/IKRigEditor/Public/RetargetEditor/IKRetargetEditorController.h
@@ -106,6 +106,8 @@ public:
 	// reinitialize retargeter without refreshing UI
 	void ReinitializeRetargeterNoUIRefresh() const;
 	FDelegateHandle RetargeterReInitDelegateHandle;
+	// refresh op stack whenever it's modified
+	FDelegateHandle OpStackModifiedDelegateHandle;
 	// callback when IK Rig asset has been swapped out
 	void HandleIKRigReplaced(ERetargetSourceOrTarget SourceOrTarget);
 	FDelegateHandle IKRigReplacedDelegateHandle;
diff --git a/Engine/Plugins/Animation/IKRig/Source/IKRigEditor/Public/RetargetEditor/IKRetargeterController.h b/Engine/Plugins/Animation/IKRig/Source/IKRigEditor/Public/RetargetEditor/IKRetargeterController.h
--- a/Engine/Plugins/Animation/IKRig/Source/IKRigEditor/Public/RetargetEditor/IKRetargeterController.h
+++ b/Engine/Plugins/Animation/IKRig/Source/IKRigEditor/Public/RetargetEditor/IKRetargeterController.h
@@ -505,6 +505,9 @@ private:
 	
 	DECLARE_MULTICAST_DELEGATE(FOnRetargeterNeedsInitialized);
 	FOnRetargeterNeedsInitialized RetargeterNeedsInitialized;
+
+	DECLARE_MULTICAST_DELEGATE(FOnOpStackModified);
+	FOnOpStackModified OpStackModified;
 	
 	DECLARE_MULTICAST_DELEGATE_OneParam(FOnIKRigReplaced, ERetargetSourceOrTarget);
 	FOnIKRigReplaced IKRigReplaced;
@@ -512,21 +515,6 @@ private:
 	DECLARE_MULTICAST_DELEGATE_OneParam(FOnPreviewMeshReplaced, ERetargetSourceOrTarget);
 	FOnPreviewMeshReplaced PreviewMeshReplaced;
 
-	void BroadcastPreviewMeshReplaced(ERetargetSourceOrTarget SourceOrTarget) const
-	{
-		PreviewMeshReplaced.Broadcast(SourceOrTarget);
-	}
-
-	void BroadcastIKRigReplaced(ERetargetSourceOrTarget SourceOrTarget) const
-	{
-		IKRigReplaced.Broadcast(SourceOrTarget);
-	}
-	
-	void BroadcastNeedsReinitialized() const
-	{
-		RetargeterNeedsInitialized.Broadcast();
-	}
-
 	// auto pose generator
 	TUniquePtr<FRetargetAutoPoseGenerator> AutoPoseGenerator;
 
@@ -551,6 +539,9 @@ public:
 	// Attach a delegate to be notified whenever the retargeter is modified in such a way that would require re-initialization of the processor.
 	FOnRetargeterNeedsInitialized& OnRetargeterNeedsInitialized(){ return RetargeterNeedsInitialized; };
 
+	// Attach a delegate to be notified whenever the op stack is modified.
+	FOnOpStackModified& OnOpStackModified(){ return OpStackModified; };
+
 	// get the singleton struct viewer for editing structs belonging to the asset this controller controls
 	UIKRigStructViewer* GetStructViewer();
 	
@@ -558,21 +549,36 @@ public:
 	friend struct FScopedReinitializeIKRetargeter;
 };
 
+/** Retargeter can reinit just the processor or other UI as well depending on the data model modification type */
+enum class ERetargetRefreshMode : uint8
+{
+	ProcessorOnly,
+	ProcessorAndOpStack,
+	ProcessorAndFullUI
+};
+
 struct FScopedReinitializeIKRetargeter
 {
-	FScopedReinitializeIKRetargeter(const UIKRetargeterController *InController)
+	FScopedReinitializeIKRetargeter(const UIKRetargeterController *InController, ERetargetRefreshMode InRefreshMode = ERetargetRefreshMode::ProcessorOnly)
 	{
 		InController->ReinitializeScopeCounter++;
 		Controller = InController;
+		RefreshMode = InRefreshMode;
 	}
 	~FScopedReinitializeIKRetargeter()
 	{
 		if (--Controller->ReinitializeScopeCounter == 0)
 		{
-			Controller->BroadcastNeedsReinitialized();
+			Controller->RetargeterNeedsInitialized.Broadcast();
+
+			if (RefreshMode == ERetargetRefreshMode::ProcessorAndOpStack)
+			{
+				Controller->OpStackModified.Broadcast();
+			}
 		}
 	};
 
 	const UIKRetargeterController* Controller;
+	ERetargetRefreshMode RefreshMode = ERetargetRefreshMode::ProcessorOnly;
 };
 
