[IK Rig] Fixes crash when reordering solvers in the solver stack.

diff --git a/Engine/Plugins/Animation/IKRig/Source/IKRigEditor/Private/RigEditor/IKRigController.cpp b/Engine/Plugins/Animation/IKRig/Source/IKRigEditor/Private/RigEditor/IKRigController.cpp
--- a/Engine/Plugins/Animation/IKRig/Source/IKRigEditor/Private/RigEditor/IKRigController.cpp
+++ b/Engine/Plugins/Animation/IKRig/Source/IKRigEditor/Private/RigEditor/IKRigController.cpp
@@ -825,9 +825,14 @@ bool UIKRigController::MoveSolverInStack(int32 SolverToMoveIndex, int32 TargetSo
 	FScopedTransaction Transaction(LOCTEXT("ReorderSolver_Label", "Reorder Solvers"));
 	FScopedReinitializeIKRig Reinitialize(this);
 	Asset->Modify();
-
-	Asset->SolverStack.Swap(SolverToMoveIndex, TargetSolverIndex);
-
+	
+	// extract the element to move
+	FInstancedStruct MovedSolver = MoveTemp(Asset->SolverStack[SolverToMoveIndex]);
+	// Remove the element without shrinking the array
+	Asset->SolverStack.RemoveAt(SolverToMoveIndex, 1, EAllowShrinking::No);
+	// Insert at the corrected TargetIndex
+	Asset->SolverStack.Insert(MoveTemp(MovedSolver), TargetSolverIndex);
+	
 	return true;
 }
 
diff --git a/Engine/Plugins/Animation/IKRig/Source/IKRigEditor/Private/RigEditor/SIKRigSolverStack.cpp b/Engine/Plugins/Animation/IKRig/Source/IKRigEditor/Private/RigEditor/SIKRigSolverStack.cpp
--- a/Engine/Plugins/Animation/IKRig/Source/IKRigEditor/Private/RigEditor/SIKRigSolverStack.cpp
+++ b/Engine/Plugins/Animation/IKRig/Source/IKRigEditor/Private/RigEditor/SIKRigSolverStack.cpp
@@ -468,26 +468,6 @@ TSharedRef<ITableRow> SIKRigSolverStack::MakeListRowWidget(
 		SharedThis(this));
 }
 
-FReply SIKRigSolverStack::OnDragDetected(
-	const FGeometry& MyGeometry,
-	const FPointerEvent& MouseEvent)
-{
-	const TArray<TSharedPtr<FSolverStackElement>> SelectedItems = ListView.Get()->GetSelectedItems();
-	if (SelectedItems.Num() != 1)
-	{
-		return FReply::Unhandled();
-	}
-	
-	if (MouseEvent.IsMouseButtonDown(EKeys::LeftMouseButton))
-	{
-		const TSharedPtr<FSolverStackElement> DraggedElement = SelectedItems[0];
-		const TSharedRef<FIKRigSolverStackDragDropOp> DragDropOp = FIKRigSolverStackDragDropOp::New(DraggedElement);
-		return FReply::Handled().BeginDragDrop(DragDropOp);
-	}
-
-	return FReply::Unhandled();
-}
-
 void SIKRigSolverStack::OnSelectionChanged(TSharedPtr<FSolverStackElement> InItem, ESelectInfo::Type SelectInfo)
 {
 	if (SelectInfo != ESelectInfo::Direct)
@@ -530,6 +510,26 @@ void SIKRigSolverStack::ShowDetailsForItem(TSharedPtr<FSolverStackElement> InIte
 	}
 }
 
+FReply SIKRigSolverStack::OnDragDetected(
+	const FGeometry& MyGeometry,
+	const FPointerEvent& MouseEvent)
+{
+	const TArray<TSharedPtr<FSolverStackElement>> SelectedItems = ListView.Get()->GetSelectedItems();
+	if (SelectedItems.Num() != 1)
+	{
+		return FReply::Unhandled();
+	}
+	
+	if (MouseEvent.IsMouseButtonDown(EKeys::LeftMouseButton))
+	{
+		const TSharedPtr<FSolverStackElement> DraggedElement = SelectedItems[0];
+		const TSharedRef<FIKRigSolverStackDragDropOp> DragDropOp = FIKRigSolverStackDragDropOp::New(DraggedElement);
+		return FReply::Handled().BeginDragDrop(DragDropOp);
+	}
+
+	return FReply::Unhandled();
+}
+
 TOptional<EItemDropZone> SIKRigSolverStack::OnCanAcceptDrop(
 	const FDragDropEvent& DragDropEvent,
 	EItemDropZone DropZone,
@@ -580,6 +580,8 @@ FReply SIKRigSolverStack::OnAcceptDrop(
 	{
 		RefreshStackView();
 		Controller->RefreshTreeView(); // update solver indices in effector items
+		Controller->ShowDetailsForSolver(TargetIndex);
+		ListView->SetSelection(ListViewItems[TargetIndex]);
 	}
 	
 	return FReply::Handled();
