[IK Retarget] Copy Base Pose is now ON by default.

diff --git a/Engine/Plugins/Animation/IKRig/Source/IKRig/Public/Retargeter/RetargetOps/CopyBasePoseOp.h b/Engine/Plugins/Animation/IKRig/Source/IKRig/Public/Retargeter/RetargetOps/CopyBasePoseOp.h
--- a/Engine/Plugins/Animation/IKRig/Source/IKRig/Public/Retargeter/RetargetOps/CopyBasePoseOp.h
+++ b/Engine/Plugins/Animation/IKRig/Source/IKRig/Public/Retargeter/RetargetOps/CopyBasePoseOp.h
@@ -17,7 +17,7 @@ struct FIKRetargetCopyBasePoseOpSettings : public FIKRetargetOpSettingsBase
 	/** When true, will copy all the source bone transforms with matching names to use as a base pose. This can be useful for partial retargeting.
 	 * NOTE: no retargeting is applied to the bone transforms, they are assumed to be fully compatible between source/target skeletons. */
 	UPROPERTY(EditAnywhere, BlueprintReadWrite, Category = RetargetPhases)
-	bool bCopyBasePose = false;
+	bool bCopyBasePose = true;
 
 	/** Filters the bones to copy when using "Copy Base Pose". If specified, will only copy all children of the specified bone (inclusive).*/
 	UPROPERTY(EditAnywhere, BlueprintReadWrite, Category = RetargetPhases, meta = (EditCondition = "bCopyBasePose"))
