[IK Retarget]
Fix regression causing auto pose alignment to produce bad results.
Fix retarget ops debug drawing being connected to op selection.

diff --git a/Engine/Plugins/Animation/IKRig/Source/IKRig/Private/Retargeter/IKRetargetProcessor.cpp b/Engine/Plugins/Animation/IKRig/Source/IKRig/Private/Retargeter/IKRetargetProcessor.cpp
--- a/Engine/Plugins/Animation/IKRig/Source/IKRig/Private/Retargeter/IKRetargetProcessor.cpp
+++ b/Engine/Plugins/Animation/IKRig/Source/IKRig/Private/Retargeter/IKRetargetProcessor.cpp
@@ -28,34 +28,34 @@ FResolvedRetargetPose& FResolvedRetargetPoseSet::AddOrUpdateRetargetPose(
 	const FName PelvisBoneName,
 	const double ScaleFactor)
 {
-	int32 NewPoseIndex = RetargetPoses.Emplace(InRetargetPoseName);
-	FResolvedRetargetPose& NewRetargetPose = RetargetPoses[NewPoseIndex];
+	// add retarget pose if it doesn't already exist
+	FResolvedRetargetPose& RetargetPose = FindOrAddRetargetPose(InRetargetPoseName);
 	
 	// record the name of the retarget pose (prevents re-initialization if profile swaps it)
-	NewRetargetPose.Version = InRetargetPose->GetVersion();
-	NewRetargetPose.Scale = ScaleFactor;
+	RetargetPose.Version = InRetargetPose->GetVersion();
+	RetargetPose.Scale = ScaleFactor;
 	
 	// initialize retarget pose to the skeletal mesh reference pose
-	NewRetargetPose.LocalPose = InSkeleton.SkeletalMesh->GetRefSkeleton().GetRefBonePose();
+	RetargetPose.LocalPose = InSkeleton.SkeletalMesh->GetRefSkeleton().GetRefBonePose();
 	// copy local pose to global
-	NewRetargetPose.GlobalPose = NewRetargetPose.LocalPose;
+	RetargetPose.GlobalPose = RetargetPose.LocalPose;
 	// convert to global space
-	InSkeleton.UpdateGlobalTransformsBelowBone(INDEX_NONE, NewRetargetPose.LocalPose, NewRetargetPose.GlobalPose);
+	InSkeleton.UpdateGlobalTransformsBelowBone(INDEX_NONE, RetargetPose.LocalPose, RetargetPose.GlobalPose);
 
 	// strip scale (done AFTER generating global pose so that scales are baked into translation)
 	for (int32 BoneIndex=0; BoneIndex<InSkeleton.BoneNames.Num(); ++BoneIndex)
 	{
-		NewRetargetPose.LocalPose[BoneIndex].SetScale3D(FVector::OneVector);
-		NewRetargetPose.GlobalPose[BoneIndex].SetScale3D(FVector::OneVector);
+		RetargetPose.LocalPose[BoneIndex].SetScale3D(FVector::OneVector);
+		RetargetPose.GlobalPose[BoneIndex].SetScale3D(FVector::OneVector);
 	}
 
 	// apply pelvis translation offset
 	const int32 PelvisBoneIndex = InSkeleton.FindBoneIndexByName(PelvisBoneName);
 	if (PelvisBoneIndex != INDEX_NONE)
 	{
-		FTransform& PelvisTransform = NewRetargetPose.GlobalPose[PelvisBoneIndex];
+		FTransform& PelvisTransform = RetargetPose.GlobalPose[PelvisBoneIndex];
 		PelvisTransform.AddToTranslation(InRetargetPose->GetRootTranslationDelta());
-		InSkeleton.UpdateLocalTransformOfSingleBone(PelvisBoneIndex, NewRetargetPose.LocalPose, NewRetargetPose.GlobalPose);
+		InSkeleton.UpdateLocalTransformOfSingleBone(PelvisBoneIndex, RetargetPose.LocalPose, RetargetPose.GlobalPose);
 	}
 
 	// apply retarget pose offsets (retarget pose is stored as offset relative to reference pose)
@@ -71,26 +71,40 @@ FResolvedRetargetPose& FResolvedRetargetPoseSet::AddOrUpdateRetargetPose(
 		}
 
 		const FQuat LocalBoneRotation = RefPoseLocal[BoneIndex].GetRotation() * BoneDelta.Value;
-		NewRetargetPose.LocalPose[BoneIndex].SetRotation(LocalBoneRotation);
+		RetargetPose.LocalPose[BoneIndex].SetRotation(LocalBoneRotation);
 	}
 
 	// update global transforms based on local pose modified by the retarget pose offsets
-	InSkeleton.UpdateGlobalTransformsBelowBone(INDEX_NONE, NewRetargetPose.LocalPose, NewRetargetPose.GlobalPose);
+	InSkeleton.UpdateGlobalTransformsBelowBone(INDEX_NONE, RetargetPose.LocalPose, RetargetPose.GlobalPose);
 
 	// scale the retarget pose
 	if (!FMath::IsNearlyEqual(ScaleFactor, 1.0f))
 	{
 		// scale translations of global transforms
-		for (FTransform& BoneTransform : NewRetargetPose.GlobalPose)
+		for (FTransform& BoneTransform : RetargetPose.GlobalPose)
 		{
 			BoneTransform.SetTranslation(BoneTransform.GetTranslation() * ScaleFactor);
 		}
 
 		// update the local transforms
-		InSkeleton.UpdateLocalTransformsBelowBone(INDEX_NONE, NewRetargetPose.LocalPose, NewRetargetPose.GlobalPose);
+		InSkeleton.UpdateLocalTransformsBelowBone(INDEX_NONE, RetargetPose.LocalPose, RetargetPose.GlobalPose);
 	}
 
-	return NewRetargetPose;
+	return RetargetPose;
+}
+
+FResolvedRetargetPose& FResolvedRetargetPoseSet::FindOrAddRetargetPose(const FName InRetargetPoseName)
+{
+	for (FResolvedRetargetPose& Pose : RetargetPoses)
+	{
+		if (Pose.Name == CurrentRetargetPoseName)
+		{
+			return Pose;
+		}
+	}
+	
+	int32 NewPoseIndex = RetargetPoses.Emplace(InRetargetPoseName);
+	return RetargetPoses[NewPoseIndex];
 }
 
 const FResolvedRetargetPose* FResolvedRetargetPoseSet::FindRetargetPoseByName(const FName InRetargetPoseName) const
@@ -1710,15 +1724,8 @@ void FIKRetargetProcessor::DebugDrawAllOps(
 	
 	for (int32 OpIndex=0; OpIndex<OpStack.Num(); ++OpIndex)
 	{
-		// filter unselected ops if any op is selected
-		if (InSelectedOpIndex != INDEX_NONE && OpIndex != InSelectedOpIndex)
-		{
-			continue;
-		}
-
-		const FIKRetargetOpBase* Op = OpStack[OpIndex].GetPtr<FIKRetargetOpBase>();
-
 		// filter disabled or uninitialized ops
+		const FIKRetargetOpBase* Op = OpStack[OpIndex].GetPtr<FIKRetargetOpBase>();
 		if (!Op->IsEnabled() || !Op->IsInitialized())
 		{
 			continue;
diff --git a/Engine/Plugins/Animation/IKRig/Source/IKRig/Public/Retargeter/IKRetargetProcessor.h b/Engine/Plugins/Animation/IKRig/Source/IKRig/Public/Retargeter/IKRetargetProcessor.h
--- a/Engine/Plugins/Animation/IKRig/Source/IKRig/Public/Retargeter/IKRetargetProcessor.h
+++ b/Engine/Plugins/Animation/IKRig/Source/IKRig/Public/Retargeter/IKRetargetProcessor.h
@@ -58,6 +58,8 @@ struct FResolvedRetargetPoseSet
 		const double ScaleFactor);
 
 	UE_API const FResolvedRetargetPose* FindRetargetPoseByName(const FName InRetargetPoseName) const;
+	
+	FResolvedRetargetPose& FindOrAddRetargetPose(const FName InRetargetPoseName);
 
 	UE_API const TArray<FTransform>& GetLocalRetargetPose() const;
 	
