[IK Retarget] Fixed crash in CopyBasePose op when swapping source mesh

diff --git a/Engine/Plugins/Animation/IKRig/Source/IKRig/Private/Retargeter/IKRetargetProcessor.cpp b/Engine/Plugins/Animation/IKRig/Source/IKRig/Private/Retargeter/IKRetargetProcessor.cpp
--- a/Engine/Plugins/Animation/IKRig/Source/IKRig/Private/Retargeter/IKRetargetProcessor.cpp
+++ b/Engine/Plugins/Animation/IKRig/Source/IKRig/Private/Retargeter/IKRetargetProcessor.cpp
@@ -992,7 +992,11 @@ void FBasePoseCopier::CopyPoseFromSource(
 	FIKRigLogger& Log)
 {
 	// check if we need to rebuild the cached bone mapping
-	if (SourceToTargetBoneIndexMap.IsEmpty() || InRootBoneToCopy != RootBoneToCopy)
+	// cache is invalid if: map is empty, source/target skel mesh swapped or root bone swapped
+	if (SourceToTargetBoneIndexMap.IsEmpty() ||
+		SourceSkeletalMesh != InSourceSkeleton.SkeletalMesh ||
+		TargetSkeletalMesh != InTargetSkeleton.SkeletalMesh ||
+		RootBoneToCopy != InRootBoneToCopy)
 	{
 		Initialize(InRootBoneToCopy, InSourceSkeleton, InTargetSkeleton, Log);
 	}
@@ -1040,6 +1044,10 @@ void FBasePoseCopier::Initialize(
 	// store the root bone we initialized with
 	RootBoneToCopy = InRootBoneToCopy;
 
+	// store the skeletal meshes we initialized with
+	SourceSkeletalMesh = InSourceSkeleton.SkeletalMesh;
+	TargetSkeletalMesh = InTargetSkeleton.SkeletalMesh;
+
 	// find the roots to copy from
 	const int32 SourceRootBoneToCopyIndex = RootBoneToCopy == NAME_None ? 0 : InSourceSkeleton.FindBoneIndexByName(RootBoneToCopy);
 	const int32 TargetRootBoneToCopyIndex = RootBoneToCopy == NAME_None ? 0 : InTargetSkeleton.FindBoneIndexByName(RootBoneToCopy);
diff --git a/Engine/Plugins/Animation/IKRig/Source/IKRig/Public/Retargeter/IKRetargetProcessor.h b/Engine/Plugins/Animation/IKRig/Source/IKRig/Public/Retargeter/IKRetargetProcessor.h
--- a/Engine/Plugins/Animation/IKRig/Source/IKRig/Public/Retargeter/IKRetargetProcessor.h
+++ b/Engine/Plugins/Animation/IKRig/Source/IKRig/Public/Retargeter/IKRetargetProcessor.h
@@ -349,6 +349,10 @@ private:
 
 	// the bone to start copying from (includes all children recursive)
 	FName RootBoneToCopy = NAME_None;
+
+	// the skeletal meshes that were used to build the bone mapping
+	const USkeletalMesh* SourceSkeletalMesh = nullptr;
+	const USkeletalMesh* TargetSkeletalMesh = nullptr;
 };
 
 /** The runtime processor that converts an input pose from a source skeleton into an output pose on a target skeleton.
