[IK Retarget] Fixes inability to assign target skeletal mesh without having an IK Rig.

diff --git a/Engine/Plugins/Animation/IKRig/Source/IKRig/Private/Retargeter/IKRetargetProcessor.cpp b/Engine/Plugins/Animation/IKRig/Source/IKRig/Private/Retargeter/IKRetargetProcessor.cpp
--- a/Engine/Plugins/Animation/IKRig/Source/IKRig/Private/Retargeter/IKRetargetProcessor.cpp
+++ b/Engine/Plugins/Animation/IKRig/Source/IKRig/Private/Retargeter/IKRetargetProcessor.cpp
@@ -166,7 +166,7 @@ void FRetargetSkeleton::Initialize(
 	const USkeletalMesh* InSkeletalMesh,
 	const ERetargetSourceOrTarget InSourceOrTarget,
 	const UIKRetargeter* InRetargetAsset,
-	const FName RetargetRootBone,
+	const FName PelvisBoneName,
 	const double ScaleFactor)
 {
 	// reset all skeleton data
@@ -189,7 +189,7 @@ void FRetargetSkeleton::Initialize(
 	// add default retarget pose for this skeleton
 	static FName DefaultRetargetPoseName = UIKRetargeter::GetDefaultPoseName();
 	const FIKRetargetPose* DefaultRetargetPose = InRetargetAsset->GetRetargetPoseByName(InSourceOrTarget, UIKRetargeter::GetDefaultPoseName());
-	RetargetPoses.AddOrUpdateRetargetPose(*this,DefaultRetargetPoseName, DefaultRetargetPose, RetargetRootBone, ScaleFactor);
+	RetargetPoses.AddOrUpdateRetargetPose(*this,DefaultRetargetPoseName, DefaultRetargetPose, PelvisBoneName, ScaleFactor);
 	RetargetPoses.CurrentRetargetPoseName = DefaultRetargetPoseName;
 
 	// add current retarget pose
@@ -197,7 +197,7 @@ void FRetargetSkeleton::Initialize(
 	const FIKRetargetPose* CurrentRetargetPose = InRetargetAsset->GetRetargetPoseByName(InSourceOrTarget, CurrentRetargetPoseName);
 	if (CurrentRetargetPose)
 	{
-		RetargetPoses.AddOrUpdateRetargetPose(*this,CurrentRetargetPoseName, CurrentRetargetPose, RetargetRootBone, ScaleFactor);
+		RetargetPoses.AddOrUpdateRetargetPose(*this,CurrentRetargetPoseName, CurrentRetargetPose, PelvisBoneName, ScaleFactor);
 		// set this as the current pose to use (may be overridden by ops)
 		RetargetPoses.CurrentRetargetPoseName = CurrentRetargetPoseName;
 	}
@@ -818,19 +818,6 @@ bool FRetargeterBoneChains::Initialize(
 	const FRetargetSkeleton& InTargetSkeleton,
 	FIKRigLogger& InLog)
 {
-	Reset();
-
-	// get source and target IK Rig assets
-	const UIKRigDefinition* SourceIKRig = InRetargetAsset->GetIKRig(ERetargetSourceOrTarget::Source);
-	if (!ensure(SourceIKRig && !InTargetIKRigs.IsEmpty()))
-	{
-		InLog.LogError(LOCTEXT("MissingIKRig", "IK Retargeter either a source or target IK Rig asset."));
-		return false;
-	}
-
-	// first element is always the default
-	DefaultTargetIKRig = InTargetIKRigs[0];
-	
 	// load ALL bone chains on source and target (ops may use them even if they aren't mapped)
 	auto LoadChains = [](
 		const FRetargetSkeleton& InRetargetSkeleton,
@@ -849,8 +836,18 @@ bool FRetargeterBoneChains::Initialize(
 		}
 	};
 
+	// reset all the chains
+	Reset();
+
+	// store the default IK Rig (may be null)
+	DefaultTargetIKRig = InTargetIKRigs.IsEmpty() ? nullptr : InTargetIKRigs[0];
+
 	// load chains for source
-	LoadChains(InSourceSkeleton, SourceIKRig->GetRetargetChains(), InLog, SourceBoneChains);
+	const UIKRigDefinition* SourceIKRig = InRetargetAsset->GetIKRig(ERetargetSourceOrTarget::Source);
+	if (SourceIKRig)
+	{
+		LoadChains(InSourceSkeleton, SourceIKRig->GetRetargetChains(), InLog, SourceBoneChains);
+	}
 
 	// load chains for all target IK rigs
 	for (const UIKRigDefinition* TargetIKRig : InTargetIKRigs)
@@ -1135,16 +1132,13 @@ void FIKRetargetProcessor::Initialize(
 
 	// record source asset
 	RetargeterAsset = InRetargeterAsset;
+
+	// check prerequisite assets
 	if (!RetargeterAsset)
 	{
 		Log.LogError(LOCTEXT("MissingRetargetAsset", "IK Retargeter unable to initialize because no IK Retargeter asset was supplied."));
 		return;
 	}
-
-	const UIKRigDefinition* SourceIKRig = RetargeterAsset->GetIKRig(ERetargetSourceOrTarget::Source);
-	const UIKRigDefinition* TargetIKRig = RetargeterAsset->GetIKRig(ERetargetSourceOrTarget::Target);
-
-	// check prerequisite assets
 	if (!SourceSkeletalMesh)
 	{
 		Log.LogError(LOCTEXT("MissingSourceMesh", "IK Retargeter unable to initialize. Missing source Skeletal Mesh asset."));
@@ -1155,23 +1149,13 @@ void FIKRetargetProcessor::Initialize(
 		Log.LogError(LOCTEXT("MissingTargetMesh", "IK Retargeter unable to initialize. Missing target Skeletal Mesh asset."));
 		return;
 	}
-	if (!SourceIKRig)
-	{
-		Log.LogError(LOCTEXT("MissingSourceIKRig", "IK Retargeter unable to initialize. Missing source IK Rig asset."));
-		return;
-	}
-	if (!TargetIKRig)
-	{
-		Log.LogError(LOCTEXT("MissingTargetIKRig", "IK Retargeter unable to initialize. Missing target IK Rig asset."));
-		return;
-	}
 	
 	// initialize source skeleton
 	SourceSkeleton.Initialize(
 		SourceSkeletalMesh,
 		ERetargetSourceOrTarget::Source,
 		RetargeterAsset,
-		SourceIKRig->GetPelvis(),
+		GetPelvisBone(ERetargetSourceOrTarget::Source, ERetargetOpsToSearch::AssetOps),
 		GetSourceScaleFactor());
 	
 	// initialize target skeleton
@@ -1179,7 +1163,7 @@ void FIKRetargetProcessor::Initialize(
 		TargetSkeletalMesh,
 		ERetargetSourceOrTarget::Target,
 		RetargeterAsset,
-		TargetIKRig->GetPelvis());
+		GetPelvisBone(ERetargetSourceOrTarget::Target, ERetargetOpsToSearch::AssetOps));
 
 	// setup the goal container
 	ResetIKRigGoalContainer();
@@ -1205,8 +1189,12 @@ void FIKRetargetProcessor::Initialize(
 void FIKRetargetProcessor::ResetIKRigGoalContainer()
 {
 	// fill the goal container with all the goals in the IK Rig
+	GoalContainer.Empty();
 	const UIKRigDefinition* TargetIKRig = RetargeterAsset->GetIKRig(ERetargetSourceOrTarget::Target);
-	GoalContainer.FillWithGoalArray(TargetIKRig->GetGoalArray());
+	if (TargetIKRig)
+	{
+		GoalContainer.FillWithGoalArray(TargetIKRig->GetGoalArray());
+	}
 	
 	// set all goals to be additive (ops may adjust them)
 	for (FIKRigGoal& Goal : GoalContainer.GetGoalArray())
@@ -1255,7 +1243,14 @@ bool FIKRetargetProcessor::InitializeOpStack(
 
 	// collect all the IK Rigs referenced by the ops
 	TArray<const UIKRigDefinition*> AllIKRigs;
-	AllIKRigs.Add(RetargeterAsset->GetIKRig(ERetargetSourceOrTarget::Target));
+
+	// add the default target IK Rig
+	const UIKRigDefinition* DefaultTargetIKRig = RetargeterAsset->GetIKRig(ERetargetSourceOrTarget::Target);
+	if (DefaultTargetIKRig)
+	{
+		AllIKRigs.Add(DefaultTargetIKRig);
+	}
+	// add any IK Rigs referenced by the ops
 	for (FInstancedStruct& OpStruct : OpStack)
 	{
 		FIKRetargetOpBase& NewOp = OpStruct.GetMutable<FIKRetargetOpBase>();
@@ -1534,7 +1529,7 @@ bool FIKRetargetProcessor::IsBoneInAMappedChain(
 	}
 
 	// pelvis is always retargeted
-	const FName PelvisBoneName = GetPelvisBone(SourceOrTarget);
+	const FName PelvisBoneName = GetPelvisBone(SourceOrTarget, ERetargetOpsToSearch::AssetOps);
 	if (BoneName == PelvisBoneName)
 	{
 		return true;
@@ -1578,7 +1573,7 @@ FName FIKRetargetProcessor::GetChainNameForBone(const FName BoneName, const ERet
 	const bool bUseSource = SourceOrTarget == ERetargetSourceOrTarget::Source;
 	const FRetargetSkeleton& Skeleton = bUseSource ? SourceSkeleton : TargetSkeleton;
 	const int32 BoneIndex = GetBoneIndexFromName(BoneName, SourceOrTarget);
-	const FName RootBoneName = GetPelvisBone(SourceOrTarget);
+	const FName RootBoneName = GetPelvisBone(SourceOrTarget, ERetargetOpsToSearch::AssetOps);
 	if (BoneName == RootBoneName)
 	{
 		return FName("Retarget Pelvis");
@@ -1804,7 +1799,7 @@ void FIKRetargetProcessor::UpdateRetargetPoseAtRuntime(
 	
 	auto AddOrUpdateNewRetargetPose = [this, NewRetargetPoseName, NewRetargetPose, SourceOrTarget]()
 	{
-		const FName PelvisBoneName = GetPelvisBone(SourceOrTarget);
+		const FName PelvisBoneName = GetPelvisBone(SourceOrTarget, ERetargetOpsToSearch::ProcessorOps);
 		const double ScaleFactor = SourceOrTarget == ERetargetSourceOrTarget::Source ? GetSourceScaleFactor() : 1.0;
 		const bool bIsSource = SourceOrTarget == ERetargetSourceOrTarget::Source;
 		FRetargetSkeleton& RetargetSkeleton = bIsSource ? SourceSkeleton : TargetSkeleton;
@@ -1858,10 +1853,10 @@ void FIKRetargetProcessor::UpdateRetargetPoseAtRuntime(
 	}
 }
 
-FName FIKRetargetProcessor::GetPelvisBone(ERetargetSourceOrTarget SourceOrTarget) const
+FName FIKRetargetProcessor::GetPelvisBone(ERetargetSourceOrTarget SourceOrTarget, ERetargetOpsToSearch InOpsToSearch) const
 {
 	// first try get it from a pelvis of
-	const FIKRetargetPelvisMotionOp* PelvisMotionOp = GetFirstRetargetOpOfType<FIKRetargetPelvisMotionOp>();
+	const FIKRetargetPelvisMotionOp* PelvisMotionOp = GetFirstRetargetOpOfType<FIKRetargetPelvisMotionOp>(InOpsToSearch);
 	if (PelvisMotionOp)
 	{
 		FName PelvisNameFromOp = PelvisMotionOp->GetPelvisBoneName(SourceOrTarget);
diff --git a/Engine/Plugins/Animation/IKRig/Source/IKRig/Private/Retargeter/RetargetOps/FKChainsOp.cpp b/Engine/Plugins/Animation/IKRig/Source/IKRig/Private/Retargeter/RetargetOps/FKChainsOp.cpp
--- a/Engine/Plugins/Animation/IKRig/Source/IKRig/Private/Retargeter/RetargetOps/FKChainsOp.cpp
+++ b/Engine/Plugins/Animation/IKRig/Source/IKRig/Private/Retargeter/RetargetOps/FKChainsOp.cpp
@@ -692,7 +692,7 @@ void FIKRetargetFKChainsOp::PostInitialize(
 	const FTargetSkeleton& TargetSkeleton,
 	FIKRigLogger& Log)
 {
-	const FName TargetPelvisBoneName = Processor.GetPelvisBone(ERetargetSourceOrTarget::Target);
+	const FName TargetPelvisBoneName = Processor.GetPelvisBone(ERetargetSourceOrTarget::Target, ERetargetOpsToSearch::ProcessorOps);
 	const int32 TargetPelvisBoneIndex = TargetSkeleton.FindBoneIndexByName(TargetPelvisBoneName);
 
 	// record intermediate bones (non-retargeted bones located BETWEEN FK chains on the target skeleton)
diff --git a/Engine/Plugins/Animation/IKRig/Source/IKRig/Public/Retargeter/IKRetargetProcessor.h b/Engine/Plugins/Animation/IKRig/Source/IKRig/Public/Retargeter/IKRetargetProcessor.h
--- a/Engine/Plugins/Animation/IKRig/Source/IKRig/Public/Retargeter/IKRetargetProcessor.h
+++ b/Engine/Plugins/Animation/IKRig/Source/IKRig/Public/Retargeter/IKRetargetProcessor.h
@@ -92,7 +92,7 @@ struct FRetargetSkeleton
 		const USkeletalMesh* InSkeletalMesh,
 		const ERetargetSourceOrTarget InSourceOrTarget,
 		const UIKRetargeter* InRetargetAsset,
-		const FName RetargetRootBone,
+		const FName PelvisBoneName,
 		const double ScaleFactor);
 	
 	UE_API void Reset();
@@ -491,8 +491,8 @@ public:
 	/** Calls AnimGraphEvaluateAnyThread() for all ops in the stack */
 	UE_API void OnAnimGraphEvaluateAnyThread(FPoseContext& Output);
 
-	/** Get name of the root bone of the given skeleton. */
-	UE_API FName GetPelvisBone(ERetargetSourceOrTarget SourceOrTarget) const;
+	/** Get name of the pelvis bone for either the source or target skeleton. */
+	UE_API FName GetPelvisBone(ERetargetSourceOrTarget SourceOrTarget, ERetargetOpsToSearch InOpsToSearch) const;
 
 	/** Determine if we should be directly copying in the source pose into the target
 	 * @param OutCopyBasePoseRoot: will contain the name of the bone to copy from if one was specified, otherwise Name_None
@@ -642,7 +642,7 @@ public:
 	}
 	FName GetRetargetRoot(ERetargetSourceOrTarget SourceOrTarget) const
 	{
-		return Processor.GetPelvisBone(SourceOrTarget);
+		return Processor.GetPelvisBone(SourceOrTarget, ERetargetOpsToSearch::AssetOps);
 	}
 	bool IsInitialized() const
 	{
diff --git a/Engine/Plugins/Animation/IKRig/Source/IKRig/Public/Retargeter/IKRetargeter.h b/Engine/Plugins/Animation/IKRig/Source/IKRig/Public/Retargeter/IKRetargeter.h
--- a/Engine/Plugins/Animation/IKRig/Source/IKRig/Public/Retargeter/IKRetargeter.h
+++ b/Engine/Plugins/Animation/IKRig/Source/IKRig/Public/Retargeter/IKRetargeter.h
@@ -325,7 +325,7 @@ private:
 
 #if WITH_EDITORONLY_DATA
 	// Optional. Override the Skeletal Mesh to copy animation from. Uses the preview mesh from the Source IK Rig asset by default. 
-	UPROPERTY(EditAnywhere, Category = Source, meta = (EditCondition = "HasSourceIKRig"))
+	UPROPERTY(EditAnywhere, Category = Source)
 	TSoftObjectPtr<USkeletalMesh> SourcePreviewMesh = nullptr;
 #endif
 	
@@ -335,7 +335,7 @@ private:
 
 #if WITH_EDITORONLY_DATA
 	// Optional. Override the Skeletal Mesh to preview the retarget on. Uses the preview mesh from the Target IK Rig asset by default. 
-	UPROPERTY(EditAnywhere, Category = Target, meta = (EditCondition = "HasTargetIKRig"))
+	UPROPERTY(EditAnywhere, Category = Target)
 	TSoftObjectPtr<USkeletalMesh> TargetPreviewMesh = nullptr;
 #endif
 	
diff --git a/Engine/Plugins/Animation/IKRig/Source/IKRigEditor/Private/RetargetEditor/IKRetargeterController.cpp b/Engine/Plugins/Animation/IKRig/Source/IKRigEditor/Private/RetargetEditor/IKRetargeterController.cpp
--- a/Engine/Plugins/Animation/IKRig/Source/IKRigEditor/Private/RetargetEditor/IKRetargeterController.cpp
+++ b/Engine/Plugins/Animation/IKRig/Source/IKRigEditor/Private/RetargetEditor/IKRetargeterController.cpp
@@ -74,12 +74,12 @@ void UIKRetargeterController::SetIKRig(const ERetargetSourceOrTarget SourceOrTar
 	if (SourceOrTarget == ERetargetSourceOrTarget::Source)
 	{
 		Asset->SourceIKRigAsset = IKRig;
-		Asset->SourcePreviewMesh = IKRig ? IKRig->GetPreviewMesh() : nullptr;
+		Asset->SourcePreviewMesh = IKRig ? IKRig->GetPreviewMesh() : Asset->SourcePreviewMesh;
 	}
 	else
 	{
 		Asset->TargetIKRigAsset = IKRig;
-		Asset->TargetPreviewMesh = IKRig ? IKRig->GetPreviewMesh() : nullptr;
+		Asset->TargetPreviewMesh = IKRig ? IKRig->GetPreviewMesh() : Asset->TargetPreviewMesh;
 	}
 	
 	// re-ask to fix root height for this mesh
@@ -170,14 +170,7 @@ USkeletalMesh* UIKRetargeterController::GetPreviewMesh(const ERetargetSourceOrTa
 {
 	FScopeLock Lock(&ControllerLock);
 	
-	// can't preview anything if target IK Rig is null
-	const UIKRigDefinition* IKRig = GetIKRig(SourceOrTarget);
-	if (!IKRig)
-	{
-		return nullptr;
-	}
-
-	// optionally prefer override if one is provided
+	// return the preview mesh if one is provided
 	const TSoftObjectPtr<USkeletalMesh> PreviewMesh = SourceOrTarget == ERetargetSourceOrTarget::Source ? Asset->SourcePreviewMesh : Asset->TargetPreviewMesh;
 	if (!PreviewMesh.IsNull())
 	{
@@ -185,7 +178,12 @@ USkeletalMesh* UIKRetargeterController::GetPreviewMesh(const ERetargetSourceOrTa
 	}
 
 	// fallback to preview mesh from IK Rig asset
-	return IKRig->GetPreviewMesh();
+	if (const UIKRigDefinition* IKRig = GetIKRig(SourceOrTarget))
+	{
+		return IKRig->GetPreviewMesh();
+	}
+	
+	return nullptr;
 }
 
 int32 UIKRetargeterController::AddRetargetOp(const FString InIKRetargetOpType) const
diff --git a/Engine/Plugins/Animation/IKRig/Source/IKRigEditor/Private/RetargetEditor/IKRetargeterPoseGenerator.cpp b/Engine/Plugins/Animation/IKRig/Source/IKRigEditor/Private/RetargetEditor/IKRetargeterPoseGenerator.cpp
--- a/Engine/Plugins/Animation/IKRig/Source/IKRigEditor/Private/RetargetEditor/IKRetargeterPoseGenerator.cpp
+++ b/Engine/Plugins/Animation/IKRig/Source/IKRigEditor/Private/RetargetEditor/IKRetargeterPoseGenerator.cpp
@@ -62,7 +62,7 @@ void FRetargetAutoPoseGenerator::AlignBones(
 	if (!bSuppressWarnings)
 	{
 		// warn user if they tried to align the pelvis (doesn't work)
-		if (BonesToAlign.Contains(Processor->GetPelvisBone(SourceOrTarget)))
+		if (BonesToAlign.Contains(Processor->GetPelvisBone(SourceOrTarget, ERetargetOpsToSearch::AssetOps)))
 		{
 			Processor->Log.LogWarning(LOCTEXT("AutoAlignPelvisNotSupported", "Skipped aligning the pelvis. This is not supported."));
 		}
@@ -100,7 +100,7 @@ void FRetargetAutoPoseGenerator::AlignBone(
 	Processor->UpdateRetargetPoseAtRuntime(CurrentPoseName, SourceOrTarget);
 
 	// use a custom procedure to auto-align pelvis
-	if (BoneToAlign == Processor->GetPelvisBone(SourceOrTarget))
+	if (BoneToAlign == Processor->GetPelvisBone(SourceOrTarget, ERetargetOpsToSearch::AssetOps))
 	{
 		AlignPelvisBone(SourceOrTarget);
 		return;
@@ -274,7 +274,7 @@ void FRetargetAutoPoseGenerator::AlignPelvisBone(const ERetargetSourceOrTarget S
 	const FQuat WorldRotationOffset = FQuat::FindBetweenNormals(DirectionToAlign, DirectionToMatch);
 
 	// convert to local space offset and store in retarget pose
-	const FName BoneToAlign = Processor->GetPelvisBone(SourceOrTarget);
+	const FName BoneToAlign = Processor->GetPelvisBone(SourceOrTarget, ERetargetOpsToSearch::AssetOps);
 	ApplyWorldRotationToRetargetPose(WorldRotationOffset, BoneToAlign, SourceOrTarget);
 }
 
