[IK Retarget] Fixed bug, Scale Source causes broken retarget output.

diff --git a/Engine/Plugins/Animation/IKRig/Source/IKRig/Private/Retargeter/IKRetargetProcessor.cpp b/Engine/Plugins/Animation/IKRig/Source/IKRig/Private/Retargeter/IKRetargetProcessor.cpp
--- a/Engine/Plugins/Animation/IKRig/Source/IKRig/Private/Retargeter/IKRetargetProcessor.cpp
+++ b/Engine/Plugins/Animation/IKRig/Source/IKRig/Private/Retargeter/IKRetargetProcessor.cpp
@@ -1168,11 +1168,11 @@ void FIKRetargetProcessor::Initialize(
 	// setup the goal container
 	ResetIKRigGoalContainer();
 
-	// initialize the retargeting ops
-	InitializeOpStack(RetargeterAsset->GetRetargetOps(), InRetargetProfile);
+	// copy ops from asset 
+	InitialOpStackSetup(RetargeterAsset->GetRetargetOps(), InRetargetProfile);
 	
-	// determine which bones are retargeted and run PostInitialize() on each op
-	PostInitializeOpStack();
+	// initialize all the retarget ops
+	InitializeRetargetOps();
 
 	// confirm for the user that the retargeter is initialized
 	Log.LogInfo(FText::Format(
@@ -1206,12 +1206,10 @@ void FIKRetargetProcessor::ResetIKRigGoalContainer()
 	}
 }
 
-bool FIKRetargetProcessor::InitializeOpStack(
+void FIKRetargetProcessor::InitialOpStackSetup(
 	const TArray<FInstancedStruct>& OpsFromAsset,
 	const FRetargetProfile& InRetargetProfile)
 {
-	bool bAllOpsInitialized = true;
-	
 	// create copies of all the ops in the asset
 	OpStack.Reset(OpsFromAsset.Num());
 	for (const FInstancedStruct& AssetOpStruct : OpsFromAsset)
@@ -1262,7 +1260,10 @@ bool FIKRetargetProcessor::InitializeOpStack(
 	
 	// resolve all bone chains onto the skeletons
 	AllBoneChains.Initialize(RetargeterAsset, AllIKRigs, SourceSkeleton, TargetSkeleton, Log);
+}
 
+void FIKRetargetProcessor::InitializeRetargetOps()
+{
 	// initialize the ops
 	for (FInstancedStruct& OpStruct : OpStack)
 	{
@@ -1270,17 +1271,7 @@ bool FIKRetargetProcessor::InitializeOpStack(
 		const FIKRetargetOpBase* ParentOp = GetRetargetOpByName(Op.GetParentOpName());
 		Op.Initialize(*this, SourceSkeleton, TargetSkeleton, ParentOp, Log);
 	}
-
-	if (!bAllOpsInitialized)
-	{
-		Log.LogWarning(FText(LOCTEXT("PostOperationInitError", "One (or more) of the retarget operations has warnings. See output log for details.")));
-	}
 	
-	return bAllOpsInitialized;
-}
-
-void FIKRetargetProcessor::PostInitializeOpStack()
-{
 	// gather retargeted bones from each op
 	TSet<int32> AllRetargetedBones;
 	for (FInstancedStruct& OpStruct : OpStack)
@@ -1291,6 +1282,7 @@ void FIKRetargetProcessor::PostInitializeOpStack()
 			Op.CollectRetargetedBones(AllRetargetedBones);
 		}
 	}
+	
 	// store retargeted bone mask on target skeleton
 	TargetSkeleton.SetRetargetedBones(AllRetargetedBones);
 
@@ -1843,14 +1835,9 @@ void FIKRetargetProcessor::UpdateRetargetPoseAtRuntime(
 
 	// re-load the updated pose into the bone chains
 	AllBoneChains.UpdatePoseFromSkeleton(RetargetSkeleton, SourceOrTarget);
-
+	
 	// re-initialize the op stack (ops sometimes cache retarget poses)
-	for (FInstancedStruct& OpStruct : OpStack)
-	{
-		FIKRetargetOpBase& Op = OpStruct.GetMutable<FIKRetargetOpBase>();
-		const FIKRetargetOpBase* ParentOp = GetRetargetOpByName(Op.GetParentOpName());
-		Op.Initialize(*this, SourceSkeleton, TargetSkeleton, ParentOp, Log);
-	}
+	InitializeRetargetOps();
 }
 
 FName FIKRetargetProcessor::GetPelvisBone(ERetargetSourceOrTarget SourceOrTarget, ERetargetOpsToSearch InOpsToSearch) const
diff --git a/Engine/Plugins/Animation/IKRig/Source/IKRig/Public/Retargeter/IKRetargetProcessor.h b/Engine/Plugins/Animation/IKRig/Source/IKRig/Public/Retargeter/IKRetargetProcessor.h
--- a/Engine/Plugins/Animation/IKRig/Source/IKRig/Public/Retargeter/IKRetargetProcessor.h
+++ b/Engine/Plugins/Animation/IKRig/Source/IKRig/Public/Retargeter/IKRetargetProcessor.h
@@ -573,22 +573,22 @@ private:
 	TArray<FInstancedStruct> OpStack;
 
 	/** Apply the settings stored in a retarget profile. Called inside RunRetargeter(). */
-	UE_API void ApplyProfileAtRuntime(const FRetargetProfile& Profile);
+	void ApplyProfileAtRuntime(const FRetargetProfile& Profile);
 
 	/** Internal retarget phase that does simple bone-to-bone copying from source*/
-	UE_API void GenerateBasePoses(TArray<FTransform>& InSourceGlobalPose);
+	void GenerateBasePoses(TArray<FTransform>& InSourceGlobalPose);
 
 	/** Setup the goal container for IK Rigs to use */
-	UE_API void ResetIKRigGoalContainer();
+	void ResetIKRigGoalContainer();
 
-	/** Initialize the retarget op stack */
-	UE_API bool InitializeOpStack(const TArray<FInstancedStruct>& OpsFromAsset, const FRetargetProfile& InRetargetProfile);
+	/** Initial setup of the retarget op stack (copies op stack from asset into processor) */
+	void InitialOpStackSetup(const TArray<FInstancedStruct>& OpsFromAsset, const FRetargetProfile& InRetargetProfile);
 
-	/** Collect retargeted bones and call PostInitialize on each Op */
-	UE_API void PostInitializeOpStack();
+	/** Run Initialize() on all ops then collects retargeted bones and call PostInitialize() on each Op */
+	void InitializeRetargetOps();
 
 	/** Run all post process operations on the retargeted result. */
-	UE_API void RunRetargetOps(
+	void RunRetargetOps(
 		const double InDeltaTime,
 		const TArray<FTransform>& InSourceGlobalPose,
 		TArray<FTransform>& OutTargetGlobalPose);
