[IK Retarget] Fixed bug, blending retarget poses can create non-normalized rotations.

diff --git a/Engine/Plugins/Animation/IKRig/Source/IKRig/Private/Retargeter/IKRetargetProcessor.cpp b/Engine/Plugins/Animation/IKRig/Source/IKRig/Private/Retargeter/IKRetargetProcessor.cpp
--- a/Engine/Plugins/Animation/IKRig/Source/IKRig/Private/Retargeter/IKRetargetProcessor.cpp
+++ b/Engine/Plugins/Animation/IKRig/Source/IKRig/Private/Retargeter/IKRetargetProcessor.cpp
@@ -97,7 +97,7 @@ FResolvedRetargetPose& FResolvedRetargetPoseSet::FindOrAddRetargetPose(const FNa
 {
 	for (FResolvedRetargetPose& Pose : RetargetPoses)
 	{
-		if (Pose.Name == CurrentRetargetPoseName)
+		if (Pose.Name == InRetargetPoseName)
 		{
 			return Pose;
 		}
diff --git a/Engine/Plugins/Animation/IKRig/Source/IKRigEditor/Private/RetargetEditor/IKRetargetAnimInstance.cpp b/Engine/Plugins/Animation/IKRig/Source/IKRigEditor/Private/RetargetEditor/IKRetargetAnimInstance.cpp
--- a/Engine/Plugins/Animation/IKRig/Source/IKRigEditor/Private/RetargetEditor/IKRetargetAnimInstance.cpp
+++ b/Engine/Plugins/Animation/IKRig/Source/IKRigEditor/Private/RetargetEditor/IKRetargetAnimInstance.cpp
@@ -121,6 +121,7 @@ void FAnimNode_PreviewRetargetPose::Evaluate_AnyThread(FPoseContext& Output)
 
 		BlendTransform<ETransformBlendMode::Overwrite>(RefPose[MeshBoneIndex], Output.Pose[BoneIndex], 1.f - RetargetPoseBlend);
 		BlendTransform<ETransformBlendMode::Accumulate>(RetargetLocalPose[MeshBoneIndex], Output.Pose[BoneIndex], RetargetPoseBlend);
+		Output.Pose.NormalizeRotations();
 	}
 }
 
