[IK Retarget] Fixed crash when trying to apply empty op settings in retarget profile

diff --git a/Engine/Plugins/Animation/IKRig/Source/IKRig/Private/Retargeter/IKRetargetProfile.cpp b/Engine/Plugins/Animation/IKRig/Source/IKRig/Private/Retargeter/IKRetargetProfile.cpp
--- a/Engine/Plugins/Animation/IKRig/Source/IKRig/Private/Retargeter/IKRetargetProfile.cpp
+++ b/Engine/Plugins/Animation/IKRig/Source/IKRig/Private/Retargeter/IKRetargetProfile.cpp
@@ -30,7 +30,10 @@ void FRetargetOpProfile::CopyFromOtherOpProfile(const FRetargetOpProfile& OtherO
 
 bool FRetargetOpProfile::CopySettingsToOp(FInstancedStruct& InOutOpStruct, ECopyOpSettingsContext InApplyContext) const
 {
-	const FIKRetargetOpSettingsBase& OpSettingsFromProfile = SettingsToApply.Get<FIKRetargetOpSettingsBase>();
+	if (!SettingsToApply.IsValid())
+	{
+		return false; // no settings stored in the op profile
+	}
 	
 	FIKRetargetOpBase* Op = InOutOpStruct.GetMutablePtr<FIKRetargetOpBase>();
 	if (Op->GetSettingsType() != SettingsToApply.GetScriptStruct())
@@ -38,6 +41,7 @@ bool FRetargetOpProfile::CopySettingsToOp(FInstancedStruct& InOutOpStruct, ECopy
 		return false; // settings were wrong type
 	}
 
+	const FIKRetargetOpSettingsBase& OpSettingsFromProfile = SettingsToApply.Get<FIKRetargetOpSettingsBase>();
 	switch (InApplyContext)
 	{
 		case ECopyOpSettingsContext::Runtime:
diff --git a/Engine/Plugins/Animation/IKRig/Source/IKRigEditor/Private/RetargetEditor/IKRetargetOpDetails.cpp b/Engine/Plugins/Animation/IKRig/Source/IKRigEditor/Private/RetargetEditor/IKRetargetOpDetails.cpp
--- a/Engine/Plugins/Animation/IKRig/Source/IKRigEditor/Private/RetargetEditor/IKRetargetOpDetails.cpp
+++ b/Engine/Plugins/Animation/IKRig/Source/IKRigEditor/Private/RetargetEditor/IKRetargetOpDetails.cpp
@@ -36,7 +36,7 @@ bool FIKRetargetOpBaseSettingsCustomization::LoadAndValidateStructToCustomize(
 
 	FIKRetargetOpSettingsBase* SettingsBeingCustomized = static_cast<FIKRetargetOpSettingsBase*>(StructMemory);
 	OpName = SettingsBeingCustomized->OwningOpName;
-	if (!ensure(OpName != NAME_None))
+	if (OpName == NAME_None)
 	{
 		return false;
 	}
