IKRig: fixed unitiliazied properties

#rnx

diff --git a/Engine/Plugins/Animation/IKRig/Source/IKRig/Public/Retargeter/RetargetOps/RetargetPoseOp.h b/Engine/Plugins/Animation/IKRig/Source/IKRig/Public/Retargeter/RetargetOps/RetargetPoseOp.h
--- a/Engine/Plugins/Animation/IKRig/Source/IKRig/Public/Retargeter/RetargetOps/RetargetPoseOp.h
+++ b/Engine/Plugins/Animation/IKRig/Source/IKRig/Public/Retargeter/RetargetOps/RetargetPoseOp.h
@@ -21,10 +21,10 @@ struct IKRIG_API FIKRetargetPoseOpSettings : public FIKRetargetOpSettingsBase
 	GENERATED_BODY()
 
 	UPROPERTY(EditAnywhere, BlueprintReadWrite, Category="Pose")
-	ERetargetSourceOrTarget SourceOrTarget;
+	ERetargetSourceOrTarget SourceOrTarget = ERetargetSourceOrTarget::Source;
 
 	UPROPERTY(EditAnywhere, BlueprintReadWrite, Category="Pose")
-	ERetargetPoseMode Mode;
+	ERetargetPoseMode Mode = ERetargetPoseMode::RetargetPose;
 	
 	UPROPERTY(EditAnywhere, BlueprintReadWrite, Category="Pose")
 	FName CurrentPose;
