[IK Retarget] Fixes crash when viewing customized op settings outside retargeter details panel.

diff --git a/Engine/Plugins/Animation/IKRig/Source/IKRigEditor/Private/RetargetEditor/IKRetargetOpDetails.cpp b/Engine/Plugins/Animation/IKRig/Source/IKRigEditor/Private/RetargetEditor/IKRetargetOpDetails.cpp
--- a/Engine/Plugins/Animation/IKRig/Source/IKRigEditor/Private/RetargetEditor/IKRetargetOpDetails.cpp
+++ b/Engine/Plugins/Animation/IKRig/Source/IKRigEditor/Private/RetargetEditor/IKRetargetOpDetails.cpp
@@ -42,6 +42,10 @@ bool FIKRetargetOpBaseSettingsCustomization::LoadAndValidateStructToCustomize(
 	}
 	
 	StructViewer = Cast<UIKRigStructViewer>(SelectedObjects[0].Get());
+	if (StructViewer == nullptr)
+	{
+		return false;
+	}
 	RetargetAsset = CastChecked<UIKRetargeter>(StructViewer->GetStructOwner());
 	AssetController = UIKRetargeterController::GetController(RetargetAsset);
 
