[IK Retarget] Fixing pose selection combobox not refreshing.

diff --git a/Engine/Plugins/Animation/IKRig/Source/IKRigEditor/Private/RetargetEditor/SIKRetargetPoseEditor.cpp b/Engine/Plugins/Animation/IKRig/Source/IKRigEditor/Private/RetargetEditor/SIKRetargetPoseEditor.cpp
--- a/Engine/Plugins/Animation/IKRig/Source/IKRigEditor/Private/RetargetEditor/SIKRetargetPoseEditor.cpp
+++ b/Engine/Plugins/Animation/IKRig/Source/IKRigEditor/Private/RetargetEditor/SIKRetargetPoseEditor.cpp
@@ -22,6 +22,8 @@ void SIKRetargetPoseEditor::Construct(
 
 	// the commands for the menus
 	TSharedPtr<FUICommandList> Commands = Controller->Editor.Pin()->GetToolkitCommands();
+
+	RefreshPoseNames();
 	
 	ChildSlot
 	[
@@ -47,8 +49,12 @@ void SIKRetargetPoseEditor::Construct(
 			+SHorizontalBox::Slot()
 			.AutoWidth()
 			[
-				SAssignNew(PoseListComboBox, SComboBox<TSharedPtr<FName>>)
+				SNew(SComboBox<TSharedPtr<FName>>)
 				.OptionsSource(&PoseNames)
+				.OnComboBoxOpening_Lambda([this]()
+				{
+					RefreshPoseNames();
+				})
 				.OnGenerateWidget_Lambda([](TSharedPtr<FName> InItem)
 				{
 					return SNew(STextBlock).Text(FText::FromName(*InItem.Get()));
@@ -89,7 +95,7 @@ void SIKRetargetPoseEditor::Construct(
 	];
 }
 
-void SIKRetargetPoseEditor::Refresh()
+void SIKRetargetPoseEditor::RefreshPoseNames()
 {
 	// get the retarget poses from the editor controller
 	const FIKRetargetEditorController* Controller = EditorController.Pin().Get();
@@ -101,8 +107,6 @@ void SIKRetargetPoseEditor::Refresh()
 	{
 		PoseNames.Add(MakeShareable(new FName(Pose.Key)));
 	}
-
-	PoseListComboBox->RefreshOptions();
 }
 
 TSharedRef<SWidget>  SIKRetargetPoseEditor::MakeToolbar(TSharedPtr<FUICommandList> Commands)
diff --git a/Engine/Plugins/Animation/IKRig/Source/IKRigEditor/Public/RetargetEditor/SIKRetargetHierarchy.h b/Engine/Plugins/Animation/IKRig/Source/IKRigEditor/Public/RetargetEditor/SIKRetargetHierarchy.h
--- a/Engine/Plugins/Animation/IKRig/Source/IKRigEditor/Public/RetargetEditor/SIKRetargetHierarchy.h
+++ b/Engine/Plugins/Animation/IKRig/Source/IKRigEditor/Public/RetargetEditor/SIKRetargetHierarchy.h
@@ -82,7 +82,7 @@ public:
 
 	void ShowItemAfterSelection(FName ItemKey);
 
-	void RefreshPoseList() const { EditPoseView.Get()->Refresh(); };
+	void RefreshPoseList() const { EditPoseView.Get()->RefreshPoseNames(); };
 
 private:
 	
diff --git a/Engine/Plugins/Animation/IKRig/Source/IKRigEditor/Public/RetargetEditor/SIKRetargetPoseEditor.h b/Engine/Plugins/Animation/IKRig/Source/IKRigEditor/Public/RetargetEditor/SIKRetargetPoseEditor.h
--- a/Engine/Plugins/Animation/IKRig/Source/IKRigEditor/Public/RetargetEditor/SIKRetargetPoseEditor.h
+++ b/Engine/Plugins/Animation/IKRig/Source/IKRigEditor/Public/RetargetEditor/SIKRetargetPoseEditor.h
@@ -19,7 +19,7 @@ public:
 
 	void Construct(const FArguments& InArgs, TSharedRef<FIKRetargetEditorController> InEditorController);
 
-	void Refresh();
+	void RefreshPoseNames();
 	
 private:
 
@@ -33,5 +33,4 @@ private:
 	
 	TWeakPtr<FIKRetargetEditorController> EditorController;
 	TArray<TSharedPtr<FName>> PoseNames;
-	TSharedPtr<SComboBox<TSharedPtr<FName>>> PoseListComboBox;
 };
