[IK Retarget] Fixes auto-retargeter not generating a retarget pose and exported anims in t-pose.

diff --git a/Engine/Plugins/Animation/IKRig/Source/IKRigEditor/Private/RetargetEditor/IKRetargetBatchOperation.cpp b/Engine/Plugins/Animation/IKRig/Source/IKRigEditor/Private/RetargetEditor/IKRetargetBatchOperation.cpp
--- a/Engine/Plugins/Animation/IKRig/Source/IKRigEditor/Private/RetargetEditor/IKRetargetBatchOperation.cpp
+++ b/Engine/Plugins/Animation/IKRig/Source/IKRigEditor/Private/RetargetEditor/IKRetargetBatchOperation.cpp
@@ -233,9 +233,6 @@ void UIKRetargetBatchOperation::RetargetAssets(
 			Controller.OpenBracket(FText::FromString("Preparing for retargeted animation."), bShouldTransact);
 			Controller.RemoveAllCurvesOfType(ERawCurveTrackTypes::RCT_Transform, bShouldTransact);
 
-			// clear bone tracks to prevent recompression
-			Controller.RemoveAllBoneTracks(bShouldTransact);
-
 			// reset all additive animation properties to ensure WYSIWYG playback of additive anims between retargeter and sequence
 			FAdditiveRetargetSettings SequenceSettings;
 			SequenceSettings.PrepareForRetarget(AnimSequenceToRetarget);
@@ -244,6 +241,8 @@ void UIKRetargetBatchOperation::RetargetAssets(
 			// set the retarget source to the target skeletal mesh
 			AnimSequenceToRetarget->RetargetSource = NAME_None;
 			AnimSequenceToRetarget->SetRetargetSourceAsset(Context.TargetMesh);
+
+			// removes all bone tracks not belonging to this skeleton, adds those that do
 			Controller.UpdateWithSkeleton(NewSkeleton, bShouldTransact);
 
 			// done editing sequence data, close bracket
@@ -332,7 +331,6 @@ void UIKRetargetBatchOperation::ConvertAnimation(
 	FScopedSlowTask& Progress)
 {
 	// initialize the retargeter
-	UObject* TransientOuter = Cast<UObject>(GetTransientPackage());
 	FIKRetargetProcessor Processor;
 	FRetargetProfile RetargetProfile;
 	RetargetProfile.FillProfileWithAssetSettings(Context.IKRetargetAsset);
diff --git a/Engine/Plugins/Animation/IKRig/Source/IKRigEditor/Private/RetargetEditor/SRetargetAnimAssetsWindow.cpp b/Engine/Plugins/Animation/IKRig/Source/IKRigEditor/Private/RetargetEditor/SRetargetAnimAssetsWindow.cpp
--- a/Engine/Plugins/Animation/IKRig/Source/IKRigEditor/Private/RetargetEditor/SRetargetAnimAssetsWindow.cpp
+++ b/Engine/Plugins/Animation/IKRig/Source/IKRigEditor/Private/RetargetEditor/SRetargetAnimAssetsWindow.cpp
@@ -26,6 +26,7 @@
 #include "Editor/EditorPerProjectUserSettings.h"
 #include "RetargetEditor/IKRetargetAnimInstance.h"
 #include "RetargetEditor/IKRetargetEditorController.h"
+#include "Retargeter/RetargetOps/FKChainsOp.h"
 #include "Retargeter/RetargetOps/PinBoneOp.h"
 #include "Retargeter/RetargetOps/RootMotionGeneratorOp.h"
 #include "Retargeter/RetargetOps/RunIKRigOp.h"
@@ -670,6 +671,16 @@ void FProceduralRetargetAssets::AutoGenerateIKRetargetAsset()
 	RetargetController->SetIKRig(ERetargetSourceOrTarget::Source, SourceIKRig);
 	RetargetController->SetIKRig(ERetargetSourceOrTarget::Target, TargetIKRig);
 
+	// clear the op stack and add default ops
+	// NOTE: this must be done BEFORE generating the retarget pose since it relies on the first op chain mapping
+	RetargetController->RemoveAllOps();
+	RetargetController->AddDefaultOps();
+
+	// use only exact chain mapping in the auto-retarget (default is fuzzy which can map toes to fingers and other chaos)
+	FIKRetargetFKChainsOp* FKChainOp = RetargetController->GetFirstRetargetOpOfType<FIKRetargetFKChainsOp>();
+	constexpr bool bForceRemap = true;
+	FKChainOp->ChainMapping.AutoMapChains(EAutoMapChainType::Exact, bForceRemap);
+
 	// reset and regenerate the target retarget pose
 	const FName TargetRetargetPoseName = RetargetController->GetCurrentRetargetPoseName(ERetargetSourceOrTarget::Target);
 	RetargetController->ResetRetargetPose(TargetRetargetPoseName, TArray<FName>() /* all bones if empty */, ERetargetSourceOrTarget::Target);
@@ -687,12 +698,6 @@ void FProceduralRetargetAssets::AutoGenerateIKRetargetAsset()
 			ERetargetSourceOrTarget::Target);
 	}
 
-	// clear the op stack and regenerate it
-	RetargetController->RemoveAllOps();
-
-	// add all the default ops for retargeting
-	RetargetController->AddDefaultOps();
-
 	// disable IK pass because auto-IK is not yet robust enough, still useful for manual editing though
 	FIKRetargetRunIKRigOp* RunIKOp = RetargetController->GetFirstRetargetOpOfType<FIKRetargetRunIKRigOp>();
 	RunIKOp->SetEnabled(false);
