[IK Retarget] Fixes issue with Pin Bone Op showing wrong bone selection menu and not updating when modifying skeletal to copy from.

diff --git a/Engine/Plugins/Animation/IKRig/Source/IKRig/Private/Retargeter/RetargetOps/PinBoneOp.cpp b/Engine/Plugins/Animation/IKRig/Source/IKRig/Private/Retargeter/RetargetOps/PinBoneOp.cpp
--- a/Engine/Plugins/Animation/IKRig/Source/IKRig/Private/Retargeter/RetargetOps/PinBoneOp.cpp
+++ b/Engine/Plugins/Animation/IKRig/Source/IKRig/Private/Retargeter/RetargetOps/PinBoneOp.cpp
@@ -83,19 +83,15 @@ void FIKRetargetPinBoneOpSettings::CopySettingsAtRuntime(const FIKRetargetOpSett
 #if WITH_EDITOR
 USkeleton* FIKRetargetPinBoneOpSettings::GetSkeleton(const FName InPropertyName)
 {
-	if (!SourceSkeletonAsset || !TargetSkeletonAsset)
-	{
-		return nullptr;
-	}
-
 	if (InPropertyName == GET_MEMBER_NAME_CHECKED(FPinBoneData, BoneToCopyTo))
 	{
-		const USkeleton* Skeleton = SkeletonToCopyFrom == ERetargetSourceOrTarget::Source ? TargetSkeletonAsset : SourceSkeletonAsset;
-		return const_cast<USkeleton*>(Skeleton);
+		// bone to copy TO is ALWAYS target
+		return const_cast<USkeleton*>(TargetSkeletonAsset);
 	}
-
+	
 	if (InPropertyName == GET_MEMBER_NAME_CHECKED(FPinBoneData, BoneToCopyFrom))
 	{
+		// bone to copy FROM may be either source or target
 		const USkeleton* Skeleton = SkeletonToCopyFrom == ERetargetSourceOrTarget::Target ? TargetSkeletonAsset : SourceSkeletonAsset;
 		return const_cast<USkeleton*>(Skeleton);
 	}
@@ -339,14 +335,17 @@ FTransform FIKRetargetPinBoneOp::GetNewBoneTransform(
 		}
 	case EPinBoneRotationMode::MaintainOffsetFromBoneToCopyFrom:
 		{
-			// get retarget pose of source and target
+			// get rotation delta between from/to in retarget pose
 			const TArray<FTransform>& FromRetargetPose = CopyFrom.RetargetPoses.GetGlobalRetargetPose();
 			const TArray<FTransform>& ToRetargetPose = CopyTo.RetargetPoses.GetGlobalRetargetPose();
 			const FTransform& BoneToCopyFromRefPoseGlobal = FromRetargetPose[PinData.BoneToCopyFrom.BoneIndex];
 			const FTransform& BoneToCopyToRefPoseGlobal = ToRetargetPose[PinData.BoneToCopyTo.BoneIndex];
-			const FTransform Offset = BoneToCopyFromRefPoseGlobal.GetRelativeTransform(BoneToCopyToRefPoseGlobal);
-			const FTransform BoneToCopyFromTransform = PoseToCopyFrom[PinData.BoneToCopyFrom.BoneIndex];
-			Result.SetRotation((Offset * BoneToCopyFromTransform).GetRotation());
+			const FQuat Delta = BoneToCopyToRefPoseGlobal.GetRotation() * BoneToCopyFromRefPoseGlobal.GetRotation().Inverse();
+
+			// apply "copy from" rotation plus delta
+			const FQuat& BoneToCopyFromRotation = PoseToCopyFrom[PinData.BoneToCopyFrom.BoneIndex].GetRotation();
+			Result.SetRotation(Delta * BoneToCopyFromRotation);
+			
 			break;
 		}
 		
diff --git a/Engine/Plugins/Animation/IKRig/Source/IKRig/Public/Retargeter/RetargetOps/PinBoneOp.h b/Engine/Plugins/Animation/IKRig/Source/IKRig/Public/Retargeter/RetargetOps/PinBoneOp.h
--- a/Engine/Plugins/Animation/IKRig/Source/IKRig/Public/Retargeter/RetargetOps/PinBoneOp.h
+++ b/Engine/Plugins/Animation/IKRig/Source/IKRig/Public/Retargeter/RetargetOps/PinBoneOp.h
@@ -97,7 +97,7 @@ struct FIKRetargetPinBoneOpSettings : public FIKRetargetOpSettingsBase
 	TArray<FPinBoneData> BonesToPin;
 
 	/** Which skeleton to copy animation from. */
-	UPROPERTY(EditAnywhere, BlueprintReadWrite, Category=Setup)
+	UPROPERTY(EditAnywhere, BlueprintReadWrite, Category=Setup, meta=(ReinitializeOnEdit))
 	ERetargetSourceOrTarget SkeletonToCopyFrom = ERetargetSourceOrTarget::Target;
 
 	/** Copy the translation of the bone. */
