[IK Retarget] Fixed issue preventing runtime swapping of retarget poses
Fixed "Retarget Pose Op" not working at all.

diff --git a/Engine/Plugins/Animation/IKRig/Source/IKRig/Private/Retargeter/IKRetargetProcessor.cpp b/Engine/Plugins/Animation/IKRig/Source/IKRig/Private/Retargeter/IKRetargetProcessor.cpp
--- a/Engine/Plugins/Animation/IKRig/Source/IKRig/Private/Retargeter/IKRetargetProcessor.cpp
+++ b/Engine/Plugins/Animation/IKRig/Source/IKRig/Private/Retargeter/IKRetargetProcessor.cpp
@@ -31,7 +31,7 @@ FResolvedRetargetPose& FResolvedRetargetPoseSet::AddOrUpdateRetargetPose(
 	// add retarget pose if it doesn't already exist
 	FResolvedRetargetPose& RetargetPose = FindOrAddRetargetPose(InRetargetPoseName);
 	
-	// record the name of the retarget pose (prevents re-initialization if profile swaps it)
+	// record the version of the retarget pose (prevents re-initialization if profile swaps it)
 	RetargetPose.Version = InRetargetPose->GetVersion();
 	RetargetPose.Scale = ScaleFactor;
 	
@@ -111,13 +111,12 @@ const FResolvedRetargetPose* FResolvedRetargetPoseSet::FindRetargetPoseByName(co
 {
 	for (const FResolvedRetargetPose& Pose : RetargetPoses)
 	{
-		if (Pose.Name == CurrentRetargetPoseName)
+		if (Pose.Name == InRetargetPoseName)
 		{
 			return &Pose;
 		}
 	}
-
-	ensureMsgf(false, TEXT("Specified retarget pose was not found. Calling code should check existence of pose."));
+	
 	return nullptr;
 }
 
@@ -230,11 +229,14 @@ void FRetargetSkeleton::UpdateGlobalTransformsBelowBone(
 	
 	for (int32 BoneIndex=StartBoneIndex+1; BoneIndex<OutGlobalPose.Num(); ++BoneIndex)
 	{
-		UpdateGlobalTransformOfSingleBone(BoneIndex,InLocalPose,OutGlobalPose);
+		UpdateGlobalTransformOfSingleBone(BoneIndex, InLocalPose[BoneIndex],OutGlobalPose);
 	}
 }
 
-void FRetargetSkeleton::SetGlobalTransformAndUpdateChildren(const int32 InBoneToSetIndex, const FTransform& InNewTransform, TArray<FTransform>& InOutGlobalPose) const
+void FRetargetSkeleton::SetGlobalTransformAndUpdateChildren(
+	const int32 InBoneToSetIndex,
+	const FTransform& InNewTransform,
+	TArray<FTransform>& InOutGlobalPose) const
 {
 	check(BoneNames.Num() == InOutGlobalPose.Num());
 	check(InOutGlobalPose.IsValidIndex(InBoneToSetIndex));
@@ -268,17 +270,17 @@ void FRetargetSkeleton::UpdateLocalTransformsBelowBone(
 
 void FRetargetSkeleton::UpdateGlobalTransformOfSingleBone(
 	const int32 BoneIndex,
-	const TArray<FTransform>& InLocalPose,
+	const FTransform& InLocalTransform,
 	TArray<FTransform>& OutGlobalPose) const
 {
 	const int32 ParentIndex = ParentIndices[BoneIndex];
 	if (ParentIndex == INDEX_NONE)
 	{
 		// root always in global space already, no conversion required
-		OutGlobalPose[BoneIndex] = InLocalPose[BoneIndex];
+		OutGlobalPose[BoneIndex] = InLocalTransform;
 		return; 
 	}
-	const FTransform& ChildLocalTransform = InLocalPose[BoneIndex];
+	const FTransform& ChildLocalTransform = InLocalTransform;
 	const FTransform& ParentGlobalTransform = OutGlobalPose[ParentIndex];
 	OutGlobalPose[BoneIndex] = ChildLocalTransform * ParentGlobalTransform;
 }
@@ -300,6 +302,22 @@ void FRetargetSkeleton::UpdateLocalTransformOfSingleBone(
 	OutLocalPose[BoneIndex] = ChildGlobalTransform.GetRelativeTransform(ParentGlobalTransform);
 }
 
+FTransform FRetargetSkeleton::GetLocalTransformOfSingleBone(
+	const int32 BoneIndex,
+	const TArray<FTransform>& InGlobalPose) const
+{
+	const FTransform& GlobalTransform = InGlobalPose[BoneIndex];
+	const int32 ParentIndex = ParentIndices[BoneIndex];
+	if (ParentIndex == INDEX_NONE)
+	{
+		// root bone, so local transform is same as global
+		return GlobalTransform;
+	}
+	
+	const FTransform& ParentGlobalTransform = InGlobalPose[ParentIndex];
+	return GlobalTransform.GetRelativeTransform(ParentGlobalTransform);
+}
+
 int32 FRetargetSkeleton::GetCachedEndOfBranchIndex(const int32 InBoneIndex) const
 {
 	if (!CachedEndOfBranchIndices.IsValidIndex(InBoneIndex))
@@ -436,22 +454,6 @@ void FTargetSkeleton::Reset()
 	OutputGlobalPose.Reset();
 }
 
-void FTargetSkeleton::UpdateGlobalTransformsWithMask(
-	const TArray<bool>& InBoneMask,
-	const TArray<FTransform>& InLocalPose,
-	TArray<FTransform>& InOutGlobalPose) const
-{
-	check(InBoneMask.Num() == InOutGlobalPose.Num());
-	
-	for (int32 BoneIndex=0; BoneIndex<InOutGlobalPose.Num(); ++BoneIndex)
-	{
-		if (!InBoneMask[BoneIndex])
-		{
-			UpdateGlobalTransformOfSingleBone(BoneIndex, InLocalPose, InOutGlobalPose);
-		}
-	}
-}
-
 void FTargetSkeleton::SetRetargetedBones(const TSet<int32>& InRetargetedBones)
 {
 	IsBoneRetargeted.Init(false, BoneNames.Num());
@@ -1780,64 +1782,66 @@ void FIKRetargetProcessor::ApplyProfileAtRuntime(const FRetargetProfile& Profile
 }
 
 void FIKRetargetProcessor::UpdateRetargetPoseAtRuntime(
-	const FName NewRetargetPoseName,
+	const FName RetargetPoseToUseName,
 	ERetargetSourceOrTarget SourceOrTarget)
 {
-	const FIKRetargetPose* NewRetargetPose = RetargeterAsset->GetRetargetPoseByName(SourceOrTarget, NewRetargetPoseName);
-	if (!NewRetargetPose)
-	{
-		return; // retarget pose not found
-	}
-	
-	auto AddOrUpdateNewRetargetPose = [this, NewRetargetPoseName, NewRetargetPose, SourceOrTarget]()
+	auto LoadRetargetPoseFromAsset = [this, RetargetPoseToUseName, SourceOrTarget](const FIKRetargetPose* PoseToLoad)
 	{
+		// load the pose and resolve it onto the current skeleton
 		const FName PelvisBoneName = GetPelvisBone(SourceOrTarget, ERetargetOpsToSearch::ProcessorOps);
 		const double ScaleFactor = SourceOrTarget == ERetargetSourceOrTarget::Source ? GetSourceScaleFactor() : 1.0;
 		const bool bIsSource = SourceOrTarget == ERetargetSourceOrTarget::Source;
 		FRetargetSkeleton& RetargetSkeleton = bIsSource ? SourceSkeleton : TargetSkeleton;
-		return RetargetSkeleton.RetargetPoses.AddOrUpdateRetargetPose(
+		FResolvedRetargetPose& ResolvedPose = RetargetSkeleton.RetargetPoses.AddOrUpdateRetargetPose(
 			RetargetSkeleton,
-			NewRetargetPoseName,
-			NewRetargetPose,
+			RetargetPoseToUseName,
+			PoseToLoad,
 			PelvisBoneName,
 			ScaleFactor);
+
+		// set this as the current pose to use
+		RetargetSkeleton.RetargetPoses.CurrentRetargetPoseName = RetargetPoseToUseName;
+
+		// re-load the updated pose into the bone chains
+		AllBoneChains.UpdatePoseFromSkeleton(RetargetSkeleton, SourceOrTarget);
+	
+		// re-initialize the op stack (ops sometimes cache retarget poses)
+		InitializeRetargetOps();
+		
+		return &ResolvedPose;
 	};
 
+	// verify that retarget pose exists in the retarget asset
+	const FIKRetargetPose* UnresolvedPoseFromAsset = RetargeterAsset->GetRetargetPoseByName(SourceOrTarget, RetargetPoseToUseName);
+	if (!UnresolvedPoseFromAsset)
+	{
+		return; // retarget pose not found
+	}
+
+	// load pose out of asset if it hasn't been loaded yet
 	const bool bIsSource = SourceOrTarget == ERetargetSourceOrTarget::Source;
 	FRetargetSkeleton& RetargetSkeleton = bIsSource ? SourceSkeleton : TargetSkeleton;
-	const FResolvedRetargetPose* NewResolvedPose = RetargetSkeleton.RetargetPoses.FindRetargetPoseByName(NewRetargetPoseName);
-	if (!NewResolvedPose)
+	const FResolvedRetargetPose* ResolvedPoseToUse = RetargetSkeleton.RetargetPoses.FindRetargetPoseByName(RetargetPoseToUseName);
+	if (!ResolvedPoseToUse)
 	{
 		// trying to switch to a retarget pose that hasn't been loaded yet, so load it
-		AddOrUpdateNewRetargetPose();
-		NewResolvedPose = RetargetSkeleton.RetargetPoses.FindRetargetPoseByName(NewRetargetPoseName);
-		if (!NewResolvedPose)
-		{
-			// could not load the new pose
-			return;
-		}
+		ResolvedPoseToUse = LoadRetargetPoseFromAsset(UnresolvedPoseFromAsset);
 	}
-
+	
 	// verify if pose needs rebuilding
 	const FName CurrentRetargetPoseName = RetargetSkeleton.RetargetPoses.CurrentRetargetPoseName;
-	const FResolvedRetargetPose* CurrentResolvedPose = RetargetSkeleton.RetargetPoses.FindRetargetPoseByName(CurrentRetargetPoseName);
-	const bool bSameRetargetPose = NewResolvedPose->Name == CurrentResolvedPose->Name;
-	const bool bSameVersion = CurrentResolvedPose->Version == NewRetargetPose->GetVersion();
+	const FResolvedRetargetPose& CurrentResolvedPose = RetargetSkeleton.RetargetPoses.FindOrAddRetargetPose(CurrentRetargetPoseName);
+	const bool bSameRetargetPose = CurrentResolvedPose.Name == ResolvedPoseToUse->Name;
+	const bool bSameVersion = CurrentResolvedPose.Version == UnresolvedPoseFromAsset->GetVersion();
 	const double SourceScaleFactor = GetSourceScaleFactor();
-	const bool bSameGlobalScale = bIsSource ? FMath::IsNearlyEqual(CurrentResolvedPose->Scale, SourceScaleFactor) : true;
+	const bool bSameGlobalScale = bIsSource ? FMath::IsNearlyEqual(CurrentResolvedPose.Scale, SourceScaleFactor) : true;
 	if (bSameRetargetPose && bSameVersion && bSameGlobalScale)
 	{
 		return; // retarget pose has not changed since it was initialized
 	}
 	
-	// re-generate the retarget pose
-	AddOrUpdateNewRetargetPose();
-
-	// re-load the updated pose into the bone chains
-	AllBoneChains.UpdatePoseFromSkeleton(RetargetSkeleton, SourceOrTarget);
-	
-	// re-initialize the op stack (ops sometimes cache retarget poses)
-	InitializeRetargetOps();
+	// reload/update the retarget pose
+	LoadRetargetPoseFromAsset(UnresolvedPoseFromAsset);
 }
 
 FName FIKRetargetProcessor::GetPelvisBone(ERetargetSourceOrTarget SourceOrTarget, ERetargetOpsToSearch InOpsToSearch) const
diff --git a/Engine/Plugins/Animation/IKRig/Source/IKRig/Private/Retargeter/IKRetargeter.cpp b/Engine/Plugins/Animation/IKRig/Source/IKRig/Private/Retargeter/IKRetargeter.cpp
--- a/Engine/Plugins/Animation/IKRig/Source/IKRig/Private/Retargeter/IKRetargeter.cpp
+++ b/Engine/Plugins/Animation/IKRig/Source/IKRig/Private/Retargeter/IKRetargeter.cpp
@@ -604,6 +604,12 @@ void UIKRetargeter::CleanOpStack()
 		}
 		return INDEX_NONE;
 	};
+
+	// remove null ops (could happen if op is in plugin that is not loaded)
+	RetargetOps.RemoveAll([](const FInstancedStruct& InOp)
+	{
+		return !InOp.IsValid();
+	});
 	
 	// enforce unique non-None names on all ops
 	for (int32 OpIndex=0; OpIndex<RetargetOps.Num(); ++OpIndex)
diff --git a/Engine/Plugins/Animation/IKRig/Source/IKRig/Private/Retargeter/RetargetOps/FKChainsOp.cpp b/Engine/Plugins/Animation/IKRig/Source/IKRig/Private/Retargeter/RetargetOps/FKChainsOp.cpp
--- a/Engine/Plugins/Animation/IKRig/Source/IKRig/Private/Retargeter/RetargetOps/FKChainsOp.cpp
+++ b/Engine/Plugins/Animation/IKRig/Source/IKRig/Private/Retargeter/RetargetOps/FKChainsOp.cpp
@@ -506,9 +506,23 @@ void FChainDecoderFK::UpdateIntermediateParents(
 	const FTargetSkeleton& TargetSkeleton,
 	TArray<FTransform>& InOutGlobalPose)
 {
+	for (int32 IntermediateIndex=0; IntermediateIndex<IntermediateParentIndices.Num(); ++IntermediateIndex)
+	{
+		const int32 ParentIndex = IntermediateParentIndices[IntermediateIndex];
+		const FTransform& ParentLocalTransform = IntermediateLocalTransforms[IntermediateIndex];
+		TargetSkeleton.UpdateGlobalTransformOfSingleBone(ParentIndex, ParentLocalTransform, InOutGlobalPose);
+	}
+}
+
+void FChainDecoderFK::UpdateIntermediateLocalTransforms(
+	const FTargetSkeleton& TargetSkeleton,
+	const TArray<FTransform>& InOutGlobalPose)
+{
+	IntermediateLocalTransforms.Reset(IntermediateParentIndices.Num());
 	for (const int32& ParentIndex : IntermediateParentIndices)
 	{
-		TargetSkeleton.UpdateGlobalTransformOfSingleBone(ParentIndex, TargetSkeleton.InputLocalPose, InOutGlobalPose);
+		const FTransform LocalTransform = TargetSkeleton.GetLocalTransformOfSingleBone(ParentIndex, InOutGlobalPose);
+		IntermediateLocalTransforms.Add(LocalTransform);
 	}
 }
 
@@ -644,6 +658,36 @@ bool FIKRetargetFKChainsOp::Initialize(
 	};
 	ChainPairsFK.Sort(ChainsSorter);
 
+	// gather all children bones that need updating after FK chains are retargeted
+	auto IsBoneInAnyChain = [this](const int32 InBoneIndex)
+	{
+		for (const FChainPairFK& ChainPair : ChainPairsFK)
+		{
+			if (ChainPair.TargetBoneChain->BoneIndices.Contains(InBoneIndex))
+			{
+				return true;
+			}
+		}
+		
+		return false;
+	};
+	NonRetargetedChildrenToUpdate.Reset();
+	for (const FChainPairFK& ChainPair : ChainPairsFK)
+	{
+		for (int32 ChainBone : ChainPair.TargetBoneChain->BoneIndices)
+		{
+			TArray<int32> AllChildren;
+			InTargetSkeleton.GetChildrenIndicesRecursive(ChainBone, AllChildren);
+			for (int32 ChildBoneIndex : AllChildren)
+			{
+				if (!IsBoneInAnyChain(ChildBoneIndex))
+				{
+					NonRetargetedChildrenToUpdate.AddUnique(ChildBoneIndex);
+				}
+			}
+		}
+	}
+
 	// consider initialized if at least 1 pair of bone chains were initialized
 	bIsInitialized = !ChainPairsFK.IsEmpty();
 	return bIsInitialized;
@@ -655,11 +699,25 @@ void FIKRetargetFKChainsOp::Run(
 	const TArray<FTransform>& InSourceGlobalPose,
 	TArray<FTransform>& OutTargetGlobalPose)
 {
-	const FIKRetargetPelvisMotionOp* PelvisMotionOp = Processor.GetFirstRetargetOpOfType<FIKRetargetPelvisMotionOp>();
-	const FRetargetSkeleton& SourceSkeleton = Processor.GetSkeleton(ERetargetSourceOrTarget::Source);
 	FTargetSkeleton& TargetSkeleton = Processor.GetTargetSkeleton();
+
+	// update the local transforms of intermediate joints (in case prior op modified them)
+	for (FChainPairFK& ChainPair : ChainPairsFK)
+	{
+		ChainPair.FKDecoder.UpdateIntermediateLocalTransforms(TargetSkeleton, OutTargetGlobalPose);
+	}
+
+	// update local transforms of all the non-retargeted children to update (in case prior op modified them)
+	ChildrenToUpdateLocalTransforms.Reset();
+	for (int32 NonRetargetedChildBoneIndex : NonRetargetedChildrenToUpdate)
+	{
+		const FTransform ChildLocalTransform = TargetSkeleton.GetLocalTransformOfSingleBone(NonRetargetedChildBoneIndex, OutTargetGlobalPose);
+		ChildrenToUpdateLocalTransforms.Add(ChildLocalTransform);
+	}
 	
 	// spin through chains and encode/decode them all using the input pose
+	const FIKRetargetPelvisMotionOp* PelvisMotionOp = Processor.GetFirstRetargetOpOfType<FIKRetargetPelvisMotionOp>();
+	const FRetargetSkeleton& SourceSkeleton = Processor.GetSkeleton(ERetargetSourceOrTarget::Source);
 	for (FChainPairFK& ChainPair : ChainPairsFK)
 	{
 		ChainPair.FKEncoder.EncodePose(
@@ -676,10 +734,13 @@ void FIKRetargetFKChainsOp::Run(
 			OutTargetGlobalPose);
 	}
 
-	TargetSkeleton.UpdateGlobalTransformsWithMask(
-		TargetSkeleton.GetRetargetedBonesMask(),
-		TargetSkeleton.InputLocalPose,
-		TargetSkeleton.OutputGlobalPose);
+	// update non-retargeted children
+	for (int32 ChildIndex=0; ChildIndex<NonRetargetedChildrenToUpdate.Num(); ++ChildIndex)
+	{
+		int32 ChildBoneIndex = NonRetargetedChildrenToUpdate[ChildIndex];
+		const FTransform& ChildLocalTransform = ChildrenToUpdateLocalTransforms[ChildIndex];
+		TargetSkeleton.UpdateGlobalTransformOfSingleBone(ChildBoneIndex, ChildLocalTransform, TargetSkeleton.OutputGlobalPose);
+	}
 
 #if WITH_EDITOR
 	SaveDebugData(OutTargetGlobalPose);
diff --git a/Engine/Plugins/Animation/IKRig/Source/IKRig/Private/Retargeter/RetargetOps/PelvisMotionOp.cpp b/Engine/Plugins/Animation/IKRig/Source/IKRig/Private/Retargeter/RetargetOps/PelvisMotionOp.cpp
--- a/Engine/Plugins/Animation/IKRig/Source/IKRig/Private/Retargeter/RetargetOps/PelvisMotionOp.cpp
+++ b/Engine/Plugins/Animation/IKRig/Source/IKRig/Private/Retargeter/RetargetOps/PelvisMotionOp.cpp
@@ -95,12 +95,14 @@ void FIKRetargetPelvisMotionOp::Run(
 	const TArray<FTransform>& InSourceGlobalPose,
 	TArray<FTransform>& OutTargetGlobalPose)
 {
+	FTransform NewPelvisGlobalTransform;
+	
 	EncodePose(InSourceGlobalPose);
-	DecodePose(OutTargetGlobalPose);
+	DecodePose(NewPelvisGlobalTransform);
 	
 	// update global transforms below root
 	FTargetSkeleton& TargetSkeleton = Processor.GetTargetSkeleton();
-	TargetSkeleton.UpdateGlobalTransformsBelowBone(Target.BoneIndex, TargetSkeleton.InputLocalPose, TargetSkeleton.OutputGlobalPose);
+	TargetSkeleton.SetGlobalTransformAndUpdateChildren(Target.BoneIndex, NewPelvisGlobalTransform, TargetSkeleton.OutputGlobalPose);
 
 #if WITH_EDITOR
 	FScopeLock ScopeLock(&DebugDataMutex);
@@ -248,7 +250,7 @@ void FIKRetargetPelvisMotionOp::EncodePose(const TArray<FTransform>& SourceGloba
 	Source.CurrentRotation = SourceTransform.GetRotation();	
 }
 
-void FIKRetargetPelvisMotionOp::DecodePose(TArray<FTransform>& OutTargetGlobalPose)
+void FIKRetargetPelvisMotionOp::DecodePose(FTransform& OutPelvisGlobalPose)
 {
 	// retarget position
 	FVector Position;
@@ -296,9 +298,8 @@ void FIKRetargetPelvisMotionOp::DecodePose(TArray<FTransform>& OutTargetGlobalPo
 	}
 
 	// apply to target
-	FTransform& TargetRootTransform = OutTargetGlobalPose[Target.BoneIndex];
-	TargetRootTransform.SetTranslation(Position);
-	TargetRootTransform.SetRotation(Rotation);
+	OutPelvisGlobalPose.SetTranslation(Position);
+	OutPelvisGlobalPose.SetRotation(Rotation);
 }
 
 #if WITH_EDITOR
diff --git a/Engine/Plugins/Animation/IKRig/Source/IKRig/Private/Retargeter/RetargetOps/RetargetPoseOp.cpp b/Engine/Plugins/Animation/IKRig/Source/IKRig/Private/Retargeter/RetargetOps/RetargetPoseOp.cpp
--- a/Engine/Plugins/Animation/IKRig/Source/IKRig/Private/Retargeter/RetargetOps/RetargetPoseOp.cpp
+++ b/Engine/Plugins/Animation/IKRig/Source/IKRig/Private/Retargeter/RetargetOps/RetargetPoseOp.cpp
@@ -2,71 +2,136 @@
 
 #include "Retargeter/RetargetOps/RetargetPoseOp.h"
 
-const UClass* FIKRetargetPoseOpSettings::GetControllerType() const
+#include "Retargeter/IKRetargetProcessor.h"
+
+const UClass* FIKRetargetAdditivePoseOpSettings::GetControllerType() const
 {
-	return UIKRetargetPoseController::StaticClass();
+	return UIKRetargetAdditivePoseController::StaticClass();
 }
 
-void FIKRetargetPoseOpSettings::CopySettingsAtRuntime(const FIKRetargetOpSettingsBase* InSettingsToCopyFrom)
+void FIKRetargetAdditivePoseOpSettings::CopySettingsAtRuntime(const FIKRetargetOpSettingsBase* InSettingsToCopyFrom)
 {
 	// copy all properties
 	static TArray<FName> PropertiesToIgnore = {};
 	FIKRetargetOpBase::CopyStructProperties(
-		FIKRetargetPoseOpSettings::StaticStruct(),
+		FIKRetargetAdditivePoseOpSettings::StaticStruct(),
 		InSettingsToCopyFrom,
 		this,
 		PropertiesToIgnore);
 }
 
-bool FIKRetargetPoseOp::Initialize(
+bool FIKRetargetAdditivePoseOp::Initialize(
 	const FIKRetargetProcessor& InProcessor,
 	const FRetargetSkeleton& InSourceSkeleton,
 	const FTargetSkeleton& InTargetSkeleton,
 	const FIKRetargetOpBase* InParentOp,
 	FIKRigLogger& Log)
 {
-	// TODO cache the retarget pose in the processor 
 	bIsInitialized = true;
+	PelvisBoneName = InProcessor.GetPelvisBone(ERetargetSourceOrTarget::Target, ERetargetOpsToSearch::ProcessorOps);
 	return true;
 }
 
-void FIKRetargetPoseOp::Run(
+void FIKRetargetAdditivePoseOp::Run(
 	FIKRetargetProcessor& InProcessor,
 	const double InDeltaTime,
 	const TArray<FTransform>& InSourceGlobalPose,
 	TArray<FTransform>& OutTargetGlobalPose)
 {
-	// TODO update/switch to current retarget pose
-	// TODO apply pose to output pose
+	ApplyAdditivePose(InProcessor, OutTargetGlobalPose);
 }
 
-FIKRetargetOpSettingsBase* FIKRetargetPoseOp::GetSettings()
+FIKRetargetOpSettingsBase* FIKRetargetAdditivePoseOp::GetSettings()
 {
 	return &Settings;
 }
 
-void FIKRetargetPoseOp::SetSettings(const FIKRetargetOpSettingsBase* InSettings)
+void FIKRetargetAdditivePoseOp::SetSettings(const FIKRetargetOpSettingsBase* InSettings)
 {
 	Settings.CopySettingsAtRuntime(InSettings);
 }
 
-const UScriptStruct* FIKRetargetPoseOp::GetSettingsType() const
+const UScriptStruct* FIKRetargetAdditivePoseOp::GetSettingsType() const
+{
+	return FIKRetargetAdditivePoseOpSettings::StaticStruct();
+}
+
+const UScriptStruct* FIKRetargetAdditivePoseOp::GetType() const
 {
-	return FIKRetargetPoseOpSettings::StaticStruct();
+	return FIKRetargetAdditivePoseOp::StaticStruct();
 }
 
-const UScriptStruct* FIKRetargetPoseOp::GetType() const
+void FIKRetargetAdditivePoseOp::ApplyAdditivePose(
+	FIKRetargetProcessor& InProcessor,
+	TArray<FTransform>& OutTargetGlobalPose)
 {
-	return FIKRetargetPoseOp::StaticStruct();
+	// verify that retarget pose exists in the retarget asset
+	const UIKRetargeter* RetargetAsset = InProcessor.GetRetargetAsset();
+	const FIKRetargetPose* RetargetPose = RetargetAsset->GetRetargetPoseByName(ERetargetSourceOrTarget::Target, Settings.PoseToApply);
+	if (!RetargetPose)
+	{
+		return; // retarget pose not found
+	}
+
+	const FRetargetSkeleton& TargetSkeleton = InProcessor.GetSkeleton(ERetargetSourceOrTarget::Target);
+
+	// apply pelvis translation offset
+	const int32 PelvisBoneIndex = TargetSkeleton.FindBoneIndexByName(PelvisBoneName);
+	if (PelvisBoneIndex != INDEX_NONE)
+	{
+		FTransform PelvisTransform = OutTargetGlobalPose[PelvisBoneIndex];
+		PelvisTransform.AddToTranslation(RetargetPose->GetRootTranslationDelta());
+		TargetSkeleton.SetGlobalTransformAndUpdateChildren(PelvisBoneIndex, PelvisTransform, OutTargetGlobalPose);
+	}
+	
+	// NOTE: we could convert the entire global pose to a local pose, apply the offsets, and then convert it back to global space
+	// BUT, for the majority of use cases, retarget poses only affect a small set of bones and so we are doing sparse updates.
+	// If the retarget pose affects many bones, this could end up being slower due to SetGlobalTransformAndUpdateChildren()
+	// We chose to optimize for the common case; if this ever shows up in a profile we could modify this to use a batch conversion if the
+	// retarget pose modifies a high percentage of the bones in the skeleton.
+
+	// apply retarget pose offsets (retarget pose is stored as offset relative to reference pose)
+	for (const TTuple<FName, FQuat>& BoneDelta : RetargetPose->GetAllDeltaRotations())
+	{
+		const int32 BoneIndex = TargetSkeleton.FindBoneIndexByName(BoneDelta.Key);
+		if (BoneIndex == INDEX_NONE)
+		{
+			// this can happen if a retarget pose recorded a bone offset for a bone that is not present in the
+			// target skeleton; ie, the retarget pose was generated from a different Skeletal Mesh with extra bones
+			continue;
+		}
+
+		// get local transform of bone
+		FTransform LocalTransform = TargetSkeleton.GetLocalTransformOfSingleBone(BoneIndex, OutTargetGlobalPose);
+
+		// apply local offset
+		FQuat DeltaRotation = BoneDelta.Value;
+		if (!FMath::IsNearlyEqual(Settings.Alpha, 1.0))
+		{
+			DeltaRotation = FQuat::FastLerp(FQuat::Identity, DeltaRotation, Settings.Alpha);
+			DeltaRotation.Normalize();
+		}
+		LocalTransform.SetRotation(LocalTransform.GetRotation() * DeltaRotation);
+
+		// convert back to global and update the output pose
+		FTransform ParentGlobalTransform = FTransform::Identity;
+		const int32 ParentIndex = TargetSkeleton.ParentIndices[BoneIndex];
+		if (ParentIndex != INDEX_NONE)
+		{
+			ParentGlobalTransform = OutTargetGlobalPose[ParentIndex];
+		}
+		FTransform GlobalTransform = LocalTransform * ParentGlobalTransform;
+		TargetSkeleton.SetGlobalTransformAndUpdateChildren(BoneIndex, GlobalTransform, OutTargetGlobalPose);
+	}
 }
 
-FIKRetargetPoseOpSettings UIKRetargetPoseController::GetSettings()
+FIKRetargetAdditivePoseOpSettings UIKRetargetAdditivePoseController::GetSettings()
 {
-	return *reinterpret_cast<FIKRetargetPoseOpSettings*>(OpSettingsToControl);
+	return *reinterpret_cast<FIKRetargetAdditivePoseOpSettings*>(OpSettingsToControl);
 }
 
-void UIKRetargetPoseController::SetSettings(FIKRetargetPoseOpSettings InSettings)
+void UIKRetargetAdditivePoseController::SetSettings(FIKRetargetAdditivePoseOpSettings InSettings)
 {
-	reinterpret_cast<FIKRetargetPoseOpSettings*>(OpSettingsToControl)->CopySettingsAtRuntime(&InSettings);
+	reinterpret_cast<FIKRetargetAdditivePoseOpSettings*>(OpSettingsToControl)->CopySettingsAtRuntime(&InSettings);
 }
 
diff --git a/Engine/Plugins/Animation/IKRig/Source/IKRig/Public/Retargeter/IKRetargetProcessor.h b/Engine/Plugins/Animation/IKRig/Source/IKRig/Public/Retargeter/IKRetargetProcessor.h
--- a/Engine/Plugins/Animation/IKRig/Source/IKRig/Public/Retargeter/IKRetargetProcessor.h
+++ b/Engine/Plugins/Animation/IKRig/Source/IKRig/Public/Retargeter/IKRetargetProcessor.h
@@ -97,12 +97,6 @@ struct FRetargetSkeleton
 	
 	UE_API void Reset();
 
-	UE_API void GenerateRetargetPose(
-		const FName InRetargetPoseName,
-		const FIKRetargetPose* InRetargetPose,
-		const FName PelvisBoneName,
-		const double ScaleFactor);
-
 	UE_API int32 FindBoneIndexByName(const FName InName) const;
 
 	UE_API int32 GetParentIndex(const int32 BoneIndex) const;
@@ -124,7 +118,7 @@ struct FRetargetSkeleton
 	
 	UE_API void UpdateGlobalTransformOfSingleBone(
 		const int32 BoneIndex,
-		const TArray<FTransform>& InLocalPose,
+		const FTransform& InLocalTransform,
 		TArray<FTransform>& OutGlobalPose) const;
 	
 	UE_API void UpdateLocalTransformOfSingleBone(
@@ -132,6 +126,10 @@ struct FRetargetSkeleton
 		TArray<FTransform>& OutLocalPose,
 		const TArray<FTransform>& InGlobalPose) const;
 
+	UE_API FTransform GetLocalTransformOfSingleBone(
+		const int32 BoneIndex,
+		const TArray<FTransform>& InGlobalPose) const;
+
 	UE_API int32 GetCachedEndOfBranchIndex(const int32 InBoneIndex) const;
 
 	UE_API void GetChildrenIndices(const int32 BoneIndex, TArray<int32>& OutChildren) const;
@@ -164,11 +162,6 @@ struct FTargetSkeleton : public FRetargetSkeleton
 
 	void Reset();
 
-	void UpdateGlobalTransformsWithMask(
-		const TArray<bool>& InBoneMask,
-		const TArray<FTransform>& InLocalPose,
-		TArray<FTransform>& InOutGlobalPose) const;
-
 	void SetRetargetedBones(const TSet<int32>& InRetargetedBones);
 
 	bool GetIsBoneRetargeted(int32 InBoneIndex) const;
@@ -417,7 +410,7 @@ public:
 
 	/** Does a partial reinitialization (at runtime) whenever the retarget pose is swapped to a different or if the
 	 * pose has been modified. Does nothing if the pose has not changed. */
-	UE_API void UpdateRetargetPoseAtRuntime(const FName NewRetargetPoseName, ERetargetSourceOrTarget SourceOrTarget);
+	UE_API void UpdateRetargetPoseAtRuntime(const FName RetargetPoseToUseName, ERetargetSourceOrTarget SourceOrTarget);
 
 	/** Get read-only access to either source or target skeleton. */
 	UE_API const FRetargetSkeleton& GetSkeleton(ERetargetSourceOrTarget SourceOrTarget) const;
diff --git a/Engine/Plugins/Animation/IKRig/Source/IKRig/Public/Retargeter/RetargetOps/FKChainsOp.h b/Engine/Plugins/Animation/IKRig/Source/IKRig/Public/Retargeter/RetargetOps/FKChainsOp.h
--- a/Engine/Plugins/Animation/IKRig/Source/IKRig/Public/Retargeter/RetargetOps/FKChainsOp.h
+++ b/Engine/Plugins/Animation/IKRig/Source/IKRig/Public/Retargeter/RetargetOps/FKChainsOp.h
@@ -38,6 +38,10 @@ struct FChainDecoderFK
 		const int32 InRetargetRootBoneIndex,
 		const int32 InChainRootBoneIndex,
 		const FTargetSkeleton& InTargetSkeleton);
+
+	void UpdateIntermediateLocalTransforms(
+		const FTargetSkeleton& TargetSkeleton,
+		const TArray<FTransform> &InOutGlobalPose);
 	
 	void DecodePose(
 		const FIKRetargetPelvisMotionOp* PelvisMotionOp,
@@ -71,6 +75,7 @@ private:
 	
 	TArray<FTransform> CurrentGlobalTransforms;
 	TArray<int32> IntermediateParentIndices;
+	TArray<FTransform> IntermediateLocalTransforms;
 	const FResolvedBoneChain* BoneChain;
 };
 
@@ -277,6 +282,10 @@ private:
 	
 	/** The Source/Target pairs of Bone Chains retargeted using the FK algorithm */
 	TArray<FChainPairFK> ChainPairsFK;
+
+	/** Children of the FK chains that need updated after the chains are retargeted */
+	TArray<int32> NonRetargetedChildrenToUpdate;
+	TArray<FTransform> ChildrenToUpdateLocalTransforms;
 };
 
 /* The blueprint/python API for editing a FK Chains Op */
diff --git a/Engine/Plugins/Animation/IKRig/Source/IKRig/Public/Retargeter/RetargetOps/PelvisMotionOp.h b/Engine/Plugins/Animation/IKRig/Source/IKRig/Public/Retargeter/RetargetOps/PelvisMotionOp.h
--- a/Engine/Plugins/Animation/IKRig/Source/IKRig/Public/Retargeter/RetargetOps/PelvisMotionOp.h
+++ b/Engine/Plugins/Animation/IKRig/Source/IKRig/Public/Retargeter/RetargetOps/PelvisMotionOp.h
@@ -176,7 +176,7 @@ private:
 	
 	void EncodePose(const TArray<FTransform> &SourceGlobalPose);
 	
-	void DecodePose(TArray<FTransform> &OutTargetGlobalPose);
+	void DecodePose(FTransform& OutPelvisGlobalPose);
 
 	// transient work data
 	FPelvisSource Source;
diff --git a/Engine/Plugins/Animation/IKRig/Source/IKRig/Public/Retargeter/RetargetOps/RetargetPoseOp.h b/Engine/Plugins/Animation/IKRig/Source/IKRig/Public/Retargeter/RetargetOps/RetargetPoseOp.h
--- a/Engine/Plugins/Animation/IKRig/Source/IKRig/Public/Retargeter/RetargetOps/RetargetPoseOp.h
+++ b/Engine/Plugins/Animation/IKRig/Source/IKRig/Public/Retargeter/RetargetOps/RetargetPoseOp.h
@@ -10,34 +10,26 @@
 
 #define LOCTEXT_NAMESPACE "RetargetPoseOp"
 
-UENUM()
-enum class ERetargetPoseMode : uint8
-{
-	RetargetPose,
-	AdditivePose,
-};
-
-USTRUCT(BlueprintType, meta = (DisplayName = "Retarget Pose Op Settings"))
-struct FIKRetargetPoseOpSettings : public FIKRetargetOpSettingsBase
+USTRUCT(BlueprintType, meta = (DisplayName = "Additive Pose Op Settings"))
+struct FIKRetargetAdditivePoseOpSettings : public FIKRetargetOpSettingsBase
 {
 	GENERATED_BODY()
-
-	UPROPERTY(EditAnywhere, BlueprintReadWrite, Category="Pose")
-	ERetargetSourceOrTarget SourceOrTarget = ERetargetSourceOrTarget::Source;
-
-	UPROPERTY(EditAnywhere, BlueprintReadWrite, Category="Pose")
-	ERetargetPoseMode Mode = ERetargetPoseMode::RetargetPose;
 	
+	// a retarget pose that is applied additively to the output pose of the target skeleton
 	UPROPERTY(EditAnywhere, BlueprintReadWrite, Category="Pose")
-	FName CurrentPose;
+	FName PoseToApply;
+
+	// blend the amount of the pose to apply
+	UPROPERTY(EditAnywhere, BlueprintReadWrite, Category="Pose", meta=(ClampMin=0, ClampMax=1))
+	float Alpha = 1.0f;
 	
 	UE_API virtual const UClass* GetControllerType() const override;
 	
 	UE_API virtual void CopySettingsAtRuntime(const FIKRetargetOpSettingsBase* InSettingsToCopyFrom) override;
 };
 
-USTRUCT(BlueprintType, meta = (DisplayName = "Retarget Pose"))
-struct FIKRetargetPoseOp : public FIKRetargetOpBase
+USTRUCT(BlueprintType, meta = (DisplayName = "Additive Pose"))
+struct FIKRetargetAdditivePoseOp : public FIKRetargetOpBase
 {
 	GENERATED_BODY()
 	
@@ -62,13 +54,20 @@ struct FIKRetargetPoseOp : public FIKRetargetOpBase
 	
 	UE_API virtual const UScriptStruct* GetType() const override;
 
+private:
+	
+	void ApplyAdditivePose(FIKRetargetProcessor& InProcessor, TArray<FTransform>& OutTargetGlobalPose);
+
+	// cached in Initialize()
+	FName PelvisBoneName;
+
 	UPROPERTY()
-	FIKRetargetPoseOpSettings Settings;
+	FIKRetargetAdditivePoseOpSettings Settings;
 };
 
 /* The blueprint/python API for editing a Retarget Pose Op */
-UCLASS(MinimalAPI, BlueprintType)
-class UIKRetargetPoseController : public UIKRetargetOpControllerBase
+UCLASS(BlueprintType)
+class UIKRetargetAdditivePoseController : public UIKRetargetOpControllerBase
 {
 	GENERATED_BODY()
 	
@@ -76,12 +75,12 @@ public:
 	/* Get the current op settings as a struct.
 	 * @return FIKRetargetPoseOpSettings struct with the current settings used by the op. */
 	UFUNCTION(BlueprintCallable, Category = Settings)
-	UE_API FIKRetargetPoseOpSettings GetSettings();
+	UE_API FIKRetargetAdditivePoseOpSettings GetSettings();
 
 	/* Set the op settings. Input is a custom struct type for this op.
 	 * @param InSettings a FIKRetargetPoseOpSettings struct containing all the settings to apply to this op */
 	UFUNCTION(BlueprintCallable, Category = Settings)
-	UE_API void SetSettings(FIKRetargetPoseOpSettings InSettings);
+	UE_API void SetSettings(FIKRetargetAdditivePoseOpSettings InSettings);
 };
 
 #undef LOCTEXT_NAMESPACE
diff --git a/Engine/Plugins/Animation/IKRig/Source/IKRigEditor/Private/IKRigEditor.cpp b/Engine/Plugins/Animation/IKRig/Source/IKRigEditor/Private/IKRigEditor.cpp
--- a/Engine/Plugins/Animation/IKRig/Source/IKRigEditor/Private/IKRigEditor.cpp
+++ b/Engine/Plugins/Animation/IKRig/Source/IKRigEditor/Private/IKRigEditor.cpp
@@ -102,8 +102,8 @@ void FIKRigEditor::StartupModule()
 		"IKRetargetAlignPoleVectorOpSettings",
 		FOnGetPropertyTypeCustomizationInstance::CreateStatic(&FPoleVectorOpCustomization::MakeInstance));
 	PropertyEditorModule.RegisterCustomPropertyTypeLayout(
-			"IKRetargetPoseOpSettings",
-			FOnGetPropertyTypeCustomizationInstance::CreateStatic(&FRetargetPoseOpCustomization::MakeInstance));
+			"IKRetargetAdditivePoseOpSettings",
+			FOnGetPropertyTypeCustomizationInstance::CreateStatic(&FAdditivePoseOpCustomization::MakeInstance));
 
 	// register a thumbnail renderer for the assets
 	UThumbnailManager::Get().RegisterCustomRenderer(UIKRigDefinition::StaticClass(), UIKRigThumbnailRenderer::StaticClass());
diff --git a/Engine/Plugins/Animation/IKRig/Source/IKRigEditor/Private/RetargetEditor/IKRetargetOpDetails.cpp b/Engine/Plugins/Animation/IKRig/Source/IKRigEditor/Private/RetargetEditor/IKRetargetOpDetails.cpp
--- a/Engine/Plugins/Animation/IKRig/Source/IKRigEditor/Private/RetargetEditor/IKRetargetOpDetails.cpp
+++ b/Engine/Plugins/Animation/IKRig/Source/IKRigEditor/Private/RetargetEditor/IKRetargetOpDetails.cpp
@@ -541,7 +541,7 @@ void FPoleVectorOpCustomization::CustomizeChildren(
 	];
 }
 
-void FRetargetPoseOpCustomization::CustomizeChildren(
+void FAdditivePoseOpCustomization::CustomizeChildren(
 	TSharedRef<IPropertyHandle> StructPropertyHandle,
 	IDetailChildrenBuilder& ChildBuilder,
 	IPropertyTypeCustomizationUtils& StructCustomizationUtils)
@@ -552,29 +552,15 @@ void FRetargetPoseOpCustomization::CustomizeChildren(
 	}
 
 	// store property handles for callbacks
-	SourceOrTargetProperty = StructPropertyHandle->GetChildHandle(GET_MEMBER_NAME_CHECKED(FIKRetargetPoseOpSettings, SourceOrTarget));
-	CurrentPoseProperty = StructPropertyHandle->GetChildHandle(GET_MEMBER_NAME_CHECKED(FIKRetargetPoseOpSettings, CurrentPose));
+	PoseToApplyProperty = StructPropertyHandle->GetChildHandle(GET_MEMBER_NAME_CHECKED(FIKRetargetAdditivePoseOpSettings, PoseToApply));
 
-	// add the source/target property
-	ChildBuilder.AddProperty(SourceOrTargetProperty.ToSharedRef());
-
-	// add the mode property
-	TSharedPtr<IPropertyHandle> ModeProperty = StructPropertyHandle->GetChildHandle(GET_MEMBER_NAME_CHECKED(FIKRetargetPoseOpSettings, Mode));
-	ChildBuilder.AddProperty(ModeProperty.ToSharedRef());
-
-	// add a callback to update the pose name list when switching source/target
-	SourceOrTargetProperty->SetOnPropertyValueChanged((FSimpleDelegate::CreateLambda([this]()
-	{
-		UpdatePoseNameOptions();
-	})));
-	
 	UpdatePoseNameOptions();
 
 	// add dropdown menu to select retarget pose
-	ChildBuilder.AddCustomRow(LOCTEXT("CurrentPoseLabel", "Current Pose"))
+	ChildBuilder.AddCustomRow(LOCTEXT("CurrentPoseLabel", "Pose To Apply"))
 	.NameContent()
 	[
-		CurrentPoseProperty->CreatePropertyNameWidget()
+		PoseToApplyProperty->CreatePropertyNameWidget()
 	]
 	.ValueContent()
 	.MinDesiredWidth(200)
@@ -590,31 +576,32 @@ void FRetargetPoseOpCustomization::CustomizeChildren(
 		{
 			if (NewSelection.IsValid())
 			{
-				CurrentPoseProperty->SetValue(*NewSelection);
+				PoseToApplyProperty->SetValue(*NewSelection);
 			}
 		})
 		[
 			SNew(STextBlock)
 			.Text_Lambda([this]() {
 				FName Value;
-				CurrentPoseProperty->GetValue(Value);
+				PoseToApplyProperty->GetValue(Value);
 				return FText::FromName(Value);
 			})
 		]
 	];
+
+	// add the alpha property
+	TSharedPtr<IPropertyHandle> AlphaProperty = StructPropertyHandle->GetChildHandle(GET_MEMBER_NAME_CHECKED(FIKRetargetAdditivePoseOpSettings, Alpha));
+	ChildBuilder.AddProperty(AlphaProperty.ToSharedRef());
 }
 
-void FRetargetPoseOpCustomization::UpdatePoseNameOptions()
+void FAdditivePoseOpCustomization::UpdatePoseNameOptions()
 {
 	// determine currently selected pose
 	FName CurrentPoseName;
-	CurrentPoseProperty->GetValue(CurrentPoseName);
+	PoseToApplyProperty->GetValue(CurrentPoseName);
 	
 	// get all the retarget poses
-	int64 EnumValueInt;
-	SourceOrTargetProperty->GetValue(EnumValueInt);
-	ERetargetSourceOrTarget SourceOrTarget = static_cast<ERetargetSourceOrTarget>(EnumValueInt);
-	const TMap<FName, FIKRetargetPose>& RetargetPoses = AssetController->GetRetargetPoses(SourceOrTarget);
+	const TMap<FName, FIKRetargetPose>& RetargetPoses = AssetController->GetRetargetPoses(ERetargetSourceOrTarget::Target);
 	TArray<FName> PoseNames;
 	RetargetPoses.GetKeys(PoseNames);
 
diff --git a/Engine/Plugins/Animation/IKRig/Source/IKRigEditor/Public/RetargetEditor/IKRetargetOpDetails.h b/Engine/Plugins/Animation/IKRig/Source/IKRigEditor/Public/RetargetEditor/IKRetargetOpDetails.h
--- a/Engine/Plugins/Animation/IKRig/Source/IKRigEditor/Public/RetargetEditor/IKRetargetOpDetails.h
+++ b/Engine/Plugins/Animation/IKRig/Source/IKRigEditor/Public/RetargetEditor/IKRetargetOpDetails.h
@@ -225,12 +225,12 @@ public:
 // Retarget Pose Op Customization
 //
 
-class FRetargetPoseOpCustomization : public FIKRetargetOpBaseSettingsCustomization
+class FAdditivePoseOpCustomization : public FIKRetargetOpBaseSettingsCustomization
 {
 public:
 	static TSharedRef<IPropertyTypeCustomization> MakeInstance()
 	{
-		return MakeShared<FRetargetPoseOpCustomization>();
+		return MakeShared<FAdditivePoseOpCustomization>();
 	}
 
 	// IPropertyTypeCustomization
@@ -244,8 +244,7 @@ private:
 	
 	void UpdatePoseNameOptions();
 
-	TSharedPtr<IPropertyHandle> SourceOrTargetProperty;
-	TSharedPtr<IPropertyHandle> CurrentPoseProperty;
+	TSharedPtr<IPropertyHandle> PoseToApplyProperty;
 	
 	TSharedPtr<FName> CurrentPoseOption;
 	TArray<TSharedPtr<FName>> PoseNameOptions;
@@ -257,5 +256,5 @@ class URetargetPoseOpSettingsWrapper : public UIKRigStructWrapperBase
 	GENERATED_BODY()
 public:
 	UPROPERTY(EditAnywhere, Category = "Retarget Pose Settings")
-	FIKRetargetPoseOpSettings Settings;
+	FIKRetargetAdditivePoseOpSettings Settings;
 };
