[IK Rig] Fixes crash loading old IK Rig with missing effectors.

diff --git a/Engine/Plugins/Animation/IKRig/Source/IKRig/Public/Rig/Solvers/IKRigFullBodyIK.h b/Engine/Plugins/Animation/IKRig/Source/IKRig/Public/Rig/Solvers/IKRigFullBodyIK.h
--- a/Engine/Plugins/Animation/IKRig/Source/IKRig/Public/Rig/Solvers/IKRigFullBodyIK.h
+++ b/Engine/Plugins/Animation/IKRig/Source/IKRig/Public/Rig/Solvers/IKRigFullBodyIK.h
@@ -408,6 +408,10 @@ public:
 		// load effectors into solver
 		for (UIKRig_FBIKEffector* OldEffector : Effectors_DEPRECATED)
 		{
+			if (!OldEffector)
+			{
+				continue;
+			}
 			FIKRigFBIKGoalSettings NewEffector;
 			NewEffector.Goal = OldEffector->GoalName;
 			NewEffector.BoneName = OldEffector->BoneName;
@@ -421,6 +425,10 @@ public:
 		// load bone settings into solver
 		for (UIKRig_FBIKBoneSettings* OldBoneSetting : BoneSettings_DEPRECATED)
 		{
+			if (!OldBoneSetting)
+			{
+				continue;
+			}
 			FIKRigFBIKBoneSettings NewBoneSetting;
 			NewBoneSetting.Bone = OldBoneSetting->Bone;
 			NewBoneSetting.RotationStiffness = OldBoneSetting->RotationStiffness;
