[IK Retarget] Fixes scale source not being applied to exported retargeted animations.

diff --git a/Engine/Plugins/Animation/IKRig/Source/IKRigEditor/Private/RetargetEditor/IKRetargetBatchOperation.cpp b/Engine/Plugins/Animation/IKRig/Source/IKRigEditor/Private/RetargetEditor/IKRetargetBatchOperation.cpp
--- a/Engine/Plugins/Animation/IKRig/Source/IKRigEditor/Private/RetargetEditor/IKRetargetBatchOperation.cpp
+++ b/Engine/Plugins/Animation/IKRig/Source/IKRigEditor/Private/RetargetEditor/IKRetargetBatchOperation.cpp
@@ -473,6 +473,9 @@ void UIKRetargetBatchOperation::ConvertAnimation(
 			FRetargetProfile SettingsProfile;
 			SettingsProfile.FillProfileWithAssetSettings(Context.IKRetargetAsset);
 
+			// give processor the pose to apply source scaling (if needed)
+			Processor.ScaleSourcePose(SourceComponentPose);
+
 			// run the retargeter
 			const TArray<FTransform>& TargetComponentPose = Processor.RunRetargeter(SourceComponentPose, SettingsProfile, DeltaTime);
 
