[IK Retarget] Fixes crash when pasting full transform into rotation field of bone details.

diff --git a/Engine/Plugins/Animation/IKRig/Source/IKRigEditor/Private/RetargetEditor/IKRetargetDetails.cpp b/Engine/Plugins/Animation/IKRig/Source/IKRigEditor/Private/RetargetEditor/IKRetargetDetails.cpp
--- a/Engine/Plugins/Animation/IKRig/Source/IKRigEditor/Private/RetargetEditor/IKRetargetDetails.cpp
+++ b/Engine/Plugins/Animation/IKRig/Source/IKRigEditor/Private/RetargetEditor/IKRetargetDetails.cpp
@@ -367,6 +367,7 @@ void UIKRetargetBoneDetails::OnPasteFromClipboard(
 			break;
 		}
 	case ESlateTransformComponent::Rotation:
+	case ESlateTransformComponent::Max:
 		{
 			FRotator Data = Transform.Rotator();
 			const TCHAR* Result = TBaseStructure<FRotator>::Get()->ImportText(*Content, &Data, nullptr, PPF_None, &ErrorPipe, TBaseStructure<FRotator>::Get()->GetName(), true);
@@ -377,8 +378,10 @@ void UIKRetargetBoneDetails::OnPasteFromClipboard(
 			}
 			break;
 		}
+	
 	default:
-		checkNoEntry();
+		// scale not supported / ignored if pasted
+		break;
 	}
 }
 
