[IK Retarget] Fixes regression causing Pelvis to get mangled if Root bone is retargeted.

diff --git a/Engine/Plugins/Animation/IKRig/Source/IKRig/Private/Retargeter/RetargetOps/FKChainsOp.cpp b/Engine/Plugins/Animation/IKRig/Source/IKRig/Private/Retargeter/RetargetOps/FKChainsOp.cpp
--- a/Engine/Plugins/Animation/IKRig/Source/IKRig/Private/Retargeter/RetargetOps/FKChainsOp.cpp
+++ b/Engine/Plugins/Animation/IKRig/Source/IKRig/Private/Retargeter/RetargetOps/FKChainsOp.cpp
@@ -659,8 +659,15 @@ bool FIKRetargetFKChainsOp::Initialize(
 	ChainPairsFK.Sort(ChainsSorter);
 
 	// gather all children bones that need updating after FK chains are retargeted
-	auto IsBoneInAnyChain = [this](const int32 InBoneIndex)
+	const FName PelvisBone = InProcessor.GetPelvisBone(ERetargetSourceOrTarget::Target, ERetargetOpsToSearch::ProcessorOps);
+	const int32 PelvisBoneIndex = InTargetSkeleton.FindBoneIndexByName(PelvisBone);
+	auto IsBoneRetargeted = [this, PelvisBoneIndex](const int32 InBoneIndex)
 	{
+		if (PelvisBoneIndex == InBoneIndex)
+		{
+			return true; // never update the pelvis
+		}
+		
 		for (const FChainPairFK& ChainPair : ChainPairsFK)
 		{
 			if (ChainPair.TargetBoneChain->BoneIndices.Contains(InBoneIndex))
@@ -680,7 +687,7 @@ bool FIKRetargetFKChainsOp::Initialize(
 			InTargetSkeleton.GetChildrenIndicesRecursive(ChainBone, AllChildren);
 			for (int32 ChildBoneIndex : AllChildren)
 			{
-				if (!IsBoneInAnyChain(ChildBoneIndex))
+				if (!IsBoneRetargeted(ChildBoneIndex))
 				{
 					NonRetargetedChildrenToUpdate.AddUnique(ChildBoneIndex);
 				}
