[IK Rig] Fixes GC crash when using goal creators

diff --git a/Engine/Plugins/Animation/IKRig/Source/IKRig/Private/AnimNodes/AnimNode_IKRig.cpp b/Engine/Plugins/Animation/IKRig/Source/IKRig/Private/AnimNodes/AnimNode_IKRig.cpp
--- a/Engine/Plugins/Animation/IKRig/Source/IKRig/Private/AnimNodes/AnimNode_IKRig.cpp
+++ b/Engine/Plugins/Animation/IKRig/Source/IKRig/Private/AnimNodes/AnimNode_IKRig.cpp
@@ -268,15 +268,16 @@ void FAnimNode_IKRig::PreUpdate(const UAnimInstance* InAnimInstance)
 			{
 				continue;
 			}
-			GoalCreators.Add(GoalCreator);
+			GoalCreators.Add(GoalCreatorComponent);
 		}
 	}
 	
 	// pull all the goals out of any goal creators on the owning actor
 	// this is done on the main thread because we're talking to actor components here
 	GoalsFromGoalCreators.Reset();
-	for (IIKGoalCreatorInterface* GoalCreator : GoalCreators)
+	for (UActorComponent* GoalCreatorComponent : GoalCreators)
 	{
+		IIKGoalCreatorInterface* GoalCreator = Cast<IIKGoalCreatorInterface>(GoalCreatorComponent);
 		GoalCreator->AddIKGoals_Implementation(GoalsFromGoalCreators);
 	}
 }
diff --git a/Engine/Plugins/Animation/IKRig/Source/IKRig/Public/AnimNodes/AnimNode_IKRig.h b/Engine/Plugins/Animation/IKRig/Source/IKRig/Public/AnimNodes/AnimNode_IKRig.h
--- a/Engine/Plugins/Animation/IKRig/Source/IKRig/Public/AnimNodes/AnimNode_IKRig.h
+++ b/Engine/Plugins/Animation/IKRig/Source/IKRig/Public/AnimNodes/AnimNode_IKRig.h
@@ -104,7 +104,8 @@ private:
 	FIKRigProcessor IKRigProcessor;
 
 	/** a cached list of components on the owning actor that implement the goal creator interface */
-	TArray<IIKGoalCreatorInterface*> GoalCreators;
+	UPROPERTY(transient)
+	TArray<TObjectPtr<UActorComponent>> GoalCreators;
 	TMap<FName, FIKRigGoal> GoalsFromGoalCreators;
 	bool bSearchedForGoalCreators = false;
 	
