[IK Retarget] Adding "Animation Relative" translation mode to FK Chains Op.

diff --git a/Engine/Plugins/Animation/IKRig/Source/IKRig/Private/Retargeter/IKRetargetProcessor.cpp b/Engine/Plugins/Animation/IKRig/Source/IKRig/Private/Retargeter/IKRetargetProcessor.cpp
--- a/Engine/Plugins/Animation/IKRig/Source/IKRig/Private/Retargeter/IKRetargetProcessor.cpp
+++ b/Engine/Plugins/Animation/IKRig/Source/IKRig/Private/Retargeter/IKRetargetProcessor.cpp
@@ -424,6 +424,17 @@ int32 FRetargetSkeleton::GetParentIndex(const int32 BoneIndex) const
 	return ParentIndices[BoneIndex];
 }
 
+const FTransform& FRetargetSkeleton::GetParentTransform(const int32 BoneIndex, const TArray<FTransform>& InPose) const
+{
+	int32 ParentIndex = GetParentIndex(BoneIndex);
+	if (ParentIndex == INDEX_NONE)
+	{
+		return FTransform::Identity;
+	}
+
+	return InPose[ParentIndex];
+}
+
 void FTargetSkeleton::Initialize(
 const USkeletalMesh* InSkeletalMesh,
 	const ERetargetSourceOrTarget InSourceOrTarget,
@@ -667,17 +678,42 @@ double FResolvedBoneChain::GetStretchAtParam(
 	{
 		return 1.0;
 	}
+
+	double InitialLength;
+	double CurrentLength;
 	
 	// end of chain always uses last bone
 	if (Param > 1.0 - UE_KINDA_SMALL_NUMBER)
 	{
 		const int32 Last = InitialTransforms.Num() - 1;
 		const int32 Prev = Last - 1;
-		const double InitialLength = (InitialTransforms[Last].GetTranslation() - InitialTransforms[Prev].GetTranslation()).Length();
-		const double CurrentLength = (CurrentTransforms[Last].GetTranslation() - CurrentTransforms[Prev].GetTranslation()).Length();
+		InitialLength = (InitialTransforms[Last].GetTranslation() - InitialTransforms[Prev].GetTranslation()).Length();
+		CurrentLength = (CurrentTransforms[Last].GetTranslation() - CurrentTransforms[Prev].GetTranslation()).Length();
 	}
+	else
+	{
+		// find the bone associated with this parameter
+		int32 ClosestBoneChainIndex = GetBoneClosestToParam(Param);
+		int32 BoneStartIndex = ClosestBoneChainIndex - 1;
+		int32 BoneEndIndex = ClosestBoneChainIndex;
 
-	// find the bone associated with this parameter
+		if (ClosestBoneChainIndex == 0)
+		{
+			BoneStartIndex = 0;
+			BoneEndIndex = 1;
+		}
+		
+		InitialLength = (InitialTransforms[BoneStartIndex].GetTranslation() - InitialTransforms[BoneEndIndex].GetTranslation()).Length();
+		CurrentLength = (CurrentTransforms[BoneStartIndex].GetTranslation() - CurrentTransforms[BoneEndIndex].GetTranslation()).Length();
+	}
+
+	// calculate the ratio of the current to the initial length of the closest bone
+	return (InitialLength < UE_KINDA_SMALL_NUMBER ? 1.0 : CurrentLength / InitialLength);
+}
+
+int32 FResolvedBoneChain::GetBoneClosestToParam(const double& Param) const
+{
+	// find the chain index of the bone closest to Param
 	double ClosestParamDistance = TNumericLimits<double>::Max();
 	int32 ClosestBoneChainIndex = INDEX_NONE;
 	for (int32 ChainIndex=0; ChainIndex<Params.Num(); ++ChainIndex)
@@ -689,20 +725,32 @@ double FResolvedBoneChain::GetStretchAtParam(
 			ClosestBoneChainIndex = ChainIndex;
 		}
 	}
+	
+	return ClosestBoneChainIndex;
+}
 
-	int32 BoneStartIndex = ClosestBoneChainIndex-1;
-	int32 BoneEndIndex = ClosestBoneChainIndex;
-
-	if (ClosestBoneChainIndex == 0)
+int32 FResolvedBoneChain::GetBoneAtParam(const double& Param) const
+{
+	if (Param > 1.0f || FMath::IsNearlyEqual(Param, 1.0))
+	{
+		return Params.Num() - 1;
+	}
+	
+	for (int32 ChainIndex=0; ChainIndex<Params.Num(); ++ChainIndex)
 	{
-		BoneStartIndex = 0;
-		BoneEndIndex = 1;
+		constexpr double Tolerance = 0.05;
+		if (FMath::IsNearlyEqual(Params[ChainIndex], Param, Tolerance))
+		{
+			return ChainIndex;
+		}
+		
+		if (Param < Params[ChainIndex])
+		{
+			return ChainIndex;
+		}
 	}
 
-	// calculate the ratio of the current to the initial length of the closest bone
-	const double InitialLength = (InitialTransforms[BoneStartIndex].GetTranslation() - InitialTransforms[BoneEndIndex].GetTranslation()).Length();
-	const double CurrentLength = (CurrentTransforms[BoneStartIndex].GetTranslation() - CurrentTransforms[BoneEndIndex].GetTranslation()).Length();
-	return (InitialLength < UE_KINDA_SMALL_NUMBER ? 1.0 : CurrentLength / InitialLength);
+	return Params.Num() - 1;
 }
 
 double FResolvedBoneChain::GetChainLength(const TArray<FTransform>& Transforms)
@@ -748,7 +796,7 @@ void FResolvedBoneChain::FillTransformsWithGlobalRetargetPoseOfChain(
 	const TArray<int32>& InBoneIndices,
 	TArray<FTransform>& OutGlobalTransforms)
 {
-	check(InBoneIndices.Num() == OutGlobalTransforms.Num())
+	OutGlobalTransforms.Reset(InBoneIndices.Num());
 	
 	// update chain current transforms to the retarget pose in global space
 	TArray<FTransform> LocalRetargetPose = InSkeleton.RetargetPoses.GetLocalRetargetPose();
@@ -758,7 +806,7 @@ void FResolvedBoneChain::FillTransformsWithGlobalRetargetPoseOfChain(
 		if (ChainIndex == 0)
 		{
 			const int32 BoneIndex = InBoneIndices[ChainIndex];
-			OutGlobalTransforms[ChainIndex] = InSkeleton.RetargetPoses.GetGlobalRetargetPoseOfSingleBone(InSkeleton, BoneIndex, InGlobalPose);
+			OutGlobalTransforms.Add(InSkeleton.RetargetPoses.GetGlobalRetargetPoseOfSingleBone(InSkeleton, BoneIndex, InGlobalPose));
 		}
 		else
 		{
@@ -766,7 +814,7 @@ void FResolvedBoneChain::FillTransformsWithGlobalRetargetPoseOfChain(
 			const int32 BoneIndex = InBoneIndices[ChainIndex];
 			const FTransform& ParentGlobalTransform = OutGlobalTransforms[ChainIndex-1];
 			const FTransform& ChildLocalTransform = LocalRetargetPose[BoneIndex];
-			OutGlobalTransforms[ChainIndex] = ChildLocalTransform * ParentGlobalTransform;
+			OutGlobalTransforms.Add(ChildLocalTransform * ParentGlobalTransform);
 		}
 	}
 }
diff --git a/Engine/Plugins/Animation/IKRig/Source/IKRig/Private/Retargeter/RetargetOps/FKChainsOp.cpp b/Engine/Plugins/Animation/IKRig/Source/IKRig/Private/Retargeter/RetargetOps/FKChainsOp.cpp
--- a/Engine/Plugins/Animation/IKRig/Source/IKRig/Private/Retargeter/RetargetOps/FKChainsOp.cpp
+++ b/Engine/Plugins/Animation/IKRig/Source/IKRig/Private/Retargeter/RetargetOps/FKChainsOp.cpp
@@ -57,11 +57,56 @@ void FChainEncoderFK::TransformCurrentChainTransforms(const FTransform& NewParen
 	}
 }
 
-void FChainDecoderFK::Initialize(const FResolvedBoneChain* InBoneChain)
+void FChainDecoderCachedData::Initialize(
+	const FResolvedBoneChain* InSourceBoneChain,
+	const FResolvedBoneChain* InTargetBoneChain,
+	const FRetargetSkeleton& InSourceSkeleton,
+	const FTargetSkeleton& InTargetSkeleton)
 {
-	check(InBoneChain);
-	BoneChain = InBoneChain;
-	CurrentGlobalTransforms = BoneChain->RefPoseGlobalTransforms;
+	// reset all data
+	*this = FChainDecoderCachedData();
+
+	// cache data for decoding the chain pose
+	for (int32 TargetChainIndex=0; TargetChainIndex<InTargetBoneChain->BoneIndices.Num(); ++TargetChainIndex)
+	{
+		// cache the source bone associated with each target bone in the chain
+		const double Param = InTargetBoneChain->Params[TargetChainIndex];
+		const int32 SourceBoneChainIndex = InSourceBoneChain->GetBoneAtParam(Param);
+		const int32 SourceBoneIndex = InSourceBoneChain->BoneIndices[SourceBoneChainIndex];
+		SourceBoneIndices.Add(SourceBoneIndex);
+
+		// calculate bone direction in SOURCE ref pose
+		const FTransform& SourceInitialGlobal = InSourceSkeleton.RetargetPoses.GetGlobalRetargetPose()[SourceBoneIndex];
+		const FTransform& SourceParentInitialGlobal = InSourceSkeleton.GetParentTransform(SourceBoneIndex, InSourceSkeleton.RetargetPoses.GetGlobalRetargetPose());
+		const FVector SourceInitialPosition = SourceInitialGlobal.GetTranslation();
+		const FVector SourceParentInitialPosition = SourceParentInitialGlobal.GetTranslation();
+		const FVector SourceInitialDirection = (SourceInitialPosition - SourceParentInitialPosition).GetSafeNormal();
+		// calculate bone direction in TARGET ref pose
+		const int32 TargetBoneIndex = InTargetBoneChain->BoneIndices[TargetChainIndex];
+		const FVector TargetInitialPosition = InTargetSkeleton.RetargetPoses.GetGlobalRetargetPose()[TargetBoneIndex].GetTranslation();
+		const FTransform& TargetParentInitialGlobal = InTargetSkeleton.GetParentTransform(TargetBoneIndex, InTargetSkeleton.RetargetPoses.GetGlobalRetargetPose());
+		const FVector TargetParentInitialPosition = TargetParentInitialGlobal.GetTranslation();
+		const FVector TargetInitialDirection = (TargetInitialPosition - TargetParentInitialPosition).GetSafeNormal();
+		// calculate FromTo rotation from source bone direction to target bone direction
+		const FQuat RefPoseGlobalDeltaRotation = FQuat::FindBetweenNormals(SourceInitialDirection, TargetInitialDirection);
+		// cache FromTo orientation delta in ref pose
+		FromToGlobalRotationDeltas.Add(RefPoseGlobalDeltaRotation);
+
+		// cache rest pose bone lengths
+		SourceRefPoseLengths.Add(FMath::Max(KINDA_SMALL_NUMBER, (SourceInitialPosition - SourceParentInitialPosition).Size()));
+		TargetRefPoseLengths.Add(FMath::Max(KINDA_SMALL_NUMBER, (TargetInitialPosition - TargetParentInitialPosition).Size()));
+	}	
+}
+
+void FChainDecoderFK::Initialize(
+	const FResolvedBoneChain* InSourceBoneChain,
+	const FResolvedBoneChain* InTargetBoneChain,
+	const FRetargetSkeleton& InSourceSkeleton,
+	const FTargetSkeleton& InTargetSkeleton)
+{
+	check(InSourceBoneChain && InTargetBoneChain);
+	BoneChain = InTargetBoneChain;
+	CachedData.Initialize(InSourceBoneChain, InTargetBoneChain, InSourceSkeleton, InTargetSkeleton);
 }
 
 void FChainDecoderFK::InitializeIntermediateParentIndices(
@@ -98,10 +143,11 @@ void FChainDecoderFK::DecodePose(
 	const FRetargetFKChainSettings& Settings,
 	const TArray<int32>& TargetBoneIndices,
     FChainEncoderFK& SourceChain,
+    const TArray<FTransform>& InSourceGlobalPose,
+    const FRetargetSkeleton& SourceSkeleton,
     const FTargetSkeleton& TargetSkeleton,
     TArray<FTransform> &InOutGlobalPose)
 {
-	check(TargetBoneIndices.Num() == CurrentGlobalTransforms.Num());
 	check(TargetBoneIndices.Num() == BoneChain->Params.Num());
 
 	// Before setting this chain pose, we need to ensure that any
@@ -127,12 +173,13 @@ void FChainDecoderFK::DecodePose(
 	if (!Settings.EnableFK)
 	{
 		// put the chain in the global ref pose (globally rotated by parent bone in it's currently retargeted state)
-		BoneChain->FillTransformsWithGlobalRetargetPoseOfChain(TargetSkeleton, InOutGlobalPose, TargetBoneIndices, CurrentGlobalTransforms);
+		TArray<FTransform> RefPoseGlobalTransforms;
+		BoneChain->FillTransformsWithGlobalRetargetPoseOfChain(TargetSkeleton, InOutGlobalPose, TargetBoneIndices, RefPoseGlobalTransforms);
 		
 		for (int32 ChainIndex=0; ChainIndex<TargetBoneIndices.Num(); ++ChainIndex)
 		{
 			const int32 BoneIndex = TargetBoneIndices[ChainIndex];
-			InOutGlobalPose[BoneIndex] = CurrentGlobalTransforms[ChainIndex];
+			InOutGlobalPose[BoneIndex] = RefPoseGlobalTransforms[ChainIndex];
 		}
 
 		return;
@@ -212,7 +259,6 @@ void FChainDecoderFK::DecodePose(
 		const FQuat RotationDelta = SourceCurrentRotation * SourceInitialRotation.Inverse();
 		const FQuat TargetInitialRotation = TargetInitialTransform.GetRotation();
 		const FQuat OutRotation = RotationDelta * TargetInitialRotation;
-		
 
 		// calculate output POSITION based on translation mode setting
 		FTransform ParentGlobalTransform = FTransform::Identity;
@@ -267,7 +313,35 @@ void FChainDecoderFK::DecodePose(
 				OutPosition = ParentGlobalTransform.TransformPosition(InitialLocalOffset * StretchRatio);
 				break;
 			}
-			default:
+		case EFKChainTranslationMode::OrientAndScale:
+			{
+				// calculate current direction of source bone
+				const int32 SourceBoneIndex = CachedData.SourceBoneIndices[ChainIndex];
+				const FTransform& SourceCurrentGlobal = InSourceGlobalPose[SourceBoneIndex];
+				const FTransform& SourceParentCurrentGlobal = SourceSkeleton.GetParentTransform(SourceBoneIndex, InSourceGlobalPose);
+				const FVector SourceCurrentDirection = SourceCurrentGlobal.GetTranslation() - SourceParentCurrentGlobal.GetTranslation();
+
+				// calculate current stretch of source bone
+				double SourceCurrentLength;
+				FVector SourceCurrentNorm;
+				SourceCurrentDirection.ToDirectionAndLength(SourceCurrentNorm, SourceCurrentLength);
+				const double SourceLengthScaleFactor = SourceCurrentLength / CachedData.SourceRefPoseLengths[ChainIndex];
+
+				// calculate equivalently stretched target bone in direction of source
+				const double TargetRefLength = CachedData.TargetRefPoseLengths[ChainIndex];
+				FVector TargetOffset = SourceCurrentNorm * (TargetRefLength * SourceLengthScaleFactor);
+
+				// apply rotational delta from the ref pose to the translational offset
+				const FQuat& RefPoseGlobalDeltaRotation = CachedData.FromToGlobalRotationDeltas[ChainIndex];
+				const FQuat CurrentGlobalRotationPlusDelta = RefPoseGlobalDeltaRotation * OutRotation;
+				const FQuat CurrentGlobalRotationDelta = CurrentGlobalRotationPlusDelta * OutRotation.Inverse();
+				TargetOffset = CurrentGlobalRotationDelta.RotateVector(TargetOffset);
+
+				// apply target offset to the parent global position
+				OutPosition = ParentGlobalTransform.GetLocation() + TargetOffset;
+				break;
+			}
+		default:
 			{
 				const FVector InitialLocalOffset = TargetSkeleton.RetargetPoses.GetLocalRetargetPose()[BoneIndex].GetTranslation();
 				OutPosition = ParentGlobalTransform.TransformPosition(InitialLocalOffset);
@@ -282,8 +356,7 @@ void FChainDecoderFK::DecodePose(
 		const FVector OutScale = SourceCurrentScale + (TargetInitialScale - SourceInitialScale);
 		
 		// apply output transform
-		CurrentGlobalTransforms[ChainIndex] = FTransform(OutRotation, OutPosition, OutScale);
-		InOutGlobalPose[BoneIndex] = CurrentGlobalTransforms[ChainIndex];
+		InOutGlobalPose[BoneIndex] = FTransform(OutRotation, OutPosition, OutScale);
 	}
 
 	// apply match chain operation on-top of translated and rotated result
@@ -292,13 +365,7 @@ void FChainDecoderFK::DecodePose(
 	{
 		const bool bScaleSourceChain = Settings.RotationMode == EFKChainRotationMode::MatchScaledChain;
 		FVector TargetChainOrigin = InOutGlobalPose[TargetBoneIndices[0]].GetTranslation();
-		MatchChain(bScaleSourceChain, SourceChain, TargetBoneIndices, TargetChainOrigin);
-		
-		// update output pose
-		for (int32 ChainIndex=0; ChainIndex<TargetBoneIndices.Num(); ++ChainIndex)
-		{
-			InOutGlobalPose[TargetBoneIndices[ChainIndex]] = CurrentGlobalTransforms[ChainIndex];
-		}
+		MatchChain(bScaleSourceChain, SourceChain, TargetBoneIndices, TargetChainOrigin, InOutGlobalPose);
 	}
 
 	// apply final blending between retarget pose of chain and newly retargeted pose
@@ -334,7 +401,8 @@ void FChainDecoderFK::MatchChain(
 	const bool bScaleSourceChain,
 	const FChainEncoderFK& SourceChain,
 	const TArray<int32>& TargetBoneIndices,
-	const FVector& TargetChainOrigin)
+	const FVector& TargetChainOrigin,
+	TArray<FTransform>& InOutGlobalPose)
 {
 	// "MatchChain" mode assumes the interpolated rotations (and any translation mode) as a starting point.
 	//
@@ -361,8 +429,15 @@ void FChainDecoderFK::MatchChain(
 	// scale chain to match target length
 	if (bScaleSourceChain)
 	{
-		double TargetChainLengthCurrent = BoneChain->GetChainLength(CurrentGlobalTransforms);
-		double SourceChainLengthCurrent = BoneChain-> GetChainLength(SourceChain.CurrentGlobalTransforms);
+		TArray <FTransform> ChainGlobalTransforms;
+		ChainGlobalTransforms.Reset(TargetBoneIndices.Num());
+		for (int32 BoneIndex : TargetBoneIndices)
+		{
+			ChainGlobalTransforms.Add(InOutGlobalPose[BoneIndex]);
+		}
+		
+		double TargetChainLengthCurrent = BoneChain->GetChainLength(ChainGlobalTransforms);
+		double SourceChainLengthCurrent = BoneChain->GetChainLength(SourceChain.CurrentGlobalTransforms);
 		double ScaleFactor = TargetChainLengthCurrent < UE_KINDA_SMALL_NUMBER ? 1.0f : SourceChainLengthCurrent / TargetChainLengthCurrent;
 		for (FVector& SplinePoint : SourceSplinePoints)
 		{
@@ -373,12 +448,14 @@ void FChainDecoderFK::MatchChain(
 	// snap orient each bone to lie on the source spline
 	int32 OriginPointIndex = 0;
 	double OriginPointAlpha = 0.0;
-	for (int32 ChainIndex=0; ChainIndex<TargetBoneIndices.Num()-1; ++ChainIndex)
+	for (int32 ChainIndex=0; ChainIndex<TargetBoneIndices.Num() - 1; ++ChainIndex)
 	{
 		// generate an aim rotation from the current joint vector to the point on the spline
-		const FVector StartBonePosition = CurrentGlobalTransforms[ChainIndex].GetTranslation();
-		const FVector EndBonePosition = CurrentGlobalTransforms[ChainIndex+1].GetTranslation();
-		const FVector EndBoneLocalPosition = CurrentGlobalTransforms[ChainIndex].InverseTransformPosition(EndBonePosition);
+		int32 BoneIndex = TargetBoneIndices[ChainIndex];
+		int32 NextBoneIndex = TargetBoneIndices[ChainIndex + 1];
+		const FVector StartBonePosition = InOutGlobalPose[BoneIndex].GetTranslation();
+		const FVector EndBonePosition = InOutGlobalPose[NextBoneIndex].GetTranslation();
+		const FVector EndBoneLocalPosition = InOutGlobalPose[BoneIndex].InverseTransformPosition(EndBonePosition);
 		FVector BoneNorm;
 		double BoneLength;
 		(EndBonePosition - StartBonePosition).ToDirectionAndLength(BoneNorm, BoneLength);
@@ -400,14 +477,15 @@ void FChainDecoderFK::MatchChain(
 		const FQuat RotationOffset = FQuat::FindBetweenNormals(BoneNorm, AimNorm);
 
 		// rotate bone to contact spline
-		CurrentGlobalTransforms[ChainIndex].SetRotation( RotationOffset * CurrentGlobalTransforms[ChainIndex].GetRotation());
+		InOutGlobalPose[BoneIndex].SetRotation( RotationOffset * InOutGlobalPose[BoneIndex].GetRotation());
 
 		// propagate translational offset of end bone to all it's children
-		const FVector NewEndBonePosition = CurrentGlobalTransforms[ChainIndex].TransformPosition(EndBoneLocalPosition);
+		const FVector NewEndBonePosition = InOutGlobalPose[BoneIndex].TransformPosition(EndBoneLocalPosition);
 		const FVector DeltaTranslation = NewEndBonePosition - EndBonePosition;
-		for (int32 ChildIndex=ChainIndex+1; ChildIndex<TargetBoneIndices.Num(); ++ChildIndex)
+		for (int32 ChildChainIndex=ChainIndex+1; ChildChainIndex<TargetBoneIndices.Num(); ++ChildChainIndex)
 		{
-			CurrentGlobalTransforms[ChildIndex].AddToTranslation(DeltaTranslation);
+			int32 ChildBoneIndex = TargetBoneIndices[ChildChainIndex];
+			InOutGlobalPose[ChildBoneIndex].AddToTranslation(DeltaTranslation);
 		}
 	}
 }
@@ -529,6 +607,8 @@ void FChainDecoderFK::UpdateIntermediateLocalTransforms(
 bool FChainPairFK::Initialize(
 	const FResolvedBoneChain& InSourceBoneChain,
 	const FResolvedBoneChain& InTargetBoneChain,
+	const FRetargetSkeleton& InSourceSkeleton,
+	const FTargetSkeleton& InTargetSkeleton,
 	const FRetargetFKChainSettings& InSettings,
 	const FIKRigLogger& InLog)
 {
@@ -546,7 +626,7 @@ bool FChainPairFK::Initialize(
 	Settings = &InSettings;
 	
 	FKEncoder.Initialize(SourceBoneChain);
-	FKDecoder.Initialize(TargetBoneChain);
+	FKDecoder.Initialize(SourceBoneChain, TargetBoneChain, InSourceSkeleton, InTargetSkeleton);
 
 	return true;
 }
@@ -631,6 +711,8 @@ bool FIKRetargetFKChainsOp::Initialize(
 		const bool bChainInitialized = ChainPair.Initialize(
 			*SourceBoneChain,
 			*TargetBoneChain,
+			InSourceSkeleton,
+			InTargetSkeleton,
 			ChainSettings,
 			InLog);
 		if (!bChainInitialized)
@@ -737,6 +819,8 @@ void FIKRetargetFKChainsOp::Run(
 			*ChainPair.Settings,
 			ChainPair.TargetBoneChain->BoneIndices,
 			ChainPair.FKEncoder,
+			InSourceGlobalPose,
+			SourceSkeleton,
 			TargetSkeleton,
 			OutTargetGlobalPose);
 	}
diff --git a/Engine/Plugins/Animation/IKRig/Source/IKRig/Private/Retargeter/RetargetOps/PinBoneOp.cpp b/Engine/Plugins/Animation/IKRig/Source/IKRig/Private/Retargeter/RetargetOps/PinBoneOp.cpp
--- a/Engine/Plugins/Animation/IKRig/Source/IKRig/Private/Retargeter/RetargetOps/PinBoneOp.cpp
+++ b/Engine/Plugins/Animation/IKRig/Source/IKRig/Private/Retargeter/RetargetOps/PinBoneOp.cpp
@@ -27,17 +27,6 @@ void FPinBoneData::PostSerialize(const FArchive& Ar)
 #endif // WITH_EDITORONLY_DATA
 }
 
-const FTransform& FPinBoneData::GetParentTransform(int32 BoneIndex, const FRetargetSkeleton& Skeleton, const TArray<FTransform>& Pose)
-{
-	int32 ParentIndex = Skeleton.GetParentIndex(BoneIndex);
-	if (ParentIndex == INDEX_NONE)
-	{
-		return FTransform::Identity;
-	}
-
-	return Pose[ParentIndex];
-}
-
 void FPinBoneData::CachePinData(ERetargetSourceOrTarget InSkeletonToCopyFrom, const FIKRetargetProcessor& Processor)
 {
 	SourceScaleFactor = Processor.GetSourceScaleFactor();
@@ -51,8 +40,8 @@ void FPinBoneData::CachePinData(ERetargetSourceOrTarget InSkeletonToCopyFrom, co
 	const TArray<FTransform>& ToRetargetPose = SkeletonToCopyTo.RetargetPoses.GetGlobalRetargetPose();
 
 	// get relevant transforms from the retarget poses
-	FTransform BoneToCopyFromParentRefPoseGlobal = GetParentTransform(BoneToCopyFrom.BoneIndex, SkeletonToCopyFrom, FromRetargetPose);
-	FTransform BoneToCopyToParentRefPoseGlobal = GetParentTransform(BoneToCopyTo.BoneIndex, SkeletonToCopyTo, ToRetargetPose);
+	FTransform BoneToCopyFromParentRefPoseGlobal = SkeletonToCopyFrom.GetParentTransform(BoneToCopyFrom.BoneIndex, FromRetargetPose);
+	FTransform BoneToCopyToParentRefPoseGlobal = SkeletonToCopyTo.GetParentTransform(BoneToCopyTo.BoneIndex, ToRetargetPose);
 	FTransform BoneToCopyFromRefPoseGlobal = FromRetargetPose[BoneToCopyFrom.BoneIndex];
 	FTransform BoneToCopyToRefPoseGlobal = ToRetargetPose[BoneToCopyTo.BoneIndex];
 
@@ -256,7 +245,7 @@ FTransform FIKRetargetPinBoneOp::GetNewBoneTransform(
 	auto GetBoneToCopyFromCurrentVector = [PoseToCopyFrom, CopyFrom, PinData]()
 	{
 		const FTransform& BoneToCopyFromCurrent = PoseToCopyFrom[PinData.BoneToCopyFrom.BoneIndex];
-		const FTransform& ParentOfBoneToCopyFromCurrent = FPinBoneData::GetParentTransform(PinData.BoneToCopyFrom.BoneIndex, CopyFrom, PoseToCopyFrom);
+		const FTransform& ParentOfBoneToCopyFromCurrent = CopyFrom.GetParentTransform(PinData.BoneToCopyFrom.BoneIndex, PoseToCopyFrom);
 		return BoneToCopyFromCurrent.GetTranslation() - ParentOfBoneToCopyFromCurrent.GetTranslation();
 	};
 	
@@ -282,7 +271,7 @@ FTransform FIKRetargetPinBoneOp::GetNewBoneTransform(
 	case EPinBoneTranslationMode::CopyLocalPosition:
 		{
 			const FVector BoneToCopyFromCurrentVector = GetBoneToCopyFromCurrentVector();
-			const FTransform& ParentOfBoneToCopyToCurrent = FPinBoneData::GetParentTransform(PinData.BoneToCopyTo.BoneIndex, CopyTo, OutTargetGlobalPose);
+			const FTransform& ParentOfBoneToCopyToCurrent = CopyTo.GetParentTransform(PinData.BoneToCopyTo.BoneIndex, OutTargetGlobalPose);
 			Result.SetTranslation(ParentOfBoneToCopyToCurrent.GetTranslation() + BoneToCopyFromCurrentVector);
 			break;
 		}
@@ -296,7 +285,7 @@ FTransform FIKRetargetPinBoneOp::GetNewBoneTransform(
 			BoneToCopyFromTranslationCurrent.ToDirectionAndLength(BoneToCopyFromTranslationNorm, CurrentLength);
 			const FVector BoneToCopyFromRelativeOffset = BoneToCopyFromTranslationNorm * (CurrentLength + RestPoseLengthDifference);
 		
-			const FTransform& ParentOfBoneToCopyToCurrent = FPinBoneData::GetParentTransform(PinData.BoneToCopyTo.BoneIndex, CopyTo, OutTargetGlobalPose);
+			const FTransform& ParentOfBoneToCopyToCurrent = CopyTo.GetParentTransform(PinData.BoneToCopyTo.BoneIndex, OutTargetGlobalPose);
 			Result.SetTranslation(ParentOfBoneToCopyToCurrent.GetTranslation() + BoneToCopyFromRelativeOffset);
 			break;
 		}
@@ -314,7 +303,7 @@ FTransform FIKRetargetPinBoneOp::GetNewBoneTransform(
 			const double LengthOfBoneToCopyToInRefPose = PinData.LocalRefPoseBoneToCopyTo.GetTranslation().Size();
 			const FVector BoneToCopyFromRelativeOffset = BoneToCopyFromTranslationNorm * (LengthOfBoneToCopyToInRefPose * BoneLengthScaleFactor);
 		
-			const FTransform& ParentOfBoneToCopyToCurrent = FPinBoneData::GetParentTransform(PinData.BoneToCopyTo.BoneIndex, CopyTo, OutTargetGlobalPose);
+			const FTransform& ParentOfBoneToCopyToCurrent = CopyTo.GetParentTransform(PinData.BoneToCopyTo.BoneIndex, OutTargetGlobalPose);
 			Result.SetTranslation(ParentOfBoneToCopyToCurrent.GetTranslation() + BoneToCopyFromRelativeOffset);
 			
 			break;
diff --git a/Engine/Plugins/Animation/IKRig/Source/IKRig/Public/Retargeter/IKRetargetProcessor.h b/Engine/Plugins/Animation/IKRig/Source/IKRig/Public/Retargeter/IKRetargetProcessor.h
--- a/Engine/Plugins/Animation/IKRig/Source/IKRig/Public/Retargeter/IKRetargetProcessor.h
+++ b/Engine/Plugins/Animation/IKRig/Source/IKRig/Public/Retargeter/IKRetargetProcessor.h
@@ -101,6 +101,8 @@ struct FRetargetSkeleton
 
 	UE_API int32 GetParentIndex(const int32 BoneIndex) const;
 
+	const FTransform& GetParentTransform(const int32 BoneIndex, const TArray<FTransform>& InPose) const;
+
 	UE_API void UpdateGlobalTransformsBelowBone(
 		const int32 StartBoneIndex,
 		const TArray<FTransform>& InLocalPose,
@@ -206,6 +208,10 @@ struct FResolvedBoneChain
 		const TArray<FTransform>& InitialTransforms,
 		const TArray<FTransform>& CurrentTransforms,
 		const double& Param) const;
+
+	int32 GetBoneClosestToParam(const double& Param) const;
+
+	int32 GetBoneAtParam(const double& Param) const;
 	
 	static UE_API double GetChainLength(const TArray<FTransform>& Transforms);
 
diff --git a/Engine/Plugins/Animation/IKRig/Source/IKRig/Public/Retargeter/RetargetOps/FKChainsOp.h b/Engine/Plugins/Animation/IKRig/Source/IKRig/Public/Retargeter/RetargetOps/FKChainsOp.h
--- a/Engine/Plugins/Animation/IKRig/Source/IKRig/Public/Retargeter/RetargetOps/FKChainsOp.h
+++ b/Engine/Plugins/Animation/IKRig/Source/IKRig/Public/Retargeter/RetargetOps/FKChainsOp.h
@@ -30,9 +30,27 @@ struct FChainEncoderFK
 	const FResolvedBoneChain* BoneChain;
 };
 
+struct FChainDecoderCachedData
+{
+	TArray<int32> SourceBoneIndices;
+	TArray<FQuat> FromToGlobalRotationDeltas;
+	TArray<double> SourceRefPoseLengths;
+	TArray<double> TargetRefPoseLengths;
+
+	void Initialize(
+		const FResolvedBoneChain* InSourceBoneChain,
+		const FResolvedBoneChain* InTargetBoneChain,
+		const FRetargetSkeleton& InSourceSkeleton,
+		const FTargetSkeleton& InTargetSkeleton);
+};
+
 struct FChainDecoderFK
 {
-	void Initialize(const FResolvedBoneChain* InBoneChain);
+	void Initialize(
+		const FResolvedBoneChain* InSourceBoneChain,
+		const FResolvedBoneChain* InTargetBoneChain,
+		const FRetargetSkeleton& InSourceSkeleton,
+		const FTargetSkeleton& InTargetSkeleton);
 	
 	void InitializeIntermediateParentIndices(
 		const int32 InRetargetRootBoneIndex,
@@ -48,8 +66,10 @@ struct FChainDecoderFK
 		const FRetargetFKChainSettings& Settings,
 		const TArray<int32>& TargetBoneIndices,
 		FChainEncoderFK& SourceChain,
+		const TArray<FTransform>& InSourceGlobalPose,
+		const FRetargetSkeleton& SourceSkeleton,
 		const FTargetSkeleton& TargetSkeleton,
-		TArray<FTransform> &InOutGlobalPose);
+		TArray<FTransform>& InOutGlobalPose);
 	
 private:
 
@@ -57,7 +77,8 @@ private:
 		const bool bScaleSourceChain,
 		const FChainEncoderFK& SourceChain,
 		const TArray<int32>& TargetBoneIndices,
-		const FVector& TargetChainOrigin);
+		const FVector& TargetChainOrigin,
+		TArray<FTransform>& InOutGlobalPose);
 
 	static FVector GetPointOnSplineFromIndexAndAlpha(const TArray<FVector>& InSplinePoints, const int32 InPointIndex, const double InSegmentAlpha);
 	
@@ -72,8 +93,8 @@ private:
 	void UpdateIntermediateParents(
 		const FTargetSkeleton& TargetSkeleton,
 		TArray<FTransform> &InOutGlobalPose);
-	
-	TArray<FTransform> CurrentGlobalTransforms;
+
+	FChainDecoderCachedData CachedData;
 	TArray<int32> IntermediateParentIndices;
 	TArray<FTransform> IntermediateLocalTransforms;
 	const FResolvedBoneChain* BoneChain;
@@ -84,6 +105,8 @@ struct FChainPairFK
 	bool Initialize(
 		const FResolvedBoneChain& InSourceBoneChain,
 		const FResolvedBoneChain& InTargetBoneChain,
+		const FRetargetSkeleton& InSourceSkeleton,
+		const FTargetSkeleton& InTargetSkeleton,
 		const FRetargetFKChainSettings& InSettings,
 		const FIKRigLogger& InLog);
 	
@@ -102,6 +125,7 @@ enum class EFKChainTranslationMode : uint8
 	Absolute						UMETA(DisplayName = "Absolute"),
 	StretchBoneLengthUniformly		UMETA(DisplayName = "Stretch Bone Length Uniformly"),
 	StretchBoneLengthNonUniformly	UMETA(DisplayName = "Stretch Bone Length Non-Uniformly"),
+	OrientAndScale					UMETA(DisplayName = "Orient and Scale"),
 };
 
 UENUM(BlueprintType)
diff --git a/Engine/Plugins/Animation/IKRig/Source/IKRig/Public/Retargeter/RetargetOps/PinBoneOp.h b/Engine/Plugins/Animation/IKRig/Source/IKRig/Public/Retargeter/RetargetOps/PinBoneOp.h
--- a/Engine/Plugins/Animation/IKRig/Source/IKRig/Public/Retargeter/RetargetOps/PinBoneOp.h
+++ b/Engine/Plugins/Animation/IKRig/Source/IKRig/Public/Retargeter/RetargetOps/PinBoneOp.h
@@ -27,8 +27,6 @@ struct FPinBoneData
 
 	void PostSerialize(const FArchive& Ar);
 	
-	static const FTransform& GetParentTransform(int32 BoneIndex, const FRetargetSkeleton& Skeleton, const TArray<FTransform>& Pose);
-
 	void CachePinData(ERetargetSourceOrTarget InSkeletonToCopyFrom, const FIKRetargetProcessor& Processor);
 
 	// The bone to copy FROM.
