[IK Retarget] Changed batch export to allow user to forcibly replace original source animation sequence assets.

diff --git a/Engine/Plugins/Animation/IKRig/Source/IKRigEditor/Private/RetargetEditor/IKRetargetBatchOperation.cpp b/Engine/Plugins/Animation/IKRig/Source/IKRigEditor/Private/RetargetEditor/IKRetargetBatchOperation.cpp
--- a/Engine/Plugins/Animation/IKRig/Source/IKRigEditor/Private/RetargetEditor/IKRetargetBatchOperation.cpp
+++ b/Engine/Plugins/Animation/IKRig/Source/IKRigEditor/Private/RetargetEditor/IKRetargetBatchOperation.cpp
@@ -674,14 +674,8 @@ void UIKRetargetBatchOperation::OverwriteExistingAssets(const FIKRetargetBatchOp
 			continue;
 		}
 
-		UObject* AssetToReplace = AssetDataToReplace.GetAsset();
-		if (AssetToReplace == OldAsset)
-		{
-			// we only replace previously retargeted animations, never the original
-			continue;
-		}
-
 		// reroute all references from old asset to new asset
+		UObject* AssetToReplace = AssetDataToReplace.GetAsset();
 		TArray<UObject*> AssetsToReplace = {AssetToReplace};
 		ObjectTools::ForceReplaceReferences(NewAsset, AssetsToReplace);
 		
