Updated headers using UnrealCodeFixup to make sure dllstorage is on methods/staticvars instead of types. Part 1/n

diff --git a/Engine/Plugins/Animation/IKRig/Source/IKRigEditor/Public/RetargetEditor/IKRetargetDetails.h b/Engine/Plugins/Animation/IKRig/Source/IKRigEditor/Public/RetargetEditor/IKRetargetDetails.h
--- a/Engine/Plugins/Animation/IKRig/Source/IKRigEditor/Public/RetargetEditor/IKRetargetDetails.h
+++ b/Engine/Plugins/Animation/IKRig/Source/IKRigEditor/Public/RetargetEditor/IKRetargetDetails.h
@@ -9,6 +9,8 @@
 
 #include "IKRetargetDetails.generated.h"
 
+#define UE_API IKRIGEDITOR_API
+
 class UDebugSkelMeshComponent;
 
 enum class EIKRetargetTransformType : int8
@@ -19,8 +21,8 @@ enum class EIKRetargetTransformType : int8
 	Bone,
 };
 
-UCLASS(config = Engine, hidecategories = UObject)
-class IKRIGEDITOR_API UIKRetargetBoneDetails : public UObject
+UCLASS(MinimalAPI, config = Engine, hidecategories = UObject)
+class UIKRetargetBoneDetails : public UObject
 {
 	GENERATED_BODY()
 
@@ -45,12 +47,12 @@ public:
 
 #if WITH_EDITOR
 
-	FEulerTransform GetTransform(EIKRetargetTransformType TransformType, const bool bLocalSpace) const;
-	bool IsComponentRelative(ESlateTransformComponent::Type Component, EIKRetargetTransformType TransformType) const;
-	void OnComponentRelativeChanged(ESlateTransformComponent::Type Component, bool bIsRelative, EIKRetargetTransformType TransformType);
-	void OnCopyToClipboard(ESlateTransformComponent::Type Component, EIKRetargetTransformType TransformType) const;
-	void OnPasteFromClipboard(ESlateTransformComponent::Type Component, EIKRetargetTransformType TransformType);
-	void OnNumericValueCommitted(
+	UE_API FEulerTransform GetTransform(EIKRetargetTransformType TransformType, const bool bLocalSpace) const;
+	UE_API bool IsComponentRelative(ESlateTransformComponent::Type Component, EIKRetargetTransformType TransformType) const;
+	UE_API void OnComponentRelativeChanged(ESlateTransformComponent::Type Component, bool bIsRelative, EIKRetargetTransformType TransformType);
+	UE_API void OnCopyToClipboard(ESlateTransformComponent::Type Component, EIKRetargetTransformType TransformType) const;
+	UE_API void OnPasteFromClipboard(ESlateTransformComponent::Type Component, EIKRetargetTransformType TransformType);
+	UE_API void OnNumericValueCommitted(
 		ESlateTransformComponent::Type Component,
 		ESlateRotationRepresentation::Type Representation,
 		ESlateTransformSubComponent::Type SubComponent,
@@ -59,14 +61,14 @@ public:
 		EIKRetargetTransformType TransformType,
 		bool bIsCommit);
 
-	TOptional<FVector::FReal> GetNumericValue(
+	UE_API TOptional<FVector::FReal> GetNumericValue(
 		EIKRetargetTransformType TransformType,
 		ESlateTransformComponent::Type Component,
 		ESlateRotationRepresentation::Type Representation,
 		ESlateTransformSubComponent::Type SubComponent);
 	
 	/** Method to react to changes of numeric values in the widget */
-	static void OnMultiNumericValueCommitted(
+	static UE_API void OnMultiNumericValueCommitted(
 		ESlateTransformComponent::Type Component,
 		ESlateRotationRepresentation::Type Representation,
 		ESlateTransformSubComponent::Type SubComponent,
@@ -83,7 +85,7 @@ public:
 
 private:
 
-	void CommitValueAsRelativeOffset(
+	UE_API void CommitValueAsRelativeOffset(
 		UIKRetargeterController* AssetController,
 		const FReferenceSkeleton& RefSkeleton,
 		const ERetargetSourceOrTarget SourceOrTarget,
@@ -95,7 +97,7 @@ private:
 		FVector::FReal Value,
 		bool bShouldTransact);
 
-	void CommitValueAsBoneSpace(
+	UE_API void CommitValueAsBoneSpace(
 		UIKRetargeterController* AssetController,
 		const FReferenceSkeleton& RefSkeleton,
 		const ERetargetSourceOrTarget SourceOrTarget,
@@ -107,9 +109,9 @@ private:
 		FVector::FReal Value,
 		bool bShouldTransact);
 
-	static TOptional<FVector::FReal> CleanRealValue(TOptional<FVector::FReal> InValue);
+	static UE_API TOptional<FVector::FReal> CleanRealValue(TOptional<FVector::FReal> InValue);
 	
-	bool IsRootBone() const;
+	UE_API bool IsRootBone() const;
 
 	bool BoneRelative[3] = { false, true, false };
 	bool RelativeOffsetTransformRelative[3] = { false, true, false };
@@ -166,3 +168,5 @@ private:
 
 	TArray<TObjectPtr<UIKRetargetBoneDetails>> Bones;
 };
+
+#undef UE_API
diff --git a/Engine/Plugins/Animation/IKRig/Source/IKRigEditor/Public/RetargetEditor/IKRetargetFactory.h b/Engine/Plugins/Animation/IKRig/Source/IKRigEditor/Public/RetargetEditor/IKRetargetFactory.h
--- a/Engine/Plugins/Animation/IKRig/Source/IKRigEditor/Public/RetargetEditor/IKRetargetFactory.h
+++ b/Engine/Plugins/Animation/IKRig/Source/IKRigEditor/Public/RetargetEditor/IKRetargetFactory.h
@@ -7,10 +7,12 @@
 #include "Factories/Factory.h"
 #include "IKRetargetFactory.generated.h"
 
+#define UE_API IKRIGEDITOR_API
+
 class SWindow;
 
-UCLASS(hidecategories=Object)
-class IKRIGEDITOR_API UIKRetargetFactory : public UFactory
+UCLASS(MinimalAPI, hidecategories=Object)
+class UIKRetargetFactory : public UFactory
 {
 	GENERATED_BODY()
 
@@ -19,13 +21,15 @@ class IKRIGEDITOR_API UIKRetargetFactory : public UFactory
 
 public:
 
-	UIKRetargetFactory();
+	UE_API UIKRetargetFactory();
 
 	// UFactory Interface
-	virtual FText GetDisplayName() const override;
-	virtual uint32 GetMenuCategories() const override;
-	virtual FText GetToolTip() const override;
-	virtual FString GetDefaultNewAssetName() const override;
-	virtual UObject* FactoryCreateNew(UClass* Class, UObject* InParent, FName Name, EObjectFlags Flags, UObject* Context, FFeedbackContext* Warn) override;
-	virtual bool ShouldShowInNewMenu() const override;
+	UE_API virtual FText GetDisplayName() const override;
+	UE_API virtual uint32 GetMenuCategories() const override;
+	UE_API virtual FText GetToolTip() const override;
+	UE_API virtual FString GetDefaultNewAssetName() const override;
+	UE_API virtual UObject* FactoryCreateNew(UClass* Class, UObject* InParent, FName Name, EObjectFlags Flags, UObject* Context, FFeedbackContext* Warn) override;
+	UE_API virtual bool ShouldShowInNewMenu() const override;
 };
+
+#undef UE_API
diff --git a/Engine/Plugins/Animation/IKRig/Source/IKRigEditor/Public/RetargetEditor/IKRetargeterController.h b/Engine/Plugins/Animation/IKRig/Source/IKRigEditor/Public/RetargetEditor/IKRetargeterController.h
--- a/Engine/Plugins/Animation/IKRig/Source/IKRigEditor/Public/RetargetEditor/IKRetargeterController.h
+++ b/Engine/Plugins/Animation/IKRig/Source/IKRigEditor/Public/RetargetEditor/IKRetargeterController.h
@@ -9,6 +9,8 @@
 
 #include "IKRetargeterController.generated.h"
 
+#define UE_API IKRIGEDITOR_API
+
 struct FRetargetGlobalSettings;
 struct FTargetRootSettings;
 struct FTargetChainSettings;
@@ -25,21 +27,21 @@ class UIKRigStructViewer;
 
 // A stateless singleton (1-per-asset) class used to make modifications to a UIKRetargeter asset.
 // Use UIKRetargeter.GetController() to get the controller for the asset you want to modify.  
-UCLASS(BlueprintType, hidecategories = UObject)
-class IKRIGEDITOR_API UIKRetargeterController : public UObject
+UCLASS(MinimalAPI, BlueprintType, hidecategories = UObject)
+class UIKRetargeterController : public UObject
 {
 	GENERATED_BODY()
 
 public:
 
-	UIKRetargeterController();
+	UE_API UIKRetargeterController();
 
 	// UObject
-	virtual void PostInitProperties() override;
+	UE_API virtual void PostInitProperties() override;
 	
 	// Get access to the retargeter asset.
 	// Warning: Do not make modifications to the asset directly. Using the controller API guarantees correctness.
-	UIKRetargeter* GetAsset() const;
+	UE_API UIKRetargeter* GetAsset() const;
   
 private:
 	
@@ -56,36 +58,36 @@ public:
 	 * @param InRetargeterAsset an IK Retarget asset
 	 * @return the controller with an API for modifying the given retarget asset */
 	UFUNCTION(BlueprintCallable, Category="IK Retarget", meta = (BlueprintThreadSafe, DisplayName = "Get IK Retarget Controller"))
-	static UIKRetargeterController* GetController(const UIKRetargeter* InRetargeterAsset);
+	static UE_API UIKRetargeterController* GetController(const UIKRetargeter* InRetargeterAsset);
 	
 	/** Set the IK Rig to use as the source or target (to copy animation FROM/TO)
 	 * @param SourceOrTarget an enum specifying either "Source" or "Target"
 	 * @param the IK Rig asset to apply as either the "Source" (to copy FROM) or "Target" (to copy TO)*/
 	UFUNCTION(BlueprintCallable, BlueprintPure=false, Category="IK Retarget", meta = (BlueprintThreadSafe, DisplayName = "Set IK Rig"))
-	void SetIKRig(const ERetargetSourceOrTarget SourceOrTarget, UIKRigDefinition* IKRig) const;
+	UE_API void SetIKRig(const ERetargetSourceOrTarget SourceOrTarget, UIKRigDefinition* IKRig) const;
 	
 	/** Get either source or target IK Rig
 	 * @param SourceOrTarget an enum specifying either "Source" or "Target"
 	 * @return the IK Rig asset associated with either the Source or Target of the IK Retargeter */ 
 	UFUNCTION(BlueprintCallable, Category="IK Retarget", meta = (BlueprintThreadSafe, DisplayName = "Get IK Rig"))
-	const UIKRigDefinition* GetIKRig(const ERetargetSourceOrTarget SourceOrTarget) const;
+	UE_API const UIKRigDefinition* GetIKRig(const ERetargetSourceOrTarget SourceOrTarget) const;
 	
 	/** Get all target IK Rigs referenced by all ops
 	 * @return an array of all target IK Rig assets associated with all the ops in the stack */ 
 	UFUNCTION(BlueprintCallable, Category="IK Retarget", meta = (BlueprintThreadSafe))
-	TArray<UIKRigDefinition*> GetAllTargetIKRigs() const;
+	UE_API TArray<UIKRigDefinition*> GetAllTargetIKRigs() const;
 
 	/** Set the preview skeletal mesh for either source or target
 	* @param SourceOrTarget an enum specifying either "Source" or "Target"
 	* @param InPreviewMesh a skeletal mesh asset to use as the preview mesh in the retarget editor (may be used on other meshes at runtime) */
 	UFUNCTION(BlueprintCallable, BlueprintPure=false, Category="IK Retarget", meta = (BlueprintThreadSafe))
-	void SetPreviewMesh(const ERetargetSourceOrTarget SourceOrTarget, USkeletalMesh* InPreviewMesh) const;
+	UE_API void SetPreviewMesh(const ERetargetSourceOrTarget SourceOrTarget, USkeletalMesh* InPreviewMesh) const;
 	
 	/** Get the preview skeletal mesh
 	 * @param SourceOrTarget an enum specifying either "Source" or "Target"
 	 * @return the skeletal mesh asset currently being used as a preview mesh */
 	UFUNCTION(BlueprintCallable, Category="IK Retarget", meta = (BlueprintThreadSafe))
-	USkeletalMesh* GetPreviewMesh(const ERetargetSourceOrTarget SourceOrTarget) const;
+	UE_API USkeletalMesh* GetPreviewMesh(const ERetargetSourceOrTarget SourceOrTarget) const;
 
 	//
 	// RETARGET OPS PUBLIC/SCRIPTING API
@@ -94,32 +96,32 @@ public:
 	/** Add a new retarget op of the given type to the bottom of the stack. Returns the stack index.
 	 * @param InIKRetargetOpType: the full package path of the UStruct type, ie /Script/IKRig.IKRetargetPinBoneOp */
 	UFUNCTION(BlueprintCallable, BlueprintPure=false, Category="IK Retarget")
-	int32 AddRetargetOp(const FString InIKRetargetOpType) const;
+	UE_API int32 AddRetargetOp(const FString InIKRetargetOpType) const;
 
 	/** Remove the retarget op at the given stack index
 	 * NOTE: if this op is a parent, all it's children will be removed as well
 	 * @param InOpIndex: the index of the op to remove
 	 * @return true if op was found and removed, false otherwise */ 
 	UFUNCTION(BlueprintCallable, BlueprintPure=false, Category="IK Retarget")
-	bool RemoveRetargetOp(const int32 InOpIndex) const;
+	UE_API bool RemoveRetargetOp(const int32 InOpIndex) const;
 
 	/** Remove all ops in the stack.
 	 * @return true if any ops were removed, false otherwise */
 	UFUNCTION(BlueprintCallable, Category="IK Retarget")
-	bool RemoveAllOps() const;
+	UE_API bool RemoveAllOps() const;
 
 	/** Set the name of the op at the given index in the stack.
 	 * @param InName the new name to use
 	 * @param InOpIndex the index of the op to be renamed
 	 * @return the new name of the op (after it's made unique) */
 	UFUNCTION(BlueprintCallable, Category="IK Retarget")
-	FName SetOpName(const FName InName, const int32 InOpIndex) const;
+	UE_API FName SetOpName(const FName InName, const int32 InOpIndex) const;
 	
 	/** Get the name of the op at the given index in the stack.
 	 * @param InOpIndex the index of the op to get the name of
 	 * @return the name of the op or None if the index was invalid */
 	UFUNCTION(BlueprintCallable, Category="IK Retarget")
-	FName GetOpName(const int32 InOpIndex) const ;
+	UE_API FName GetOpName(const int32 InOpIndex) const ;
 
 	/** Set the name of the op to parent this op to
 	 * Children ops are forced to execute before their parent
@@ -127,87 +129,87 @@ public:
 	 * @param InParentOpName the name of the op to parent to
 	 * @return true if both parent and child were found, and the parent relationship was allowed */
 	UFUNCTION(BlueprintCallable, Category="IK Retarget")
-	bool SetParentOpByName(const FName InChildOpName, const FName InParentOpName) const;
+	UE_API bool SetParentOpByName(const FName InChildOpName, const FName InParentOpName) const;
 	
 	/** Get the name of the parent op for the given op.
 	 * @param InOpName the name of the op to get the parent of
 	 * @return the name of the parent op or None if the op was not found or did not have a parent */
 	UFUNCTION(BlueprintCallable, Category="IK Retarget")
-	FName GetParentOpByName(const FName InOpName) const ;
+	UE_API FName GetParentOpByName(const FName InOpName) const ;
 
 	/** Get the index of an op.
 	 * @param InOpName the name of the op to get the index of
 	 * @return the integer index of the op in the stack or -1 if op not found. */
 	UFUNCTION(BlueprintCallable, Category="IK Retarget")
-	int32 GetIndexOfOpByName(const FName InOpName) const;
+	UE_API int32 GetIndexOfOpByName(const FName InOpName) const;
 
 	/** Automatically add basic retargeting operations
 	 * Adds these ops in the following order: Pelvis Motion, FK Chains, IK Chains, IK Solve and Root Motion
 	 * If any of these ops are already present, they will not be re-added.
 	 * @param bRunInitialSetup: if true, will run the initial setup routine on each op after they are all added to the stack. */
 	UFUNCTION(BlueprintCallable, Category="IK Retarget", meta = (BlueprintThreadSafe))
-	void AddDefaultOps() const;
+	UE_API void AddDefaultOps() const;
 
 	/** Force the op to run the initial setup. This is normally run when an Op is added through the editor UI.
 	 * NOTE: ops may not have custom initial setup routines; the exact behavior is op-dependent
 	 * @param InOpIndex the index of the op to run setup on */
 	UFUNCTION(BlueprintCallable, Category="IK Retarget", meta = (BlueprintThreadSafe))
-	void RunOpInitialSetup(const int32 InOpIndex) const;
+	UE_API void RunOpInitialSetup(const int32 InOpIndex) const;
 
 	// Tell each op to reset any settings associated with this chain
 	UFUNCTION(BlueprintCallable, Category="IK Retarget")
-	void ResetChainSettingsInAllOps(const FName InTargetChainName) const;
+	UE_API void ResetChainSettingsInAllOps(const FName InTargetChainName) const;
 
 	/** Get the number of Ops in the stack.
 	 * @return int, the number of ops */
 	UFUNCTION(BlueprintCallable, Category="IK Retarget")
-	int32 GetNumRetargetOps() const;
+	UE_API int32 GetNumRetargetOps() const;
 
 	/** Move the retarget op at the given index to the target index.
 	 * NOTE: due to constraints on execution order, the actual index may differ from what is requested
 	 * @InOpToMoveIndex the index of the op to be moved
 	 * @InTargetIndex the index where the op should be moved to */ 
 	UFUNCTION(BlueprintCallable, BlueprintPure=false, Category="IK Retarget")
-	bool MoveRetargetOpInStack(int32 InOpToMoveIndex, int32 InTargetIndex) const;
+	UE_API bool MoveRetargetOpInStack(int32 InOpToMoveIndex, int32 InTargetIndex) const;
 
 	/** Toggle an op on/off.
 	 * @param InRetargetOpIndex the index of the op to modify
 	 * @param bIsEnabled if true, turns the op On, else Off
 	 * @return true if op was found at index */ 
 	UFUNCTION(BlueprintCallable, BlueprintPure=false, Category="IK Retarget")
-	bool SetRetargetOpEnabled(int32 InRetargetOpIndex, bool bIsEnabled) const;
+	UE_API bool SetRetargetOpEnabled(int32 InRetargetOpIndex, bool bIsEnabled) const;
 
 	/** Get enabled status of the given Op.
 	 * @param InRetargetOpIndex the index of the op to get the enabled state for
 	 * @return true if op is enabled, false if op is disabled or not found */
 	UFUNCTION(BlueprintCallable, Category="IK Retarget")
-	bool GetRetargetOpEnabled(int32 InRetargetOpIndex) const;
+	UE_API bool GetRetargetOpEnabled(int32 InRetargetOpIndex) const;
 	
 	/** Get a controller to get/set the settings for a given op in the stack
 	 * NOTE: the returned UIKRetargetOpControllerBase* should be cast to the type specific to your op to get full functionality for that op.
 	 * @param InOpIndex: the index of the op in the stack to return a controller for
 	 * @return a pointer to the base op controller UClass or null if the op was not found */
 	UFUNCTION(BlueprintCallable, Category="IK Rig")
-	UIKRetargetOpControllerBase* GetOpController(int32 InOpIndex);
+	UE_API UIKRetargetOpControllerBase* GetOpController(int32 InOpIndex);
 
 	//
 	// RETARGET OPS C++ ONLY API
 	//
 	
 	// add a new retarget op to the stack
-	int32 AddRetargetOp(const UScriptStruct* InRetargetOpType, const FName InParentOpName = NAME_None) const;
+	UE_API int32 AddRetargetOp(const UScriptStruct* InRetargetOpType, const FName InParentOpName = NAME_None) const;
 
 	// Gather up all the ops that are children of this op
-	TArray<int32> GetChildOpIndices(const int32 InOpIndex) const;
+	UE_API TArray<int32> GetChildOpIndices(const int32 InOpIndex) const;
 
 	// Returns true if this op type can contain child ops (ops forced to be ordered before the parent)
-	bool GetCanOpHaveChildren(const int32 InOpIndex) const;
+	UE_API bool GetCanOpHaveChildren(const int32 InOpIndex) const;
 
 	// Get the index of the parent op or INDEX_NONE if it's root level
-	int32 GetParentOpIndex(const int32 InOpIndex) const;
+	UE_API int32 GetParentOpIndex(const int32 InOpIndex) const;
 
 	// Return NameToMakeUnique with a numbered suffix that makes it unique in the stack
-	FName GetUniqueOpName(const FName NameToMakeUnique, int32 OpIndexToIgnore) const;
+	UE_API FName GetUniqueOpName(const FName NameToMakeUnique, int32 OpIndexToIgnore) const;
 
 	// Get access to the first retarget operation of the given type
 	template <typename T>
@@ -217,19 +219,19 @@ public:
 	}
 
 	// Get access to the given retarget operation by index.
-	FIKRetargetOpBase* GetRetargetOpByIndex(const int32 InOpIndex) const;
+	UE_API FIKRetargetOpBase* GetRetargetOpByIndex(const int32 InOpIndex) const;
 
 	// Get access to the given retarget operation by name.
-	FIKRetargetOpBase* GetRetargetOpByName(const FName InOpName) const;
+	UE_API FIKRetargetOpBase* GetRetargetOpByName(const FName InOpName) const;
 
 	// Get access to the given retarget operation.
-	FInstancedStruct* GetRetargetOpStructAtIndex(int32 Index) const;
+	UE_API FInstancedStruct* GetRetargetOpStructAtIndex(int32 Index) const;
 
 	// Get the index of a given retarget operation. 
-	int32 GetIndexOfRetargetOp(FIKRetargetOpBase* RetargetOp) const;
+	UE_API int32 GetIndexOfRetargetOp(FIKRetargetOpBase* RetargetOp) const;
 
 	// A callback whenever the property of an op is modified
-	void OnOpPropertyChanged(const FName& InOpName, const FPropertyChangedEvent& InPropertyChangedEvent) const;
+	UE_API void OnOpPropertyChanged(const FName& InOpName, const FPropertyChangedEvent& InPropertyChangedEvent) const;
 
 	//
 	// GENERAL C++ ONLY API
@@ -238,19 +240,19 @@ public:
 	// Ensures all internal data is compatible with assigned meshes and ready to edit.
 	// - Removes bones from retarget poses that are no longer in skeleton
 	// - Removes chain settings for chains that are no longer in target IK Rig
-	void CleanAsset() const;
+	UE_API void CleanAsset() const;
 	
 	// Get either source or target IK Rig 
-	UIKRigDefinition* GetIKRigWriteable(const ERetargetSourceOrTarget SourceOrTarget) const;
+	UE_API UIKRigDefinition* GetIKRigWriteable(const ERetargetSourceOrTarget SourceOrTarget) const;
 	
 	// Get if we've already asked to fix the root height for the given skeletal mesh 
-	bool GetAskedToFixRootHeightForMesh(USkeletalMesh* Mesh) const;
+	UE_API bool GetAskedToFixRootHeightForMesh(USkeletalMesh* Mesh) const;
 	
 	// Set if we've asked to fix the root height for the given skeletal mesh 
-	void SetAskedToFixRootHeightForMesh(USkeletalMesh* Mesh, bool InAsked) const;
+	UE_API void SetAskedToFixRootHeightForMesh(USkeletalMesh* Mesh, bool InAsked) const;
 	
 	// Get name of the Root bone used for retargeting. 
-	FName GetPelvisBone(const ERetargetSourceOrTarget SourceOrTarget) const;
+	UE_API FName GetPelvisBone(const ERetargetSourceOrTarget SourceOrTarget) const;
 
 	//
 	// RETARGET CHAIN MAPPING PUBLIC/SCRIPTING API
@@ -261,27 +263,27 @@ public:
 	 * @param bForceRemap if false, will only remap those mappings that unset/None
 	 * @param InOpName if specified, will auto-map only the chain mapping in the specified op, otherwise applies to all ops*/
 	UFUNCTION(BlueprintCallable, BlueprintPure=false, Category="IK Retarget", meta = (BlueprintThreadSafe))
-	void AutoMapChains(const EAutoMapChainType AutoMapType, const bool bForceRemap, const FName InOpName=NAME_None) const;
+	UE_API void AutoMapChains(const EAutoMapChainType AutoMapType, const bool bForceRemap, const FName InOpName=NAME_None) const;
 
 	/** Assign a source chain to the given target chain. Animation will be copied from the source to the target.
 	 * @param InSourceChainName the name of the source retarget chain to assign
 	 * @param InTargetChainName the name of the target chain to map
 	 * @param InOpName if specified, will set the source chain only for that op (otherwise applies to all ops with a chain mapping)*/
 	UFUNCTION(BlueprintCallable, BlueprintPure=false, Category="IK Retarget", meta = (BlueprintThreadSafe))
-	bool SetSourceChain(FName InSourceChainName, FName InTargetChainName, const FName InOpName=NAME_None) const;
+	UE_API bool SetSourceChain(FName InSourceChainName, FName InTargetChainName, const FName InOpName=NAME_None) const;
 
 	/** Get the name of the source chain mapped to a given target chain (the chain animation is copied FROM).
 	 * @param InTargetChainName the name of the target retarget chain to get a mapping for
 	 * @param InOpName if specified, will get the source chain for that op (otherwise gets the first chain mapping in the op stack)
 	 * @return the name of the source chain mapped to the given target chain (or None if not mapped) */
 	UFUNCTION(BlueprintCallable, Category="IK Retarget", meta = (BlueprintThreadSafe))
-	FName GetSourceChain(const FName& InTargetChainName, const FName InOpName=NAME_None) const;
+	UE_API FName GetSourceChain(const FName& InTargetChainName, const FName InOpName=NAME_None) const;
 
 	/** Reset the settings for the given chain in the given op back to the defaults.
 	 * @param InTargetChainName the name of the chain
 	 * @param InOpName the name of the op containing the chain settings */
 	UFUNCTION(BlueprintCallable, Category="IK Retarget", meta = (BlueprintThreadSafe))
-	void ResetChainSettingsToDefault(const FName InTargetChainName, const FName InOpName) const;
+	UE_API void ResetChainSettingsToDefault(const FName InTargetChainName, const FName InOpName) const;
 
 	/** Get the target IK Rig associated with the given Op. May be different than the default IK Rig.
 	 * NOTE: all ops use the global SOURCE IK rig, but each op may use its own custom TARGET IK Rig.
@@ -290,7 +292,7 @@ public:
 	 * @param InOpName the name of the retarget op
 	 * @return the IK Rig assset associated with this op */
 	UFUNCTION(BlueprintCallable, Category="IK Retarget", meta = (BlueprintThreadSafe))
-	const UIKRigDefinition* GetTargetIKRigForOp(const FName InOpName) const;
+	UE_API const UIKRigDefinition* GetTargetIKRigForOp(const FName InOpName) const;
 
 	//
 	// RETARGET CHAIN MAPPING C++ ONLY API
@@ -298,25 +300,25 @@ public:
 
 	/** Get access to the mapping between source/target chains for the particular op (may be null)
 	 * If no op name is specified, it returns the first chain mapping it finds*/
-	const FRetargetChainMapping* GetChainMapping(const FName InOpName = NAME_None) const;
+	UE_API const FRetargetChainMapping* GetChainMapping(const FName InOpName = NAME_None) const;
 
 	// Clean all chain mappings in all ops, or if InOpName is specified, just that op
-	void CleanChainMaps(const FName InOpName=NAME_None) const;
+	UE_API void CleanChainMaps(const FName InOpName=NAME_None) const;
 	
 	// Ask if the given op has settings for the given chain that are not at the default
-	bool AreChainSettingsAtDefault(const FName InTargetChainName, const FName InOpName) const;
+	UE_API bool AreChainSettingsAtDefault(const FName InTargetChainName, const FName InOpName) const;
 	
 	// Get whether the given chain's IK goal is connected to a solver 
-	bool IsChainGoalConnectedToASolver(const FName& GoalName) const;
+	UE_API bool IsChainGoalConnectedToASolver(const FName& GoalName) const;
 	
 	// Call this when IK Rig chain is added or removed. 
-	void HandleRetargetChainAdded(UIKRigDefinition* IKRig) const;
+	UE_API void HandleRetargetChainAdded(UIKRigDefinition* IKRig) const;
 	
 	// Call this when IK Rig chain is renamed. Retains existing mappings using the new name 
-	void HandleRetargetChainRenamed(UIKRigDefinition* InIKRig, FName OldChainName, FName NewChainName) const;
+	UE_API void HandleRetargetChainRenamed(UIKRigDefinition* InIKRig, FName OldChainName, FName NewChainName) const;
 	
 	// Call this when IK Rig chain is removed. 
-	void HandleRetargetChainRemoved(UIKRigDefinition* IKRig, const FName& InChainRemoved) const;
+	UE_API void HandleRetargetChainRemoved(UIKRigDefinition* IKRig, const FName& InChainRemoved) const;
 
 	//
 	// RETARGET POSE PUBLIC/SCRIPTING API
@@ -327,14 +329,14 @@ public:
 	 * @param SourceOrTarget Specifies whether the pose applies to the source or target skeleton
 	 * @return The name of the newly created retarget pose */
 	UFUNCTION(BlueprintCallable, BlueprintPure=false, Category="IK Retarget", meta =(BlueprintThreadSafe))
-	FName CreateRetargetPose(const FName& NewPoseName, const ERetargetSourceOrTarget SourceOrTarget) const;
+	UE_API FName CreateRetargetPose(const FName& NewPoseName, const ERetargetSourceOrTarget SourceOrTarget) const;
 
 	/** Remove a retarget pose.
 	 * @param PoseToRemove The name of the retarget pose to remove
 	 * @param SourceOrTarget Specifies whether the pose belongs to the source or target skeleton
 	 * @return True if the pose was found and successfully removed, false otherwise */
 	UFUNCTION(BlueprintCallable, BlueprintPure=false, Category="IK Retarget", meta =(BlueprintThreadSafe))
-	bool RemoveRetargetPose(const FName& PoseToRemove, const ERetargetSourceOrTarget SourceOrTarget) const;
+	UE_API bool RemoveRetargetPose(const FName& PoseToRemove, const ERetargetSourceOrTarget SourceOrTarget) const;
 
 	/** Duplicate a retarget pose.
 	 * @param PoseToDuplicate The name of the retarget pose to duplicate
@@ -342,7 +344,7 @@ public:
 	 * @param SourceOrTarget Specifies whether the pose belongs to the source or target skeleton
 	 * @return The name of the new duplicated pose, or NAME_None if PoseToDuplicate is not found */
 	UFUNCTION(BlueprintCallable, BlueprintPure=false, Category="IK Retarget", meta =(BlueprintThreadSafe))
-	FName DuplicateRetargetPose(const FName PoseToDuplicate, const FName NewName, const ERetargetSourceOrTarget SourceOrTarget) const;
+	UE_API FName DuplicateRetargetPose(const FName PoseToDuplicate, const FName NewName, const ERetargetSourceOrTarget SourceOrTarget) const;
 
 	/** Rename current retarget pose.
 	 * @param OldPoseName The current name of the retarget pose to rename
@@ -350,98 +352,98 @@ public:
 	 * @param SourceOrTarget Specifies whether the pose belongs to the source or target skeleton
 	 * @return True if the pose was found and successfully renamed, false otherwise */
 	UFUNCTION(BlueprintCallable, BlueprintPure=false, Category="IK Retarget", meta =(BlueprintThreadSafe))
-	bool RenameRetargetPose(const FName OldPoseName, const FName NewPoseName, const ERetargetSourceOrTarget SourceOrTarget) const;
+	UE_API bool RenameRetargetPose(const FName OldPoseName, const FName NewPoseName, const ERetargetSourceOrTarget SourceOrTarget) const;
 
 	/** Change which retarget pose is used by the retargeter at runtime.
 	 * @param CurrentPose The name of the retarget pose to set as current
 	 * @param SourceOrTarget Specifies whether the pose applies to the source or target skeleton
 	 * @return True if the current pose was successfully set, false otherwise */
 	UFUNCTION(BlueprintCallable, BlueprintPure=false, Category="IK Retarget", meta =(BlueprintThreadSafe))
-	bool SetCurrentRetargetPose(FName CurrentPose, const ERetargetSourceOrTarget SourceOrTarget) const;
+	UE_API bool SetCurrentRetargetPose(FName CurrentPose, const ERetargetSourceOrTarget SourceOrTarget) const;
 
 	/** Get the name of the current retarget pose.
 	 * @param SourceOrTarget Specifies whether to retrieve the pose from the source or target skeleton
 	 * @return The name of the current retarget pose */
 	UFUNCTION(BlueprintCallable, Category="IK Retarget", meta =(BlueprintThreadSafe))
-	FName GetCurrentRetargetPoseName(const ERetargetSourceOrTarget SourceOrTarget) const;
+	UE_API FName GetCurrentRetargetPoseName(const ERetargetSourceOrTarget SourceOrTarget) const;
 
 	/** Get access to array of retarget poses.
 	 * @param SourceOrTarget Specifies whether to retrieve poses from the source or target skeleton
 	 * @return A reference to the map of retarget pose names to their pose data */
 	UFUNCTION(BlueprintCallable, Category="IK Retarget", meta =(BlueprintThreadSafe))
-	TMap<FName, FIKRetargetPose>& GetRetargetPoses(const ERetargetSourceOrTarget SourceOrTarget) const;
+	UE_API TMap<FName, FIKRetargetPose>& GetRetargetPoses(const ERetargetSourceOrTarget SourceOrTarget) const;
 
 	/** Get the current retarget pose.
 	 * @param SourceOrTarget Specifies whether to retrieve the pose from the source or target skeleton
 	 * @return A reference to the current retarget pose data */
 	UFUNCTION(BlueprintCallable, Category="IK Retarget", meta =(BlueprintThreadSafe))
-	FIKRetargetPose& GetCurrentRetargetPose(const ERetargetSourceOrTarget SourceOrTarget) const;
+	UE_API FIKRetargetPose& GetCurrentRetargetPose(const ERetargetSourceOrTarget SourceOrTarget) const;
 
 	/** Reset a retarget pose for the specified bones.
 	 * @param PoseToReset The name of the retarget pose to reset
 	 * @param BonesToReset The array of bone names to reset; if empty, resets all bones to the reference pose
 	 * @param SourceOrTarget Specifies whether the pose belongs to the source or target skeleton */
 	UFUNCTION(BlueprintCallable, BlueprintPure=false, Category="IK Retarget", meta =(BlueprintThreadSafe))
-	void ResetRetargetPose(const FName& PoseToReset, const TArray<FName>& BonesToReset, const ERetargetSourceOrTarget SourceOrTarget) const;
+	UE_API void ResetRetargetPose(const FName& PoseToReset, const TArray<FName>& BonesToReset, const ERetargetSourceOrTarget SourceOrTarget) const;
 
 	/** Set a delta rotation for a given bone in the current retarget pose.
 	 * @param BoneName The name of the bone to apply the rotation offset to
 	 * @param RotationOffset The quaternion representing the rotation offset to apply
 	 * @param SkeletonMode Specifies whether the bone belongs to the source or target skeleton */
 	UFUNCTION(BlueprintCallable, BlueprintPure=false, Category="IK Retarget", meta =(BlueprintThreadSafe))
-	void SetRotationOffsetForRetargetPoseBone(const FName& BoneName, const FQuat& RotationOffset, const ERetargetSourceOrTarget SkeletonMode) const;
+	UE_API void SetRotationOffsetForRetargetPoseBone(const FName& BoneName, const FQuat& RotationOffset, const ERetargetSourceOrTarget SkeletonMode) const;
 
 	/** Get a delta rotation for a given bone in the current retarget pose.
 	 * @param BoneName The name of the bone to retrieve the rotation offset for
 	 * @param SourceOrTarget Specifies whether the bone belongs to the source or target skeleton
 	 * @return The quaternion representing the rotation offset for the specified bone */
 	UFUNCTION(BlueprintCallable, BlueprintPure=false, Category="IK Retarget", meta =(BlueprintThreadSafe))
-	FQuat GetRotationOffsetForRetargetPoseBone(const FName& BoneName, const ERetargetSourceOrTarget SourceOrTarget) const;
+	UE_API FQuat GetRotationOffsetForRetargetPoseBone(const FName& BoneName, const ERetargetSourceOrTarget SourceOrTarget) const;
 
 	/** Set the translation offset on the retarget pelvis bone for the current retarget pose.
 	 * @param TranslationOffset The vector representing the translation offset to apply to the pelvis bone
 	 * @param SourceOrTarget Specifies whether the pelvis bone belongs to the source or target skeleton */
 	UFUNCTION(BlueprintCallable, BlueprintPure=false, Category="IK Retarget", meta =(BlueprintThreadSafe))
-	void SetRootOffsetInRetargetPose(const FVector& TranslationOffset, const ERetargetSourceOrTarget SourceOrTarget) const;
+	UE_API void SetRootOffsetInRetargetPose(const FVector& TranslationOffset, const ERetargetSourceOrTarget SourceOrTarget) const;
 
 	/** Get the translation offset on the retarget pelvis bone for the current retarget pose.
 	 * @param SourceOrTarget Specifies whether the pelvis bone belongs to the source or target skeleton
 	 * @return The vector representing the translation offset of the pelvis bone */
 	UFUNCTION(BlueprintCallable, Category="IK Retarget", meta = (BlueprintThreadSafe))
-	FVector GetRootOffsetInRetargetPose(const ERetargetSourceOrTarget SourceOrTarget) const;
+	UE_API FVector GetRootOffsetInRetargetPose(const ERetargetSourceOrTarget SourceOrTarget) const;
 
 	/** Automatically align all bones in mapped chains and store in the current retarget pose.
 	 * @param SourceOrTarget Specifies whether to align bones in the source or target skeleton
 	 * @param Method The method to use for automatic alignment (defaults to ChainToChain) */
 	UFUNCTION(BlueprintCallable, Category="IK Retarget")
-	void AutoAlignAllBones(ERetargetSourceOrTarget SourceOrTarget, const ERetargetAutoAlignMethod Method = ERetargetAutoAlignMethod::ChainToChain) const;
+	UE_API void AutoAlignAllBones(ERetargetSourceOrTarget SourceOrTarget, const ERetargetAutoAlignMethod Method = ERetargetAutoAlignMethod::ChainToChain) const;
 
 	/** Automatically align an array of bones and store in the current retarget pose.
 	 * @param BonesToAlign The array of bone names to align; bones not in mapped chains are ignored
 	 * @param Method The method to use for automatic alignment
 	 * @param SourceOrTarget Specifies whether to align bones in the source or target skeleton */
 	UFUNCTION(BlueprintCallable, Category="IK Retarget")
-	void AutoAlignBones(const TArray<FName>& BonesToAlign, const ERetargetAutoAlignMethod Method, ERetargetSourceOrTarget SourceOrTarget) const;
+	UE_API void AutoAlignBones(const TArray<FName>& BonesToAlign, const ERetargetAutoAlignMethod Method, ERetargetSourceOrTarget SourceOrTarget) const;
 
 	/** Moves the entire skeleton vertically until the specified bone is the same height off the ground as in the reference pose.
 	 * @param ReferenceBone The name of the bone to use as a reference for ground snapping
 	 * @param SourceOrTarget Specifies whether the skeleton belongs to the source or target */
 	UFUNCTION(BlueprintCallable, Category="IK Retarget")
-	void SnapBoneToGround(FName ReferenceBone, ERetargetSourceOrTarget SourceOrTarget);
+	UE_API void SnapBoneToGround(FName ReferenceBone, ERetargetSourceOrTarget SourceOrTarget);
 
 	//
 	// RETARGET POSE C++ ONLY API
 	//
 	
 	// Add a numbered suffix to the given pose name to make it unique. 
-	FName MakePoseNameUnique(
+	UE_API FName MakePoseNameUnique(
 		const FString& PoseName,
 		const ERetargetSourceOrTarget SourceOrTarget) const;
 	
 private:
 	
 	// Remove bones from retarget poses that are no longer in skeleton 
-	void CleanPoseList(const ERetargetSourceOrTarget SourceOrTarget) const;
+	UE_API void CleanPoseList(const ERetargetSourceOrTarget SourceOrTarget) const;
 	
 	//
 	// BEGIN DEPRECATED PUBLIC API
@@ -456,41 +458,41 @@ public:
 	// Get a copy of the retarget root settings for this asset.
 	UE_DEPRECATED(5.6, "Root settings are now accessed through the Pelvis Motion Op controller.")
 	UFUNCTION(BlueprintCallable, BlueprintPure=false, Category="IK Retarget", meta = (BlueprintThreadSafe))
-	FTargetRootSettings GetRootSettings() const;
+	UE_API FTargetRootSettings GetRootSettings() const;
 	
 	// Set the retarget root settings for this asset.
 	UE_DEPRECATED(5.6, "Root settings are now accessed through the Pelvis Motion Op controller.")
 	UFUNCTION(BlueprintCallable, BlueprintPure=false, Category="IK Retarget", meta = (BlueprintThreadSafe))
-	void SetRootSettings(const FTargetRootSettings& RootSettings) const;
+	UE_API void SetRootSettings(const FTargetRootSettings& RootSettings) const;
 
 	// Get a copy of the global settings for this asset.
 	UE_DEPRECATED(5.6, "Global settings are now accessed through Op controllers that perform the same duties.")
 	UFUNCTION(BlueprintCallable, BlueprintPure=false, Category="IK Retarget", meta = (BlueprintThreadSafe))
-	FRetargetGlobalSettings GetGlobalSettings() const;
+	UE_API FRetargetGlobalSettings GetGlobalSettings() const;
 
 	// Get a copy of the global settings for this asset.
 	UE_DEPRECATED(5.6, "Global settings are now accessed through Op controllers that perform the same duties.")
 	UFUNCTION(BlueprintCallable, BlueprintPure=false, Category="IK Retarget", meta = (BlueprintThreadSafe))
-	void SetGlobalSettings(const FRetargetGlobalSettings& GlobalSettings) const;
+	UE_API void SetGlobalSettings(const FRetargetGlobalSettings& GlobalSettings) const;
 
 	// Get a copy of the settings for the target chain by name.
 	UE_DEPRECATED(5.6, "Access to chain settings must go through an Op controller now.")
 	UFUNCTION(BlueprintCallable, BlueprintPure=false, Category="IK Retarget", meta = (BlueprintThreadSafe))
-	FTargetChainSettings GetRetargetChainSettings(const FName& TargetChainName) const;
+	UE_API FTargetChainSettings GetRetargetChainSettings(const FName& TargetChainName) const;
 
 	// Set the settings for the target chain by name. Returns true if the chain settings were applied, false otherwise.
 	UE_DEPRECATED(5.6, "Access to chain settings must go through an Op controller now.")
 	UFUNCTION(BlueprintCallable, BlueprintPure=false, Category="IK Retarget", meta = (BlueprintThreadSafe))
-	bool SetRetargetChainSettings(const FName& TargetChainName, const FTargetChainSettings& Settings) const;
+	UE_API bool SetRetargetChainSettings(const FName& TargetChainName, const FTargetChainSettings& Settings) const;
 
 	// Get read-only access to the list of settings for each target chain
 	UE_DEPRECATED(5.6, "Access to chain settings must go through an Op controller now.")
 	UFUNCTION(BlueprintCallable, Category="IK Retarget", meta = (BlueprintThreadSafe))
-	const TArray<URetargetChainSettings*>& GetAllChainSettings() const;
+	UE_API const TArray<URetargetChainSettings*>& GetAllChainSettings() const;
 
 	// convenience to get chain settings UObject by name
 	UE_DEPRECATED(5.6, "Access to chain settings must go through an Op controller now.")
-	URetargetChainSettings* GetChainSettings(const FName& TargetChainName) const;
+	UE_API URetargetChainSettings* GetChainSettings(const FName& TargetChainName) const;
 
 	PRAGMA_ENABLE_DEPRECATION_WARNINGS
 	//
@@ -543,7 +545,7 @@ public:
 	FOnOpStackModified& OnOpStackModified(){ return OpStackModified; };
 
 	// get the singleton struct viewer for editing structs belonging to the asset this controller controls
-	UIKRigStructViewer* GetStructViewer();
+	UE_API UIKRigStructViewer* GetStructViewer();
 	
 	friend class UIKRetargeter;
 	friend struct FScopedReinitializeIKRetargeter;
@@ -582,3 +584,4 @@ struct FScopedReinitializeIKRetargeter
 	ERetargetRefreshMode RefreshMode = ERetargetRefreshMode::ProcessorOnly;
 };
 
+#undef UE_API
diff --git a/Engine/Plugins/Animation/IKRig/Source/IKRigEditor/Public/RetargetEditor/IKRetargeterThumbnailRenderer.h b/Engine/Plugins/Animation/IKRig/Source/IKRigEditor/Public/RetargetEditor/IKRetargeterThumbnailRenderer.h
--- a/Engine/Plugins/Animation/IKRig/Source/IKRigEditor/Public/RetargetEditor/IKRetargeterThumbnailRenderer.h
+++ b/Engine/Plugins/Animation/IKRig/Source/IKRigEditor/Public/RetargetEditor/IKRetargeterThumbnailRenderer.h
@@ -6,23 +6,25 @@
 #include "ThumbnailRendering/SkeletalMeshThumbnailRenderer.h"
 #include "IKRetargeterThumbnailRenderer.generated.h"
 
+#define UE_API IKRIGEDITOR_API
+
 class FCanvas;
 class FRenderTarget;
 class ASkeletalMeshActor;
 enum class ERetargetSourceOrTarget : uint8;
 
-class IKRIGEDITOR_API FIKRetargeterThumbnailScene : public FThumbnailPreviewScene
+class FIKRetargeterThumbnailScene : public FThumbnailPreviewScene
 {
 public:
 	
-	FIKRetargeterThumbnailScene();
+	UE_API FIKRetargeterThumbnailScene();
 
 	// sets the retargeter to use in the next CreateView()
-	void SetSkeletalMeshes(USkeletalMesh* SourceMesh, USkeletalMesh* TargetMesh) const;
+	UE_API void SetSkeletalMeshes(USkeletalMesh* SourceMesh, USkeletalMesh* TargetMesh) const;
 
 protected:
 	// FThumbnailPreviewScene implementation
-	virtual void GetViewMatrixParameters(const float InFOVDegrees, FVector& OutOrigin, float& OutOrbitPitch, float& OutOrbitYaw, float& OutOrbitZoom) const override;
+	UE_API virtual void GetViewMatrixParameters(const float InFOVDegrees, FVector& OutOrigin, float& OutOrbitPitch, float& OutOrbitYaw, float& OutOrbitZoom) const override;
 
 private:
 	// the skeletal mesh actors used to display all skeletal mesh thumbnails
@@ -54,3 +56,4 @@ protected:
 	TObjectInstanceThumbnailScene<FIKRetargeterThumbnailScene, 128> ThumbnailSceneCache;
 };
 
+#undef UE_API
diff --git a/Engine/Plugins/Animation/IKRig/Source/IKRigEditor/Public/RigEditor/IKRigController.h b/Engine/Plugins/Animation/IKRig/Source/IKRigEditor/Public/RigEditor/IKRigController.h
--- a/Engine/Plugins/Animation/IKRig/Source/IKRigEditor/Public/RigEditor/IKRigController.h
+++ b/Engine/Plugins/Animation/IKRig/Source/IKRigEditor/Public/RigEditor/IKRigController.h
@@ -8,6 +8,8 @@
 
 #include "IKRigController.generated.h"
 
+#define UE_API IKRIGEDITOR_API
+
 class UIKRigStructViewer;
 struct FAutoCharacterizeResults;
 struct FAutoCharacterizer;
@@ -36,14 +38,14 @@ struct FBoneChain;
 // of a running IK Rig processor instance.
 // The API here is split into public/scripting sections which are accessible from Blueprint/Python and sections that
 // are only relevant to editors in C++.
-UCLASS(BlueprintType, hidecategories = UObject)
-class IKRIGEDITOR_API UIKRigController : public UObject
+UCLASS(MinimalAPI, BlueprintType, hidecategories = UObject)
+class UIKRigController : public UObject
 {
 	GENERATED_BODY()
 
 public:
 
-	UIKRigController();
+	UE_API UIKRigController();
 
 	//
 	// GENERAL PUBLIC/SCRIPTING API
@@ -51,20 +53,20 @@ public:
 	
 	// Use this to get the controller for the given IKRig
 	UFUNCTION(BlueprintCallable, Category="IK Rig", meta = (DisplayName = "Get IK Rig Controller"))
-	static UIKRigController* GetController(const UIKRigDefinition* InIKRigDefinition);
+	static UE_API UIKRigController* GetController(const UIKRigDefinition* InIKRigDefinition);
 	
 	// Sets the preview mesh to use. Loads the hierarchy into the asset's IKRigSkeleton.
 	// Returns true if the mesh was able to be set. False if it was incompatible for any reason. 
 	UFUNCTION(BlueprintCallable, BlueprintPure=false, Category="IK Rig", meta = (DisplayName = "Set IK Rig Skeletal Mesh"))
-	bool SetSkeletalMesh(USkeletalMesh* SkeletalMesh) const;
+	UE_API bool SetSkeletalMesh(USkeletalMesh* SkeletalMesh) const;
 
 	// Get the skeletal mesh this asset is initialized with 
 	UFUNCTION(BlueprintCallable, Category="IK Rig", meta = (DisplayName = "Get IK Rig Skeletal Mesh"))
-	USkeletalMesh* GetSkeletalMesh() const;
+	UE_API USkeletalMesh* GetSkeletalMesh() const;
 
 	// Returns true if the provided skeletal mesh could be used with this IK Rig.
 	UFUNCTION(BlueprintCallable, Category="IK Rig")
-	bool IsSkeletalMeshCompatible(USkeletalMesh* SkeletalMeshToCheck) const;
+	UE_API bool IsSkeletalMeshCompatible(USkeletalMesh* SkeletalMeshToCheck) const;
 
 	//
 	// GENERAL C++ ONLY API
@@ -73,18 +75,18 @@ public:
 	// Get the asset this controller controls.
 	// Warning: Do not use for general editing of the data model.
 	// If you need to make modifications to the asset, do so through the API provided by this controller. 
-	UIKRigDefinition* GetAsset() const;
-	TObjectPtr<UIKRigDefinition>& GetAssetPtr();  
+	UE_API UIKRigDefinition* GetAsset() const;
+	UE_API TObjectPtr<UIKRigDefinition>& GetAssetPtr();  
 	
 	// force all currently connected processors to reinitialize using latest asset state 
-	void BroadcastNeedsReinitialized() const;
+	UE_API void BroadcastNeedsReinitialized() const;
 	
 	//
 	// SKELETON C++ ONLY API
 	//
 	
 	// Get read-access to the IKRig skeleton representation 
-	const FIKRigSkeleton& GetIKRigSkeleton() const;
+	UE_API const FIKRigSkeleton& GetIKRigSkeleton() const;
 
 	//
 	// SOLVERS PUBLIC/SCRIPTING API
@@ -93,74 +95,74 @@ public:
 	// Add a new solver of the given type to the bottom of the stack. Returns the stack index.
 	// The type here is the full package path of the UStruct type. For example: /Script/IKRig.FullBodyIKSolver
 	UFUNCTION(BlueprintCallable, BlueprintPure=false, Category="IK Rig")
-	int32 AddSolver(const FString InIKRigSolverType) const;
+	UE_API int32 AddSolver(const FString InIKRigSolverType) const;
 
 	// Remove the solver at the given stack index. 
 	UFUNCTION(BlueprintCallable, BlueprintPure=false, Category="IK Rig")
-	bool RemoveSolver(const int32 SolverIndex) const;
+	UE_API bool RemoveSolver(const int32 SolverIndex) const;
 
 	// Get the number of solvers in the stack. 
 	UFUNCTION(BlueprintCallable, Category="IK Rig")
-	int32 GetNumSolvers() const;
+	UE_API int32 GetNumSolvers() const;
 
 	// Move the solver at the given index to the target index. 
 	UFUNCTION(BlueprintCallable, BlueprintPure=false, Category="IK Rig")
-	bool MoveSolverInStack(int32 SolverToMoveIndex, int32 TargetSolverIndex) const;
+	UE_API bool MoveSolverInStack(int32 SolverToMoveIndex, int32 TargetSolverIndex) const;
 
 	// Set enabled/disabled status of the given solver. 
 	UFUNCTION(BlueprintCallable, BlueprintPure=false, Category="IK Rig")
-	bool SetSolverEnabled(int32 SolverIndex, bool bIsEnabled) const;
+	UE_API bool SetSolverEnabled(int32 SolverIndex, bool bIsEnabled) const;
 
 	// Get enabled status of the given solver. 
 	UFUNCTION(BlueprintCallable, Category="IK Rig")
-	bool GetSolverEnabled(int32 SolverIndex) const;
+	UE_API bool GetSolverEnabled(int32 SolverIndex) const;
 
 	// Set the root bone on a given solver. (not all solvers support root bones, checks CanSetStartBone() first) 
 	UFUNCTION(BlueprintCallable, BlueprintPure=false, Category="IK Rig")
-	bool SetStartBone(const FName RootBoneName, int32 SolverIndex) const;
+	UE_API bool SetStartBone(const FName RootBoneName, int32 SolverIndex) const;
 
 	// Get the name of the root bone on a given solver. (not all solvers support start bones, checks CanSetStartBone() first) 
 	UFUNCTION(BlueprintCallable, Category="IK Rig")
-	FName GetStartBone(int32 SolverIndex) const;
+	UE_API FName GetStartBone(int32 SolverIndex) const;
 
 	// Set the end bone on a given solver. (not all solvers require extra end bones, checks CanSetEndBone() first) 
 	UFUNCTION(BlueprintCallable, BlueprintPure=false, Category="IK Rig")
-	bool SetEndBone(const FName EndBoneName, int32 SolverIndex) const;
+	UE_API bool SetEndBone(const FName EndBoneName, int32 SolverIndex) const;
 
 	// Get the name of the end bone on a given solver. (not all solvers require extra end bones, checks CanSetEndBone() first) 
 	UFUNCTION(BlueprintCallable, Category="IK Rig")
-	FName GetEndBone(int32 SolverIndex) const;
+	UE_API FName GetEndBone(int32 SolverIndex) const;
 
 	// Get a controller to customize the settings for a given solver in the stack 
 	UFUNCTION(BlueprintCallable, Category="IK Rig")
-	UIKRigSolverControllerBase* GetSolverController(int32 SolverIndex);
+	UE_API UIKRigSolverControllerBase* GetSolverController(int32 SolverIndex);
 
 	// Get the index for a given solver in the stack
 	// The input here is a solver controller (as returned by GetSolverController())
 	UFUNCTION(BlueprintCallable, Category="IK Rig")
-	int32 GetIndexOfSolver(UIKRigSolverControllerBase* InController);
+	UE_API int32 GetIndexOfSolver(UIKRigSolverControllerBase* InController);
 
 	//
 	// SOLVERS C++ ONLY API
 	//
 
 	// Add a solver to the IK Rig asset's solver stack
-	int32 AddSolver(UScriptStruct* InIKRigSolverType) const;
+	UE_API int32 AddSolver(UScriptStruct* InIKRigSolverType) const;
 
 	// Get access to the given solver. 
-	FIKRigSolverBase* GetSolverAtIndex(int32 Index) const;
+	UE_API FIKRigSolverBase* GetSolverAtIndex(int32 Index) const;
 
 	// Get access to the instanced struct for a given solver. 
-	FInstancedStruct* GetSolverStructAtIndex(int32 Index) const;
+	UE_API FInstancedStruct* GetSolverStructAtIndex(int32 Index) const;
 
 	// Get access to the given solver. 
-	int32 GetIndexOfSolver(FIKRigSolverBase* Solver) const;
+	UE_API int32 GetIndexOfSolver(FIKRigSolverBase* Solver) const;
 	
 	// Get read-only access to the array of solvers. 
-	const TArray<FIKRigSolverBase*> GetSolverArray() const;
+	UE_API const TArray<FIKRigSolverBase*> GetSolverArray() const;
 
 	// Get unique label for a given solver. Returns dash separated index and name like so, "0 - SolverName". 
-	FString GetSolverUniqueName(int32 SolverIndex) const;
+	UE_API FString GetSolverUniqueName(int32 SolverIndex) const;
 
 	//
 	// GOALS PUBLIC/SCRIPTING API
@@ -168,89 +170,89 @@ public:
 	
 	// Add a new Goal associated with the given Bone. GoalName must be unique. Bones can have multiple Goals (rare). 
 	UFUNCTION(BlueprintCallable, BlueprintPure=false, Category="IK Rig")
-	FName AddNewGoal(const FName GoalName, const FName BoneName) const;
+	UE_API FName AddNewGoal(const FName GoalName, const FName BoneName) const;
 
 	// Remove the Goal by name. 
 	UFUNCTION(BlueprintCallable, BlueprintPure=false, Category="IK Rig")
-	bool RemoveGoal(const FName GoalName) const;
+	UE_API bool RemoveGoal(const FName GoalName) const;
 
 	// Rename a Goal. Returns new name, which may be different after being sanitized. Returns NAME_None if this fails.
 	UFUNCTION(BlueprintCallable, BlueprintPure=false, Category="IK Rig")
-	FName RenameGoal(const FName OldName, const FName PotentialNewName) const;
+	UE_API FName RenameGoal(const FName OldName, const FName PotentialNewName) const;
 
 	// The the Bone that the given Goal should be parented to / associated with. 
 	UFUNCTION(BlueprintCallable, BlueprintPure=false, Category="IK Rig")
-	bool SetGoalBone(const FName GoalName, const FName NewBoneName) const;
+	UE_API bool SetGoalBone(const FName GoalName, const FName NewBoneName) const;
 
 	// The the Bone associated with the given Goal. 
 	UFUNCTION(BlueprintCallable, Category="IK Rig")
-	FName GetBoneForGoal(const FName GoalName) const;
+	UE_API FName GetBoneForGoal(const FName GoalName) const;
 
 	// Get the Goal associated with the given Bone (may be NAME_None) 
 	UFUNCTION(BlueprintCallable, Category="IK Rig")
-	FName GetGoalNameForBone(const FName BoneName) const;
+	UE_API FName GetGoalNameForBone(const FName BoneName) const;
 
 	// Connect the given Goal to the given Solver. This creates an "Effector" with settings specific to this Solver.
 	UFUNCTION(BlueprintCallable, BlueprintPure=false, Category="IK Rig")
-	bool ConnectGoalToSolver(const FName GoalName, int32 SolverIndex) const;
+	UE_API bool ConnectGoalToSolver(const FName GoalName, int32 SolverIndex) const;
 
 	// Disconnect the given Goal from the given Solver. This removes the Effector that associates the Goal with the Solver.
 	UFUNCTION(BlueprintCallable, BlueprintPure=false, Category="IK Rig")
-	bool DisconnectGoalFromSolver(const FName GoalToRemove, int32 SolverIndex) const;
+	UE_API bool DisconnectGoalFromSolver(const FName GoalToRemove, int32 SolverIndex) const;
 
 	// Returns true if the given Goal is connected to the given Solver. False otherwise. 
 	UFUNCTION(BlueprintCallable, Category="IK Rig")
-	bool IsGoalConnectedToSolver(const FName GoalName, int32 SolverIndex) const;
+	UE_API bool IsGoalConnectedToSolver(const FName GoalName, int32 SolverIndex) const;
 
 	// Returns true if the given Goal is connected to ANY solver. False otherwise. 
 	UFUNCTION(BlueprintCallable, Category="IK Rig")
-	bool IsGoalConnectedToAnySolver(const FName GoalName) const;
+	UE_API bool IsGoalConnectedToAnySolver(const FName GoalName) const;
 
 	// Get read-write access to the Goal with the given name.
 	UFUNCTION(BlueprintCallable, Category="IK Rig")
-	UIKRigEffectorGoal* GetGoal(const FName GoalName) const;
+	UE_API UIKRigEffectorGoal* GetGoal(const FName GoalName) const;
 
 	// Get access to the list of Goals. 
 	UFUNCTION(BlueprintCallable, Category="IK Rig")
-	const TArray<UIKRigEffectorGoal*>& GetAllGoals() const;
+	UE_API const TArray<UIKRigEffectorGoal*>& GetAllGoals() const;
 
 	// Get the UObject for the settings associated with the given Goal in the given Solver.
 	// Solvers can define their own per-Goal settings depending on their needs. These are termed "Effectors".
 	UE_DEPRECATED(5.6, "In 5.6+ each solver provides it's own controller. Use GetSolverController(SolverIndex).GetGoalSettings(GoalName) instead.")
 	UFUNCTION(BlueprintCallable, Category="IK Rig")
-	UObject* GetGoalSettingsForSolver(const FName GoalName, int32 SolverIndex) const;
+	UE_API UObject* GetGoalSettingsForSolver(const FName GoalName, int32 SolverIndex) const;
 
 	// returns true if the goal has custom settings for the given solver
 	UFUNCTION(BlueprintCallable, Category="IK Rig")
-	bool DoesSolverHaveCustomGoalSettings(const FName GoalName, int32 SolverIndex) const;
+	UE_API bool DoesSolverHaveCustomGoalSettings(const FName GoalName, int32 SolverIndex) const;
 	
 	//
 	// GOALS C++ ONLY API
 	//
 	
 	// Get the index of the given Goal in the list of Goals. 
-	int32 GetGoalIndex(const FName InGoalName, const ENameCase CompareMethod = ENameCase::IgnoreCase) const;
+	UE_API int32 GetGoalIndex(const FName InGoalName, const ENameCase CompareMethod = ENameCase::IgnoreCase) const;
 
 	// Get the global-space transform of the given Goal. This may be set by the user in the editor, or at runtime. 
-	FTransform GetGoalCurrentTransform(const FName GoalName) const;
+	UE_API FTransform GetGoalCurrentTransform(const FName GoalName) const;
 
 	// Set the Goal to the given transform. 
-	void SetGoalCurrentTransform(const FName GoalName, const FTransform& Transform) const;
+	UE_API void SetGoalCurrentTransform(const FName GoalName, const FTransform& Transform) const;
 
 	// Reset all Goals back to their initial transforms. 
-	void ResetGoalTransforms() const;
+	UE_API void ResetGoalTransforms() const;
 
 	// Reset initial goal transform 
-	void ResetInitialGoalTransforms() const;
+	UE_API void ResetInitialGoalTransforms() const;
 
 	// Ensure that the given name adheres to required standards for Goal names (no special characters etc..)
-	static void SanitizeGoalName(FString& InOutName);
+	static UE_API void SanitizeGoalName(FString& InOutName);
 
 	// Add a suffix as needed to ensure the Goal name is unique 
-	FName GetUniqueGoalName(const FName NameToMakeUnique) const;
+	UE_API FName GetUniqueGoalName(const FName NameToMakeUnique) const;
 
 	// Modify a Goal for a transaction. Returns true if Goal found.
-	bool ModifyGoal(const FName GoalName) const;
+	UE_API bool ModifyGoal(const FName GoalName) const;
 
 	//
 	// BONES PUBLIC/SCRIPTING API
@@ -258,44 +260,44 @@ public:
 
 	// Include/exclude a bone from all the solvers. All bones are included by default. 
 	UFUNCTION(BlueprintCallable, BlueprintPure=false, Category="IK Rig")
-	bool SetBoneExcluded(const FName BoneName, const bool bExclude) const;
+	UE_API bool SetBoneExcluded(const FName BoneName, const bool bExclude) const;
 
 	// Returns true if the given Bone is excluded, false otherwise. 
 	UFUNCTION(BlueprintCallable, Category="IK Rig")
-	bool GetBoneExcluded(const FName BoneName) const;
+	UE_API bool GetBoneExcluded(const FName BoneName) const;
 	
 	// Add settings to the given Bone/Solver. Does nothing if Bone already has settings in this Solver.
 	UFUNCTION(BlueprintCallable, BlueprintPure=false, Category="IK Rig")
-	bool AddBoneSetting(const FName BoneName, int32 SolverIndex) const;
+	UE_API bool AddBoneSetting(const FName BoneName, int32 SolverIndex) const;
 
 	// Remove settings for the given Bone/Solver. Does nothing if Bone doesn't have setting in this Solver.
 	UFUNCTION(BlueprintCallable, BlueprintPure=false, Category="IK Rig")
-	bool RemoveBoneSetting(const FName BoneName, int32 SolverIndex) const;
+	UE_API bool RemoveBoneSetting(const FName BoneName, int32 SolverIndex) const;
 
 	// Get the global-space retarget pose transform of the given Bone.
 	UFUNCTION(BlueprintCallable, Category="IK Rig")
-	FTransform GetRefPoseTransformOfBone(const FName BoneName) const;
+	UE_API FTransform GetRefPoseTransformOfBone(const FName BoneName) const;
 
 	// Get the generic (Solver-specific) Bone settings UObject for this Bone in the given Solver.
 	UE_DEPRECATED(5.6, "In 5.6+ each solver provides it's own controller. Use GetSolverController(SolverIndex).GetBoneSettings(BoneName) instead.")
 	UFUNCTION(BlueprintCallable, BlueprintPure=false, Category="IK Rig")
-	UObject* GetBoneSettings(const FName BoneName, int32 SolverIndex) const;
+	UE_API UObject* GetBoneSettings(const FName BoneName, int32 SolverIndex) const;
 
 	//
 	// BONES C++ ONLY API
 	//
 
 	// Get the generic (Solver-specific) Bone settings UStruct for this Bone in the given Solver.
-	FIKRigBoneSettingsBase* GetSettingsForBone(const FName BoneName, int32 SolverIndex) const;
+	UE_API FIKRigBoneSettingsBase* GetSettingsForBone(const FName BoneName, int32 SolverIndex) const;
 
 	// Returns true if this Bone can have settings in the given Solver. 
-	bool CanAddBoneSetting(const FName BoneName, int32 SolverIndex) const;
+	UE_API bool CanAddBoneSetting(const FName BoneName, int32 SolverIndex) const;
 
 	// Returns true if settings for this Bone can be removed from the given Solver. 
-	bool CanRemoveBoneSetting(const FName BoneName, int32 SolverIndex) const;
+	UE_API bool CanRemoveBoneSetting(const FName BoneName, int32 SolverIndex) const;
 
 	// Returns true if the given Bone has any settings in any Solver. 
-	bool DoesBoneHaveSettings(const FName BoneName) const;
+	UE_API bool DoesBoneHaveSettings(const FName BoneName) const;
 
 	//
 	// RETARGETING PUBLIC/SCRIPTING API
@@ -304,102 +306,102 @@ public:
 	// Add a new chain with the given Chain and Bone names. Returns newly created chain name (uniquified).
 	// Note: only the ChainName is required here, all else can be set later.
 	UFUNCTION(BlueprintCallable, BlueprintPure=false, Category="IK Rig")
-	FName AddRetargetChain(const FName ChainName, const FName StartBoneName, const FName EndBoneName, const FName GoalName) const;
+	UE_API FName AddRetargetChain(const FName ChainName, const FName StartBoneName, const FName EndBoneName, const FName GoalName) const;
 
 	// Remove a Chain with the given name. Returns true if a Chain was removed. 
 	UFUNCTION(BlueprintCallable, BlueprintPure=false, Category="IK Rig")
-	bool RemoveRetargetChain(const FName ChainName) const;
+	UE_API bool RemoveRetargetChain(const FName ChainName) const;
 
 	// Renamed the given Chain. Returns the new name (same as old if unsuccessful). 
 	UFUNCTION(BlueprintCallable, BlueprintPure=false, Category="IK Rig")
-	FName RenameRetargetChain(const FName ChainName, const FName NewChainName) const;
+	UE_API FName RenameRetargetChain(const FName ChainName, const FName NewChainName) const;
 
 	// Set the Start Bone for the given Chain. Returns true if operation was successful. 
 	UFUNCTION(BlueprintCallable, BlueprintPure=false, Category="IK Rig")
-	bool SetRetargetChainStartBone(const FName ChainName, const FName StartBoneName) const;
+	UE_API bool SetRetargetChainStartBone(const FName ChainName, const FName StartBoneName) const;
 
 	// Set the End Bone for the given Chain. Returns true if operation was successful. 
 	UFUNCTION(BlueprintCallable, BlueprintPure=false, Category="IK Rig")
-	bool SetRetargetChainEndBone(const FName ChainName, const FName EndBoneName) const;
+	UE_API bool SetRetargetChainEndBone(const FName ChainName, const FName EndBoneName) const;
 
 	// Set the Goal for the given Chain. Returns true if operation was successful. 
 	UFUNCTION(BlueprintCallable, BlueprintPure=false, Category="IK Rig")
-	bool SetRetargetChainGoal(const FName ChainName, const FName GoalName) const;
+	UE_API bool SetRetargetChainGoal(const FName ChainName, const FName GoalName) const;
 
 	// Get the Goal name for the given Chain. 
 	UFUNCTION(BlueprintCallable, Category="IK Rig")
-	FName GetRetargetChainGoal(const FName ChainName) const;
+	UE_API FName GetRetargetChainGoal(const FName ChainName) const;
 
 	// Get the End Bone name for the given Chain. 
 	UFUNCTION(BlueprintCallable, Category="IK Rig")
-    FName GetRetargetChainStartBone(const FName ChainName) const;
+    UE_API FName GetRetargetChainStartBone(const FName ChainName) const;
 
 	// Get the Start Bone name for the given Chain. 
 	UFUNCTION(BlueprintCallable, Category="IK Rig")
-    FName GetRetargetChainEndBone(const FName ChainName) const;
+    UE_API FName GetRetargetChainEndBone(const FName ChainName) const;
 
 	// Get read-only access to the list of Chains. 
 	UFUNCTION(BlueprintCallable, Category="IK Rig")
-	const TArray<FBoneChain>& GetRetargetChains() const;
+	UE_API const TArray<FBoneChain>& GetRetargetChains() const;
 
 	// Set the Root Bone of the retargeting (can only be one). 
 	UFUNCTION(BlueprintCallable, BlueprintPure=false, Category="IK Rig", DisplayName="SetPelvisBone")
-	bool SetRetargetRoot(const FName RootBoneName) const;
+	UE_API bool SetRetargetRoot(const FName RootBoneName) const;
 
 	// Get the name of the Root Bone of the retargeting (can only be one). 
 	UFUNCTION(BlueprintCallable, Category="IK Rig", DisplayName="GetPelvisBone")
-	FName GetRetargetRoot() const;
+	UE_API FName GetRetargetRoot() const;
 
 	// Analyse the skeleton to see if it matches a known template and automatically generates all retarget chains and sets the pelvis
 	// Returns true if a matching skeletal template was found and the retarget definition for it was applied.
 	UFUNCTION(BlueprintCallable, Category="IK Rig")
-	bool ApplyAutoGeneratedRetargetDefinition();
+	UE_API bool ApplyAutoGeneratedRetargetDefinition();
 
 	// Analyse the skeleton to see if it matches a known template and automatically generates a full body IK setup
 	// Returns true if a matching skeletal template was found and the FBIK setup for it was applied.
 	UFUNCTION(BlueprintCallable, Category="IK Rig")
-	bool ApplyAutoFBIK();
+	UE_API bool ApplyAutoFBIK();
 	
 	//
 	// RETARGETING C++ ONLY API
 	//
 
 	// Replace the entire retarget definition (includes all bone chains and the pelvis setting)
-	void SetRetargetDefinition(const FRetargetDefinition& RetargetDefinition) const;
+	UE_API void SetRetargetDefinition(const FRetargetDefinition& RetargetDefinition) const;
 
 	// Auto generates a retarget definition and returns the results
-	void AutoGenerateRetargetDefinition(FAutoCharacterizeResults& Results) const;
+	UE_API void AutoGenerateRetargetDefinition(FAutoCharacterizeResults& Results) const;
 
 	// Auto generates an FBIK setup for the current skeletal mesh
-	void AutoGenerateFBIK(FAutoFBIKResults& Results) const;
+	UE_API void AutoGenerateFBIK(FAutoFBIKResults& Results) const;
 
 	// Get read-only access to characterizer
-	const FAutoCharacterizer& GetAutoCharacterizer() const;
+	UE_API const FAutoCharacterizer& GetAutoCharacterizer() const;
 	
 	// Add a Chain with the given BoneChain settings. Returns newly created chain name.
-	FName AddRetargetChainInternal(const FBoneChain& BoneChain) const;
+	UE_API FName AddRetargetChainInternal(const FBoneChain& BoneChain) const;
 	
 	// Get read-only access to a single retarget chain with the given name 
-	const FBoneChain* GetRetargetChainByName(const FName ChainName) const;
+	UE_API const FBoneChain* GetRetargetChainByName(const FName ChainName) const;
 	
 	// Get the name of the retarget chain that contains the given Bone. Returns NAME_None if Bone not in a Chain. 
-	FName GetRetargetChainFromBone(const FName BoneName, const FIKRigSkeleton* OptionalSkeleton) const;
+	UE_API FName GetRetargetChainFromBone(const FName BoneName, const FIKRigSkeleton* OptionalSkeleton) const;
 
 	// Get the name of the retarget chain that contains the given Goal. Returns NAME_None if Goal not in a Chain. 
-	FName GetRetargetChainFromGoal(const FName GoalName) const;
+	UE_API FName GetRetargetChainFromGoal(const FName GoalName) const;
 
 	// Sorts the Chains from Root to tip based on the Start Bone of each Chain. 
-	void SortRetargetChains() const;
+	UE_API void SortRetargetChains() const;
 
 	// Make unique name for a retargeting bone chain. Adds a numbered suffix to make it unique.
-	FName GetUniqueRetargetChainName(FName NameToMakeUnique) const;
+	UE_API FName GetUniqueRetargetChainName(FName NameToMakeUnique) const;
 
 	// Returns true if this is a valid chain. Produces array of bone indices between start and end (inclusive).
 	// Optionally provide a runtime skeleton from an IKRigProcessor to get indices for a running instance (otherwise uses stored hierarchy in asset)
-	bool ValidateChain( const FName ChainName, const FIKRigSkeleton* OptionalSkeleton, TSet<int32>& OutChainIndices) const;
+	UE_API bool ValidateChain( const FName ChainName, const FIKRigSkeleton* OptionalSkeleton, TSet<int32>& OutChainIndices) const;
 
 	// a singleton wrapper to display structs in the details panel
-	UIKRigStructViewer* GetStructViewer();
+	UE_API UIKRigStructViewer* GetStructViewer();
 
 private:
 
@@ -447,10 +449,10 @@ private:
 	TObjectPtr<UIKRigStructViewer> StructViewer = nullptr;
 
 	// Static global map of assets to their associated controller
-	static TMap<UIKRigDefinition*, UIKRigController*> Controllers;
+	static UE_API TMap<UIKRigDefinition*, UIKRigController*> Controllers;
 
 	// broadcast changes within the asset goals array
-	void BroadcastGoalsChange() const;
+	UE_API void BroadcastGoalsChange() const;
 	
 	friend class UIKRigDefinition;
 	friend struct FScopedReinitializeIKRig;
@@ -480,3 +482,5 @@ struct FScopedReinitializeIKRig
 
 	const UIKRigController* Controller;
 };
+
+#undef UE_API
diff --git a/Engine/Plugins/Animation/IKRig/Source/IKRigEditor/Public/RigEditor/IKRigDefinitionFactory.h b/Engine/Plugins/Animation/IKRig/Source/IKRigEditor/Public/RigEditor/IKRigDefinitionFactory.h
--- a/Engine/Plugins/Animation/IKRig/Source/IKRigEditor/Public/RigEditor/IKRigDefinitionFactory.h
+++ b/Engine/Plugins/Animation/IKRig/Source/IKRigEditor/Public/RigEditor/IKRigDefinitionFactory.h
@@ -7,30 +7,32 @@
 #include "Factories/Factory.h"
 #include "IKRigDefinitionFactory.generated.h"
 
+#define UE_API IKRIGEDITOR_API
+
 class SWindow;
 
-UCLASS(BlueprintType, hidecategories=Object)
-class IKRIGEDITOR_API UIKRigDefinitionFactory : public UFactory
+UCLASS(MinimalAPI, BlueprintType, hidecategories=Object)
+class UIKRigDefinitionFactory : public UFactory
 {
 	GENERATED_BODY()
 
 public:
 
-	UIKRigDefinitionFactory();
+	UE_API UIKRigDefinitionFactory();
 
 	// UFactory interface
-	virtual UObject* FactoryCreateNew(
+	UE_API virtual UObject* FactoryCreateNew(
         UClass* Class,
         UObject* InParent,
         FName InName,
         EObjectFlags InFlags,
         UObject* Context,
         FFeedbackContext* Warn) override;
-	virtual FText GetDisplayName() const override;
-	virtual uint32 GetMenuCategories() const override;
-	virtual FText GetToolTip() const override;
-	virtual FString GetDefaultNewAssetName() const override;
-	virtual bool ShouldShowInNewMenu() const override;
+	UE_API virtual FText GetDisplayName() const override;
+	UE_API virtual uint32 GetMenuCategories() const override;
+	UE_API virtual FText GetToolTip() const override;
+	UE_API virtual FString GetDefaultNewAssetName() const override;
+	UE_API virtual bool ShouldShowInNewMenu() const override;
 	// END UFactory interface
 
 	/**
@@ -39,5 +41,7 @@ public:
 	 * @param InAssetName The name of the new asset (ie "IK_MyNewRig")
 	 */
 	UFUNCTION(BlueprintCallable, Category = "IK Rig", DisplayName="Create New IK Rig Asset")
-	static UIKRigDefinition* CreateNewIKRigAsset(const FString& InPackagePath, const FString& InAssetName);
+	static UE_API UIKRigDefinition* CreateNewIKRigAsset(const FString& InPackagePath, const FString& InAssetName);
 };
+
+#undef UE_API
diff --git a/Engine/Plugins/Animation/IKRig/Source/IKRigEditor/Public/RigEditor/IKRigEditorController.h b/Engine/Plugins/Animation/IKRig/Source/IKRigEditor/Public/RigEditor/IKRigEditorController.h
--- a/Engine/Plugins/Animation/IKRig/Source/IKRigEditor/Public/RigEditor/IKRigEditorController.h
+++ b/Engine/Plugins/Animation/IKRig/Source/IKRigEditor/Public/RigEditor/IKRigEditorController.h
@@ -10,6 +10,8 @@
 
 #include "IKRigEditorController.generated.h"
 
+#define UE_API IKRIGEDITOR_API
+
 class SIKRigAssetBrowser;
 class SIKRigOutputLog;
 class UIKRigAnimInstance;
@@ -31,8 +33,8 @@ enum EIKRigSelectionType : int8
 	RetargetChains,
 };
 
-UCLASS(config = Engine, hidecategories = UObject)
-class IKRIGEDITOR_API UIKRigBoneDetails : public UObject
+UCLASS(MinimalAPI, config = Engine, hidecategories = UObject)
+class UIKRigBoneDetails : public UObject
 {
 	GENERATED_BODY()
 
@@ -61,11 +63,11 @@ public:
 
 #if WITH_EDITOR
 
-	TOptional<FTransform> GetTransform(EIKRigTransformType::Type TransformType) const;
-	bool IsComponentRelative(ESlateTransformComponent::Type Component, EIKRigTransformType::Type TransformType) const;
-	void OnComponentRelativeChanged(ESlateTransformComponent::Type Component, bool bIsRelative, EIKRigTransformType::Type TransformType);
-	void OnCopyToClipboard(ESlateTransformComponent::Type Component, EIKRigTransformType::Type TransformType) const;
-	void OnPasteFromClipboard(ESlateTransformComponent::Type Component, EIKRigTransformType::Type TransformType);
+	UE_API TOptional<FTransform> GetTransform(EIKRigTransformType::Type TransformType) const;
+	UE_API bool IsComponentRelative(ESlateTransformComponent::Type Component, EIKRigTransformType::Type TransformType) const;
+	UE_API void OnComponentRelativeChanged(ESlateTransformComponent::Type Component, bool bIsRelative, EIKRigTransformType::Type TransformType);
+	UE_API void OnCopyToClipboard(ESlateTransformComponent::Type Component, EIKRigTransformType::Type TransformType) const;
+	UE_API void OnPasteFromClipboard(ESlateTransformComponent::Type Component, EIKRigTransformType::Type TransformType);
 
 	template <typename DataType>
 	void GetContentFromData(const DataType& InData, FString& Content) const
@@ -77,8 +79,8 @@ public:
 	
 private:
 	
-	static bool CurrentTransformRelative[3];
-	static bool ReferenceTransformRelative[3];
+	static UE_API bool CurrentTransformRelative[3];
+	static UE_API bool ReferenceTransformRelative[3];
 };
 
 enum class EChainSide
@@ -285,4 +287,6 @@ struct FIKRigSolverTypeAndName
 {
 	FText NiceName;
 	UScriptStruct* SolverType;
-};
\ No newline at end of file
+};
+
+#undef UE_API
diff --git a/Engine/Plugins/Animation/LiveLink/Source/LiveLinkEditor/Private/LiveLinkClientPanel.h b/Engine/Plugins/Animation/LiveLink/Source/LiveLinkEditor/Private/LiveLinkClientPanel.h
--- a/Engine/Plugins/Animation/LiveLink/Source/LiveLinkEditor/Private/LiveLinkClientPanel.h
+++ b/Engine/Plugins/Animation/LiveLink/Source/LiveLinkEditor/Private/LiveLinkClientPanel.h
@@ -8,6 +8,8 @@
 #include "UObject/GCObject.h"
 #include "Widgets/SCompoundWidget.h"
 
+#define UE_API LIVELINKEDITOR_API
+
 namespace ESelectInfo { enum Type : int; }
 namespace ETextCommit { enum Type : int; }
 struct FPropertyChangedEvent;
@@ -24,17 +26,17 @@ class ULiveLinkSourceFactory;
 class SLiveLinkSourceListView;
 class STableViewBase;
 
-class LIVELINKEDITOR_API SLiveLinkClientPanel : public SCompoundWidget, public FGCObject, public FEditorUndoClient
+class SLiveLinkClientPanel : public SCompoundWidget, public FGCObject, public FEditorUndoClient
 {
 	SLATE_BEGIN_ARGS(SLiveLinkClientPanel) {}
 	SLATE_END_ARGS()
 
-	virtual ~SLiveLinkClientPanel() override;
+	UE_API virtual ~SLiveLinkClientPanel() override;
 	
-	void Construct(const FArguments& Args, FLiveLinkClient* InClient);
+	UE_API void Construct(const FArguments& Args, FLiveLinkClient* InClient);
 
 	// FGCObject interface
-	virtual void AddReferencedObjects(FReferenceCollector& Collector) override;
+	UE_API virtual void AddReferencedObjects(FReferenceCollector& Collector) override;
 	virtual FString GetReferencerName() const override
 	{
 		return TEXT("SLiveLinkClientPanel");
@@ -42,33 +44,33 @@ class LIVELINKEDITOR_API SLiveLinkClientPanel : public SCompoundWidget, public F
 	// End of FGCObject interface
 
 	// FEditorUndoClient interface
-	virtual void PostUndo(bool bSuccess) override;
-	virtual void PostRedo(bool bSuccess) override;
+	UE_API virtual void PostUndo(bool bSuccess) override;
+	UE_API virtual void PostRedo(bool bSuccess) override;
 	// End FEditorUndoClient interface
 
 private:
 	// Controls whether the editor performance throttling warning should be visible
-	EVisibility ShowEditorPerformanceThrottlingWarning() const;
+	UE_API EVisibility ShowEditorPerformanceThrottlingWarning() const;
 
 	// Handle disabling of editor performance throttling
-	FReply DisableEditorPerformanceThrottling();
+	UE_API FReply DisableEditorPerformanceThrottling();
 
 	// Return the message count text
-	FText GetMessageCountText() const;
+	UE_API FText GetMessageCountText() const;
 
 	// Return the occurrence count and last time occurred text
-	FText GetSelectedMessageOccurrenceText() const;
+	UE_API FText GetSelectedMessageOccurrenceText() const;
 
-	int32 GetDetailWidgetIndex() const;
+	UE_API int32 GetDetailWidgetIndex() const;
 	
 	// Handler for the subject tree context menu opening
-	TSharedPtr<SWidget> OnOpenVirtualSubjectContextMenu();
+	UE_API TSharedPtr<SWidget> OnOpenVirtualSubjectContextMenu();
 
 	// Returns whether the panel is in read only mode or not.
-	bool IsInReadOnlyMode() const;
+	UE_API bool IsInReadOnlyMode() const;
 
 	/** Get widget visibility according to whether the panel is in read-only mode. */
-	EVisibility GetVisibilityBasedOnReadOnly() const;
+	UE_API EVisibility GetVisibilityBasedOnReadOnly() const;
 
 private:
 	// Handles callback connections between the sources, subjects and details views
@@ -83,3 +85,5 @@ private:
 	// Details index
 	int32 DetailWidgetIndex = -1;
 };
+
+#undef UE_API
diff --git a/Engine/Plugins/Animation/LiveLink/Source/LiveLinkEditor/Public/LiveLinkClientPanelToolbar.h b/Engine/Plugins/Animation/LiveLink/Source/LiveLinkEditor/Public/LiveLinkClientPanelToolbar.h
--- a/Engine/Plugins/Animation/LiveLink/Source/LiveLinkEditor/Public/LiveLinkClientPanelToolbar.h
+++ b/Engine/Plugins/Animation/LiveLink/Source/LiveLinkEditor/Public/LiveLinkClientPanelToolbar.h
@@ -5,6 +5,8 @@
 #include "Widgets/SCompoundWidget.h"
 #include "UObject/GCObject.h"
 
+#define UE_API LIVELINKEDITOR_API
+
 
 class ULiveLinkPreset;
 template <typename T> class TSubclassOf;
@@ -19,7 +21,7 @@ class STextEntryPopup;
 class ULiveLinkSourceFactory;
 struct FAssetData;
 
-class LIVELINKEDITOR_API SLiveLinkClientPanelToolbar : public SCompoundWidget, public FGCObject
+class SLiveLinkClientPanelToolbar : public SCompoundWidget, public FGCObject
 {
 	SLATE_BEGIN_ARGS(SLiveLinkClientPanelToolbar)
 		: _SourceButtonAlignment(HAlign_Fill)
@@ -39,41 +41,38 @@ class LIVELINKEDITOR_API SLiveLinkClientPanelToolbar : public SCompoundWidget, p
 	SLATE_END_ARGS()
 
 
-	void Construct(const FArguments& Args, FLiveLinkClient* InClient);
+	UE_API void Construct(const FArguments& Args, FLiveLinkClient* InClient);
 
 	//~ FEditorUndoClient interface
-	virtual void AddReferencedObjects(FReferenceCollector& Collector) override;
+	UE_API virtual void AddReferencedObjects(FReferenceCollector& Collector) override;
 	virtual FString GetReferencerName() const override
 	{
 		return TEXT("SLiveLinkClientPanelToolbar");
 	}
 
 private:
-	TSharedRef<SWidget> OnGenerateSourceMenu();
-	/** Open a modal window that hosts the source factory widget. */
-	void OpenCreateMenuWindow(int32 FactoryIndex);
-	void ExecuteCreateSource(int32 FactoryIndex);
-	void OnSourceCreated(TSharedPtr<ILiveLinkSource> NewSource, FString ConnectionString, TSubclassOf<ULiveLinkSourceFactory> Factory);
-
-	TSharedRef<SWidget> OnPresetGeneratePresetsMenu();
-	TSharedRef<SWidget> HandleSourceSelectionComboButton();
-	void HandleVirtualSourceSelection(const FGuid& InSourceGuid);
-	bool IsVirtualSourceSelected(const FGuid& InSourceGuid);
-	void OnSaveAsPreset();
-	void OnImportPreset(const FAssetData& InPreset);
-	FReply OnRevertChanges();
-	bool HasLoadedLiveLinkPreset() const;
-
-	void PopulateVirtualSubjectSourceCreationMenu(FMenuBuilder& InMenuBuilder);
-	FReply OnAddVirtualSubjectSource();
-	void AddVirtualSubject();
+	UE_API TSharedRef<SWidget> OnGenerateSourceMenu();
+	UE_API void RetrieveFactorySourcePanel(FMenuBuilder& MenuBuilder, int32 FactoryIndex);
+	UE_API void ExecuteCreateSource(int32 FactoryIndex);
+	UE_API void OnSourceCreated(TSharedPtr<ILiveLinkSource> NewSource, FString ConnectionString, TSubclassOf<ULiveLinkSourceFactory> Factory);
+
+	UE_API TSharedRef<SWidget> OnPresetGeneratePresetsMenu();
+	UE_API TSharedRef<SWidget> HandleSourceSelectionComboButton();
+	UE_API void HandleVirtualSourceSelection(const FGuid& InSourceGuid);
+	UE_API bool IsVirtualSourceSelected(const FGuid& InSourceGuid);
+	UE_API void OnSaveAsPreset();
+	UE_API void OnImportPreset(const FAssetData& InPreset);
+	UE_API FReply OnRevertChanges();
+	UE_API bool HasLoadedLiveLinkPreset() const;
+
+	UE_API void PopulateVirtualSubjectSourceCreationMenu(FMenuBuilder& InMenuBuilder);
+	UE_API FReply OnAddVirtualSubjectSource();
+	UE_API void AddVirtualSubject();
 
 
 private:
 	FLiveLinkClient* Client;
 
-	/** Source creation modal window. */
-	TSharedPtr<SWindow> SourceCreationWindow;
 	TWeakPtr<IMenu> AddSubjectMenu;
 	TWeakPtr<IMenu> VirtualSubjectMenu;
 	TWeakPtr<STextEntryPopup> VirtualSubjectPopup;
@@ -85,3 +84,5 @@ private:
 	/** Parent window override used for creating asset dialogs when running in LiveLinkHub. */
 	TSharedPtr<SWindow> ParentWindowOverride;
 };
+
+#undef UE_API
diff --git a/Engine/Plugins/Animation/LiveLink/Source/LiveLinkEditor/Public/LiveLinkClientPanelViews.h b/Engine/Plugins/Animation/LiveLink/Source/LiveLinkEditor/Public/LiveLinkClientPanelViews.h
--- a/Engine/Plugins/Animation/LiveLink/Source/LiveLinkEditor/Public/LiveLinkClientPanelViews.h
+++ b/Engine/Plugins/Animation/LiveLink/Source/LiveLinkEditor/Public/LiveLinkClientPanelViews.h
@@ -13,6 +13,8 @@
 #include "Widgets/Views/SListView.h"
 #include "Widgets/Views/STreeView.h"
 
+#define UE_API LIVELINKEDITOR_API
+
 class FLiveLinkClient;
 struct FLiveLinkSourceUIEntry;
 class FLiveLinkSourcesView;
@@ -168,40 +170,40 @@ public:
 	}
 };
 
-class LIVELINKEDITOR_API FLiveLinkSubjectsView : public TSharedFromThis<FLiveLinkSubjectsView>
+class FLiveLinkSubjectsView : public TSharedFromThis<FLiveLinkSubjectsView>
 {
 public:
 	DECLARE_DELEGATE_TwoParams(FOnSubjectSelectionChanged, FLiveLinkSubjectUIEntryPtr, ESelectInfo::Type);
 
-	FLiveLinkSubjectsView(FOnSubjectSelectionChanged InOnSubjectSelectionChanged, const TSharedPtr<FUICommandList>& InCommandList, TAttribute<bool> bInReadOnly);
+	UE_API FLiveLinkSubjectsView(FOnSubjectSelectionChanged InOnSubjectSelectionChanged, const TSharedPtr<FUICommandList>& InCommandList, TAttribute<bool> bInReadOnly);
 
 	// Helper functions for building the subject tree UI
-	TSharedRef<ITableRow> MakeTreeRowWidget(FLiveLinkSubjectUIEntryPtr InInfo, const TSharedRef<STableViewBase>& OwnerTable);
+	UE_API TSharedRef<ITableRow> MakeTreeRowWidget(FLiveLinkSubjectUIEntryPtr InInfo, const TSharedRef<STableViewBase>& OwnerTable);
 	// Get a subjectsview's children
-	void GetChildrenForInfo(FLiveLinkSubjectUIEntryPtr InInfo, TArray< FLiveLinkSubjectUIEntryPtr >& OutChildren);
+	UE_API void GetChildrenForInfo(FLiveLinkSubjectUIEntryPtr InInfo, TArray< FLiveLinkSubjectUIEntryPtr >& OutChildren);
 	// Handle subject selection changed
-	void OnSubjectSelectionChanged(FLiveLinkSubjectUIEntryPtr SubjectEntry, ESelectInfo::Type SelectInfo);
+	UE_API void OnSubjectSelectionChanged(FLiveLinkSubjectUIEntryPtr SubjectEntry, ESelectInfo::Type SelectInfo);
 	// Handler for the subject tree context menu opening
-	TSharedPtr<SWidget> OnOpenVirtualSubjectContextMenu(TSharedPtr<FUICommandList> InCommandList);
+	UE_API TSharedPtr<SWidget> OnOpenVirtualSubjectContextMenu(TSharedPtr<FUICommandList> InCommandList);
 	// Return whether a subject can be removed
-	bool CanRemoveSubject() const;
+	UE_API bool CanRemoveSubject() const;
 	// Refresh the list of subjects using the livelink client.
-	void RefreshSubjects();
+	UE_API void RefreshSubjects();
 	// Return whether a subject can be paused
-	bool CanPauseSubject() const;
+	UE_API bool CanPauseSubject() const;
 	// Pauses/Unpauses a subject
-	void HandlePauseSubject();
+	UE_API void HandlePauseSubject();
 
 private:
 	// Create the subjects tree view
-	void CreateSubjectsTreeView(const TSharedPtr<FUICommandList>& InCommandList);
+	UE_API void CreateSubjectsTreeView(const TSharedPtr<FUICommandList>& InCommandList);
 
 	/** Get the label for the pause subject context menu entry. */
-	FText GetPauseSubjectLabel() const;
+	UE_API FText GetPauseSubjectLabel() const;
 	/** Get the tooltip for the pause subject context menu entry. */
-	FText GetPauseSubjectToolTip() const;
+	UE_API FText GetPauseSubjectToolTip() const;
 	/** Returns whether the selected subject is paused. */
-	bool IsSelectedSubjectPaused() const;
+	UE_API bool IsSelectedSubjectPaused() const;
 
 public:
 	// Subject tree widget
@@ -250,3 +252,5 @@ public:
 	// Returns whether the panel is in read-only mode.
 	TAttribute<bool> bReadOnly;
 };
+
+#undef UE_API
diff --git a/Engine/Plugins/Animation/LiveLink/Source/LiveLinkEditor/Public/LiveLinkEditorSettings.h b/Engine/Plugins/Animation/LiveLink/Source/LiveLinkEditor/Public/LiveLinkEditorSettings.h
--- a/Engine/Plugins/Animation/LiveLink/Source/LiveLinkEditor/Public/LiveLinkEditorSettings.h
+++ b/Engine/Plugins/Animation/LiveLink/Source/LiveLinkEditor/Public/LiveLinkEditorSettings.h
@@ -9,8 +9,8 @@
 /**
  * Editor preferences settings for LiveLink.
  */
-UCLASS(config = EditorPerProjectUserSettings)
-class LIVELINKEDITOR_API ULiveLinkEditorSettings : public UObject
+UCLASS(MinimalAPI, config = EditorPerProjectUserSettings)
+class ULiveLinkEditorSettings : public UObject
 {
 	GENERATED_BODY()
 
diff --git a/Engine/Plugins/Animation/LiveLink/Source/LiveLinkEditor/Public/LiveLinkPanelController.h b/Engine/Plugins/Animation/LiveLink/Source/LiveLinkEditor/Public/LiveLinkPanelController.h
--- a/Engine/Plugins/Animation/LiveLink/Source/LiveLinkEditor/Public/LiveLinkPanelController.h
+++ b/Engine/Plugins/Animation/LiveLink/Source/LiveLinkEditor/Public/LiveLinkPanelController.h
@@ -8,6 +8,8 @@
 #include "Misc/Attribute.h"
 #include "Templates/SharedPointer.h"
 
+#define UE_API LIVELINKEDITOR_API
+
 class IDetailsView;
 class FLiveLinkClient;
 struct FLiveLinkSourceUIEntry;
@@ -25,11 +27,11 @@ typedef TSharedPtr<FLiveLinkSubjectUIEntry> FLiveLinkSubjectUIEntryPtr;
 
 
 /** Handles callback connections between the sources, subjects and details views. */
-class LIVELINKEDITOR_API FLiveLinkPanelController : public TSharedFromThis<FLiveLinkPanelController>
+class FLiveLinkPanelController : public TSharedFromThis<FLiveLinkPanelController>
 {
 public:
-	FLiveLinkPanelController(TAttribute<bool> bInReadOnly = false);
-	~FLiveLinkPanelController();
+	UE_API FLiveLinkPanelController(TAttribute<bool> bInReadOnly = false);
+	UE_API ~FLiveLinkPanelController();
 
 	DECLARE_MULTICAST_DELEGATE_OneParam(FOnSubjectSelectionChanged, const FLiveLinkSubjectKey&);
 	/** Subject Selection changed callback. */
@@ -40,35 +42,35 @@ public:
 
 private:
 	// Bind live link commands 
-	void BindCommands();
+	UE_API void BindCommands();
 	// Handles the source collection changing.
-	void OnSourcesChangedHandler();
+	UE_API void OnSourcesChangedHandler();
 	// Handles the subject collection changing.
-	void OnSubjectsChangedHandler();
+	UE_API void OnSubjectsChangedHandler();
 	// Remove all sources command handlers
-	bool HasSource() const;
+	UE_API bool HasSource() const;
 	// Handles the remove source command.
-	void HandleRemoveSource();
+	UE_API void HandleRemoveSource();
 	// Handles the remove all sources command.
-	void HandleRemoveAllSources();
+	UE_API void HandleRemoveAllSources();
 	// Returns whether a source could be removed.
-	bool CanRemoveSource() const;
+	UE_API bool CanRemoveSource() const;
 	// Returns whether a subject can be removed.
-	bool CanRemoveSubject() const;
+	UE_API bool CanRemoveSubject() const;
 	// Returns whethera subject can be paused
-	bool CanPauseSubject() const;
+	UE_API bool CanPauseSubject() const;
 	// Pauses/Unpauses a subject.
-	void HandlePauseSubject();
+	UE_API void HandlePauseSubject();
 	// Handles the remove subject command.
-	void HandleRemoveSubject();
+	UE_API void HandleRemoveSubject();
 	// Recreates the source list data behind the list view.
-	void RebuildSourceList();
+	UE_API void RebuildSourceList();
 	// Recreates the subject list data behind the tree view.
-	void RebuildSubjectList();
+	UE_API void RebuildSubjectList();
 	// Handles source selection changing.
-	void OnSourceSelectionChangedHandler(FLiveLinkSourceUIEntryPtr Entry, ESelectInfo::Type SelectionType) const;
+	UE_API void OnSourceSelectionChangedHandler(FLiveLinkSourceUIEntryPtr Entry, ESelectInfo::Type SelectionType) const;
 	// Hadnles subject selection changing.
-	void OnSubjectSelectionChangedHandler(FLiveLinkSubjectUIEntryPtr SubjectEntry, ESelectInfo::Type SelectInfo);
+	UE_API void OnSubjectSelectionChangedHandler(FLiveLinkSubjectUIEntryPtr SubjectEntry, ESelectInfo::Type SelectInfo);
 
 	// Sources and subjects live in different tabs in LiveLink Hub.
 	bool bSeparateSourcesSubjects = false;
@@ -95,3 +97,5 @@ public:
 	// Delegate called when the subject selection changes
 	FOnSubjectSelectionChanged SubjectSelectionChangedDelegate;
 };
+
+#undef UE_API
diff --git a/Engine/Plugins/Animation/LiveLink/Source/LiveLinkEditor/Public/SLiveLinkBoneSelectionWidget.h b/Engine/Plugins/Animation/LiveLink/Source/LiveLinkEditor/Public/SLiveLinkBoneSelectionWidget.h
--- a/Engine/Plugins/Animation/LiveLink/Source/LiveLinkEditor/Public/SLiveLinkBoneSelectionWidget.h
+++ b/Engine/Plugins/Animation/LiveLink/Source/LiveLinkEditor/Public/SLiveLinkBoneSelectionWidget.h
@@ -12,6 +12,8 @@
 #include "Widgets/Views/STableRow.h"
 #include "Widgets/Views/STreeView.h"
 
+#define UE_API LIVELINKEDITOR_API
+
 
 class IEditableSkeleton;
 class SComboButton;
@@ -19,7 +21,7 @@ class SComboButton;
 DECLARE_DELEGATE_OneParam(FOnBoneSelectionChanged, FName);
 DECLARE_DELEGATE_RetVal(FName, FGetSelectedBone);
 
-class LIVELINKEDITOR_API SLiveLinkBoneTreeMenu : public SCompoundWidget
+class SLiveLinkBoneTreeMenu : public SCompoundWidget
 {
 public:
 	/** Storage object for bone hierarchy */
@@ -45,32 +47,32 @@ public:
 	*
 	* @param	InArgs	The declaration data for this widget
 	*/
-	void Construct(const FArguments& InArgs, TOptional<FLiveLinkSkeletonStaticData> SkeletonStaticData);
+	UE_API void Construct(const FArguments& InArgs, TOptional<FLiveLinkSkeletonStaticData> SkeletonStaticData);
 
 	/** Get the filter text widget, e.g. for focus */
-	TSharedPtr<SWidget> GetFilterTextWidget();
+	UE_API TSharedPtr<SWidget> GetFilterTextWidget();
 
 private:
 	/** SWidget interface */
-	virtual FReply OnKeyDown(const FGeometry& MyGeometry, const FKeyEvent& InKeyEvent) override;
+	UE_API virtual FReply OnKeyDown(const FGeometry& MyGeometry, const FKeyEvent& InKeyEvent) override;
 
 	/** Using the current filter, repopulate the tree view */
-	void RebuildBoneList(const FName& SelectedBone);
+	UE_API void RebuildBoneList(const FName& SelectedBone);
 
 	/** Make a single tree row widget */
-	TSharedRef<ITableRow> MakeTreeRowWidget(TSharedPtr<FBoneNameInfo> InInfo, const TSharedRef<STableViewBase>& OwnerTable);
+	UE_API TSharedRef<ITableRow> MakeTreeRowWidget(TSharedPtr<FBoneNameInfo> InInfo, const TSharedRef<STableViewBase>& OwnerTable);
 
 	/** Get the children for the provided bone info */
-	void GetChildrenForInfo(TSharedPtr<FBoneNameInfo> InInfo, TArray< TSharedPtr<FBoneNameInfo> >& OutChildren);
+	UE_API void GetChildrenForInfo(TSharedPtr<FBoneNameInfo> InInfo, TArray< TSharedPtr<FBoneNameInfo> >& OutChildren);
 
 	/** Called when the user changes the search filter */
-	void OnFilterTextChanged(const FText& InFilterText);
+	UE_API void OnFilterTextChanged(const FText& InFilterText);
 
 	/** Handle the tree view selection changing */
-	void OnSelectionChanged(TSharedPtr<SLiveLinkBoneTreeMenu::FBoneNameInfo> BoneInfo, ESelectInfo::Type SelectInfo);
+	UE_API void OnSelectionChanged(TSharedPtr<SLiveLinkBoneTreeMenu::FBoneNameInfo> BoneInfo, ESelectInfo::Type SelectInfo);
 
 	/** Select a specific bone, helper for UI handler functions */
-	void SelectBone(TSharedPtr<SLiveLinkBoneTreeMenu::FBoneNameInfo> BoneInfo);
+	UE_API void SelectBone(TSharedPtr<SLiveLinkBoneTreeMenu::FBoneNameInfo> BoneInfo);
 	
 	/** Tree info entries for bone picker */
 	TArray<TSharedPtr<FBoneNameInfo>> SkeletonTreeInfo;
@@ -93,7 +95,7 @@ private:
 	FLiveLinkSkeletonStaticData SkeletonStaticData;
 };
 
-class LIVELINKEDITOR_API SLiveLinkBoneSelectionWidget : public SCompoundWidget
+class SLiveLinkBoneSelectionWidget : public SCompoundWidget
 {
 public: 
 
@@ -114,23 +116,23 @@ public:
 	 *
 	 * @param	InArgs	The declaration data for this widget
 	 */
-	void Construct(const FArguments& InArgs, const FLiveLinkSubjectKey& InSubjectKey);
+	UE_API void Construct(const FArguments& InArgs, const FLiveLinkSubjectKey& InSubjectKey);
 
 	/** Set the subject for which we will list the bones. */
-	void SetSubject(const FLiveLinkSubjectKey& InSubjectKey);
+	UE_API void SetSubject(const FLiveLinkSubjectKey& InSubjectKey);
 
 	/** Utility method to create fake static data from a frame translator. */
-	FLiveLinkSkeletonStaticData MakeStaticDataFromTranslator(const FLiveLinkSubjectKey& SubjectKey) const;
+	UE_API FLiveLinkSkeletonStaticData MakeStaticDataFromTranslator(const FLiveLinkSubjectKey& SubjectKey) const;
 
 private: 
 
 	/** Creates the combo button menu when clicked */
-	TSharedRef<SWidget> CreateSkeletonWidgetMenu();
+	UE_API TSharedRef<SWidget> CreateSkeletonWidgetMenu();
 	/** Called when the user selects a bone name */
-	void OnSelectionChanged(FName BoneName);
+	UE_API void OnSelectionChanged(FName BoneName);
 
 	/** Gets the current bone name, used to get the right name for the combo button */
-	FText GetCurrentBoneName() const;
+	UE_API FText GetCurrentBoneName() const;
 
 	/** Base combo button */
 	TSharedPtr<SComboButton> BonePickerButton;
@@ -142,3 +144,5 @@ private:
 	FOnBoneSelectionChanged OnBoneSelectionChanged;
 	FGetSelectedBone OnGetSelectedBone;
 };
+
+#undef UE_API
diff --git a/Engine/Plugins/Animation/LiveLink/Source/LiveLinkEditor/Public/SLiveLinkDataView.h b/Engine/Plugins/Animation/LiveLink/Source/LiveLinkEditor/Public/SLiveLinkDataView.h
--- a/Engine/Plugins/Animation/LiveLink/Source/LiveLinkEditor/Public/SLiveLinkDataView.h
+++ b/Engine/Plugins/Animation/LiveLink/Source/LiveLinkEditor/Public/SLiveLinkDataView.h
@@ -6,6 +6,8 @@
 
 #include "LiveLinkTypes.h"
 
+#define UE_API LIVELINKEDITOR_API
+
 namespace ETextCommit { enum Type : int; }
 struct FPropertyChangedEvent;
 
@@ -13,7 +15,7 @@ class IDetailsView;
 class FLiveLinkClient;
 class IStructureDetailsView;
 
-class LIVELINKEDITOR_API SLiveLinkDataView : public SCompoundWidget
+class SLiveLinkDataView : public SCompoundWidget
 {
 private:
 	using Super = SCompoundWidget;
@@ -25,11 +27,11 @@ public:
 	SLATE_ATTRIBUTE(bool, ReadOnly)
 	SLATE_END_ARGS()
 
-	void Construct(const FArguments& Args, FLiveLinkClient* InClient);
+	UE_API void Construct(const FArguments& Args, FLiveLinkClient* InClient);
 
-	virtual void Tick(const FGeometry& AllottedGeometry, const double InCurrentTime, const float InDeltaTime) override;
+	UE_API virtual void Tick(const FGeometry& AllottedGeometry, const double InCurrentTime, const float InDeltaTime) override;
 
-	void SetSubjectKey(FLiveLinkSubjectKey InSubjectKey);
+	UE_API void SetSubjectKey(FLiveLinkSubjectKey InSubjectKey);
 	FLiveLinkSubjectKey GetSubjectKey() const { return SubjectKey; }
 
 	void SetRefreshDelay(double DelaySeconds) { UpdateDelay = DelaySeconds; }
@@ -43,9 +45,9 @@ private:
 		FrameData,
 	};
 
-	void OnPropertyChanged(const FPropertyChangedEvent& InEvent);
-	int32 GetDetailWidgetIndex() const;
-	void OnSelectDetailWidget(EDetailType InDetailType);
+	UE_API void OnPropertyChanged(const FPropertyChangedEvent& InEvent);
+	UE_API int32 GetDetailWidgetIndex() const;
+	UE_API void OnSelectDetailWidget(EDetailType InDetailType);
 	bool IsSelectedDetailWidget(EDetailType InDetailType) const { return InDetailType == DetailType; }
 	bool CanEditRefreshDelay() const { return DetailType != EDetailType::Property; }
 	void SetRefreshDelayInternal(double InDelaySeconds, ETextCommit::Type) { SetRefreshDelay(InDelaySeconds); }
@@ -60,3 +62,5 @@ private:
 	TSharedPtr<IStructureDetailsView> StructureDetailsView;
 	TSharedPtr<IDetailsView> SettingsDetailsView;
 };
+
+#undef UE_API
diff --git a/Engine/Plugins/Animation/LiveLink/Source/LiveLinkEditor/Public/SLiveLinkSubjectRepresentationPicker.h b/Engine/Plugins/Animation/LiveLink/Source/LiveLinkEditor/Public/SLiveLinkSubjectRepresentationPicker.h
--- a/Engine/Plugins/Animation/LiveLink/Source/LiveLinkEditor/Public/SLiveLinkSubjectRepresentationPicker.h
+++ b/Engine/Plugins/Animation/LiveLink/Source/LiveLinkEditor/Public/SLiveLinkSubjectRepresentationPicker.h
@@ -9,6 +9,8 @@
 #include "Styling/SlateTypes.h"
 #include "Styling/SlateWidgetStyleAsset.h"
 
+#define UE_API LIVELINKEDITOR_API
+
 class ULiveLinkPreset;
 namespace ESelectInfo { enum Type : int; }
 namespace ETextCommit { enum Type : int; }
@@ -31,7 +33,7 @@ DECLARE_DELEGATE_OneParam(FOnGetSubjects, TArray<FLiveLinkSubjectKey>& /**OutSub
 /**
  * A widget which allows the user to enter a subject name or discover it from a drop menu.
  */
-class LIVELINKEDITOR_API SLiveLinkSubjectRepresentationPicker : public SCompoundWidget
+class SLiveLinkSubjectRepresentationPicker : public SCompoundWidget
 {
 public:
 	struct FLiveLinkSourceSubjectRole
@@ -106,34 +108,34 @@ public:
 	/**
 	 * Slate widget construction method
 	 */
-	void Construct(const FArguments& InArgs);
+	UE_API void Construct(const FArguments& InArgs);
 
 	/**
 	 * Access the current value of this picker
 	 */
-	FLiveLinkSourceSubjectRole GetCurrentValue() const;
+	UE_API FLiveLinkSourceSubjectRole GetCurrentValue() const;
 
 private:
-	FText GetSourceNameValueText() const;
-	FText GetSubjectNameValueText() const;
-	const FSlateBrush* GetRoleIcon() const;
-	FText GetRoleText() const;
-
-	TSharedRef<SWidget> BuildMenu();
-	FText GetPresetSelectedText() const;
-	FSlateColor GetSelectPresetForegroundColor() const;
-	FReply ClearCurrentPreset();
-	bool HasCurrentPreset() const;
+	UE_API FText GetSourceNameValueText() const;
+	UE_API FText GetSubjectNameValueText() const;
+	UE_API const FSlateBrush* GetRoleIcon() const;
+	UE_API FText GetRoleText() const;
+
+	UE_API TSharedRef<SWidget> BuildMenu();
+	UE_API FText GetPresetSelectedText() const;
+	UE_API FSlateColor GetSelectPresetForegroundColor() const;
+	UE_API FReply ClearCurrentPreset();
+	UE_API bool HasCurrentPreset() const;
 	
-	TSharedRef<ITableRow> MakeSubjectRepListViewWidget(FLiveLinkSubjectRepresentationPickerEntryPtr Entry, const TSharedRef<STableViewBase>& OwnerTable) const;
-	void OnSubjectRepListSelectionChanged(FLiveLinkSubjectRepresentationPickerEntryPtr Entry, ESelectInfo::Type SelectionType);
+	UE_API TSharedRef<ITableRow> MakeSubjectRepListViewWidget(FLiveLinkSubjectRepresentationPickerEntryPtr Entry, const TSharedRef<STableViewBase>& OwnerTable) const;
+	UE_API void OnSubjectRepListSelectionChanged(FLiveLinkSubjectRepresentationPickerEntryPtr Entry, ESelectInfo::Type SelectionType);
 
-	TSharedRef<SWidget> BuildPresetSubMenu();
-	void NewPresetSelected(const FAssetData& AssetData);
-	void OnComboTextCommitted(const FText& NewText, ETextCommit::Type InTextCommit);
-	void SetValue(const FLiveLinkSourceSubjectRole& InValue);
+	UE_API TSharedRef<SWidget> BuildPresetSubMenu();
+	UE_API void NewPresetSelected(const FAssetData& AssetData);
+	UE_API void OnComboTextCommitted(const FText& NewText, ETextCommit::Type InTextCommit);
+	UE_API void SetValue(const FLiveLinkSourceSubjectRole& InValue);
 
-	void BuildSubjectRepDataList();
+	UE_API void BuildSubjectRepDataList();
 
 private:
 	TWeakObjectPtr<ULiveLinkPreset> SelectedLiveLinkPreset;
@@ -152,3 +154,5 @@ private:
 	/** Delegate used to retrieve a custom list of subjects. */
 	FOnGetSubjects GetSubjectsDelegate;
 };
+
+#undef UE_API
diff --git a/Engine/Plugins/Animation/LiveLink/Source/LiveLinkEditor/Public/VirtualSubjects/LiveLinkBlueprintVirtualSubjectFactory.h b/Engine/Plugins/Animation/LiveLink/Source/LiveLinkEditor/Public/VirtualSubjects/LiveLinkBlueprintVirtualSubjectFactory.h
--- a/Engine/Plugins/Animation/LiveLink/Source/LiveLinkEditor/Public/VirtualSubjects/LiveLinkBlueprintVirtualSubjectFactory.h
+++ b/Engine/Plugins/Animation/LiveLink/Source/LiveLinkEditor/Public/VirtualSubjects/LiveLinkBlueprintVirtualSubjectFactory.h
@@ -6,24 +6,28 @@
 
 #include "LiveLinkBlueprintVirtualSubjectFactory.generated.h"
 
+#define UE_API LIVELINKEDITOR_API
+
 class ULiveLinkRole;
 
-UCLASS(hidecategories = Object)
-class LIVELINKEDITOR_API ULiveLinkBlueprintVirtualSubjectFactory : public UBlueprintFactory
+UCLASS(MinimalAPI, hidecategories = Object)
+class ULiveLinkBlueprintVirtualSubjectFactory : public UBlueprintFactory
 {
 	GENERATED_BODY()
 	
 public:
-	ULiveLinkBlueprintVirtualSubjectFactory();
+	UE_API ULiveLinkBlueprintVirtualSubjectFactory();
 
 	UPROPERTY(BlueprintReadWrite, Category = "Live Link Blueprint Virtual Subject Factory")
 	TSubclassOf<ULiveLinkRole> Role;
 
 	//~ Begin UFactory Interface
-	virtual FText GetDisplayName() const override;
-	virtual bool ConfigureProperties() override;
-	virtual UObject* FactoryCreateNew(UClass* Class, UObject* InParent, FName Name, EObjectFlags Flags, UObject* Context, FFeedbackContext* Warn, FName CallingContext) override;
-	virtual bool ShouldShowInNewMenu() const override;
-	virtual uint32 GetMenuCategories() const override;
+	UE_API virtual FText GetDisplayName() const override;
+	UE_API virtual bool ConfigureProperties() override;
+	UE_API virtual UObject* FactoryCreateNew(UClass* Class, UObject* InParent, FName Name, EObjectFlags Flags, UObject* Context, FFeedbackContext* Warn, FName CallingContext) override;
+	UE_API virtual bool ShouldShowInNewMenu() const override;
+	UE_API virtual uint32 GetMenuCategories() const override;
 	//~ Begin UFactory Interface
 };
+
+#undef UE_API
