[IK Rig] Added "Refresh Hierarchy" menu option

diff --git a/Engine/Plugins/Animation/IKRig/Source/IKRigEditor/Private/RigEditor/IKRigSkeletonCommands.cpp b/Engine/Plugins/Animation/IKRig/Source/IKRigEditor/Private/RigEditor/IKRigSkeletonCommands.cpp
--- a/Engine/Plugins/Animation/IKRig/Source/IKRigEditor/Private/RigEditor/IKRigSkeletonCommands.cpp
+++ b/Engine/Plugins/Animation/IKRig/Source/IKRigEditor/Private/RigEditor/IKRigSkeletonCommands.cpp
@@ -20,6 +20,7 @@ void FIKRigSkeletonCommands::RegisterCommands()
 	UI_COMMAND(SetPelvis, "Set Pelvis", "Set the Pelvis bone used for retargeting.", EUserInterfaceActionType::Button, FInputChord());
 	UI_COMMAND(ClearPelvis, "Clear Pelvis", "Clear the Pelvis bone used for retargeting.", EUserInterfaceActionType::Button, FInputChord());
 	UI_COMMAND(RenameGoal, "Rename Goal", "Rename the selected goal.", EUserInterfaceActionType::Button, FInputChord(EKeys::F2));
+	UI_COMMAND(RefreshHierarchy, "Refresh Hierarchy", "Reload bones from current skeletal mesh.", EUserInterfaceActionType::Button, FInputChord());
 }
 
 #undef LOCTEXT_NAMESPACE
diff --git a/Engine/Plugins/Animation/IKRig/Source/IKRigEditor/Private/RigEditor/SIKRigHierarchy.cpp b/Engine/Plugins/Animation/IKRig/Source/IKRigEditor/Private/RigEditor/SIKRigHierarchy.cpp
--- a/Engine/Plugins/Animation/IKRig/Source/IKRigEditor/Private/RigEditor/SIKRigHierarchy.cpp
+++ b/Engine/Plugins/Animation/IKRig/Source/IKRigEditor/Private/RigEditor/SIKRigHierarchy.cpp
@@ -821,6 +821,12 @@ void SIKRigHierarchy::BindCommands()
 	CommandList->MapAction(Commands.RenameGoal,
 	FExecuteAction::CreateSP(this, &SIKRigHierarchy::HandleRenameGoal),
 		FCanExecuteAction::CreateSP(this, &SIKRigHierarchy::CanRenameGoal));
+
+	CommandList->MapAction(Commands.RefreshHierarchy,
+	FExecuteAction::CreateSP(this, &SIKRigHierarchy::HandleRefreshHierarchy),
+		FCanExecuteAction::CreateSP(this, &SIKRigHierarchy::CanRefreshHierarchy));
+
+	
 }
 
 void SIKRigHierarchy::FillContextMenu(FMenuBuilder& MenuBuilder)
@@ -864,6 +870,10 @@ void SIKRigHierarchy::FillContextMenu(FMenuBuilder& MenuBuilder)
 	MenuBuilder.AddMenuEntry(Actions.ClearPelvis);
 	MenuBuilder.AddMenuEntry(Actions.NewRetargetChain);
 	MenuBuilder.EndSection();
+
+	MenuBuilder.BeginSection("Hierarchy", LOCTEXT("HierarchyOperations", "Hierarchy"));
+	MenuBuilder.AddMenuEntry(Actions.RefreshHierarchy);
+	MenuBuilder.EndSection();
 }
 
 void SIKRigHierarchy::HandleNewGoal() const
@@ -1455,6 +1465,31 @@ bool SIKRigHierarchy::CanClearPelvis()
 	return Controller->AssetController->GetRetargetRoot() != NAME_None;
 }
 
+bool SIKRigHierarchy::CanRefreshHierarchy()
+{
+	FIKRigEditorController* Controller = EditorController.Pin().Get();
+	if (!Controller)
+	{
+		return false;
+	}
+
+	return Controller->AssetController->GetSkeletalMesh() != nullptr;
+}
+
+void SIKRigHierarchy::HandleRefreshHierarchy()
+{
+	FIKRigEditorController* Controller = EditorController.Pin().Get();
+	if (!Controller)
+	{
+		return;
+	}
+
+	if (USkeletalMesh* SkeletalMesh = Controller->AssetController->GetSkeletalMesh())
+	{
+		Controller->AssetController->SetSkeletalMesh(SkeletalMesh);
+	}
+}
+
 bool SIKRigHierarchy::IsBoneInSelection(TArray<TSharedPtr<FIKRigTreeElement>>& SelectedBoneItems, const FName& BoneName)
 {
 	for (const TSharedPtr<FIKRigTreeElement>& Item : SelectedBoneItems)
diff --git a/Engine/Plugins/Animation/IKRig/Source/IKRigEditor/Public/RigEditor/IKRigSkeletonCommands.h b/Engine/Plugins/Animation/IKRig/Source/IKRigEditor/Public/RigEditor/IKRigSkeletonCommands.h
--- a/Engine/Plugins/Animation/IKRig/Source/IKRigEditor/Public/RigEditor/IKRigSkeletonCommands.h
+++ b/Engine/Plugins/Animation/IKRig/Source/IKRigEditor/Public/RigEditor/IKRigSkeletonCommands.h
@@ -52,6 +52,9 @@ public:
 
 	/** rename selected goal */
 	TSharedPtr< FUICommandInfo > RenameGoal;
+
+	/** reload the hierarchy */
+	TSharedPtr< FUICommandInfo > RefreshHierarchy;
 	
 	/** initialize commands */
 	virtual void RegisterCommands() override;
diff --git a/Engine/Plugins/Animation/IKRig/Source/IKRigEditor/Public/RigEditor/SIKRigHierarchy.h b/Engine/Plugins/Animation/IKRig/Source/IKRigEditor/Public/RigEditor/SIKRigHierarchy.h
--- a/Engine/Plugins/Animation/IKRig/Source/IKRigEditor/Public/RigEditor/SIKRigHierarchy.h
+++ b/Engine/Plugins/Animation/IKRig/Source/IKRigEditor/Public/RigEditor/SIKRigHierarchy.h
@@ -282,6 +282,10 @@ private:
 	bool CanClearPelvis();
 	/** END retarget chains */
 
+	/** refresh the bone hierarchy */
+	bool CanRefreshHierarchy();
+	void HandleRefreshHierarchy();
+
 	/** centralized editor controls (facilitate cross-communication between multiple UI elements)*/
 	TWeakPtr<FIKRigEditorController> EditorController;
 	
