[IK Retarget] Cleaning up IK Goals Op to make it easier to extend.

diff --git a/Engine/Plugins/Animation/IKRig/Source/IKRig/Private/Retargeter/RetargetOps/IKChainsOp.cpp b/Engine/Plugins/Animation/IKRig/Source/IKRig/Private/Retargeter/RetargetOps/IKChainsOp.cpp
--- a/Engine/Plugins/Animation/IKRig/Source/IKRig/Private/Retargeter/RetargetOps/IKChainsOp.cpp
+++ b/Engine/Plugins/Animation/IKRig/Source/IKRig/Private/Retargeter/RetargetOps/IKChainsOp.cpp
@@ -6,7 +6,6 @@
 #include "Retargeter/RetargetOps/FKChainsOp.h"
 #include "Retargeter/RetargetOps/PelvisMotionOp.h"
 #include "Retargeter/RetargetOps/RunIKRigOp.h"
-#include "Retargeter/RetargetOps/StrideWarpingOp.h"
 
 #if WITH_EDITOR
 #include "PrimitiveDrawInterface.h"
@@ -17,8 +16,11 @@
 
 #define LOCTEXT_NAMESPACE "IKChainsOp"
 
-
-bool FSourceChainIK::Initialize(const FResolvedBoneChain& InSourceBoneChain, FIKRigLogger& InLog)
+bool FIKChainRetargeter::Initialize(
+	const FResolvedBoneChain& InSourceBoneChain,
+	const FResolvedBoneChain& InTargetBoneChain,
+	const FRetargetIKChainSettings& InSettings,
+	FIKRigLogger& InLog)
 {
 	if (InSourceBoneChain.BoneIndices.Num() < 2)
 	{
@@ -26,121 +28,113 @@ bool FSourceChainIK::Initialize(const FResolvedBoneChain& InSourceBoneChain, FIK
 		return false;
 	}
 
-	SourceBoneChain = &InSourceBoneChain;
-	StartBoneIndex = InSourceBoneChain.BoneIndices[0];
-	EndBoneIndex = InSourceBoneChain.BoneIndices.Last();
-	
-	const FTransform& End = InSourceBoneChain.RefPoseGlobalTransforms.Last();
-	PreviousEndPosition = End.GetTranslation();
-	InitialEndPosition = End.GetTranslation();
-	InitialEndRotation = End.GetRotation();
-
-	const FTransform& Start = InSourceBoneChain.RefPoseGlobalTransforms[0];
-	const float Length = static_cast<float>((Start.GetTranslation() - InitialEndPosition).Size());
-
-	if (Length <= UE_KINDA_SMALL_NUMBER)
+	if (InTargetBoneChain.BoneIndices.Num() < 3)
 	{
-		InLog.LogWarning(LOCTEXT("SourceZeroLengthIK", "IK Chains Op: found zero-length source bone chain."));
+		InLog.LogWarning(LOCTEXT("TargetChainLessThanThree", "IK Chains Op: trying to retarget target bone chain with IK but it has less than 3 joints."));
 		return false;
 	}
 
-	InvInitialLength = 1.0f / Length;
-
-	return true;
-}
-
-bool FTargetChainIK::Initialize(const FResolvedBoneChain& InTargetBoneChain, FIKRigLogger& InLog)
-{
-	if (InTargetBoneChain.BoneIndices.Num() < 3)
+	// store pointers to required data
+	SourceBoneChain = &InSourceBoneChain;
+	TargetBoneChain = &InTargetBoneChain;
+	Settings = &InSettings;
+	
+	// cache initial length of SOURCE chain
+	const FTransform& InitialSourceStartGlobal = InSourceBoneChain.RefPoseGlobalTransforms[0];
+	const FTransform& InitialSourceEndGlobal = InSourceBoneChain.RefPoseGlobalTransforms.Last();
+	const float InitialStartLength = static_cast<float>((InitialSourceStartGlobal.GetLocation() - InitialSourceEndGlobal.GetLocation()).Size());
+	if (InitialStartLength <= UE_KINDA_SMALL_NUMBER)
 	{
-		InLog.LogWarning(LOCTEXT("TargetChainLessThanThree", "IK Chains Op: trying to retarget target bone chain with IK but it has less than 3 joints."));
+		InLog.LogWarning(LOCTEXT("SourceZeroLengthIK", "IK Chains Op: found zero-length source bone chain."));
 		return false;
 	}
+	InvInitialLengthSource = 1.0f / InitialStartLength;
 
-	TargetBoneChain = &InTargetBoneChain;
-	BoneIndexA = InTargetBoneChain.BoneIndices[0];
-	BoneIndexC = InTargetBoneChain.BoneIndices.Last();
-	const FTransform& Last = InTargetBoneChain.RefPoseGlobalTransforms.Last();
-	PrevEndPosition = Last.GetLocation();
-	InitialEndPosition = Last.GetTranslation();
-	InitialEndRotation = Last.GetRotation();
-	InitialLength = static_cast<float>((InTargetBoneChain.RefPoseGlobalTransforms[0].GetTranslation() - Last.GetTranslation()).Size());
-
-	if (InitialLength <= UE_KINDA_SMALL_NUMBER)
+	// cache initial length of TARGET chain
+	const FTransform& InitialTargetStartGlobal = InTargetBoneChain.RefPoseGlobalTransforms[0];
+	const FTransform& InitialTargetEndGlobal = InTargetBoneChain.RefPoseGlobalTransforms.Last();
+	InitialTargetChainLength = static_cast<float>((InitialTargetStartGlobal.GetTranslation() - InitialTargetEndGlobal.GetTranslation()).Size());
+	if (InitialTargetChainLength <= UE_KINDA_SMALL_NUMBER)
 	{
 		InLog.LogWarning(LOCTEXT("TargetZeroLengthIK", "IK Retargeter trying to retarget target bone chain with IK, but it is zero length!"));
 		return false;
 	}
-	
+
+	// TODO
+	InitialTargetEndRelativeToSourceEnd = InitialTargetEndGlobal.GetRelativeTransform(InitialSourceEndGlobal);
+
 	return true;
 }
 
-bool FIKChainRetargeter::Initialize(
-	const FResolvedBoneChain& InSourceBoneChain,
-	const FResolvedBoneChain& InTargetBoneChain,
-	const FRetargetIKChainSettings& InSettings,
-	FIKRigLogger& InLog)
+const FTransform& FIKChainRetargeter::GenerateGoalTransform(
+	const FIKRetargetPelvisMotionOp* InPelvisMotionOp,
+	const TArray<FTransform>& InSourceGlobalPose,
+	const TArray<FTransform>& InTargetGlobalPose,
+	const FRetargetSkeleton& InSourceSkeleton,
+	const FRetargetSkeleton& InTargetSkeleton)
 {
-	bool bIsInitialized = true;
-	Settings = &InSettings;
-	bIsInitialized &= Source.Initialize(InSourceBoneChain, InLog);
-	bIsInitialized &= Target.Initialize(InTargetBoneChain, InLog);
-	return bIsInitialized;
+	// calculate ROTATION of IK goal
+	GoalOutput.SetRotation( CalculateGoalRotation(InSourceGlobalPose) );
+	// calculate POSITION of IK goal ...
+	GoalOutput.SetTranslation( CalculateGoalPosition(InPelvisMotionOp, InSourceGlobalPose, InTargetGlobalPose, InSourceSkeleton, InTargetSkeleton) );
+	return GoalOutput;
 }
 
-void FIKChainRetargeter::EncodePose(const TArray<FTransform>& InSourceGlobalPose)
+FQuat FIKChainRetargeter::CalculateGoalRotation(const TArray<FTransform>& InSourceGlobalPose) const
 {
-	const FVector Start = InSourceGlobalPose[Source.StartBoneIndex].GetTranslation();
-	const FVector End = InSourceGlobalPose[Source.EndBoneIndex].GetTranslation();
-
-    // get the normalized direction / length of the IK limb (how extended it is as percentage of original length)
-    const FVector ChainVector = End - Start;
-	double ChainLength;
-	FVector ChainDirection;
-	ChainVector.ToDirectionAndLength(ChainDirection, ChainLength);
-	const double NormalizedLimbLength = ChainLength * Source.InvInitialLength;
-
-	Source.PreviousEndPosition = Source.CurrentEndPosition;
-	Source.CurrentEndPosition = End;
-	Source.CurrentEndDirectionNormalized = ChainDirection * NormalizedLimbLength;
-	Source.CurrentEndRotation = InSourceGlobalPose[Source.EndBoneIndex].GetRotation();
-	Source.CurrentHeightFromGroundNormalized = static_cast<float>(End.Z - Source.InitialEndPosition.Z)  * Source.InvInitialLength;
-}
+	// get relevant bone transforms
+	const int32 SourceEndBoneIndex = SourceBoneChain->BoneIndices.Last();
+	const FTransform& InitialSourceEndGlobal = SourceBoneChain->RefPoseGlobalTransforms.Last();
+	const FTransform& InitialTargetEndGlobal = TargetBoneChain->RefPoseGlobalTransforms.Last();
+	const FTransform& CurrentSourceEndGlobal = InSourceGlobalPose[SourceEndBoneIndex];
 	
-void FIKChainRetargeter::DecodePose(
-	const FIKRetargetPelvisMotionOp* PelvisMotionOp,
-    const TArray<FTransform>& InGlobalPose)
-{
-	//
-	// calculate ROTATION of IK goal ...
-	//
-	
-	// apply delta rotation from input
-	const FQuat DeltaRotation = Source.CurrentEndRotation * Source.InitialEndRotation.Inverse();
-	FQuat GoalRotation = DeltaRotation * Target.InitialEndRotation;
+	// apply delta rotation from source input
+	const FQuat CurrentSourceEndRotation = CurrentSourceEndGlobal.GetRotation();
+	const FQuat SourceDeltaRotation = CurrentSourceEndRotation * InitialSourceEndGlobal.GetRotation().Inverse();
+	FQuat GoalRotation = SourceDeltaRotation * InitialTargetEndGlobal.GetRotation();
 
 	// blend to source rotation
 	const double BlendToSourceRotation = Settings->BlendToSource * Settings->BlendToSourceRotation;
 	if (BlendToSourceRotation > UE_KINDA_SMALL_NUMBER)
 	{
-		GoalRotation = FQuat::FastLerp(GoalRotation, Source.CurrentEndRotation, BlendToSourceRotation);
+		GoalRotation = FQuat::FastLerp(GoalRotation, CurrentSourceEndRotation, BlendToSourceRotation);
 		GoalRotation.Normalize();
 	}
 
-	// apply static rotation offset in the local space of the foot
+	// apply static rotation offset in the local space of the foot (AFTER blend to source)
 	GoalRotation = GoalRotation * Settings->StaticRotationOffset.Quaternion();
+	
+	return MoveTemp(GoalRotation);
+}
 
-	//
-	// calculate POSITION of IK goal ...
-	//
+FVector FIKChainRetargeter::CalculateGoalPosition(
+	const FIKRetargetPelvisMotionOp* InPelvisMotionOp,
+	const TArray<FTransform>& InSourceGlobalPose,
+	const TArray<FTransform>& InTargetGlobalPose,
+	const FRetargetSkeleton& InSourceSkeleton,
+	const FRetargetSkeleton& InTargetSkeleton) const
+{
+	const int32 SourceStartBoneIndex = SourceBoneChain->BoneIndices[0];
+	const int32 SourceEndBoneIndex = SourceBoneChain->BoneIndices.Last();
+	const FTransform& CurrentSourceStartGlobal = InSourceGlobalPose[SourceStartBoneIndex];
+	const FTransform& CurrentSourceEndGlobal = InSourceGlobalPose[SourceEndBoneIndex];
+	
+	// get the current normalized direction / length of the IK limb (how extended it is as percentage of original length)
+	const FVector CurrentSourceChainVector = CurrentSourceEndGlobal.GetLocation() - CurrentSourceStartGlobal.GetLocation();
+	double CurrentSourceChainLength;
+	FVector CurrentSourceChainDirection;
+	CurrentSourceChainVector.ToDirectionAndLength(CurrentSourceChainDirection, CurrentSourceChainLength);
+	const double NormalizedLimbLength = CurrentSourceChainLength * InvInitialLengthSource;
+	const FVector CurrentSourceEndDirectionNormalized = CurrentSourceChainDirection * NormalizedLimbLength;
 	
 	// set position to length-scaled direction from source limb
-	const FVector PelvisTranslationDelta = PelvisMotionOp ? PelvisMotionOp->GetPelvisTranslationOffset() : FVector::ZeroVector;
-	const FVector AffectIKWeights = PelvisMotionOp ? PelvisMotionOp->GetAffectIKWeightAsVector() : FVector::ZeroVector;
+	const FVector PelvisTranslationDelta = InPelvisMotionOp ? InPelvisMotionOp->GetPelvisTranslationOffset() : FVector::ZeroVector;
+	const FVector AffectIKWeights = InPelvisMotionOp ? InPelvisMotionOp->GetAffectIKWeightAsVector() : FVector::ZeroVector;
 	const FVector InvAffectIKWeights = FVector::OneVector - AffectIKWeights;
 	const FVector InvRootModification = PelvisTranslationDelta * InvAffectIKWeights;
-	const FVector Start = InGlobalPose[Target.BoneIndexA].GetTranslation() - InvRootModification;
-	FVector GoalPosition = Start + (Source.CurrentEndDirectionNormalized * Target.InitialLength);
+	const int32 TargetStartBoneIndex = TargetBoneChain->BoneIndices[0];
+	const FVector CurrentTargetStartLocation = InTargetGlobalPose[TargetStartBoneIndex].GetTranslation() - InvRootModification;
+	FVector GoalPosition = CurrentTargetStartLocation + (CurrentSourceEndDirectionNormalized * InitialTargetChainLength);
 
 	// blend to source location
 	const double BlendToSourceTranslation = Settings->BlendToSource * Settings->BlendToSourceTranslation;
@@ -148,7 +142,7 @@ void FIKChainRetargeter::DecodePose(
 	{
 		const FVector RootModification = PelvisTranslationDelta * AffectIKWeights;
 		const FVector Weight = BlendToSourceTranslation * Settings->BlendToSourceWeights;
-		const FVector SourceLocation = Source.CurrentEndPosition + RootModification;
+		const FVector SourceLocation = CurrentSourceEndGlobal.GetLocation() + RootModification;
 		GoalPosition.X = FMath::Lerp(GoalPosition.X, SourceLocation.X, Weight.X);
 		GoalPosition.Y = FMath::Lerp(GoalPosition.Y, SourceLocation.Y, Weight.Y);
 		GoalPosition.Z = FMath::Lerp(GoalPosition.Z, SourceLocation.Z, Weight.Z);
@@ -158,7 +152,7 @@ void FIKChainRetargeter::DecodePose(
 	GoalPosition += Settings->StaticOffset;
 
 	// apply local static offset
-	GoalPosition += GoalRotation.RotateVector(Settings->StaticLocalOffset);
+	GoalPosition += GoalOutput.GetRotation().RotateVector(Settings->StaticLocalOffset);
 
 	// apply vertical scale
 	GoalPosition.Z *= Settings->ScaleVertical;
@@ -166,13 +160,10 @@ void FIKChainRetargeter::DecodePose(
 	// apply extension
 	if (!FMath::IsNearlyEqual(Settings->Extension, 1.0f))
 	{
-		GoalPosition = Start + (GoalPosition - Start) * Settings->Extension;	
+		GoalPosition = CurrentTargetStartLocation + (GoalPosition - CurrentTargetStartLocation) * Settings->Extension;	
 	}
-	
-	// output transform
-	Results.EndEffectorPosition = GoalPosition;
-	Results.EndEffectorRotation = GoalRotation;
-	Target.PrevEndPosition = GoalPosition;
+
+	return MoveTemp(GoalPosition);
 }
 
 bool FRetargetIKChainSettings::operator==(const FRetargetIKChainSettings& Other) const
@@ -327,30 +318,30 @@ void FIKRetargetIKChainsOp::Run(
 	{
 		return; // skip this op when IK is off
 	}
-	
+
+	// used to optionally apply pelvis offsets to the goals
 	const FIKRetargetPelvisMotionOp* PelvisMotionOp = InProcessor.GetFirstRetargetOpOfType<FIKRetargetPelvisMotionOp>();
-	
-	// retarget the IK goals to their new locations based on input pose
-	for (FIKChainRetargeter& IKChainRetargeter : IKChainRetargeters)
-	{
-		// encode them all using the input pose
-		IKChainRetargeter.EncodePose(InSourceGlobalPose);
-		// decode the IK goal and apply to IKRig
-		IKChainRetargeter.DecodePose(PelvisMotionOp, OutTargetGlobalPose);
-	}
 
-	// set the goal transform on the IK Rig
+	// retarget the IK goals to their new locations based on input pose 
 	FIKRigGoalContainer& GoalContainer = InProcessor.GetIKRigGoalContainer();
-	for (const FIKChainRetargeter& IKChain : IKChainRetargeters)
+	for (FIKChainRetargeter& IKChain : IKChainRetargeters)
 	{
+		// generate the new goal transform for the chain
+		const FTransform& GoalTransform = IKChain.GenerateGoalTransform(
+			PelvisMotionOp,
+			InSourceGlobalPose,
+			OutTargetGlobalPose,
+			InProcessor.GetSkeleton(ERetargetSourceOrTarget::Source),
+			InProcessor.GetSkeleton(ERetargetSourceOrTarget::Target));
+
+		// set the goal transform on the IK Rig
 		constexpr double PositionAlpha = 1.0;
 		constexpr double RotationAlpha = 1.0;
-		
 		FIKRigGoal Goal = FIKRigGoal(
 			IKChain.GetTargetChain().IKGoalName,
 			IKChain.GetTargetChain().EndBone,
-			IKChain.GetResults().EndEffectorPosition,
-			IKChain.GetResults().EndEffectorRotation,
+			GoalTransform.GetLocation(),
+			GoalTransform.GetRotation(),
 			PositionAlpha,
 			RotationAlpha,
 			EIKRigGoalSpace::Component,
@@ -482,7 +473,7 @@ void FIKRetargetIKChainsOp::DebugDraw(
 				InPDI,
 				SourceGoalTransform,
 				Color,
-				Settings.GoalDrawThickness,
+				Settings.GoalDrawSize,
 				12,
 				SDPG_World,
 				0.0f,
@@ -517,14 +508,13 @@ void FIKRetargetIKChainsOp::SaveDebugData(
 	for (const FIKChainRetargeter& IKChainPair : IKChainRetargeters)
 	{
 		FChainDebugData NewChainData;
-		const FSourceChainIK& Source = IKChainPair.GetSource();
-		const FTargetChainIK& Target = IKChainPair.GetTarget();
-		const FDecodedIKChain& Results = IKChainPair.GetResults();
+		const FResolvedBoneChain& SourceChain = IKChainPair.GetSourceChain();
+		const FResolvedBoneChain& TargetChain = IKChainPair.GetTargetChain();
 		NewChainData.TargetChainName = IKChainPair.GetSettings()->TargetChainName;
-		NewChainData.InputTransformStart = OutTargetGlobalPose[Target.BoneIndexA];
-		NewChainData.InputTransformEnd = OutTargetGlobalPose[Target.BoneIndexC];
-		NewChainData.OutputTransformEnd = FTransform(Results.EndEffectorRotation, Results.EndEffectorPosition);
-		NewChainData.SourceTransformEnd = InSourceGlobalPose[Source.EndBoneIndex];
+		NewChainData.InputTransformStart = OutTargetGlobalPose[TargetChain.BoneIndices[0]];
+		NewChainData.InputTransformEnd = OutTargetGlobalPose[TargetChain.BoneIndices.Last()];
+		NewChainData.OutputTransformEnd = IKChainPair.GetGoalTransform();
+		NewChainData.SourceTransformEnd = InSourceGlobalPose[SourceChain.BoneIndices.Last()];
 		AllChainsDebugData.Add(NewChainData);
 	}
 
diff --git a/Engine/Plugins/Animation/IKRig/Source/IKRig/Private/Retargeter/RetargetOps/PelvisMotionOp.cpp b/Engine/Plugins/Animation/IKRig/Source/IKRig/Private/Retargeter/RetargetOps/PelvisMotionOp.cpp
--- a/Engine/Plugins/Animation/IKRig/Source/IKRig/Private/Retargeter/RetargetOps/PelvisMotionOp.cpp
+++ b/Engine/Plugins/Animation/IKRig/Source/IKRig/Private/Retargeter/RetargetOps/PelvisMotionOp.cpp
@@ -95,18 +95,17 @@ void FIKRetargetPelvisMotionOp::Run(
 	const TArray<FTransform>& InSourceGlobalPose,
 	TArray<FTransform>& OutTargetGlobalPose)
 {
-	FTransform NewPelvisGlobalTransform;
-	
 	EncodePose(InSourceGlobalPose);
-	DecodePose(NewPelvisGlobalTransform);
+
+	FTransform PelvisGlobal;
+	DecodePose(PelvisGlobal);
 	
 	// update global transforms below root
 	FTargetSkeleton& TargetSkeleton = Processor.GetTargetSkeleton();
-	TargetSkeleton.SetGlobalTransformAndUpdateChildren(Target.BoneIndex, NewPelvisGlobalTransform, TargetSkeleton.OutputGlobalPose);
+	TargetSkeleton.SetGlobalTransformAndUpdateChildren(Target.BoneIndex, PelvisGlobal, TargetSkeleton.OutputGlobalPose);
 
 #if WITH_EDITOR
-	FScopeLock ScopeLock(&DebugDataMutex);
-	DebugPelvisTransform = OutTargetGlobalPose[Target.BoneIndex];
+	CurrentPelvisTransform = PelvisGlobal;
 #endif
 }
 
@@ -236,6 +235,11 @@ FVector FIKRetargetPelvisMotionOp::GetPelvisTranslationOffset() const
 	return Target.PelvisTranslationDelta;
 }
 
+int32 FIKRetargetPelvisMotionOp::GetPelvisBoneIndex(ERetargetSourceOrTarget SourceOrTarget) const
+{
+	return SourceOrTarget == ERetargetSourceOrTarget::Source ? Source.BoneIndex : Target.BoneIndex;
+}
+
 void FIKRetargetPelvisMotionOp::Reset()
 {
 	Source = FPelvisSource();
@@ -329,7 +333,7 @@ void FIKRetargetPelvisMotionOp::DebugDraw(
 
 	FScopeLock ScopeLock(&DebugDataMutex);
 	
-	const FTransform PelvisTransform = DebugPelvisTransform * InComponentTransform;
+	const FTransform PelvisTransform = CurrentPelvisTransform * InComponentTransform;
 	const FVector PelvisCircleLocation = PelvisTransform.GetLocation() * FVector(1,1,0);
 	const FLinearColor CircleColor = InEditorState.bIsRootSelected ? InEditorState.MainColor : InEditorState.MainColor * InEditorState.NonSelected;
 
diff --git a/Engine/Plugins/Animation/IKRig/Source/IKRig/Public/Retargeter/RetargetOps/IKChainsOp.h b/Engine/Plugins/Animation/IKRig/Source/IKRig/Public/Retargeter/RetargetOps/IKChainsOp.h
--- a/Engine/Plugins/Animation/IKRig/Source/IKRig/Public/Retargeter/RetargetOps/IKChainsOp.h
+++ b/Engine/Plugins/Animation/IKRig/Source/IKRig/Public/Retargeter/RetargetOps/IKChainsOp.h
@@ -13,45 +13,6 @@ struct FIKRetargetRunIKRigOp;
 struct FResolvedBoneChain;
 struct FIKRetargetPelvisMotionOp;
 
-struct FSourceChainIK
-{
-	bool Initialize(const FResolvedBoneChain& InSourceBoneChain, FIKRigLogger& InLog);
-
-	const FResolvedBoneChain* SourceBoneChain;
-	int32 StartBoneIndex = INDEX_NONE;
-	int32 EndBoneIndex = INDEX_NONE;
-	
-	FVector InitialEndPosition = FVector::ZeroVector;
-	FQuat InitialEndRotation = FQuat::Identity;
-	double InvInitialLength = 1.0f;
-
-	// results after encoding...
-	FVector PreviousEndPosition = FVector::ZeroVector;
-	FVector CurrentEndPosition = FVector::ZeroVector;
-	FVector CurrentEndDirectionNormalized = FVector::ZeroVector;
-	FQuat CurrentEndRotation = FQuat::Identity;
-	double CurrentHeightFromGroundNormalized = 0.0f;
-};
-
-struct FTargetChainIK
-{
-	bool Initialize(const FResolvedBoneChain& InTargetBoneChain, FIKRigLogger& InLog);
-
-	const FResolvedBoneChain* TargetBoneChain;
-	int32 BoneIndexA = INDEX_NONE;
-	int32 BoneIndexC = INDEX_NONE;
-	
-	double InitialLength = 1.0f;
-	FVector InitialEndPosition = FVector::ZeroVector;
-	FQuat InitialEndRotation = FQuat::Identity;
-	FVector PrevEndPosition = FVector::ZeroVector;
-};
-
-struct FDecodedIKChain
-{
-	FVector EndEffectorPosition = FVector::ZeroVector;
-	FQuat EndEffectorRotation = FQuat::Identity;
-};
 
 struct FIKChainRetargeter
 {
@@ -61,28 +22,42 @@ struct FIKChainRetargeter
 		const FRetargetIKChainSettings& InSettings,
 		FIKRigLogger& InLog);
 
-	void EncodePose(const TArray<FTransform> &SourceInputGlobalPose);
-	
-	void DecodePose(
-		const FIKRetargetPelvisMotionOp* PelvisMotionOp,
-		const TArray<FTransform>& InGlobalPose);
-	
-	const FDecodedIKChain& GetResults() const { return Results; };
-
-	const FSourceChainIK& GetSource() const {return Source; };
+	const FTransform& GenerateGoalTransform(
+		const FIKRetargetPelvisMotionOp* InPelvisMotionOp,
+		const TArray<FTransform>& InSourceGlobalPose,
+		const TArray<FTransform>& InTargetGlobalPose,
+		const FRetargetSkeleton& InSourceSkeleton,
+		const FRetargetSkeleton& InTargetSkeleton);
 
-	const FTargetChainIK& GetTarget() const {return Target; };
+	const FResolvedBoneChain& GetSourceChain() const { return *SourceBoneChain; };
 	
-	const FResolvedBoneChain& GetTargetChain() const { return *Target.TargetBoneChain; };
+	const FResolvedBoneChain& GetTargetChain() const { return *TargetBoneChain; };
 	
 	const FRetargetIKChainSettings* GetSettings() const { return Settings; };
 
+	const FTransform& GetGoalTransform() const {return GoalOutput; }
+
 private:
 
-	FSourceChainIK Source;
-	FTargetChainIK Target;
-	FDecodedIKChain Results;
+	FQuat CalculateGoalRotation(const TArray<FTransform>& InSourceGlobalPose) const;
+
+	FVector CalculateGoalPosition(
+		const FIKRetargetPelvisMotionOp* InPelvisMotionOp,
+		const TArray<FTransform>& InSourceGlobalPose,
+		const TArray<FTransform>& InTargetGlobalPose,
+		const FRetargetSkeleton& InSourceSkeleton,
+		const FRetargetSkeleton& InTargetSkeleton) const;
+
+	// cached data
+	FTransform GoalOutput;
+	const FResolvedBoneChain* SourceBoneChain;
+	const FResolvedBoneChain* TargetBoneChain;
 	const FRetargetIKChainSettings* Settings;
+	double InvInitialLengthSource = 1.0f;
+	double InitialTargetChainLength = 1.0f;
+
+	
+	FTransform InitialTargetEndRelativeToSourceEnd;
 };
 
 USTRUCT(BlueprintType)
diff --git a/Engine/Plugins/Animation/IKRig/Source/IKRig/Public/Retargeter/RetargetOps/PelvisMotionOp.h b/Engine/Plugins/Animation/IKRig/Source/IKRig/Public/Retargeter/RetargetOps/PelvisMotionOp.h
--- a/Engine/Plugins/Animation/IKRig/Source/IKRig/Public/Retargeter/RetargetOps/PelvisMotionOp.h
+++ b/Engine/Plugins/Animation/IKRig/Source/IKRig/Public/Retargeter/RetargetOps/PelvisMotionOp.h
@@ -157,6 +157,8 @@ struct FIKRetargetPelvisMotionOp : public FIKRetargetOpBase
 
 	IKRIG_API FVector GetPelvisTranslationOffset() const;
 
+	int32 GetPelvisBoneIndex(ERetargetSourceOrTarget SourceOrTarget) const;
+
 	UPROPERTY()
 	FIKRetargetPelvisMotionOpSettings Settings;
 
@@ -191,7 +193,7 @@ public:
 		const FTransform& InComponentTransform,
 		const double InComponentScale,
 		const FIKRetargetDebugDrawState& InEditorState) const override;
-	FTransform DebugPelvisTransform;
+	FTransform CurrentPelvisTransform;
 	static IKRIG_API FCriticalSection DebugDataMutex;
 #endif
 };
