[IK Retarget] Fixed bug with FKChain Op's Orient and Scale on bones with zero length.

diff --git a/Engine/Plugins/Animation/IKRig/Source/IKRig/Private/Retargeter/RetargetOps/FKChainsOp.cpp b/Engine/Plugins/Animation/IKRig/Source/IKRig/Private/Retargeter/RetargetOps/FKChainsOp.cpp
--- a/Engine/Plugins/Animation/IKRig/Source/IKRig/Private/Retargeter/RetargetOps/FKChainsOp.cpp
+++ b/Engine/Plugins/Animation/IKRig/Source/IKRig/Private/Retargeter/RetargetOps/FKChainsOp.cpp
@@ -315,30 +315,50 @@ void FChainDecoderFK::DecodePose(
 			}
 		case EFKChainTranslationMode::OrientAndScale:
 			{
+				const double TargetRefLength = CachedData.TargetRefPoseLengths[ChainIndex];
+				const double SourceRefLength = CachedData.SourceRefPoseLengths[ChainIndex];
+
 				// calculate current direction of source bone
 				const int32 SourceBoneIndex = CachedData.SourceBoneIndices[ChainIndex];
 				const FTransform& SourceCurrentGlobal = InSourceGlobalPose[SourceBoneIndex];
 				const FTransform& SourceParentCurrentGlobal = SourceSkeleton.GetParentTransform(SourceBoneIndex, InSourceGlobalPose);
-				const FVector SourceCurrentDirection = SourceCurrentGlobal.GetTranslation() - SourceParentCurrentGlobal.GetTranslation();
+				const FVector SourceCurrentTranslation = SourceCurrentGlobal.GetTranslation() - SourceParentCurrentGlobal.GetTranslation();
 
 				// calculate current stretch of source bone
 				double SourceCurrentLength;
 				FVector SourceCurrentNorm;
-				SourceCurrentDirection.ToDirectionAndLength(SourceCurrentNorm, SourceCurrentLength);
-				const double SourceLengthScaleFactor = SourceCurrentLength / CachedData.SourceRefPoseLengths[ChainIndex];
+				SourceCurrentTranslation.ToDirectionAndLength(SourceCurrentNorm, SourceCurrentLength);
 
-				// calculate equivalently stretched target bone in direction of source
-				const double TargetRefLength = CachedData.TargetRefPoseLengths[ChainIndex];
-				FVector TargetOffset = SourceCurrentNorm * (TargetRefLength * SourceLengthScaleFactor);
+				// generate a translational offset for the target bone
+				FVector TargetNewGlobalOffset = FVector::ZeroVector;
 
-				// apply rotational delta from the ref pose to the translational offset
-				const FQuat& RefPoseGlobalDeltaRotation = CachedData.FromToGlobalRotationDeltas[ChainIndex];
-				const FQuat CurrentGlobalRotationPlusDelta = RefPoseGlobalDeltaRotation * OutRotation;
-				const FQuat CurrentGlobalRotationDelta = CurrentGlobalRotationPlusDelta * OutRotation.Inverse();
-				TargetOffset = CurrentGlobalRotationDelta.RotateVector(TargetOffset);
+				// in the absence of a proper source and target bone length, we treat translational deltas as additive rather than scaled
+				constexpr double ThresholdLength = 0.1f;
+				if (TargetRefLength <= ThresholdLength || SourceRefLength <= ThresholdLength)
+				{
+					// additive offset
+					const FVector SourceCurrentGlobalOffset = SourceCurrentNorm * (SourceCurrentLength - SourceRefLength);
+					// add it to the current local offset
+					const FVector InitialLocalOffset = TargetSkeleton.RetargetPoses.GetLocalRetargetPose()[BoneIndex].GetTranslation();
+					const FVector TargetCurrentGlobalOffset = ParentGlobalTransform.TransformPosition(InitialLocalOffset) - ParentGlobalTransform.GetLocation();
+					// add the source offset to the target
+					TargetNewGlobalOffset = TargetCurrentGlobalOffset + SourceCurrentGlobalOffset;
+				}
+				else
+				{
+					// scaled offset
+					const double SourceLengthScaleFactor = SourceCurrentLength / SourceRefLength;
+					// calculate equivalently stretched target bone in direction of source
+					TargetNewGlobalOffset = SourceCurrentNorm * (TargetRefLength * SourceLengthScaleFactor);
+					// apply rotational delta from the ref pose to the translational offset
+					const FQuat& RefPoseGlobalDeltaRotation = CachedData.FromToGlobalRotationDeltas[ChainIndex];
+					const FQuat CurrentGlobalRotationPlusDelta = RefPoseGlobalDeltaRotation * OutRotation;
+					const FQuat CurrentGlobalRotationDelta = CurrentGlobalRotationPlusDelta * OutRotation.Inverse();
+					TargetNewGlobalOffset = CurrentGlobalRotationDelta.RotateVector(TargetNewGlobalOffset);
+				}
 
 				// apply target offset to the parent global position
-				OutPosition = ParentGlobalTransform.GetLocation() + TargetOffset;
+				OutPosition = ParentGlobalTransform.GetLocation() + TargetNewGlobalOffset;
 				break;
 			}
 		default:
