[IK Rig] Added null-check to goal creators in IK Rig anim node.

diff --git a/Engine/Plugins/Animation/IKRig/Source/IKRig/Private/AnimNodes/AnimNode_IKRig.cpp b/Engine/Plugins/Animation/IKRig/Source/IKRig/Private/AnimNodes/AnimNode_IKRig.cpp
--- a/Engine/Plugins/Animation/IKRig/Source/IKRig/Private/AnimNodes/AnimNode_IKRig.cpp
+++ b/Engine/Plugins/Animation/IKRig/Source/IKRig/Private/AnimNodes/AnimNode_IKRig.cpp
@@ -276,10 +276,12 @@ void FAnimNode_IKRig::PreUpdate(const UAnimInstance* InAnimInstance)
 	// this is done on the main thread because we're talking to actor components here
 	GoalsFromGoalCreators.Reset();
 	for (UActorComponent* GoalCreatorComponent : GoalCreators)
-	{
-		IIKGoalCreatorInterface* GoalCreator = Cast<IIKGoalCreatorInterface>(GoalCreatorComponent);
-		GoalCreator->AddIKGoals_Implementation(GoalsFromGoalCreators);
-	}
+    {
+    	if (IIKGoalCreatorInterface* GoalCreator = Cast<IIKGoalCreatorInterface>(GoalCreatorComponent))
+    	{
+    		GoalCreator->AddIKGoals_Implementation(GoalsFromGoalCreators);
+    	}
+    }
 }
 
 void FAnimNode_IKRig::OnInitializeAnimInstance(const FAnimInstanceProxy* InProxy, const UAnimInstance* InAnimInstance)
