[IK Retarget] Adding new "Stretch Chain" op

diff --git a/Engine/Plugins/Animation/IKRig/Source/IKRig/Private/Retargeter/IKRetargetProcessor.cpp b/Engine/Plugins/Animation/IKRig/Source/IKRig/Private/Retargeter/IKRetargetProcessor.cpp
--- a/Engine/Plugins/Animation/IKRig/Source/IKRig/Private/Retargeter/IKRetargetProcessor.cpp
+++ b/Engine/Plugins/Animation/IKRig/Source/IKRig/Private/Retargeter/IKRetargetProcessor.cpp
@@ -318,6 +318,55 @@ FTransform FRetargetSkeleton::GetLocalTransformOfSingleBone(
 	return GlobalTransform.GetRelativeTransform(ParentGlobalTransform);
 }
 
+TArray<FTransform> FRetargetSkeleton::GetLocalTransformsOfMultipleBones(
+	const TArray<int32>& InBoneIndices,
+	const TArray<FTransform>& InGlobalPose) const
+{
+	TArray<FTransform> OutputLocalTransforms;
+
+	for (const int32 ChildIndex : InBoneIndices)
+	{
+		const int32 ParentIndex = ParentIndices[ChildIndex];
+		if (ParentIndex == INDEX_NONE)
+		{
+			OutputLocalTransforms.Add(InGlobalPose[ChildIndex]);
+			continue;
+		}
+		const FTransform& ChildGlobalTransform = InGlobalPose[ChildIndex];
+		const FTransform& ParentGlobalTransform = InGlobalPose[ParentIndex];
+		const FTransform LocalTransform = ChildGlobalTransform.GetRelativeTransform(ParentGlobalTransform);
+		OutputLocalTransforms.Add(LocalTransform);
+	}
+	
+	return MoveTemp(OutputLocalTransforms);
+}
+
+void FRetargetSkeleton::UpdateGlobalTransformsOfMultipleBones(
+	const TArray<int32>& InBoneIndices,
+	const TArray<FTransform>& InLocalTransforms,
+	TArray<FTransform>& OutGlobalPose) const
+{
+	if (!ensure(InLocalTransforms.Num() == InBoneIndices.Num()))
+	{
+		return;
+	}
+	
+	for (int32 BoneArrayIndex=0; BoneArrayIndex < InBoneIndices.Num(); ++BoneArrayIndex)
+	{
+		const int32 ChildIndex = InBoneIndices[BoneArrayIndex];
+		const int32 ParentIndex = ParentIndices[ChildIndex];
+		if (ParentIndex == INDEX_NONE)
+		{
+			// root always in global space already, no conversion required
+			OutGlobalPose[ChildIndex] = InLocalTransforms[BoneArrayIndex];
+			continue; 
+		}
+		const FTransform& ChildLocalTransform = InLocalTransforms[BoneArrayIndex];
+		const FTransform& ParentGlobalTransform = OutGlobalPose[ParentIndex];
+		OutGlobalPose[ChildIndex] = ChildLocalTransform * ParentGlobalTransform;
+	}
+}
+
 int32 FRetargetSkeleton::GetCachedEndOfBranchIndex(const int32 InBoneIndex) const
 {
 	if (!CachedEndOfBranchIndices.IsValidIndex(InBoneIndex))
diff --git a/Engine/Plugins/Animation/IKRig/Source/IKRig/Private/Retargeter/RetargetOps/AlignPoleVectorOp.cpp b/Engine/Plugins/Animation/IKRig/Source/IKRig/Private/Retargeter/RetargetOps/AlignPoleVectorOp.cpp
--- a/Engine/Plugins/Animation/IKRig/Source/IKRig/Private/Retargeter/RetargetOps/AlignPoleVectorOp.cpp
+++ b/Engine/Plugins/Animation/IKRig/Source/IKRig/Private/Retargeter/RetargetOps/AlignPoleVectorOp.cpp
@@ -60,20 +60,8 @@ void FPoleVectorMatcher::MatchPoleVector(
 		return;
 	}
 	
-	// update local spaces of all bones within chain
-	TArray<FTransform> LocalSpaces;
-	for (const int32 ChildIndex : AllChildrenWithinChain)
-	{
-		const int32 ParentIndex = TargetSkeleton.ParentIndices[ChildIndex];
-		if (ParentIndex == INDEX_NONE)
-		{
-			LocalSpaces.Add(OutTargetGlobalPose[ChildIndex]);
-		}
-		const FTransform& ChildGlobalTransform = OutTargetGlobalPose[ChildIndex];
-		const FTransform& ParentGlobalTransform = OutTargetGlobalPose[ParentIndex];
-		const FTransform LocalTransform = ChildGlobalTransform.GetRelativeTransform(ParentGlobalTransform);
-		LocalSpaces.Add(LocalTransform);
-	}
+	// update local spaces of all bones within chain as they are passed in
+	TArray<FTransform> AllChildrenInputLocalTransforms = TargetSkeleton.GetLocalTransformsOfMultipleBones(AllChildrenWithinChain, OutTargetGlobalPose);
 
 	// normalized vector pointing from root to tip of chain
 	const FVector TargetChainAxisNorm = GetChainAxisNormalized(TargetBoneChain->BoneIndices, OutTargetGlobalPose);
@@ -103,21 +91,8 @@ void FPoleVectorMatcher::MatchPoleVector(
 	FTransform& BaseOfChain = OutTargetGlobalPose[TargetBoneChain->BoneIndices[0]];
 	BaseOfChain.SetRotation(MatchingRotation * OffsetRotation * BaseOfChain.GetRotation());
 
-	// now update global pose of all bones within chain
-	for (int32 ChainIndex=0; ChainIndex < AllChildrenWithinChain.Num(); ++ChainIndex)
-	{
-		const int32 ChildIndex = AllChildrenWithinChain[ChainIndex];
-		const int32 ParentIndex = TargetSkeleton.ParentIndices[ChildIndex];
-		if (ParentIndex == INDEX_NONE)
-		{
-			// root always in global space already, no conversion required
-			OutTargetGlobalPose[ChildIndex] = LocalSpaces[ChainIndex];
-			return; 
-		}
-		const FTransform& ChildLocalTransform = LocalSpaces[ChainIndex];
-		const FTransform& ParentGlobalTransform = OutTargetGlobalPose[ParentIndex];
-		OutTargetGlobalPose[ChildIndex] = ChildLocalTransform * ParentGlobalTransform;
-	}
+	// now propagate changes to children of the chain
+	TargetSkeleton.UpdateGlobalTransformsOfMultipleBones(AllChildrenWithinChain, AllChildrenInputLocalTransforms,OutTargetGlobalPose);
 }
 
 EAxis::Type FPoleVectorMatcher::CalculateBestPoleAxisForChain(
diff --git a/Engine/Plugins/Animation/IKRig/Source/IKRig/Private/Retargeter/RetargetOps/StretchChainOp.cpp b/Engine/Plugins/Animation/IKRig/Source/IKRig/Private/Retargeter/RetargetOps/StretchChainOp.cpp
new file mode 100644
--- /dev/null
+++ b/Engine/Plugins/Animation/IKRig/Source/IKRig/Private/Retargeter/RetargetOps/StretchChainOp.cpp
@@ -0,0 +1,338 @@
+﻿// Copyright Epic Games, Inc. All Rights Reserved.
+
+#include "Retargeter/RetargetOps/StretchChainOp.h"
+
+#include "Engine/SkeletalMesh.h"
+#include "Retargeter/IKRetargetProcessor.h"
+
+#define LOCTEXT_NAMESPACE "StretchChainOp"
+
+bool FChainStretcher::Initialize(
+	const FRetargetStretchChainSettings* InSettings,
+	const FResolvedBoneChain& InSourceBoneChain,
+	const FResolvedBoneChain& InTargetBoneChain,
+	const FRetargetSkeleton& InSourceSkeleton,
+	const FRetargetSkeleton& InTargetSkeleton)
+{
+	Settings = InSettings;
+	SourceBoneChain = &InSourceBoneChain;
+	TargetBoneChain = &InTargetBoneChain;
+
+	// cache indices of all children of all bones in chain (recursively)
+	AllChildrenOfChain.Reset();
+	InTargetSkeleton.GetChildrenIndicesRecursive(TargetBoneChain->BoneIndices[0],AllChildrenOfChain);
+	// filter out indices of the chain itself
+	const TArray<int32>& IndicesToRemove = TargetBoneChain->BoneIndices;
+	AllChildrenOfChain.RemoveAll([&IndicesToRemove](const int32& Element)
+	{
+		return IndicesToRemove.Contains(Element);
+	});
+	
+	return true;
+}
+
+void FChainStretcher::StretchChain(
+	const TArray<FTransform> &InSourceGlobalPose,
+	const FRetargetSkeleton& InTargetSkeleton,
+	TArray<FTransform> &OutTargetGlobalPose) const
+{
+	const bool bMatchingSourceLength = !FMath::IsNearlyZero(Settings->MatchSourceLength);
+	const bool bScalingChainLength = !FMath::IsNearlyEqual(Settings->ScaleChainLength, 1.0f);
+	if (!bMatchingSourceLength && !bScalingChainLength)
+	{
+		// no operations to perform
+		return;
+	}
+
+	// store input local spaces of all children to propagate to
+	TArray<FTransform> AllChildrenInputLocalTransforms = InTargetSkeleton.GetLocalTransformsOfMultipleBones(AllChildrenOfChain, OutTargetGlobalPose);
+
+	// get the lengths of the chains as they currently are at this point in the op stack
+	const TArray<FTransform> InputTargetChainGlobal = TargetBoneChain->GetChainTransformsFromPose(OutTargetGlobalPose);
+	const TArray<FTransform> InputSourceChainGlobal = SourceBoneChain->GetChainTransformsFromPose(InSourceGlobalPose);
+	const double CurrentTargetChainLength = FResolvedBoneChain::GetChainLength(InputTargetChainGlobal);
+	const double CurrentSourceChainLength = FResolvedBoneChain::GetChainLength(InputSourceChainGlobal);
+	if (FMath::IsNearlyZero(CurrentTargetChainLength))
+	{
+		return; // cannot scale a collapsed chain
+	}
+
+	// determine amount to scale the chain to match the source
+	double TargetLengthToMatch = FMath::Lerp(CurrentTargetChainLength, CurrentSourceChainLength, Settings->MatchSourceLength);
+	const double MatchLengthMultiplier = TargetLengthToMatch / CurrentTargetChainLength;
+
+	// scale the length of each bone in the chain
+	for (int32 ChainIndex = 1; ChainIndex < TargetBoneChain->BoneIndices.Num(); ChainIndex++)
+	{
+		const FVector InputBoneVectorGlobal = InputTargetChainGlobal[ChainIndex].GetLocation() - InputTargetChainGlobal[ChainIndex-1].GetLocation();
+		const FVector NewBoneVector = InputBoneVectorGlobal * (MatchLengthMultiplier * Settings->ScaleChainLength);
+		const int32 BoneIndex = TargetBoneChain->BoneIndices[ChainIndex];
+		const int32 ParentBoneIndex = TargetBoneChain->BoneIndices[ChainIndex-1];
+		const FVector NewLocation = OutTargetGlobalPose[ParentBoneIndex].GetLocation() + NewBoneVector;
+		OutTargetGlobalPose[BoneIndex].SetLocation(NewLocation);
+	}
+
+	// now propagate changes to children of the chain
+	InTargetSkeleton.UpdateGlobalTransformsOfMultipleBones(AllChildrenOfChain, AllChildrenInputLocalTransforms,OutTargetGlobalPose);
+}
+
+bool FRetargetStretchChainSettings::operator==(const FRetargetStretchChainSettings& Other) const
+{
+	return bEnabled == Other.bEnabled
+		&& FMath::IsNearlyEqualByULP(MatchSourceLength,Other.MatchSourceLength)
+		&& FMath::IsNearlyEqualByULP(ScaleChainLength, Other.ScaleChainLength);
+}
+
+const UClass* FIKRetargetStretchChainOpSettings::GetControllerType() const
+{
+	return UIKRetargetStretchChainController::StaticClass();
+}
+
+void FIKRetargetStretchChainOpSettings::CopySettingsAtRuntime(const FIKRetargetOpSettingsBase* InSettingsToCopyFrom)
+{
+	// copy settings only for chains that the op has initialized
+	const FIKRetargetStretchChainOpSettings* NewSettings = reinterpret_cast<const FIKRetargetStretchChainOpSettings*>(InSettingsToCopyFrom);
+	for (const FRetargetStretchChainSettings& NewChainSettings : NewSettings->ChainsToStretch)
+	{
+		for (FRetargetStretchChainSettings& ChainSettings : ChainsToStretch)
+		{
+			if (ChainSettings.TargetChainName == NewChainSettings.TargetChainName)
+			{
+				ChainSettings = NewChainSettings;
+			}
+		}
+	}
+}
+
+bool FIKRetargetStretchChainOp::Initialize(
+	const FIKRetargetProcessor& InProcessor,
+	const FRetargetSkeleton& InSourceSkeleton,
+	const FTargetSkeleton& InTargetSkeleton,
+	const FIKRetargetOpBase* InParentOp,
+	FIKRigLogger& InLog)
+{
+	bIsInitialized = false;
+
+	ChainStretchers.Reset();
+
+	// spin through all the mapped retarget bone chains and load them
+	const FRetargeterBoneChains& BoneChains = InProcessor.GetBoneChains();
+	for (const FRetargetStretchChainSettings& ChainSettings : Settings.ChainsToStretch)
+	{
+		if (!ChainSettings.bEnabled)
+		{
+			continue; // if the enabled flag is toggled it will trigger reinit to make a Chain Stretcher
+		}
+		
+		const FResolvedBoneChain* TargetBoneChain = BoneChains.GetResolvedBoneChainByName(
+			ChainSettings.TargetChainName,
+			ERetargetSourceOrTarget::Target,
+			Settings.IKRigAsset);
+		if (TargetBoneChain == nullptr)
+		{
+			InLog.LogWarning( FText::Format(
+			LOCTEXT("ChainStretchOpMissingChain", "Chain Stretch Op: chain data is out of sync with IK Rig. Missing target chain, '{0}."),
+			FText::FromName(ChainSettings.TargetChainName)));
+			continue;
+		}
+
+		// which source chain was this target chain mapped to?
+		const FName SourceChainName = ChainMapping.GetChainMappedTo(ChainSettings.TargetChainName, ERetargetSourceOrTarget::Target);
+		const FResolvedBoneChain* SourceBoneChain = BoneChains.GetResolvedBoneChainByName(SourceChainName, ERetargetSourceOrTarget::Source);
+		if (SourceBoneChain == nullptr)
+		{
+			InLog.LogWarning( FText::Format(
+			LOCTEXT("ChainStretchOpMissingSource", "Chain Stretch Op: missing source chain. Cannot stretch a target chain without a source chain, '{0}."),
+			FText::FromName(ChainSettings.TargetChainName)));
+			continue;
+		}
+		
+		// initialize the mapped pair of source/target bone chains
+		FChainStretcher ChainStretcher;
+		const bool bChainStretcherInitialized = ChainStretcher.Initialize(
+			&ChainSettings,
+			*SourceBoneChain,
+			*TargetBoneChain,
+			InSourceSkeleton,
+			InTargetSkeleton);
+		if (!bChainStretcherInitialized)
+		{
+			InLog.LogWarning( FText::Format(
+				LOCTEXT("ChainStretchOpFailedInit", "Chain Stretch Op: failed to initialize stretcher for chain, '{0}', on Skeletal Mesh: '{1}'"),
+				FText::FromName(ChainSettings.TargetChainName), FText::FromString(InTargetSkeleton.SkeletalMesh->GetName())));
+			return false;
+		}
+
+		// store valid chain pair to be retargeted
+		ChainStretchers.Add(ChainStretcher);
+	}
+
+	bIsInitialized = !ChainStretchers.IsEmpty();
+
+	return bIsInitialized;
+}
+
+void FIKRetargetStretchChainOp::Run(
+	FIKRetargetProcessor& InProcessor,
+	const double InDeltaTime,
+	const TArray<FTransform>& InSourceGlobalPose,
+	TArray<FTransform>& OutTargetGlobalPose)
+{
+	const FTargetSkeleton& TargetSkeleton = InProcessor.GetTargetSkeleton();
+	for (FChainStretcher& ChainStretcher : ChainStretchers)
+	{
+		ChainStretcher.StretchChain(InSourceGlobalPose, TargetSkeleton, OutTargetGlobalPose);
+	}
+}
+
+void FIKRetargetStretchChainOp::OnAddedToStack(const UIKRetargeter* InRetargetAsset, const FIKRetargetOpBase* InParentOp)
+{
+	// on initial setup, use the default source/target IK rigs
+	const UIKRigDefinition* SourceIKRig = InRetargetAsset->GetIKRig(ERetargetSourceOrTarget::Source);
+	const UIKRigDefinition* TargetIKRig = InRetargetAsset->GetIKRig(ERetargetSourceOrTarget::Target);
+	ApplyIKRigs(SourceIKRig, TargetIKRig);
+
+	// auto map
+	static bool bForceRemap = true;
+	ChainMapping.AutoMapChains(EAutoMapChainType::Fuzzy, bForceRemap);
+}
+
+FIKRetargetOpSettingsBase* FIKRetargetStretchChainOp::GetSettings()
+{
+	return &Settings;
+}
+
+const UScriptStruct* FIKRetargetStretchChainOp::GetSettingsType() const
+{
+	return FIKRetargetStretchChainOpSettings::StaticStruct();
+}
+
+const UScriptStruct* FIKRetargetStretchChainOp::GetType() const
+{
+	return FIKRetargetStretchChainOp::StaticStruct();
+}
+
+const UIKRigDefinition* FIKRetargetStretchChainOp::GetCustomTargetIKRig() const
+{
+	return Settings.IKRigAsset;
+}
+
+FRetargetChainMapping* FIKRetargetStretchChainOp::GetChainMapping()
+{
+	return &ChainMapping;
+}
+
+void FIKRetargetStretchChainOp::OnTargetChainRenamed(const FName InOldChainName, const FName InNewChainName)
+{
+	for (FRetargetStretchChainSettings& ChainSettings : Settings.ChainsToStretch)
+	{
+		if (ChainSettings.TargetChainName == InOldChainName)
+		{
+			ChainSettings.TargetChainName = InNewChainName;
+		}
+	}
+}
+
+void FIKRetargetStretchChainOp::OnReinitPropertyEdited(const FPropertyChangedEvent* InPropertyChangedEvent)
+{
+	const UIKRigDefinition* SourceIKRig = ChainMapping.GetIKRig(ERetargetSourceOrTarget::Source);
+	const UIKRigDefinition* TargetIKRig = Settings.IKRigAsset;
+	ApplyIKRigs(SourceIKRig, TargetIKRig);
+}
+
+#if WITH_EDITOR
+FText FIKRetargetStretchChainOp::GetWarningMessage() const
+{
+	if (!bIsInitialized && Settings.ChainsToStretch.IsEmpty())
+	{
+		return LOCTEXT("NoChainsAssigned", "Not initialized. No chains assigned.");
+	}
+	
+	return FIKRetargetOpBase::GetWarningMessage();
+}
+
+void FIKRetargetStretchChainOp::ResetChainSettingsToDefault(const FName& InChainName)
+{
+	for (FRetargetStretchChainSettings& ChainToStretch : Settings.ChainsToStretch)
+	{
+		if (ChainToStretch.TargetChainName == InChainName)
+		{
+			ChainToStretch = FRetargetStretchChainSettings(InChainName);
+			return;
+		}
+	}
+}
+
+bool FIKRetargetStretchChainOp::AreChainSettingsAtDefault(const FName& InChainName)
+{
+	for (FRetargetStretchChainSettings& ChainToStretch : Settings.ChainsToStretch)
+	{
+		if (ChainToStretch.TargetChainName == InChainName)
+		{
+			return ChainToStretch == FRetargetStretchChainSettings(InChainName);
+		}
+	}
+	
+	return true;
+}
+#endif
+
+void FIKRetargetStretchChainOp::ApplyIKRigs(const UIKRigDefinition* InSourceIKRig, const UIKRigDefinition* InTargetIKRig)
+{
+	// store IK Rig
+	Settings.IKRigAsset = InTargetIKRig;
+	
+	// update chain mapping
+	ChainMapping.ReinitializeWithIKRigs(InSourceIKRig, InTargetIKRig);
+
+	// update settings only if we have a valid mapping
+	if (!ChainMapping.IsReady())
+	{
+		// don't remove settings, instead we want to preserve existing settings at least until the next valid rig is loaded
+		return;
+	}
+	
+	// get the required target chains
+	const TArray<FBoneChain>& AllTargetChains = Settings.IKRigAsset->GetRetargetChains();
+	TArray<FName> RequiredTargetChains;
+	for (const FBoneChain& ChainToRetarget : AllTargetChains)
+	{
+		RequiredTargetChains.Add(ChainToRetarget.ChainName);
+	}
+	
+	// remove chains that are not required
+	Settings.ChainsToStretch.RemoveAll([&RequiredTargetChains](const FRetargetStretchChainSettings& InChainSettings)
+	{
+		return !RequiredTargetChains.Contains(InChainSettings.TargetChainName);
+	});
+
+	// add any required chains not already present
+	for (const FName RequiredTargetChain : RequiredTargetChains)
+	{
+		bool bFound = false;
+		for (const FRetargetStretchChainSettings& ChainToStretch : Settings.ChainsToStretch)
+		{
+			if (ChainToStretch.TargetChainName == RequiredTargetChain)
+			{
+				bFound = true;
+				break;
+			}
+		}
+		if (!bFound)
+		{
+			Settings.ChainsToStretch.Emplace(RequiredTargetChain);
+		}
+	}
+}
+
+FIKRetargetStretchChainOpSettings UIKRetargetStretchChainController::GetSettings()
+{
+	return *reinterpret_cast<FIKRetargetStretchChainOpSettings*>(OpSettingsToControl);
+}
+
+void UIKRetargetStretchChainController::SetSettings(FIKRetargetStretchChainOpSettings InSettings)
+{
+	OpSettingsToControl->CopySettingsAtRuntime(&InSettings);
+}
+
+#undef LOCTEXT_NAMESPACE
diff --git a/Engine/Plugins/Animation/IKRig/Source/IKRig/Public/Retargeter/IKRetargetProcessor.h b/Engine/Plugins/Animation/IKRig/Source/IKRig/Public/Retargeter/IKRetargetProcessor.h
--- a/Engine/Plugins/Animation/IKRig/Source/IKRig/Public/Retargeter/IKRetargetProcessor.h
+++ b/Engine/Plugins/Animation/IKRig/Source/IKRig/Public/Retargeter/IKRetargetProcessor.h
@@ -132,6 +132,15 @@ struct FRetargetSkeleton
 		const int32 BoneIndex,
 		const TArray<FTransform>& InGlobalPose) const;
 
+	UE_API TArray<FTransform> GetLocalTransformsOfMultipleBones(
+		const TArray<int32>& InBoneIndices,
+		const TArray<FTransform>& InGlobalPose) const;
+
+	UE_API void UpdateGlobalTransformsOfMultipleBones(
+		const TArray<int32>& InBoneIndices,
+		const TArray<FTransform>& InLocalTransforms,
+		TArray<FTransform>& OutGlobalPose) const;
+
 	UE_API int32 GetCachedEndOfBranchIndex(const int32 InBoneIndex) const;
 
 	UE_API void GetChildrenIndices(const int32 BoneIndex, TArray<int32>& OutChildren) const;
diff --git a/Engine/Plugins/Animation/IKRig/Source/IKRig/Public/Retargeter/RetargetOps/StretchChainOp.h b/Engine/Plugins/Animation/IKRig/Source/IKRig/Public/Retargeter/RetargetOps/StretchChainOp.h
new file mode 100644
--- /dev/null
+++ b/Engine/Plugins/Animation/IKRig/Source/IKRig/Public/Retargeter/RetargetOps/StretchChainOp.h
@@ -0,0 +1,159 @@
+﻿// Copyright Epic Games, Inc. All Rights Reserved.
+#pragma once
+
+#include "Retargeter/IKRetargetChainMapping.h"
+#include "Retargeter/IKRetargetOps.h"
+
+#include "StretchChainOp.generated.h"
+
+#define UE_API IKRIG_API
+
+#define LOCTEXT_NAMESPACE "StretchChainOp"
+
+struct FResolvedBoneChain;
+struct FRetargetStretchChainSettings;
+
+struct FChainStretcher
+{
+	bool Initialize(
+		const FRetargetStretchChainSettings* InSettings,
+		const FResolvedBoneChain& InSourceBoneChain,
+		const FResolvedBoneChain& InTargetBoneChain,
+		const FRetargetSkeleton& InSourceSkeleton,
+		const FRetargetSkeleton& InTargetSkeleton);
+
+	void StretchChain(
+		const TArray<FTransform> &InSourceGlobalPose,
+		const FRetargetSkeleton& InTargetSkeleton,
+		TArray<FTransform> &OutTargetGlobalPose) const;
+	
+private:
+	
+	TArray<int32> AllChildrenOfChain;
+	const FRetargetStretchChainSettings* Settings;
+	const FResolvedBoneChain* SourceBoneChain;
+	const FResolvedBoneChain* TargetBoneChain;
+};
+
+USTRUCT(BlueprintType)
+struct FRetargetStretchChainSettings
+{
+	GENERATED_BODY()
+
+	FRetargetStretchChainSettings() = default;
+	FRetargetStretchChainSettings(FName InTargetChainName) : TargetChainName(InTargetChainName) {}
+	
+	/** The name of the target chain to stretch. */
+	UPROPERTY(EditAnywhere, BlueprintReadWrite, Category = "Setup", meta=(ReinitializeOnEdit))
+	FName TargetChainName;
+
+	/** Enable/disable stretching. */
+	UPROPERTY(EditAnywhere, BlueprintReadWrite, Category = "Stretch Chain", meta=(ReinitializeOnEdit))
+	bool bEnabled = false;
+
+	/** Range 0 to 1. Default 0. Matches the total length of this chain with the mapped source chain.
+	*  At 0, the chain's length will be left alone
+	*  At 1, the chain will be stretched to match the length of the source.*/
+	UPROPERTY(EditAnywhere, BlueprintReadWrite, Category = "Stretch Chain", meta = (EditCondition = "bEnabled", UIMin = "0.0", UIMax = "1.0", ClampMin = "0.0", ClampMax = "1.0"))
+	double MatchSourceLength = 0.0f;
+
+	/** Range 0 to inf. Default 1. A multiplier on the length of the chain, applied after matching the source length. */
+	UPROPERTY(EditAnywhere, BlueprintReadWrite, Category = "Stretch Chain", meta = (EditCondition = "bEnabled", UIMin = "0.0", UIMax = "5.0", ClampMin = "0.0"))
+	double ScaleChainLength = 1.0f;
+	
+	bool operator==(const FRetargetStretchChainSettings& Other) const;
+};
+
+USTRUCT(BlueprintType, meta = (DisplayName = "Stretch Chain Settings"))
+struct FIKRetargetStretchChainOpSettings : public FIKRetargetOpSettingsBase
+{
+	GENERATED_BODY()
+
+	/** The target IK Rig asset that contains the bones chains this op will modify. */
+	UPROPERTY(EditAnywhere, BlueprintReadWrite, Category = "IK Rig Asset", meta=(ReinitializeOnEdit))
+	TObjectPtr<const UIKRigDefinition> IKRigAsset = nullptr;
+	
+	UPROPERTY(EditAnywhere, BlueprintReadWrite, Category = "Stretch Chain Settings")
+	TArray<FRetargetStretchChainSettings> ChainsToStretch;
+
+	virtual const UClass* GetControllerType() const override;
+	
+	virtual void CopySettingsAtRuntime(const FIKRetargetOpSettingsBase* InSettingsToCopyFrom) override;
+};
+
+USTRUCT(BlueprintType, meta = (DisplayName = "Stretch Chains"))
+struct FIKRetargetStretchChainOp : public FIKRetargetOpBase
+{
+	GENERATED_BODY()
+	
+	UE_API virtual bool Initialize(
+		const FIKRetargetProcessor& InProcessor,
+		const FRetargetSkeleton& InSourceSkeleton,
+		const FTargetSkeleton& InTargetSkeleton,
+		const FIKRetargetOpBase* InParentOp,
+		FIKRigLogger& InLog) override;
+	
+	UE_API virtual void Run(
+		FIKRetargetProcessor& InProcessor,
+		const double InDeltaTime,
+		const TArray<FTransform>& InSourceGlobalPose,
+		TArray<FTransform>& OutTargetGlobalPose) override;
+
+	UE_API virtual void OnAddedToStack(const UIKRetargeter* InRetargetAsset, const FIKRetargetOpBase* InParentOp) override;
+
+	UE_API virtual FIKRetargetOpSettingsBase* GetSettings() override;
+	
+	UE_API virtual const UScriptStruct* GetSettingsType() const override;
+
+	UE_API virtual const UScriptStruct* GetType() const override;
+	
+	UE_API virtual const UIKRigDefinition* GetCustomTargetIKRig() const override;
+
+	UE_API virtual FRetargetChainMapping* GetChainMapping() override;
+
+	UE_API virtual void OnTargetChainRenamed(const FName InOldChainName, const FName InNewChainName) override;
+	
+	UE_API virtual void OnReinitPropertyEdited(const FPropertyChangedEvent* InPropertyChangedEvent) override;
+
+#if WITH_EDITOR
+	UE_API virtual FText GetWarningMessage() const override;
+
+	UE_API virtual void ResetChainSettingsToDefault(const FName& InChainName) override;
+	UE_API virtual bool AreChainSettingsAtDefault(const FName& InChainName) override;
+#endif
+
+	UPROPERTY()
+	FIKRetargetStretchChainOpSettings Settings;
+
+	/* This op maintains its own chain mapping table to allow per-op mapping */
+	UPROPERTY()
+	FRetargetChainMapping ChainMapping;
+
+private:
+	
+	void ApplyIKRigs(const UIKRigDefinition* InSourceIKRig, const UIKRigDefinition* InTargetIKRig);
+	
+	TArray<FChainStretcher> ChainStretchers;
+};
+
+/* The blueprint/python API for editing an Align Pole Vector Op */
+UCLASS(MinimalAPI, BlueprintType)
+class UIKRetargetStretchChainController : public UIKRetargetOpControllerBase
+{
+	GENERATED_BODY()
+	
+public:
+	/* Get the current op settings as a struct.
+	 * @return FIKRetargetStretchChainOpSettings struct with the current settings used by the op. */
+	UFUNCTION(BlueprintCallable, Category = Settings)
+	UE_API FIKRetargetStretchChainOpSettings GetSettings();
+
+	/* Set the op settings. Input is a custom struct type for this op.
+	 * @param InSettings a FIKRetargetStretchChainOpSettings struct containing all the settings to apply to this op */
+	UFUNCTION(BlueprintCallable, Category = Settings)
+	UE_API void SetSettings(FIKRetargetStretchChainOpSettings InSettings);
+};
+
+#undef LOCTEXT_NAMESPACE
+
+#undef UE_API
diff --git a/Engine/Plugins/Animation/IKRig/Source/IKRigEditor/Private/IKRigEditor.cpp b/Engine/Plugins/Animation/IKRig/Source/IKRigEditor/Private/IKRigEditor.cpp
--- a/Engine/Plugins/Animation/IKRig/Source/IKRigEditor/Private/IKRigEditor.cpp
+++ b/Engine/Plugins/Animation/IKRig/Source/IKRigEditor/Private/IKRigEditor.cpp
@@ -104,6 +104,9 @@ void FIKRigEditor::StartupModule()
 	PropertyEditorModule.RegisterCustomPropertyTypeLayout(
 			"IKRetargetAdditivePoseOpSettings",
 			FOnGetPropertyTypeCustomizationInstance::CreateStatic(&FAdditivePoseOpCustomization::MakeInstance));
+	PropertyEditorModule.RegisterCustomPropertyTypeLayout(
+			"IKRetargetStretchChainOpSettings",
+			FOnGetPropertyTypeCustomizationInstance::CreateStatic(&FStretchChainOpCustomization::MakeInstance));
 
 	// register a thumbnail renderer for the assets
 	UThumbnailManager::Get().RegisterCustomRenderer(UIKRigDefinition::StaticClass(), UIKRigThumbnailRenderer::StaticClass());
diff --git a/Engine/Plugins/Animation/IKRig/Source/IKRigEditor/Private/RetargetEditor/IKRetargetOpDetails.cpp b/Engine/Plugins/Animation/IKRig/Source/IKRigEditor/Private/RetargetEditor/IKRetargetOpDetails.cpp
--- a/Engine/Plugins/Animation/IKRig/Source/IKRigEditor/Private/RetargetEditor/IKRetargetOpDetails.cpp
+++ b/Engine/Plugins/Animation/IKRig/Source/IKRigEditor/Private/RetargetEditor/IKRetargetOpDetails.cpp
@@ -626,4 +626,66 @@ void FAdditivePoseOpCustomization::UpdatePoseNameOptions()
 	}
 }
 
+void FStretchChainOpCustomization::CustomizeChildren(
+	TSharedRef<IPropertyHandle> StructPropertyHandle,
+	IDetailChildrenBuilder& ChildBuilder,
+	IPropertyTypeCustomizationUtils& StructCustomizationUtils)
+{
+	if (!LoadAndValidateStructToCustomize(StructPropertyHandle, StructCustomizationUtils))
+	{
+		return;
+	}
+
+	// add the IK Rig asset input field
+	TSharedPtr<IPropertyHandle> IKRigHandle = StructPropertyHandle->GetChildHandle(GET_MEMBER_NAME_CHECKED(FIKRetargetStretchChainOpSettings, IKRigAsset));
+	ChildBuilder.AddProperty(IKRigHandle.ToSharedRef());
+
+	// configure chain map list for the IK op
+	FChainMapListConfig ChainListConfig;
+	ChainListConfig.OpWithChainSettings = OpName;
+	ChainListConfig.OpWithChainMapping = OpName;
+	ChainListConfig.Controller = AssetController;
+	ChainListConfig.bEnableGoalColumn = true;
+	ChainListConfig.bEnableChainMapping = true;
+	ChainListConfig.ChainSettingsGetterFunc = [this](const FName InTargetChainName) -> UObject*
+	{
+		// this lambda fetches the chain settings struct for the given target chain in the given op (by name)
+		auto MemoryProvider = [this, InTargetChainName]() -> uint8*
+		{
+			if (FIKRetargetOpBase* Op = AssetController->GetRetargetOpByName(OpName))
+			{
+				FIKRetargetStretchChainOpSettings* OpSettings = reinterpret_cast<FIKRetargetStretchChainOpSettings*>(Op->GetSettings());
+				for (FRetargetStretchChainSettings& ChainToStretch : OpSettings->ChainsToStretch)
+				{
+					if (ChainToStretch.TargetChainName == InTargetChainName)
+					{
+						return reinterpret_cast<uint8*>(&ChainToStretch);
+					}
+				}
+			}
+			return nullptr;
+		};
+			
+		FIKRigStructToView StructToView;
+		StructToView.Owner = AssetController->GetAsset();
+		StructToView.Type = FRetargetStretchChainSettings::StaticStruct();
+		StructToView.MemoryProvider = MemoryProvider;
+		StructToView.UniqueName = OpName;
+		
+		UStretchChainSettingsWrapper* Wrapper = NewObject<UStretchChainSettingsWrapper>(GetTransientPackage(), NAME_None, RF_Standalone);
+		const FName SettingsPropertyName = GET_MEMBER_NAME_CHECKED(UStretchChainSettingsWrapper, Settings);
+		Wrapper->InitializeWithRetargeter(StructToView, SettingsPropertyName, AssetController);
+		Wrapper->SetPropertyHidden(GET_MEMBER_NAME_CHECKED(FRetargetStretchChainSettings, TargetChainName), true);
+		return Wrapper;
+	};
+
+	// add the chain mapping list
+	FText ChainMapRowName = FText::FromString("Chain Map");
+	ChildBuilder.AddCustomRow(ChainMapRowName)
+	.WholeRowContent()
+	[
+		SNew(SIKRetargetChainMapList).InChainMapListConfig(ChainListConfig)
+	];
+}
+
 #undef LOCTEXT_NAMESPACE
diff --git a/Engine/Plugins/Animation/IKRig/Source/IKRigEditor/Public/RetargetEditor/IKRetargetOpDetails.h b/Engine/Plugins/Animation/IKRig/Source/IKRigEditor/Public/RetargetEditor/IKRetargetOpDetails.h
--- a/Engine/Plugins/Animation/IKRig/Source/IKRigEditor/Public/RetargetEditor/IKRetargetOpDetails.h
+++ b/Engine/Plugins/Animation/IKRig/Source/IKRigEditor/Public/RetargetEditor/IKRetargetOpDetails.h
@@ -7,6 +7,7 @@
 #include "Retargeter/RetargetOps/IKChainsOp.h"
 #include "Retargeter/RetargetOps/RetargetPoseOp.h"
 #include "Retargeter/RetargetOps/SpeedPlantingOp.h"
+#include "Retargeter/RetargetOps/StretchChainOp.h"
 #include "RigEditor/IKRigStructViewer.h"
 #include "Templates/SharedPointer.h"
 #include "IKRetargetOpDetails.generated.h"
@@ -258,3 +259,32 @@ public:
 	UPROPERTY(EditAnywhere, Category = "Retarget Pose Settings")
 	FIKRetargetAdditivePoseOpSettings Settings;
 };
+
+//
+// Stretch Chain Op Customization
+//
+
+class FStretchChainOpCustomization : public FIKRetargetOpBaseSettingsCustomization
+{
+public:
+	static TSharedRef<IPropertyTypeCustomization> MakeInstance()
+	{
+		return MakeShared<FStretchChainOpCustomization>();
+	}
+
+	// IPropertyTypeCustomization
+	virtual void CustomizeChildren(
+		TSharedRef<IPropertyHandle> StructPropertyHandle,
+		IDetailChildrenBuilder& ChildBuilder,
+		IPropertyTypeCustomizationUtils& StructCustomizationUtils) override;
+	// End IPropertyTypeCustomization
+};
+
+UCLASS(BlueprintType)
+class UStretchChainSettingsWrapper : public UIKRigStructWrapperBase
+{
+	GENERATED_BODY()
+public:
+	UPROPERTY(EditAnywhere, Category = "Stretch Chain Settings")
+	FRetargetStretchChainSettings Settings;
+};
