[IK Retarget] Fixes crash when adding virtual bone to currently retargeted skeletal mesh.

diff --git a/Engine/Plugins/Animation/IKRig/Source/IKRig/Private/Retargeter/IKRetargetProcessor.cpp b/Engine/Plugins/Animation/IKRig/Source/IKRig/Private/Retargeter/IKRetargetProcessor.cpp
--- a/Engine/Plugins/Animation/IKRig/Source/IKRig/Private/Retargeter/IKRetargetProcessor.cpp
+++ b/Engine/Plugins/Animation/IKRig/Source/IKRig/Private/Retargeter/IKRetargetProcessor.cpp
@@ -1217,11 +1217,7 @@ void FIKRetargetProcessor::Initialize(
 	check(InRetargeterAsset && SourceSkeletalMesh && TargetSkeletalMesh);
 	
 	// don't attempt reinitialization unless inputs have changed
-	const bool bSameTargetMesh = TargetSkeleton.SkeletalMesh == TargetSkeletalMesh;
-	const bool bSameSourceMesh = SourceSkeleton.SkeletalMesh == SourceSkeletalMesh;
-	const bool bSameAsset = RetargeterAsset == InRetargeterAsset;
-	const bool bSameAssetVersion = AssetVersionInitializedWith == InRetargeterAsset->GetVersion();
-	if (bSameAsset && bSameAssetVersion && bSameTargetMesh && bSameSourceMesh)
+	if (WasInitializedWithTheseAssets(SourceSkeletalMesh, TargetSkeletalMesh, InRetargeterAsset))
 	{
 		return;
 	}
