IKRetargetOps: Add source transform to DebugDraw calls.

diff --git a/Engine/Plugins/Animation/IKRig/Source/IKRig/Private/Retargeter/IKRetargetProcessor.cpp b/Engine/Plugins/Animation/IKRig/Source/IKRig/Private/Retargeter/IKRetargetProcessor.cpp
--- a/Engine/Plugins/Animation/IKRig/Source/IKRig/Private/Retargeter/IKRetargetProcessor.cpp
+++ b/Engine/Plugins/Animation/IKRig/Source/IKRig/Private/Retargeter/IKRetargetProcessor.cpp
@@ -1806,6 +1806,7 @@ FName FIKRetargetProcessor::GetFirstChainMappedToChain(
 
 void FIKRetargetProcessor::DebugDrawAllOps(
 	FPrimitiveDrawInterface* InPDI,
+	const FTransform& InSourceTransform,
 	const FTransform& InComponentTransform,
 	const FIKRetargetDebugDrawState& EditorState,
 	const int32 InSelectedOpIndex) const
@@ -1821,7 +1822,7 @@ void FIKRetargetProcessor::DebugDrawAllOps(
 			continue;
 		}
 
-		Op->DebugDraw(InPDI, InComponentTransform, ComponentScale, EditorState);
+		Op->DebugDraw(InPDI, InSourceTransform, InComponentTransform, ComponentScale, EditorState);
 	}
 }
 
diff --git a/Engine/Plugins/Animation/IKRig/Source/IKRig/Private/Retargeter/RetargetOps/FKChainsOp.cpp b/Engine/Plugins/Animation/IKRig/Source/IKRig/Private/Retargeter/RetargetOps/FKChainsOp.cpp
--- a/Engine/Plugins/Animation/IKRig/Source/IKRig/Private/Retargeter/RetargetOps/FKChainsOp.cpp
+++ b/Engine/Plugins/Animation/IKRig/Source/IKRig/Private/Retargeter/RetargetOps/FKChainsOp.cpp
@@ -947,6 +947,7 @@ FCriticalSection FIKRetargetFKChainsOp::DebugDataMutex;
 
 void FIKRetargetFKChainsOp::DebugDraw(
 	FPrimitiveDrawInterface* InPDI,
+	const FTransform& InSourceTransform,
 	const FTransform& InComponentTransform,
 	const double InComponentScale,
 	const FIKRetargetDebugDrawState& InEditorState) const
diff --git a/Engine/Plugins/Animation/IKRig/Source/IKRig/Private/Retargeter/RetargetOps/IKChainsOp.cpp b/Engine/Plugins/Animation/IKRig/Source/IKRig/Private/Retargeter/RetargetOps/IKChainsOp.cpp
--- a/Engine/Plugins/Animation/IKRig/Source/IKRig/Private/Retargeter/RetargetOps/IKChainsOp.cpp
+++ b/Engine/Plugins/Animation/IKRig/Source/IKRig/Private/Retargeter/RetargetOps/IKChainsOp.cpp
@@ -426,6 +426,7 @@ FCriticalSection FIKRetargetIKChainsOp::DebugDataMutex;
 
 void FIKRetargetIKChainsOp::DebugDraw(
 	FPrimitiveDrawInterface* InPDI,
+	const FTransform& InSourceTransform,
 	const FTransform& InComponentTransform,
 	const double InComponentScale,
 	const FIKRetargetDebugDrawState& InEditorState) const
diff --git a/Engine/Plugins/Animation/IKRig/Source/IKRig/Private/Retargeter/RetargetOps/PelvisMotionOp.cpp b/Engine/Plugins/Animation/IKRig/Source/IKRig/Private/Retargeter/RetargetOps/PelvisMotionOp.cpp
--- a/Engine/Plugins/Animation/IKRig/Source/IKRig/Private/Retargeter/RetargetOps/PelvisMotionOp.cpp
+++ b/Engine/Plugins/Animation/IKRig/Source/IKRig/Private/Retargeter/RetargetOps/PelvisMotionOp.cpp
@@ -322,6 +322,7 @@ FCriticalSection FIKRetargetPelvisMotionOp::DebugDataMutex;
 
 void FIKRetargetPelvisMotionOp::DebugDraw(
 	FPrimitiveDrawInterface* InPDI,
+	const FTransform& InSourceTransform,
 	const FTransform& InComponentTransform,
 	const double InComponentScale,
 	const FIKRetargetDebugDrawState& InEditorState) const
diff --git a/Engine/Plugins/Animation/IKRig/Source/IKRig/Private/Retargeter/RetargetOps/RunIKRigOp.cpp b/Engine/Plugins/Animation/IKRig/Source/IKRig/Private/Retargeter/RetargetOps/RunIKRigOp.cpp
--- a/Engine/Plugins/Animation/IKRig/Source/IKRig/Private/Retargeter/RetargetOps/RunIKRigOp.cpp
+++ b/Engine/Plugins/Animation/IKRig/Source/IKRig/Private/Retargeter/RetargetOps/RunIKRigOp.cpp
@@ -151,6 +151,7 @@ FCriticalSection FIKRetargetRunIKRigOp::DebugDataMutex;
 
 void FIKRetargetRunIKRigOp::DebugDraw(
 	FPrimitiveDrawInterface* InPDI,
+	const FTransform& InSourceTransform,
 	const FTransform& InComponentTransform,
 	const double InComponentScale,
 	const FIKRetargetDebugDrawState& InEditorState) const
diff --git a/Engine/Plugins/Animation/IKRig/Source/IKRig/Private/Retargeter/RetargetOps/StrideWarpingOp.cpp b/Engine/Plugins/Animation/IKRig/Source/IKRig/Private/Retargeter/RetargetOps/StrideWarpingOp.cpp
--- a/Engine/Plugins/Animation/IKRig/Source/IKRig/Private/Retargeter/RetargetOps/StrideWarpingOp.cpp
+++ b/Engine/Plugins/Animation/IKRig/Source/IKRig/Private/Retargeter/RetargetOps/StrideWarpingOp.cpp
@@ -389,6 +389,7 @@ FCriticalSection FIKRetargetStrideWarpingOp::DebugDataMutex;
 
 void FIKRetargetStrideWarpingOp::DebugDraw(
 	FPrimitiveDrawInterface* InPDI,
+	const FTransform& InSourceTransform,
 	const FTransform& InComponentTransform,
 	const double InComponentScale,
 	const FIKRetargetDebugDrawState& InEditorState) const
diff --git a/Engine/Plugins/Animation/IKRig/Source/IKRig/Public/Retargeter/IKRetargetOps.h b/Engine/Plugins/Animation/IKRig/Source/IKRig/Public/Retargeter/IKRetargetOps.h
--- a/Engine/Plugins/Animation/IKRig/Source/IKRig/Public/Retargeter/IKRetargetOps.h
+++ b/Engine/Plugins/Animation/IKRig/Source/IKRig/Public/Retargeter/IKRetargetOps.h
@@ -202,6 +202,7 @@ public:
 	/** (optional) override to draw debug info into the editor viewport when the Op is selected */
 	virtual void DebugDraw(
 		FPrimitiveDrawInterface* InPDI,
+		const FTransform& InSourceTransform,
 		const FTransform& InComponentTransform,
 		const double InComponentScale,
 		const FIKRetargetDebugDrawState& InEditorState) const {};
diff --git a/Engine/Plugins/Animation/IKRig/Source/IKRig/Public/Retargeter/IKRetargetProcessor.h b/Engine/Plugins/Animation/IKRig/Source/IKRig/Public/Retargeter/IKRetargetProcessor.h
--- a/Engine/Plugins/Animation/IKRig/Source/IKRig/Public/Retargeter/IKRetargetProcessor.h
+++ b/Engine/Plugins/Animation/IKRig/Source/IKRig/Public/Retargeter/IKRetargetProcessor.h
@@ -543,6 +543,7 @@ public:
 	/** Run all debug drawing on all ops in the stack (or just selected op if one is provided) */
 	UE_API void DebugDrawAllOps(
 		FPrimitiveDrawInterface* InPDI,
+		const FTransform& InSourceTransform,
 		const FTransform& InComponentTransform,
 		const FIKRetargetDebugDrawState& EditorState,
 		const int32 InSelectedOpIndex) const;
diff --git a/Engine/Plugins/Animation/IKRig/Source/IKRig/Public/Retargeter/RetargetOps/FKChainsOp.h b/Engine/Plugins/Animation/IKRig/Source/IKRig/Public/Retargeter/RetargetOps/FKChainsOp.h
--- a/Engine/Plugins/Animation/IKRig/Source/IKRig/Public/Retargeter/RetargetOps/FKChainsOp.h
+++ b/Engine/Plugins/Animation/IKRig/Source/IKRig/Public/Retargeter/RetargetOps/FKChainsOp.h
@@ -288,6 +288,7 @@ public:
 	
 	IKRIG_API virtual void DebugDraw(
 		FPrimitiveDrawInterface* InPDI,
+		const FTransform& InSourceTransform,
 		const FTransform& InComponentTransform,
 		const double InComponentScale,
 		const FIKRetargetDebugDrawState& InEditorState) const override;
diff --git a/Engine/Plugins/Animation/IKRig/Source/IKRig/Public/Retargeter/RetargetOps/IKChainsOp.h b/Engine/Plugins/Animation/IKRig/Source/IKRig/Public/Retargeter/RetargetOps/IKChainsOp.h
--- a/Engine/Plugins/Animation/IKRig/Source/IKRig/Public/Retargeter/RetargetOps/IKChainsOp.h
+++ b/Engine/Plugins/Animation/IKRig/Source/IKRig/Public/Retargeter/RetargetOps/IKChainsOp.h
@@ -210,6 +210,7 @@ struct FIKRetargetIKChainsOp : public FIKRetargetOpBase
 public:
 	UE_API virtual void DebugDraw(
 		FPrimitiveDrawInterface* InPDI,
+		const FTransform& InSourceTransform,
 		const FTransform& InComponentTransform,
 		const double InComponentScale,
 		const FIKRetargetDebugDrawState& InEditorState) const override;
diff --git a/Engine/Plugins/Animation/IKRig/Source/IKRig/Public/Retargeter/RetargetOps/PelvisMotionOp.h b/Engine/Plugins/Animation/IKRig/Source/IKRig/Public/Retargeter/RetargetOps/PelvisMotionOp.h
--- a/Engine/Plugins/Animation/IKRig/Source/IKRig/Public/Retargeter/RetargetOps/PelvisMotionOp.h
+++ b/Engine/Plugins/Animation/IKRig/Source/IKRig/Public/Retargeter/RetargetOps/PelvisMotionOp.h
@@ -190,6 +190,7 @@ public:
 	IKRIG_API virtual FText GetWarningMessage() const override;
 	IKRIG_API virtual void DebugDraw(
 		FPrimitiveDrawInterface* InPDI,
+		const FTransform& InSourceTransform,
 		const FTransform& InComponentTransform,
 		const double InComponentScale,
 		const FIKRetargetDebugDrawState& InEditorState) const override;
diff --git a/Engine/Plugins/Animation/IKRig/Source/IKRig/Public/Retargeter/RetargetOps/RunIKRigOp.h b/Engine/Plugins/Animation/IKRig/Source/IKRig/Public/Retargeter/RetargetOps/RunIKRigOp.h
--- a/Engine/Plugins/Animation/IKRig/Source/IKRig/Public/Retargeter/RetargetOps/RunIKRigOp.h
+++ b/Engine/Plugins/Animation/IKRig/Source/IKRig/Public/Retargeter/RetargetOps/RunIKRigOp.h
@@ -98,6 +98,7 @@ struct FIKRetargetRunIKRigOp : public FIKRetargetOpBase
 #if WITH_EDITOR
 	UE_API virtual void DebugDraw(
 		FPrimitiveDrawInterface* InPDI,
+		const FTransform& InSourceTransform,
 		const FTransform& InComponentTransform,
 		const double InComponentScale,
 		const FIKRetargetDebugDrawState& InEditorState) const override;
diff --git a/Engine/Plugins/Animation/IKRig/Source/IKRig/Public/Retargeter/RetargetOps/StrideWarpingOp.h b/Engine/Plugins/Animation/IKRig/Source/IKRig/Public/Retargeter/RetargetOps/StrideWarpingOp.h
--- a/Engine/Plugins/Animation/IKRig/Source/IKRig/Public/Retargeter/RetargetOps/StrideWarpingOp.h
+++ b/Engine/Plugins/Animation/IKRig/Source/IKRig/Public/Retargeter/RetargetOps/StrideWarpingOp.h
@@ -152,6 +152,7 @@ struct FIKRetargetStrideWarpingOp : public FIKRetargetOpBase
 #if WITH_EDITOR
 	UE_API virtual void DebugDraw(
 		FPrimitiveDrawInterface* InPDI,
+		const FTransform& InSourceTransform,
 		const FTransform& InComponentTransform,
 		const double InComponentScale,
 		const FIKRetargetDebugDrawState& InEditorState) const override;
diff --git a/Engine/Plugins/Animation/IKRig/Source/IKRigEditor/Private/RetargetEditor/IKRetargetDefaultMode.cpp b/Engine/Plugins/Animation/IKRig/Source/IKRigEditor/Private/RetargetEditor/IKRetargetDefaultMode.cpp
--- a/Engine/Plugins/Animation/IKRig/Source/IKRigEditor/Private/RetargetEditor/IKRetargetDefaultMode.cpp
+++ b/Engine/Plugins/Animation/IKRig/Source/IKRigEditor/Private/RetargetEditor/IKRetargetDefaultMode.cpp
@@ -80,9 +80,10 @@ void FIKRetargetDefaultMode::RenderDebugProxies(FPrimitiveDrawInterface* PDI, co
 
 	// debug draw the currently selected op
 	int32 SelectedOpIndex = Controller->GetSelectedOpIndex();
+	UDebugSkelMeshComponent* SourceSkelMesh = Controller->GetSkeletalMeshComponent(ERetargetSourceOrTarget::Source);
 	UDebugSkelMeshComponent* TargetSkelMesh = Controller->GetSkeletalMeshComponent(ERetargetSourceOrTarget::Target);
 	const FIKRetargetDebugDrawState& EditorState = Controller->GetSelectionState();
-	RetargetProcessor->DebugDrawAllOps(PDI, TargetSkelMesh->GetComponentTransform(), EditorState, SelectedOpIndex);
+	RetargetProcessor->DebugDrawAllOps(PDI, SourceSkelMesh->GetComponentTransform(), TargetSkelMesh->GetComponentTransform(), EditorState, SelectedOpIndex);
 }
 
 bool FIKRetargetDefaultMode::HandleClick(FEditorViewportClient* InViewportClient, HHitProxy* HitProxy, const FViewportClick& Click)
