[IK Retarget] Adding "Excluded Bones" to Copy Base Pose.

diff --git a/Engine/Plugins/Animation/IKRig/Source/IKRig/Private/Retargeter/IKRetargetProcessor.cpp b/Engine/Plugins/Animation/IKRig/Source/IKRig/Private/Retargeter/IKRetargetProcessor.cpp
--- a/Engine/Plugins/Animation/IKRig/Source/IKRig/Private/Retargeter/IKRetargetProcessor.cpp
+++ b/Engine/Plugins/Animation/IKRig/Source/IKRig/Private/Retargeter/IKRetargetProcessor.cpp
@@ -1078,135 +1078,6 @@ void FRetargeterBoneChains::Reset()
 	TargetBoneChains.Reset();
 }
 
-void FBasePoseCopier::CopyPoseFromSource(
-	const FName InRootBoneToCopy,
-	const FRetargetSkeleton& InSourceSkeleton,
-	const FTargetSkeleton& InTargetSkeleton,
-	TArray<FTransform>& InOutSourcePose,
-	TArray<FTransform>& OutTargetGlobalPose,
-	TArray<FTransform>& OutTargetInputLocalPose,
-	FIKRigLogger& Log)
-{
-	// check if we need to rebuild the cached bone mapping
-	// cache is invalid if: map is empty, source/target skel mesh swapped or root bone swapped
-	if (SourceToTargetBoneIndexMap.IsEmpty() ||
-		SourceSkeletalMesh != InSourceSkeleton.SkeletalMesh ||
-		TargetSkeletalMesh != InTargetSkeleton.SkeletalMesh ||
-		RootBoneToCopy != InRootBoneToCopy)
-	{
-		Initialize(InRootBoneToCopy, InSourceSkeleton, InTargetSkeleton, Log);
-	}
-	
-	// copy source bones directly to the target (no retargeting)
-	for (const TPair<int32, int32>& Pair : SourceToTargetBoneIndexMap)
-	{
-		const int32 SourceIndex = Pair.Value;
-		const int32 TargetIndex = Pair.Key;
-
-		// copy the pose in local space (this matches CopyPoseFromMesh behavior)
-		const int32 TargetParentIndex = InTargetSkeleton.GetParentIndex(TargetIndex);
-		const FTransform TargetParentGlobal = TargetParentIndex != INDEX_NONE ? OutTargetGlobalPose[TargetParentIndex] : FTransform::Identity;
-		const int32 SourceParentIndex = InSourceSkeleton.GetParentIndex(SourceIndex);
-		const FTransform SourceParentGlobal = SourceParentIndex != INDEX_NONE ? InOutSourcePose[SourceParentIndex] : FTransform::Identity;
-		const FTransform& SourceGlobal = InOutSourcePose[SourceIndex];
-		const FTransform& SourceLocal = SourceGlobal.GetRelativeTransform(SourceParentGlobal);
-
-		// convert to global space and store
-		OutTargetGlobalPose[TargetIndex] = SourceLocal * TargetParentGlobal;
-	}
-
-	// update children of RootBoneToCopy that were not in the source
-	for (int32 TargetChildIndex : ChildrenToUpdate)
-	{
-		const int32 ParentIndex = InTargetSkeleton.GetParentIndex(TargetChildIndex);
-		const FTransform ParentGlobalTransform = ParentIndex != INDEX_NONE ? OutTargetGlobalPose[ParentIndex] : FTransform::Identity;
-		const FTransform ChildLocalTransform = InTargetSkeleton.InputLocalPose[TargetChildIndex];
-		OutTargetGlobalPose[TargetChildIndex] = ChildLocalTransform * ParentGlobalTransform;
-	}
-	
-	// update input local pose to reflect the newly copied pose
-	// (this local pose is used when updating intermediate bones in subsequent retargeting steps)
-	const int32 TargetRootBoneToCopyIndex = RootBoneToCopy == NAME_None ? InTargetSkeleton.BoneNames.Num() : InTargetSkeleton.FindBoneIndexByName(RootBoneToCopy);
-	const int32 ParentOfCopyRootIndex = InTargetSkeleton.GetParentIndex(TargetRootBoneToCopyIndex);
-	InTargetSkeleton.UpdateLocalTransformsBelowBone(ParentOfCopyRootIndex, OutTargetInputLocalPose, InTargetSkeleton.OutputGlobalPose);
-}
-
-void FBasePoseCopier::Initialize(
-	const FName InRootBoneToCopy,
-	const FRetargetSkeleton& InSourceSkeleton,
-	const FTargetSkeleton& InTargetSkeleton,
-	FIKRigLogger& Log)
-{
-	// store the root bone we initialized with
-	RootBoneToCopy = InRootBoneToCopy;
-
-	// store the skeletal meshes we initialized with
-	SourceSkeletalMesh = InSourceSkeleton.SkeletalMesh;
-	TargetSkeletalMesh = InTargetSkeleton.SkeletalMesh;
-
-	// find the roots to copy from
-	const int32 SourceRootBoneToCopyIndex = RootBoneToCopy == NAME_None ? 0 : InSourceSkeleton.FindBoneIndexByName(RootBoneToCopy);
-	const int32 TargetRootBoneToCopyIndex = RootBoneToCopy == NAME_None ? 0 : InTargetSkeleton.FindBoneIndexByName(RootBoneToCopy);
-
-	// warn if user specified a root bone but it wasn't found
-	if (RootBoneToCopy != NAME_None)
-	{
-		if (SourceRootBoneToCopyIndex == INDEX_NONE)
-		{
-			Log.LogWarning( FText::Format(
-				LOCTEXT("MissingSourceRootToCopyFrom", "The root bone to copy from, {0} was not found in the source mesh {1}"),
-				FText::FromName(RootBoneToCopy), FText::FromString(InSourceSkeleton.SkeletalMesh->GetName())));
-		}
-
-		if (TargetRootBoneToCopyIndex == INDEX_NONE)
-		{
-			Log.LogWarning( FText::Format(
-				LOCTEXT("MissingTargetRootToCopyFrom", "The root bone to copy from, {0} was not found in the target mesh {1}"),
-				FText::FromName(RootBoneToCopy), FText::FromString(InSourceSkeleton.SkeletalMesh->GetName())));
-		}
-	}
-
-	// update the mapping from source-to-target bones (by name)
-	SourceToTargetBoneIndexMap.Reset();
-	for (int32 SourceBoneIndex=0; SourceBoneIndex<InSourceSkeleton.BoneNames.Num(); ++SourceBoneIndex)
-	{
-		// filter out bones above the RootBoneToCopy
-		if (SourceRootBoneToCopyIndex != INDEX_NONE &&
-			SourceBoneIndex != SourceRootBoneToCopyIndex &&
-			!InSourceSkeleton.IsParentOf(SourceRootBoneToCopyIndex, SourceBoneIndex))
-		{
-			continue;
-		}
-	
-		// filter out bones not in the target skeleton
-		const int32 TargetBoneIndex = InTargetSkeleton.FindBoneIndexByName(InSourceSkeleton.BoneNames[SourceBoneIndex]);
-		if (TargetBoneIndex == INDEX_NONE)
-		{
-			continue;
-		}
-	
-		// store map of source/target bones in common
-		SourceToTargetBoneIndexMap.Emplace(TargetBoneIndex, SourceBoneIndex);
-	}
-
-	// cache list of children that need updated
-	ChildrenToUpdate.Reset();
-	const int32 LastBranchIndex = InTargetSkeleton.GetCachedEndOfBranchIndex(TargetRootBoneToCopyIndex);
-	const bool bIsLeafBone = LastBranchIndex == INDEX_NONE;
-	if (!bIsLeafBone)
-	{
-		for (int32 ChildBoneIndex = TargetRootBoneToCopyIndex + 1; ChildBoneIndex <= LastBranchIndex; ChildBoneIndex++)
-		{
-			if (SourceToTargetBoneIndexMap.Contains(ChildBoneIndex))
-			{
-				continue;
-			}
-
-			ChildrenToUpdate.Add(ChildBoneIndex);
-		}
-	}
-}
-
 void FIKRetargetProcessor::Initialize(
 		const USkeletalMesh* SourceSkeletalMesh,
 		const USkeletalMesh* TargetSkeletalMesh,
@@ -1463,21 +1334,6 @@ void FIKRetargetProcessor::GenerateBasePoses(TArray<FTransform>& InSourceGlobalP
 	// start from the retarget pose
 	TargetSkeleton.OutputGlobalPose = TargetSkeleton.RetargetPoses.GetGlobalRetargetPose();
 	TargetSkeleton.InputLocalPose = TargetSkeleton.RetargetPoses.GetLocalRetargetPose();
-
-	// partially override with input pose
-	FName CopyBasePoseRootBoneName;
-	if (ShouldCopyBasePose(CopyBasePoseRootBoneName))
-	{
-		// do a straight copy of bone transforms from the source (no retargeting!)
-		BasePoseCopier.CopyPoseFromSource(
-			CopyBasePoseRootBoneName,
-			SourceSkeleton,
-			TargetSkeleton,
-			InSourceGlobalPose,
-			TargetSkeleton.OutputGlobalPose,
-			TargetSkeleton.InputLocalPose,
-			Log);
-	}
 }
 
 void FIKRetargetProcessor::RunRetargetOps(
@@ -1961,23 +1817,6 @@ FName FIKRetargetProcessor::GetPelvisBone(ERetargetSourceOrTarget SourceOrTarget
 	return NAME_None;
 }
 
-bool FIKRetargetProcessor::ShouldCopyBasePose(FName& OutCopyBasePoseRoot) const
-{
-	// first try get it from the processor
-	const FIKRetargetCopyBasePoseOp* CopyBasePoseOp = GetFirstRetargetOpOfType<FIKRetargetCopyBasePoseOp>(ERetargetOpsToSearch::ProcessorOps);
-	// if the processor doesn't have the op, check the asset
-	CopyBasePoseOp = CopyBasePoseOp ? CopyBasePoseOp : GetFirstRetargetOpOfType<FIKRetargetCopyBasePoseOp>(ERetargetOpsToSearch::AssetOps);
-	if (CopyBasePoseOp && CopyBasePoseOp->IsEnabled())
-	{
-		OutCopyBasePoseRoot = CopyBasePoseOp->Settings.CopyBasePoseRoot;
-		return CopyBasePoseOp->Settings.bCopyBasePose;
-	}
-
-	// off by default
-	OutCopyBasePoseRoot = NAME_None;
-	return false;
-}
-
 double FIKRetargetProcessor::GetSourceScaleFactor() const
 {
 	// first try get it from the processor
diff --git a/Engine/Plugins/Animation/IKRig/Source/IKRig/Private/Retargeter/RetargetOps/CopyBasePoseOp.cpp b/Engine/Plugins/Animation/IKRig/Source/IKRig/Private/Retargeter/RetargetOps/CopyBasePoseOp.cpp
--- a/Engine/Plugins/Animation/IKRig/Source/IKRig/Private/Retargeter/RetargetOps/CopyBasePoseOp.cpp
+++ b/Engine/Plugins/Animation/IKRig/Source/IKRig/Private/Retargeter/RetargetOps/CopyBasePoseOp.cpp
@@ -2,6 +2,12 @@
 
 #include "Retargeter/RetargetOps/CopyBasePoseOp.h"
 
+#include "IKRigObjectVersion.h"
+#include "Retargeter/IKRetargetProcessor.h"
+#include "Engine/SkeletalMesh.h"
+
+#define LOCTEXT_NAMESPACE "CopyBasePoseOp"
+
 const UClass* FIKRetargetCopyBasePoseOpSettings::GetControllerType() const
 {
 	return UIKRetargetCopyBasePoseController::StaticClass();
@@ -18,6 +24,23 @@ void FIKRetargetCopyBasePoseOpSettings::CopySettingsAtRuntime(const FIKRetargetO
 		PropertiesToIgnore);
 }
 
+void FIKRetargetCopyBasePoseOpSettings::PostSerialize(const FArchive& Ar)
+{
+#if WITH_EDITORONLY_DATA
+	
+	if (Ar.IsLoading())
+	{
+		if (Ar.CustomVer(FIKRigObjectVersion::GUID) < FIKRigObjectVersion::CopyBasePoseConvertedToBoneReference)
+		{
+			PRAGMA_DISABLE_DEPRECATION_WARNINGS
+			CopyFromStart.BoneName = CopyBasePoseRoot_DEPRECATED;
+			PRAGMA_ENABLE_DEPRECATION_WARNINGS
+		}
+	}
+	
+#endif // WITH_EDITORONLY_DATA
+}
+
 bool FIKRetargetCopyBasePoseOp::Initialize(
 	const FIKRetargetProcessor& InProcessor,
 	const FRetargetSkeleton& InSourceSkeleton,
@@ -26,9 +49,140 @@ bool FIKRetargetCopyBasePoseOp::Initialize(
 	FIKRigLogger& Log)
 {
 	bIsInitialized = true;
+	
+	// find the roots to copy from
+	const FName RootBoneToCopy = Settings.CopyFromStart.BoneName;
+	const int32 SourceRootBoneToCopyIndex = RootBoneToCopy == NAME_None ? 0 : InSourceSkeleton.FindBoneIndexByName(RootBoneToCopy);
+	const int32 TargetRootBoneToCopyIndex = RootBoneToCopy == NAME_None ? 0 : InTargetSkeleton.FindBoneIndexByName(RootBoneToCopy);
+
+	// warn if user specified a root bone but it wasn't found
+	if (RootBoneToCopy != NAME_None)
+	{
+		if (SourceRootBoneToCopyIndex == INDEX_NONE)
+		{
+			Log.LogWarning( FText::Format(
+				LOCTEXT("MissingSourceRootToCopyFrom", "The root bone to copy from, {0} was not found in the source mesh {1}"),
+				FText::FromName(RootBoneToCopy), FText::FromString(InSourceSkeleton.SkeletalMesh->GetName())));
+		}
+
+		if (TargetRootBoneToCopyIndex == INDEX_NONE)
+		{
+			Log.LogWarning( FText::Format(
+				LOCTEXT("MissingTargetRootToCopyFrom", "The root bone to copy from, {0} was not found in the target mesh {1}"),
+				FText::FromName(RootBoneToCopy), FText::FromString(InSourceSkeleton.SkeletalMesh->GetName())));
+		}
+	}
+
+	// gather list of all bones to exclude / filter out of the copy operation
+	TArray<int32> BoneIndicesToExclude;
+	for (const FBoneReference BoneToExclude : Settings.BonesToExclude)
+	{
+		const int32 BoneToExcludeIndex = InTargetSkeleton.FindBoneIndexByName(BoneToExclude.BoneName);
+		if (BoneToExcludeIndex != INDEX_NONE)
+		{
+			BoneIndicesToExclude.Add(BoneToExcludeIndex);
+			InTargetSkeleton.GetChildrenIndicesRecursive(BoneToExcludeIndex, BoneIndicesToExclude);
+		}
+	}
+
+	// update the mapping from source-to-target bones (by name)
+	SourceToTargetBoneIndexMap.Reset();
+	for (int32 SourceBoneIndex=0; SourceBoneIndex<InSourceSkeleton.BoneNames.Num(); ++SourceBoneIndex)
+	{
+		// filter out bones above the RootBoneToCopy
+		if (SourceRootBoneToCopyIndex != INDEX_NONE &&
+			SourceBoneIndex != SourceRootBoneToCopyIndex &&
+			!InSourceSkeleton.IsParentOf(SourceRootBoneToCopyIndex, SourceBoneIndex))
+		{
+			continue;
+		}
+	
+		// filter out bones not in the target skeleton
+		const int32 TargetBoneIndex = InTargetSkeleton.FindBoneIndexByName(InSourceSkeleton.BoneNames[SourceBoneIndex]);
+		if (TargetBoneIndex == INDEX_NONE)
+		{
+			continue;
+		}
+
+		// filter out excluded bones
+		if (BoneIndicesToExclude.Contains(TargetBoneIndex))
+		{
+			continue;
+		}
+	
+		// store map of source/target bones in common
+		SourceToTargetBoneIndexMap.Emplace(TargetBoneIndex, SourceBoneIndex);
+	}
+
+	// cache list of children that need updated
+	ChildrenToUpdate.Reset();
+	const int32 LastBranchIndex = InTargetSkeleton.GetCachedEndOfBranchIndex(TargetRootBoneToCopyIndex);
+	const bool bIsLeafBone = LastBranchIndex == INDEX_NONE;
+	if (!bIsLeafBone)
+	{
+		for (int32 ChildBoneIndex = TargetRootBoneToCopyIndex + 1; ChildBoneIndex <= LastBranchIndex; ChildBoneIndex++)
+		{
+			if (SourceToTargetBoneIndexMap.Contains(ChildBoneIndex))
+			{
+				continue;
+			}
+
+			ChildrenToUpdate.Add(ChildBoneIndex);
+		}
+	}
+	
 	return true;
 }
 
+void FIKRetargetCopyBasePoseOp::Run(
+	FIKRetargetProcessor& InProcessor,
+	const double InDeltaTime,
+	const TArray<FTransform>& InSourceGlobalPose,
+	TArray<FTransform>& OutTargetGlobalPose)
+{
+	// skip if disabled
+	if (!Settings.bCopyBasePose)
+	{
+		return;
+	}
+	
+	const FRetargetSkeleton& SourceSkeleton = InProcessor.GetSkeleton(ERetargetSourceOrTarget::Source);
+	FTargetSkeleton& TargetSkeleton = InProcessor.GetTargetSkeleton();
+	
+	// copy source bones directly to the target (no retargeting)
+	for (const TPair<int32, int32>& Pair : SourceToTargetBoneIndexMap)
+	{
+		const int32 SourceIndex = Pair.Value;
+		const int32 TargetIndex = Pair.Key;
+
+		// copy the pose in local space (this matches CopyPoseFromMesh behavior)
+		const int32 TargetParentIndex = TargetSkeleton.GetParentIndex(TargetIndex);
+		const FTransform TargetParentGlobal = TargetParentIndex != INDEX_NONE ? OutTargetGlobalPose[TargetParentIndex] : FTransform::Identity;
+		const int32 SourceParentIndex = SourceSkeleton.GetParentIndex(SourceIndex);
+		const FTransform SourceParentGlobal = SourceParentIndex != INDEX_NONE ? InSourceGlobalPose[SourceParentIndex] : FTransform::Identity;
+		const FTransform& SourceGlobal = InSourceGlobalPose[SourceIndex];
+		const FTransform& SourceLocal = SourceGlobal.GetRelativeTransform(SourceParentGlobal);
+
+		// convert to global space and store
+		OutTargetGlobalPose[TargetIndex] = SourceLocal * TargetParentGlobal;
+	}
+
+	// update children of RootBoneToCopy that were not in the source
+	for (int32 TargetChildIndex : ChildrenToUpdate)
+	{
+		const int32 ParentIndex = TargetSkeleton.GetParentIndex(TargetChildIndex);
+		const FTransform ParentGlobalTransform = ParentIndex != INDEX_NONE ? OutTargetGlobalPose[ParentIndex] : FTransform::Identity;
+		const FTransform ChildLocalTransform = TargetSkeleton.InputLocalPose[TargetChildIndex];
+		OutTargetGlobalPose[TargetChildIndex] = ChildLocalTransform * ParentGlobalTransform;
+	}
+	
+	// update input local pose to reflect the newly copied pose
+	// (this local pose is used when updating intermediate bones in subsequent retargeting ops)
+	const int32 TargetRootBoneToCopyIndex = Settings.CopyFromStart.BoneName == NAME_None ? TargetSkeleton.BoneNames.Num() : TargetSkeleton.FindBoneIndexByName(Settings.CopyFromStart.BoneName);
+	const int32 ParentOfCopyRootIndex = TargetSkeleton.GetParentIndex(TargetRootBoneToCopyIndex);
+	TargetSkeleton.UpdateLocalTransformsBelowBone(ParentOfCopyRootIndex, TargetSkeleton.InputLocalPose, TargetSkeleton.OutputGlobalPose);
+}
+
 FIKRetargetOpSettingsBase* FIKRetargetCopyBasePoseOp::GetSettings()
 {
 	return &Settings;
@@ -53,3 +207,49 @@ void UIKRetargetCopyBasePoseController::SetSettings(FIKRetargetCopyBasePoseOpSet
 {
 	OpSettingsToControl->CopySettingsAtRuntime(&InSettings);
 }
+
+void UIKRetargetCopyBasePoseController::SetCopyFromRootBone(const FName InBoneName) const
+{
+	FIKRetargetCopyBasePoseOpSettings* Settings = reinterpret_cast<FIKRetargetCopyBasePoseOpSettings*>(OpSettingsToControl);
+	if (!Settings)
+	{
+		return;
+	}
+
+	Settings->CopyFromStart.BoneName = InBoneName;
+}
+
+FName UIKRetargetCopyBasePoseController::GetCopyFromRootBone(const FName InBoneName) const
+{
+	FIKRetargetCopyBasePoseOpSettings* Settings = reinterpret_cast<FIKRetargetCopyBasePoseOpSettings*>(OpSettingsToControl);
+	if (!Settings)
+	{
+		return NAME_None;
+	}
+
+	return Settings->CopyFromStart.BoneName;
+}
+
+void UIKRetargetCopyBasePoseController::AddBoneToExclude(const FName InBoneName) const
+{
+	FIKRetargetCopyBasePoseOpSettings* Settings = reinterpret_cast<FIKRetargetCopyBasePoseOpSettings*>(OpSettingsToControl);
+	if (!Settings)
+	{
+		return;
+	}
+
+	Settings->BonesToExclude.Add(InBoneName);
+}
+
+void UIKRetargetCopyBasePoseController::ResetBonesToExclude() const
+{
+	FIKRetargetCopyBasePoseOpSettings* Settings = reinterpret_cast<FIKRetargetCopyBasePoseOpSettings*>(OpSettingsToControl);
+	if (!Settings)
+	{
+		return;
+	}
+
+	Settings->BonesToExclude.Reset();
+}
+
+#undef LOCTEXT_NAMESPACE
diff --git a/Engine/Plugins/Animation/IKRig/Source/IKRig/Public/IKRigObjectVersion.h b/Engine/Plugins/Animation/IKRig/Source/IKRig/Public/IKRigObjectVersion.h
--- a/Engine/Plugins/Animation/IKRig/Source/IKRig/Public/IKRigObjectVersion.h
+++ b/Engine/Plugins/Animation/IKRig/Source/IKRig/Public/IKRigObjectVersion.h
@@ -35,6 +35,9 @@ struct FIKRigObjectVersion
 		// ops control their own chain mapping (no more global chain map on retarget asset)
 		OpsOwnChainMapping,
 
+		// copy base pose op converted to use FBoneReference 
+		CopyBasePoseConvertedToBoneReference,
+
 		// -----<new versions can be added above this line>-------------------------------------------------
 		VersionPlusOne,
 		LatestVersion = VersionPlusOne - 1
diff --git a/Engine/Plugins/Animation/IKRig/Source/IKRig/Public/Retargeter/IKRetargetOps.h b/Engine/Plugins/Animation/IKRig/Source/IKRig/Public/Retargeter/IKRetargetOps.h
--- a/Engine/Plugins/Animation/IKRig/Source/IKRig/Public/Retargeter/IKRetargetOps.h
+++ b/Engine/Plugins/Animation/IKRig/Source/IKRig/Public/Retargeter/IKRetargetOps.h
@@ -72,8 +72,8 @@ public:
 	/** a reference to the skeletons for bone selector widgets (can return these from GetSkeleton) */
 	mutable const USkeleton* SourceSkeletonAsset;
 	mutable const USkeleton* TargetSkeletonAsset;
-	/** allow settings to provide a skeleton for any given FBoneReference widget */
-	virtual USkeleton* GetSkeleton(const FName InPropertyName) { return nullptr; };
+	/** allow settings to provide a skeleton for any given FBoneReference widget. Defaults to the target skeleton. */
+	virtual USkeleton* GetSkeleton(const FName InPropertyName) { return const_cast<USkeleton*>(TargetSkeletonAsset); };
 #endif
 };
 
diff --git a/Engine/Plugins/Animation/IKRig/Source/IKRig/Public/Retargeter/IKRetargetProcessor.h b/Engine/Plugins/Animation/IKRig/Source/IKRig/Public/Retargeter/IKRetargetProcessor.h
--- a/Engine/Plugins/Animation/IKRig/Source/IKRig/Public/Retargeter/IKRetargetProcessor.h
+++ b/Engine/Plugins/Animation/IKRig/Source/IKRig/Public/Retargeter/IKRetargetProcessor.h
@@ -328,41 +328,6 @@ enum class ERetargetOpsToSearch : uint8
 	AssetOps
 };
 
-/** the system to do a straight copy of bone transform with matching names (does not do any retargeting!) */
-struct FBasePoseCopier
-{
-	void CopyPoseFromSource(
-		const FName InRootBoneToCopy,
-		const FRetargetSkeleton& InSourceSkeleton,
-		const FTargetSkeleton& InTargetSkeleton,
-		TArray<FTransform>& InOutSourcePose,
-		TArray<FTransform>& OutTargetGlobalPose,
-		TArray<FTransform>& OutTargetInputLocalPose,
-		FIKRigLogger& Log);
-
-private:
-
-	void Initialize(
-		const FName InRootBoneToCopy,
-		const FRetargetSkeleton& InSourceSkeleton,
-		const FTargetSkeleton& InTargetSkeleton,
-		FIKRigLogger& Log);
-	
-	// mapping from source bones to target bone indices
-	// NOTE: this map will only contain bones with the same name in both skeletons that are below the RootBoneToCopy (if specified)
-	TMap<int32, int32> SourceToTargetBoneIndexMap;
-
-	// children of bones that are copied that need manually updated
-	TArray<int32> ChildrenToUpdate;
-
-	// the bone to start copying from (includes all children recursive)
-	FName RootBoneToCopy = NAME_None;
-
-	// the skeletal meshes that were used to build the bone mapping
-	const USkeletalMesh* SourceSkeletalMesh = nullptr;
-	const USkeletalMesh* TargetSkeletalMesh = nullptr;
-};
-
 /** The runtime processor that converts an input pose from a source skeleton into an output pose on a target skeleton.
  * To use:
  * 1. Initialize a processor with a Source/Target skeletal mesh and a UIKRetargeter asset.
@@ -502,11 +467,6 @@ public:
 	/** Get name of the pelvis bone for either the source or target skeleton. */
 	UE_API FName GetPelvisBone(ERetargetSourceOrTarget SourceOrTarget, ERetargetOpsToSearch InOpsToSearch) const;
 
-	/** Determine if we should be directly copying in the source pose into the target
-	 * @param OutCopyBasePoseRoot: will contain the name of the bone to copy from if one was specified, otherwise Name_None
-	 * @return true if a CopyBasePoseOp is present and set to enable copying*/
-	UE_API bool ShouldCopyBasePose(FName& OutCopyBasePoseRoot) const;
-
 	/** Get the scale factor for the source pose (comes from presence of SourceScaleOp)*/
 	UE_API double GetSourceScaleFactor() const;
 
@@ -575,9 +535,6 @@ private:
 	/** Storage for all bone chain data used by Ops and editor tools */
 	FRetargeterBoneChains AllBoneChains;
 
-	/** The system to do a straight copy of source-to-target bone transforms based on name */
-	FBasePoseCopier BasePoseCopier;
-
 	/** The collection of operations to run to transfer animation from source to target. */
 	TArray<FInstancedStruct> OpStack;
 
diff --git a/Engine/Plugins/Animation/IKRig/Source/IKRig/Public/Retargeter/RetargetOps/CopyBasePoseOp.h b/Engine/Plugins/Animation/IKRig/Source/IKRig/Public/Retargeter/RetargetOps/CopyBasePoseOp.h
--- a/Engine/Plugins/Animation/IKRig/Source/IKRig/Public/Retargeter/RetargetOps/CopyBasePoseOp.h
+++ b/Engine/Plugins/Animation/IKRig/Source/IKRig/Public/Retargeter/RetargetOps/CopyBasePoseOp.h
@@ -2,6 +2,7 @@
 #pragma once
 
 #include "Retargeter/IKRetargetOps.h"
+#include "Animation/BoneReference.h"
 
 #include "CopyBasePoseOp.generated.h"
 
@@ -19,13 +20,34 @@ struct FIKRetargetCopyBasePoseOpSettings : public FIKRetargetOpSettingsBase
 	UPROPERTY(EditAnywhere, BlueprintReadWrite, Category = RetargetPhases)
 	bool bCopyBasePose = true;
 
-	/** Filters the bones to copy when using "Copy Base Pose". If specified, will only copy all children of the specified bone (inclusive).*/
-	UPROPERTY(EditAnywhere, BlueprintReadWrite, Category = RetargetPhases, meta = (EditCondition = "bCopyBasePose"))
-	FName CopyBasePoseRoot = NAME_None;
+	/** Filters the bones to copy when using "Copy Base Pose". If specified, will only copy all children of the specified bone (inclusive). */
+	UPROPERTY(EditAnywhere, Category = RetargetPhases, meta = (ReinitializeOnEdit, EditCondition = "bCopyBasePose"))
+	FBoneReference CopyFromStart;
+
+	/** Bones added to this list (and their children) will not have their pose copied. They will be left at the reference pose.*/
+	UPROPERTY(EditAnywhere, Category = RetargetPhases, meta = (ReinitializeOnEdit, EditCondition = "bCopyBasePose"))
+	TArray<FBoneReference> BonesToExclude;
 
 	UE_API virtual const UClass* GetControllerType() const override;
 
 	UE_API virtual void CopySettingsAtRuntime(const FIKRetargetOpSettingsBase* InSettingsToCopyFrom) override;
+
+	/** Custom serialization to handle conversion from deprecated CopyBasePoseRoot to CopyFromRoot */
+	void PostSerialize(const FArchive& Ar);
+
+private:
+	/** converted on PoseSerialize */
+	UPROPERTY(meta = (DeprecatedProperty))
+	FName CopyBasePoseRoot_DEPRECATED = NAME_None;
+};
+
+// Custom serialization to load deprecated settings
+template<> struct TStructOpsTypeTraits<FIKRetargetCopyBasePoseOpSettings> : public TStructOpsTypeTraitsBase2<FIKRetargetCopyBasePoseOpSettings>
+{
+	enum
+	{
+		WithPostSerialize = true,
+	};
 };
 
 USTRUCT(BlueprintType, meta = (DisplayName = "Copy Base Pose"))
@@ -47,7 +69,7 @@ struct FIKRetargetCopyBasePoseOp : public FIKRetargetOpBase
 		FIKRetargetProcessor& InProcessor,
 		const double InDeltaTime,
 		const TArray<FTransform>& InSourceGlobalPose,
-		TArray<FTransform>& OutTargetGlobalPose) override {};
+		TArray<FTransform>& OutTargetGlobalPose) override;
 
 	UE_API virtual FIKRetargetOpSettingsBase* GetSettings() override;
 	
@@ -59,6 +81,14 @@ struct FIKRetargetCopyBasePoseOp : public FIKRetargetOpBase
 
 	UPROPERTY()
 	FIKRetargetCopyBasePoseOpSettings Settings;
+
+private:
+	// mapping from source bones to target bone indices
+	// NOTE: this map will only contain bones with the same name in both skeletons that are below the RootBoneToCopy (if specified)
+	TMap<int32, int32> SourceToTargetBoneIndexMap;
+
+	// children of bones that are copied that need manually updated
+	TArray<int32> ChildrenToUpdate;
 };
 
 /* The blueprint/python API for editing a Copy Base Pose Op */
@@ -77,6 +107,22 @@ public:
 	 * @param InSettings a FIKRetargetCopyBasePoseOpSettings struct containing all the settings to apply to this op */
 	UFUNCTION(BlueprintCallable, Category = Settings)
 	UE_API void SetSettings(FIKRetargetCopyBasePoseOpSettings InSettings);
+
+	/* Set the name of the bone to use as the "Copy From Root" in the op */
+	UFUNCTION(BlueprintCallable, Category = CopyFromRoot)
+	UE_API void SetCopyFromRootBone(const FName InBoneName) const;
+
+	/* Get the name of the bone to use as the "Copy From Root" in the op */
+	UFUNCTION(BlueprintCallable, Category = CopyFromRoot)
+	UE_API FName GetCopyFromRootBone(const FName InBoneName) const;
+
+	/* Add a bone to the "BonesToExclude" array */
+	UFUNCTION(BlueprintCallable, Category = ExcludedBones)
+	UE_API void AddBoneToExclude(const FName InBoneName) const;
+
+	/* Reset the array of Bones to Exclude */
+	UFUNCTION(BlueprintCallable, Category = ExcludedBones)
+	UE_API void ResetBonesToExclude() const;
 };
 
 #undef LOCTEXT_NAMESPACE
