Fix -Wrange-loop-construct

#rnx

diff --git a/Engine/Plugins/Animation/IKRig/Source/IKRig/Private/Retargeter/RetargetOps/CopyBasePoseOp.cpp b/Engine/Plugins/Animation/IKRig/Source/IKRig/Private/Retargeter/RetargetOps/CopyBasePoseOp.cpp
--- a/Engine/Plugins/Animation/IKRig/Source/IKRig/Private/Retargeter/RetargetOps/CopyBasePoseOp.cpp
+++ b/Engine/Plugins/Animation/IKRig/Source/IKRig/Private/Retargeter/RetargetOps/CopyBasePoseOp.cpp
@@ -75,7 +75,7 @@ bool FIKRetargetCopyBasePoseOp::Initialize(
 
 	// gather list of all bones to exclude / filter out of the copy operation
 	TArray<int32> BoneIndicesToExclude;
-	for (const FBoneReference BoneToExclude : Settings.BonesToExclude)
+	for (const FBoneReference& BoneToExclude : Settings.BonesToExclude)
 	{
 		const int32 BoneToExcludeIndex = InTargetSkeleton.FindBoneIndexByName(BoneToExclude.BoneName);
 		if (BoneToExcludeIndex != INDEX_NONE)
