[IK Retarget] Adding ability to update IK Rigs of ops when assigning IK Rig to retarget asset.

diff --git a/Engine/Plugins/Animation/IKRig/Source/IKRig/Private/Retargeter/RetargetOps/AlignPoleVectorOp.cpp b/Engine/Plugins/Animation/IKRig/Source/IKRig/Private/Retargeter/RetargetOps/AlignPoleVectorOp.cpp
--- a/Engine/Plugins/Animation/IKRig/Source/IKRig/Private/Retargeter/RetargetOps/AlignPoleVectorOp.cpp
+++ b/Engine/Plugins/Animation/IKRig/Source/IKRig/Private/Retargeter/RetargetOps/AlignPoleVectorOp.cpp
@@ -293,6 +293,23 @@ void FIKRetargetAlignPoleVectorOp::OnAddedToStack(const UIKRetargeter* InRetarge
 	ChainMapping.AutoMapChains(EAutoMapChainType::Fuzzy, bForceRemap);
 }
 
+void FIKRetargetAlignPoleVectorOp::OnAssignIKRig(
+	const ERetargetSourceOrTarget SourceOrTarget,
+	const UIKRigDefinition* InIKRig,
+	const FIKRetargetOpBase* InParentOp)
+{
+	// get the newly assigned IK Rig
+	const bool bAssignedSource = SourceOrTarget == ERetargetSourceOrTarget::Source;
+	const UIKRigDefinition* SourceIKRig = bAssignedSource ? InIKRig : ChainMapping.GetIKRig(ERetargetSourceOrTarget::Source);
+	const bool bAssignedTarget = SourceOrTarget == ERetargetSourceOrTarget::Target;
+	const UIKRigDefinition* TargetIKRig = bAssignedTarget ? InIKRig : ChainMapping.GetIKRig(ERetargetSourceOrTarget::Target);
+	ApplyIKRigs(SourceIKRig, TargetIKRig);
+
+	// auto map, but don't force it (leave existing mappings alone)
+	static bool bForceRemap = false;
+	ChainMapping.AutoMapChains(EAutoMapChainType::Fuzzy, bForceRemap);
+}
+
 FIKRetargetOpSettingsBase* FIKRetargetAlignPoleVectorOp::GetSettings()
 {
 	return &Settings;
diff --git a/Engine/Plugins/Animation/IKRig/Source/IKRig/Private/Retargeter/RetargetOps/FKChainsOp.cpp b/Engine/Plugins/Animation/IKRig/Source/IKRig/Private/Retargeter/RetargetOps/FKChainsOp.cpp
--- a/Engine/Plugins/Animation/IKRig/Source/IKRig/Private/Retargeter/RetargetOps/FKChainsOp.cpp
+++ b/Engine/Plugins/Animation/IKRig/Source/IKRig/Private/Retargeter/RetargetOps/FKChainsOp.cpp
@@ -889,6 +889,23 @@ void FIKRetargetFKChainsOp::OnAddedToStack(const UIKRetargeter* InRetargetAsset,
 	ChainMapping.AutoMapChains(EAutoMapChainType::Fuzzy, bForceRemap);
 }
 
+void FIKRetargetFKChainsOp::OnAssignIKRig(
+	const ERetargetSourceOrTarget SourceOrTarget,
+	const UIKRigDefinition* InIKRig,
+	const FIKRetargetOpBase* InParentOp)
+{
+	// get the newly assigned IK Rig
+	const bool bAssignedSource = SourceOrTarget == ERetargetSourceOrTarget::Source;
+	const UIKRigDefinition* SourceIKRig = bAssignedSource ? InIKRig : ChainMapping.GetIKRig(ERetargetSourceOrTarget::Source);
+	const bool bAssignedTarget = SourceOrTarget == ERetargetSourceOrTarget::Target;
+	const UIKRigDefinition* TargetIKRig = bAssignedTarget ? InIKRig : ChainMapping.GetIKRig(ERetargetSourceOrTarget::Target);
+	ApplyIKRigs(SourceIKRig, TargetIKRig);
+
+	// auto map, but don't force it (leave existing mappings alone)
+	static bool bForceRemap = false;
+	ChainMapping.AutoMapChains(EAutoMapChainType::Fuzzy, bForceRemap);
+}
+
 void FIKRetargetFKChainsOp::CollectRetargetedBones(TSet<int32>& OutRetargetedBones) const
 {
 	// all bones in an FK chain are retargeted
diff --git a/Engine/Plugins/Animation/IKRig/Source/IKRig/Private/Retargeter/RetargetOps/IKChainsOp.cpp b/Engine/Plugins/Animation/IKRig/Source/IKRig/Private/Retargeter/RetargetOps/IKChainsOp.cpp
--- a/Engine/Plugins/Animation/IKRig/Source/IKRig/Private/Retargeter/RetargetOps/IKChainsOp.cpp
+++ b/Engine/Plugins/Animation/IKRig/Source/IKRig/Private/Retargeter/RetargetOps/IKChainsOp.cpp
@@ -361,6 +361,14 @@ void FIKRetargetIKChainsOp::OnAddedToStack(const UIKRetargeter* InRetargetAsset,
 	RegenerateChainSettings(InParentOp);
 }
 
+void FIKRetargetIKChainsOp::OnAssignIKRig(
+	const ERetargetSourceOrTarget SourceOrTarget,
+	const UIKRigDefinition* InIKRig,
+	const FIKRetargetOpBase* InParentOp)
+{
+	RegenerateChainSettings(InParentOp);
+}
+
 FIKRetargetOpSettingsBase* FIKRetargetIKChainsOp::GetSettings()
 {
 	return &Settings;
diff --git a/Engine/Plugins/Animation/IKRig/Source/IKRig/Private/Retargeter/RetargetOps/PelvisMotionOp.cpp b/Engine/Plugins/Animation/IKRig/Source/IKRig/Private/Retargeter/RetargetOps/PelvisMotionOp.cpp
--- a/Engine/Plugins/Animation/IKRig/Source/IKRig/Private/Retargeter/RetargetOps/PelvisMotionOp.cpp
+++ b/Engine/Plugins/Animation/IKRig/Source/IKRig/Private/Retargeter/RetargetOps/PelvisMotionOp.cpp
@@ -123,6 +123,15 @@ void FIKRetargetPelvisMotionOp::OnAddedToStack(const UIKRetargeter* InRetargetAs
 	}
 }
 
+void FIKRetargetPelvisMotionOp::OnAssignIKRig(
+	const ERetargetSourceOrTarget SourceOrTarget,
+	const UIKRigDefinition* InIKRig,
+	const FIKRetargetOpBase* InParentOp)
+{
+	FName& PelvisNameSetting = SourceOrTarget == ERetargetSourceOrTarget::Source ? Settings.SourcePelvisBone.BoneName : Settings.TargetPelvisBone.BoneName;
+	PelvisNameSetting = InIKRig->GetPelvis();
+}
+
 bool FIKRetargetPelvisMotionOp::InitializeSource(
 	const FName SourcePelvisBoneName,
 	const FRetargetSkeleton& SourceSkeleton,
diff --git a/Engine/Plugins/Animation/IKRig/Source/IKRig/Private/Retargeter/RetargetOps/RunIKRigOp.cpp b/Engine/Plugins/Animation/IKRig/Source/IKRig/Private/Retargeter/RetargetOps/RunIKRigOp.cpp
--- a/Engine/Plugins/Animation/IKRig/Source/IKRig/Private/Retargeter/RetargetOps/RunIKRigOp.cpp
+++ b/Engine/Plugins/Animation/IKRig/Source/IKRig/Private/Retargeter/RetargetOps/RunIKRigOp.cpp
@@ -115,6 +115,20 @@ void FIKRetargetRunIKRigOp::OnAddedToStack(const UIKRetargeter* InRetargetAsset,
 	ChainMapping.AutoMapChains(EAutoMapChainType::Fuzzy, bForceRemap);
 }
 
+void FIKRetargetRunIKRigOp::OnAssignIKRig(
+	const ERetargetSourceOrTarget SourceOrTarget,
+	const UIKRigDefinition* InIKRig,
+	const FIKRetargetOpBase* InParentOp)
+{
+	const bool bAssignedTarget = SourceOrTarget == ERetargetSourceOrTarget::Target;
+	Settings.IKRigAsset = bAssignedTarget ? InIKRig : ChainMapping.GetIKRig(ERetargetSourceOrTarget::Target);
+
+	// re-initialize the chain mapping
+	ChainMapping.ReinitializeWithIKRigs(ChainMapping.GetIKRig(ERetargetSourceOrTarget::Source), Settings.IKRigAsset);
+	static bool bForceRemap = false; // don't force mapping, keeps existing mappings
+	ChainMapping.AutoMapChains(EAutoMapChainType::Fuzzy, bForceRemap);
+}
+
 FIKRetargetOpSettingsBase* FIKRetargetRunIKRigOp::GetSettings()
 {
 	return &Settings;
diff --git a/Engine/Plugins/Animation/IKRig/Source/IKRig/Private/Retargeter/RetargetOps/SpeedPlantingOp.cpp b/Engine/Plugins/Animation/IKRig/Source/IKRig/Private/Retargeter/RetargetOps/SpeedPlantingOp.cpp
--- a/Engine/Plugins/Animation/IKRig/Source/IKRig/Private/Retargeter/RetargetOps/SpeedPlantingOp.cpp
+++ b/Engine/Plugins/Animation/IKRig/Source/IKRig/Private/Retargeter/RetargetOps/SpeedPlantingOp.cpp
@@ -152,6 +152,14 @@ void FIKRetargetSpeedPlantingOp::OnAddedToStack(const UIKRetargeter* InRetargetA
 	RegenerateChainSettings(InParentOp);
 }
 
+void FIKRetargetSpeedPlantingOp::OnAssignIKRig(
+	const ERetargetSourceOrTarget SourceOrTarget,
+	const UIKRigDefinition* InIKRig,
+	const FIKRetargetOpBase* InParentOp)
+{
+	RegenerateChainSettings(InParentOp);
+}
+
 FIKRetargetOpSettingsBase* FIKRetargetSpeedPlantingOp::GetSettings()
 {
 	return &Settings;
diff --git a/Engine/Plugins/Animation/IKRig/Source/IKRig/Private/Retargeter/RetargetOps/StretchChainOp.cpp b/Engine/Plugins/Animation/IKRig/Source/IKRig/Private/Retargeter/RetargetOps/StretchChainOp.cpp
--- a/Engine/Plugins/Animation/IKRig/Source/IKRig/Private/Retargeter/RetargetOps/StretchChainOp.cpp
+++ b/Engine/Plugins/Animation/IKRig/Source/IKRig/Private/Retargeter/RetargetOps/StretchChainOp.cpp
@@ -197,6 +197,23 @@ void FIKRetargetStretchChainOp::OnAddedToStack(const UIKRetargeter* InRetargetAs
 	ChainMapping.AutoMapChains(EAutoMapChainType::Fuzzy, bForceRemap);
 }
 
+void FIKRetargetStretchChainOp::OnAssignIKRig(
+	const ERetargetSourceOrTarget SourceOrTarget,
+	const UIKRigDefinition* InIKRig,
+	const FIKRetargetOpBase* InParentOp)
+{
+	// get the newly assigned IK Rig
+	const bool bAssignedSource = SourceOrTarget == ERetargetSourceOrTarget::Source;
+	const UIKRigDefinition* SourceIKRig = bAssignedSource ? InIKRig : ChainMapping.GetIKRig(ERetargetSourceOrTarget::Source);
+	const bool bAssignedTarget = SourceOrTarget == ERetargetSourceOrTarget::Target;
+	const UIKRigDefinition* TargetIKRig = bAssignedTarget ? InIKRig : ChainMapping.GetIKRig(ERetargetSourceOrTarget::Target);
+	ApplyIKRigs(SourceIKRig, TargetIKRig);
+
+	// auto map, but don't force it (leave existing mappings alone)
+	static bool bForceRemap = false;
+	ChainMapping.AutoMapChains(EAutoMapChainType::Fuzzy, bForceRemap);
+}
+
 FIKRetargetOpSettingsBase* FIKRetargetStretchChainOp::GetSettings()
 {
 	return &Settings;
diff --git a/Engine/Plugins/Animation/IKRig/Source/IKRig/Private/Retargeter/RetargetOps/StrideWarpingOp.cpp b/Engine/Plugins/Animation/IKRig/Source/IKRig/Private/Retargeter/RetargetOps/StrideWarpingOp.cpp
--- a/Engine/Plugins/Animation/IKRig/Source/IKRig/Private/Retargeter/RetargetOps/StrideWarpingOp.cpp
+++ b/Engine/Plugins/Animation/IKRig/Source/IKRig/Private/Retargeter/RetargetOps/StrideWarpingOp.cpp
@@ -345,6 +345,14 @@ void FIKRetargetStrideWarpingOp::OnAddedToStack(const UIKRetargeter* InRetargetA
 	RegenerateChainSettings(InParentOp);
 }
 
+void FIKRetargetStrideWarpingOp::OnAssignIKRig(
+	const ERetargetSourceOrTarget SourceOrTarget,
+	const UIKRigDefinition* InIKRig,
+	const FIKRetargetOpBase* InParentOp)
+{
+	RegenerateChainSettings(InParentOp);
+}
+
 FIKRetargetOpSettingsBase* FIKRetargetStrideWarpingOp::GetSettings()
 {
 	return &Settings;
diff --git a/Engine/Plugins/Animation/IKRig/Source/IKRig/Public/Retargeter/IKRetargetOps.h b/Engine/Plugins/Animation/IKRig/Source/IKRig/Public/Retargeter/IKRetargetOps.h
--- a/Engine/Plugins/Animation/IKRig/Source/IKRig/Public/Retargeter/IKRetargetOps.h
+++ b/Engine/Plugins/Animation/IKRig/Source/IKRig/Public/Retargeter/IKRetargetOps.h
@@ -16,6 +16,7 @@ struct FInstancedStruct;
 class USkeleton;
 struct FPoseContext;
 struct FRetargetChainMapping;
+enum class ERetargetSourceOrTarget : uint8;
 
 #if WITH_EDITOR
 class FPrimitiveDrawInterface;
@@ -118,6 +119,12 @@ public:
 	
 	/** (optional) override to automate initial setup after being added to the stack */
 	virtual void OnAddedToStack(const UIKRetargeter* InRetargetAsset, const FIKRetargetOpBase* InParentOp) {};
+
+	/** (optional) override to react after user assigns a new IK Rig to all ops
+	 * @param SourceOrTarget: whether the source or target IK Rig was changed
+	 * @param InIKRig: the IK rig asset to assign to this op
+	 * @param InParentOp: the parent for this op (may be null if unparented) */
+	virtual void OnAssignIKRig(const ERetargetSourceOrTarget SourceOrTarget, const UIKRigDefinition* InIKRig, const FIKRetargetOpBase* InParentOp) {};
 	
 	/** get if this op is enabled */
 	bool IsEnabled() const { return bIsEnabled; };
diff --git a/Engine/Plugins/Animation/IKRig/Source/IKRig/Public/Retargeter/RetargetOps/AlignPoleVectorOp.h b/Engine/Plugins/Animation/IKRig/Source/IKRig/Public/Retargeter/RetargetOps/AlignPoleVectorOp.h
--- a/Engine/Plugins/Animation/IKRig/Source/IKRig/Public/Retargeter/RetargetOps/AlignPoleVectorOp.h
+++ b/Engine/Plugins/Animation/IKRig/Source/IKRig/Public/Retargeter/RetargetOps/AlignPoleVectorOp.h
@@ -128,6 +128,8 @@ struct FIKRetargetAlignPoleVectorOp : public FIKRetargetOpBase
 
 	UE_API virtual void OnAddedToStack(const UIKRetargeter* InRetargetAsset, const FIKRetargetOpBase* InParentOp) override;
 
+	IKRIG_API virtual void OnAssignIKRig(const ERetargetSourceOrTarget SourceOrTarget, const UIKRigDefinition* InIKRig, const FIKRetargetOpBase* InParentOp) override;
+
 	UE_API virtual FIKRetargetOpSettingsBase* GetSettings() override;
 	
 	UE_API virtual const UScriptStruct* GetSettingsType() const override;
diff --git a/Engine/Plugins/Animation/IKRig/Source/IKRig/Public/Retargeter/RetargetOps/FKChainsOp.h b/Engine/Plugins/Animation/IKRig/Source/IKRig/Public/Retargeter/RetargetOps/FKChainsOp.h
--- a/Engine/Plugins/Animation/IKRig/Source/IKRig/Public/Retargeter/RetargetOps/FKChainsOp.h
+++ b/Engine/Plugins/Animation/IKRig/Source/IKRig/Public/Retargeter/RetargetOps/FKChainsOp.h
@@ -259,6 +259,8 @@ struct FIKRetargetFKChainsOp : public FIKRetargetOpBase
 		FIKRigLogger& Log) override;
 
 	IKRIG_API virtual void OnAddedToStack(const UIKRetargeter* InRetargetAsset, const FIKRetargetOpBase* InParentOp) override;
+
+	IKRIG_API virtual void OnAssignIKRig(const ERetargetSourceOrTarget SourceOrTarget, const UIKRigDefinition* InIKRig, const FIKRetargetOpBase* InParentOp) override;
 	
 	IKRIG_API virtual void CollectRetargetedBones(TSet<int32>& OutRetargetedBones) const override;
 
diff --git a/Engine/Plugins/Animation/IKRig/Source/IKRig/Public/Retargeter/RetargetOps/IKChainsOp.h b/Engine/Plugins/Animation/IKRig/Source/IKRig/Public/Retargeter/RetargetOps/IKChainsOp.h
--- a/Engine/Plugins/Animation/IKRig/Source/IKRig/Public/Retargeter/RetargetOps/IKChainsOp.h
+++ b/Engine/Plugins/Animation/IKRig/Source/IKRig/Public/Retargeter/RetargetOps/IKChainsOp.h
@@ -189,6 +189,8 @@ struct FIKRetargetIKChainsOp : public FIKRetargetOpBase
 
 	UE_API virtual void OnAddedToStack(const UIKRetargeter* InRetargetAsset, const FIKRetargetOpBase* InParentOp) override;
 
+	UE_API virtual void OnAssignIKRig(const ERetargetSourceOrTarget SourceOrTarget, const UIKRigDefinition* InIKRig, const FIKRetargetOpBase* InParentOp) override;
+
 	UE_API virtual FIKRetargetOpSettingsBase* GetSettings() override;
 	
 	UE_API virtual void SetSettings(const FIKRetargetOpSettingsBase* InSettings) override;
diff --git a/Engine/Plugins/Animation/IKRig/Source/IKRig/Public/Retargeter/RetargetOps/PelvisMotionOp.h b/Engine/Plugins/Animation/IKRig/Source/IKRig/Public/Retargeter/RetargetOps/PelvisMotionOp.h
--- a/Engine/Plugins/Animation/IKRig/Source/IKRig/Public/Retargeter/RetargetOps/PelvisMotionOp.h
+++ b/Engine/Plugins/Animation/IKRig/Source/IKRig/Public/Retargeter/RetargetOps/PelvisMotionOp.h
@@ -141,6 +141,8 @@ struct FIKRetargetPelvisMotionOp : public FIKRetargetOpBase
 
 	IKRIG_API virtual void OnAddedToStack(const UIKRetargeter* InRetargetAsset, const FIKRetargetOpBase* InParentOp) override;
 
+	IKRIG_API virtual void OnAssignIKRig(const ERetargetSourceOrTarget SourceOrTarget, const UIKRigDefinition* InIKRig, const FIKRetargetOpBase* InParentOp) override;
+
 	IKRIG_API virtual FIKRetargetOpSettingsBase* GetSettings() override;
 	
 	IKRIG_API virtual const UScriptStruct* GetSettingsType() const override;
diff --git a/Engine/Plugins/Animation/IKRig/Source/IKRig/Public/Retargeter/RetargetOps/RunIKRigOp.h b/Engine/Plugins/Animation/IKRig/Source/IKRig/Public/Retargeter/RetargetOps/RunIKRigOp.h
--- a/Engine/Plugins/Animation/IKRig/Source/IKRig/Public/Retargeter/RetargetOps/RunIKRigOp.h
+++ b/Engine/Plugins/Animation/IKRig/Source/IKRig/Public/Retargeter/RetargetOps/RunIKRigOp.h
@@ -74,6 +74,8 @@ struct FIKRetargetRunIKRigOp : public FIKRetargetOpBase
 
 	UE_API virtual void OnAddedToStack(const UIKRetargeter* InRetargetAsset, const FIKRetargetOpBase* InParentOp) override;
 
+	UE_API virtual void OnAssignIKRig(const ERetargetSourceOrTarget SourceOrTarget, const UIKRigDefinition* InIKRig, const FIKRetargetOpBase* InParentOp) override;
+
 	UE_API virtual FIKRetargetOpSettingsBase* GetSettings() override;
 	
 	UE_API virtual const UScriptStruct* GetSettingsType() const override;
diff --git a/Engine/Plugins/Animation/IKRig/Source/IKRig/Public/Retargeter/RetargetOps/SpeedPlantingOp.h b/Engine/Plugins/Animation/IKRig/Source/IKRig/Public/Retargeter/RetargetOps/SpeedPlantingOp.h
--- a/Engine/Plugins/Animation/IKRig/Source/IKRig/Public/Retargeter/RetargetOps/SpeedPlantingOp.h
+++ b/Engine/Plugins/Animation/IKRig/Source/IKRig/Public/Retargeter/RetargetOps/SpeedPlantingOp.h
@@ -96,6 +96,8 @@ struct FIKRetargetSpeedPlantingOp : public FIKRetargetOpBase
 
 	UE_API virtual void OnAddedToStack(const UIKRetargeter* InRetargetAsset, const FIKRetargetOpBase* InParentOp) override;
 
+	UE_API virtual void OnAssignIKRig(const ERetargetSourceOrTarget SourceOrTarget, const UIKRigDefinition* InIKRig, const FIKRetargetOpBase* InParentOp) override;
+
 	UE_API virtual FIKRetargetOpSettingsBase* GetSettings() override;
 	
 	UE_API virtual const UScriptStruct* GetSettingsType() const override;
diff --git a/Engine/Plugins/Animation/IKRig/Source/IKRig/Public/Retargeter/RetargetOps/StretchChainOp.h b/Engine/Plugins/Animation/IKRig/Source/IKRig/Public/Retargeter/RetargetOps/StretchChainOp.h
--- a/Engine/Plugins/Animation/IKRig/Source/IKRig/Public/Retargeter/RetargetOps/StretchChainOp.h
+++ b/Engine/Plugins/Animation/IKRig/Source/IKRig/Public/Retargeter/RetargetOps/StretchChainOp.h
@@ -101,6 +101,8 @@ struct FIKRetargetStretchChainOp : public FIKRetargetOpBase
 
 	UE_API virtual void OnAddedToStack(const UIKRetargeter* InRetargetAsset, const FIKRetargetOpBase* InParentOp) override;
 
+	UE_API virtual void OnAssignIKRig(const ERetargetSourceOrTarget SourceOrTarget, const UIKRigDefinition* InIKRig, const FIKRetargetOpBase* InParentOp) override;
+
 	UE_API virtual FIKRetargetOpSettingsBase* GetSettings() override;
 	
 	UE_API virtual const UScriptStruct* GetSettingsType() const override;
diff --git a/Engine/Plugins/Animation/IKRig/Source/IKRig/Public/Retargeter/RetargetOps/StrideWarpingOp.h b/Engine/Plugins/Animation/IKRig/Source/IKRig/Public/Retargeter/RetargetOps/StrideWarpingOp.h
--- a/Engine/Plugins/Animation/IKRig/Source/IKRig/Public/Retargeter/RetargetOps/StrideWarpingOp.h
+++ b/Engine/Plugins/Animation/IKRig/Source/IKRig/Public/Retargeter/RetargetOps/StrideWarpingOp.h
@@ -134,6 +134,8 @@ struct FIKRetargetStrideWarpingOp : public FIKRetargetOpBase
 
 	UE_API virtual void OnAddedToStack(const UIKRetargeter* InRetargetAsset, const FIKRetargetOpBase* InParentOp) override;
 
+	UE_API virtual void OnAssignIKRig(const ERetargetSourceOrTarget SourceOrTarget, const UIKRigDefinition* InIKRig, const FIKRetargetOpBase* InParentOp) override;
+
 	UE_API virtual FIKRetargetOpSettingsBase* GetSettings() override;
 	
 	UE_API virtual const UScriptStruct* GetSettingsType() const override;
diff --git a/Engine/Plugins/Animation/IKRig/Source/IKRigEditor/Private/RetargetEditor/IKRetargetCommands.cpp b/Engine/Plugins/Animation/IKRig/Source/IKRigEditor/Private/RetargetEditor/IKRetargetCommands.cpp
--- a/Engine/Plugins/Animation/IKRig/Source/IKRigEditor/Private/RetargetEditor/IKRetargetCommands.cpp
+++ b/Engine/Plugins/Animation/IKRig/Source/IKRigEditor/Private/RetargetEditor/IKRetargetCommands.cpp
@@ -11,9 +11,7 @@ void FIKRetargetCommands::RegisterCommands()
 	UI_COMMAND(ShowRetargetPose, "Show Retarget Pose", "Display the retarget pose and retarget to the target skeleton.", EUserInterfaceActionType::Button, FInputChord());
 
 	UI_COMMAND(ShowAssetSettings, "Asset Settings", "Set the source or target IK Rigs, preview meshes. Adjust debug settings.", EUserInterfaceActionType::ToggleButton, FInputChord());
-
-	UI_COMMAND(AddDefaultOps, "Add Default Ops", "Add Pelvis Motion, IK/FK Chains and Root Motion Ops. Runs auto-setup on each op.", EUserInterfaceActionType::ToggleButton, FInputChord());
-
+	
 	UI_COMMAND(ResetAllBones, "Reset All", "Sets the selected bones to the mesh reference pose.", EUserInterfaceActionType::Button, FInputChord());
 	UI_COMMAND(ResetSelectedBones, "Reset Selected Bones", "Sets the selected bones to the mesh reference pose.", EUserInterfaceActionType::Button, FInputChord());
 	UI_COMMAND(ResetSelectedAndChildrenBones, "Reset Selected and Children Bones", "Sets the selected bones (and all children recursively) to the mesh reference pose.", EUserInterfaceActionType::Button, FInputChord());
diff --git a/Engine/Plugins/Animation/IKRig/Source/IKRigEditor/Private/RetargetEditor/IKRetargetEditor.cpp b/Engine/Plugins/Animation/IKRig/Source/IKRigEditor/Private/RetargetEditor/IKRetargetEditor.cpp
--- a/Engine/Plugins/Animation/IKRig/Source/IKRigEditor/Private/RetargetEditor/IKRetargetEditor.cpp
+++ b/Engine/Plugins/Animation/IKRig/Source/IKRigEditor/Private/RetargetEditor/IKRetargetEditor.cpp
@@ -146,15 +146,6 @@ void FIKRetargetEditor::BindCommands()
 			return EditorController->IsObjectInDetailsView(Asset);	
 		}));
 
-	// add default ops to the stack
-	ToolkitCommands->MapAction(
-		Commands.AddDefaultOps,
-		FExecuteAction::CreateLambda([this]()
-		{
-			EditorController->AssetController->AddDefaultOps();
-		}),
-		FCanExecuteAction(), FIsActionChecked());
-
 	//
 	// Edit pose commands
 	//
@@ -300,13 +291,6 @@ void FIKRetargetEditor::FillToolbar(FToolBarBuilder& ToolbarBuilder)
 		TAttribute<FText>(),
 		TAttribute<FText>(),
 		FSlateIcon(FIKRetargetEditorStyle::Get().GetStyleSetName(),"IKRetarget.AssetSettings"));
-
-		ToolbarBuilder.AddToolBarButton(
-		FIKRetargetCommands::Get().AddDefaultOps,
-		NAME_None,
-		TAttribute<FText>(),
-		TAttribute<FText>(),
-		FSlateIcon(FIKRetargetEditorStyle::Get().GetStyleSetName(),"IKRetarget.PostSettings"));
 	}
 	ToolbarBuilder.EndSection();
 }
diff --git a/Engine/Plugins/Animation/IKRig/Source/IKRigEditor/Private/RetargetEditor/IKRetargetEditorController.cpp b/Engine/Plugins/Animation/IKRig/Source/IKRigEditor/Private/RetargetEditor/IKRetargetEditorController.cpp
--- a/Engine/Plugins/Animation/IKRig/Source/IKRigEditor/Private/RetargetEditor/IKRetargetEditorController.cpp
+++ b/Engine/Plugins/Animation/IKRig/Source/IKRigEditor/Private/RetargetEditor/IKRetargetEditorController.cpp
@@ -289,6 +289,26 @@ void FIKRetargetEditorController::ReinitializeRetargeterNoUIRefresh() const
 
 void FIKRetargetEditorController::HandleIKRigReplaced(ERetargetSourceOrTarget SourceOrTarget)
 {
+	const FText SourceOrTargetText = SourceOrTarget == ERetargetSourceOrTarget::Source ? FText::FromString("Source") : FText::FromString("Target");
+
+	TSharedRef<SCustomDialog> Dialog = SNew(SCustomDialog)
+		.Title(FText(LOCTEXT("AssignIKRigTitle", "Assign IK Rig to All Ops?")))
+		.Content()
+		[
+			SNew(STextBlock)
+			.Text(FText::Format(LOCTEXT("AssignToAllOpsLabel", "Assign new {0} IK Rig to all Ops?"), SourceOrTargetText))
+		]
+		.Buttons({
+			SCustomDialog::FButton(LOCTEXT("AssignLabel", "Assign")),
+			SCustomDialog::FButton(LOCTEXT("NoLabel", "No"))
+	});
+
+	if (Dialog->ShowModal() == 0)
+	{
+		const UIKRigDefinition* IKRig = AssetController->GetIKRig(SourceOrTarget);
+		AssetController->AssignIKRigToAllOps(SourceOrTarget, IKRig);
+	}
+
 	BindToIKRigAssets();
 	RefreshAllViews();
 }
diff --git a/Engine/Plugins/Animation/IKRig/Source/IKRigEditor/Private/RetargetEditor/IKRetargetFactory.cpp b/Engine/Plugins/Animation/IKRig/Source/IKRigEditor/Private/RetargetEditor/IKRetargetFactory.cpp
--- a/Engine/Plugins/Animation/IKRig/Source/IKRigEditor/Private/RetargetEditor/IKRetargetFactory.cpp
+++ b/Engine/Plugins/Animation/IKRig/Source/IKRigEditor/Private/RetargetEditor/IKRetargetFactory.cpp
@@ -5,6 +5,7 @@
 #include "AssetTypeCategories.h"
 #include "ContentBrowserModule.h"
 #include "IContentBrowserSingleton.h"
+#include "RetargetEditor/IKRetargeterController.h"
 
 #include UE_INLINE_GENERATED_CPP_BY_NAME(IKRetargetFactory)
 
@@ -26,7 +27,10 @@ UObject* UIKRetargetFactory::FactoryCreateNew(
 	UObject* Context,
 	FFeedbackContext* Warn)
 {
-	return NewObject<UIKRetargeter>(InParent, Class, Name, Flags);
+	UIKRetargeter* NewRetargeter = NewObject<UIKRetargeter>(InParent, Class, Name, Flags);
+	UIKRetargeterController* Controller = UIKRetargeterController::GetController(NewRetargeter);
+	Controller->AddDefaultOps();
+	return NewRetargeter;
 }
 
 bool UIKRetargetFactory::ShouldShowInNewMenu() const
diff --git a/Engine/Plugins/Animation/IKRig/Source/IKRigEditor/Private/RetargetEditor/IKRetargeterController.cpp b/Engine/Plugins/Animation/IKRig/Source/IKRigEditor/Private/RetargetEditor/IKRetargeterController.cpp
--- a/Engine/Plugins/Animation/IKRig/Source/IKRigEditor/Private/RetargetEditor/IKRetargeterController.cpp
+++ b/Engine/Plugins/Animation/IKRig/Source/IKRigEditor/Private/RetargetEditor/IKRetargeterController.cpp
@@ -366,6 +366,22 @@ void UIKRetargeterController::RunOpInitialSetup(const int32 InOpIndex) const
 	Op->OnAddedToStack(GetAsset(), ParentOp);
 }
 
+void UIKRetargeterController::AssignIKRigToAllOps(const ERetargetSourceOrTarget InSourceOrTarget, const UIKRigDefinition* InIKRig) const
+{
+	FScopedTransaction Transaction(LOCTEXT("RunOpInitialSetup_Label", "Run Op Initial Setup"));
+	FScopedReinitializeIKRetargeter Reinitialize(this, ERetargetRefreshMode::ProcessorAndOpStack);
+	Asset->Modify();
+
+	// assign ik rig in reverse order so parent ops consume it before children
+	for (int32 OpIndex = Asset->RetargetOps.Num() - 1; OpIndex >= 0; --OpIndex)
+	{
+		FInstancedStruct& OpStruct = Asset->RetargetOps[OpIndex];
+		FIKRetargetOpBase& Op = OpStruct.GetMutable<FIKRetargetOpBase>();
+		FIKRetargetOpBase* ParentOp = GetRetargetOpByName(GetParentOpByName(Op.GetName()));
+		Op.OnAssignIKRig(InSourceOrTarget, InIKRig, ParentOp);
+	}
+}
+
 FIKRetargetOpBase* UIKRetargeterController::GetRetargetOpByIndex(const int32 InOpIndex) const
 {
 	check(Asset)
diff --git a/Engine/Plugins/Animation/IKRig/Source/IKRigEditor/Private/RetargetEditor/SRetargetOpStack.cpp b/Engine/Plugins/Animation/IKRig/Source/IKRigEditor/Private/RetargetEditor/SRetargetOpStack.cpp
--- a/Engine/Plugins/Animation/IKRig/Source/IKRigEditor/Private/RetargetEditor/SRetargetOpStack.cpp
+++ b/Engine/Plugins/Animation/IKRig/Source/IKRigEditor/Private/RetargetEditor/SRetargetOpStack.cpp
@@ -827,12 +827,13 @@ void SRetargetOpStack::Construct(
 				+ SVerticalBox::Slot()
 				.AutoHeight()
 				.VAlign(VAlign_Top)
+				.HAlign(HAlign_Left)
 				[
 					SNew(SHorizontalBox)
 					+ SHorizontalBox::Slot()
 					.VAlign(VAlign_Center)
 					.HAlign(HAlign_Left)
-					.FillWidth(1.f)
+			
 					.Padding(6.0f, 4.0f)
 					[
 						SNew(SPositiveActionButton)
@@ -844,6 +845,23 @@ void SRetargetOpStack::Construct(
 							return ListView->CreateAddNewOpMenu();	
 						})
 					]
+
+					+ SHorizontalBox::Slot()
+					.VAlign(VAlign_Center)
+					.HAlign(HAlign_Left)
+				
+					.Padding(6.0f, 4.0f)
+					[
+						SNew(SPositiveActionButton)
+						.Icon(FIKRetargetEditorStyle::Get().GetBrush("IKRetarget.PostSettingsSmall"))
+						.Text(LOCTEXT("AddDefaultOpsLabel", "Add Default Ops"))
+						.ToolTipText(LOCTEXT("AddDefaultToolTip", "Add Pelvis Motion, IK/FK Chains and Root Motion Ops. Runs auto-setup on each op."))
+						.OnClicked_Lambda([this]()
+						{
+							EditorController.Pin()->AssetController->AddDefaultOps();
+							return FReply::Handled();
+						})
+					]
 				]
 			]
 		]
diff --git a/Engine/Plugins/Animation/IKRig/Source/IKRigEditor/Public/RetargetEditor/IKRetargetCommands.h b/Engine/Plugins/Animation/IKRig/Source/IKRigEditor/Public/RetargetEditor/IKRetargetCommands.h
--- a/Engine/Plugins/Animation/IKRig/Source/IKRigEditor/Public/RetargetEditor/IKRetargetCommands.h
+++ b/Engine/Plugins/Animation/IKRig/Source/IKRigEditor/Public/RetargetEditor/IKRetargetCommands.h
@@ -27,9 +27,6 @@ public:
 
 	// open asset settings
 	TSharedPtr< FUICommandInfo > ShowAssetSettings;
-	
-	// add default ops to the stack
-	TSharedPtr< FUICommandInfo > AddDefaultOps;
 
 	// enable Root retarget pass
 	TSharedPtr< FUICommandInfo > EnableRoot;
diff --git a/Engine/Plugins/Animation/IKRig/Source/IKRigEditor/Public/RetargetEditor/IKRetargetEditorStyle.h b/Engine/Plugins/Animation/IKRig/Source/IKRigEditor/Public/RetargetEditor/IKRetargetEditorStyle.h
--- a/Engine/Plugins/Animation/IKRig/Source/IKRigEditor/Public/RetargetEditor/IKRetargetEditorStyle.h
+++ b/Engine/Plugins/Animation/IKRig/Source/IKRigEditor/Public/RetargetEditor/IKRetargetEditorStyle.h
@@ -26,6 +26,7 @@ public:
 		Set("IKRetarget.GlobalSettings", new IMAGE_BRUSH_SVG("Slate/GlobalSettings", Icon64x64));
 		Set("IKRetarget.RootSettings", new IMAGE_BRUSH_SVG("Slate/RootSettings", Icon64x64));
 		Set("IKRetarget.PostSettings", new IMAGE_BRUSH_SVG("Slate/PostSettings", Icon64x64));
+		Set("IKRetarget.PostSettingsSmall", new IMAGE_BRUSH_SVG("Slate/PostSettings", Icon16x16));
 		Set("IKRetarget.ChainMapping", new IMAGE_BRUSH_SVG("Slate/ChainMapping", Icon64x64));
 
 		Set("IKRetarget.RunRetargeter", new IMAGE_BRUSH_SVG("Slate/RunRetargeter", Icon64x64));
diff --git a/Engine/Plugins/Animation/IKRig/Source/IKRigEditor/Public/RetargetEditor/IKRetargeterController.h b/Engine/Plugins/Animation/IKRig/Source/IKRigEditor/Public/RetargetEditor/IKRetargeterController.h
--- a/Engine/Plugins/Animation/IKRig/Source/IKRigEditor/Public/RetargetEditor/IKRetargeterController.h
+++ b/Engine/Plugins/Animation/IKRig/Source/IKRigEditor/Public/RetargetEditor/IKRetargeterController.h
@@ -156,6 +156,13 @@ public:
 	UFUNCTION(BlueprintCallable, Category="IK Retarget", meta = (BlueprintThreadSafe))
 	UE_API void RunOpInitialSetup(const int32 InOpIndex) const;
 
+	/** Force all ops to use the assigned IK Rig and update their chain mappings.
+	 * NOTE: some ops may not reference an IK Rig at all; the exact behavior is op-dependent
+	 * @param InSourceOrTarget whether to assign the provided IK Rig as a source or target IK Rig
+	 * @param InIKRig the IK Rig asset to assign */
+	UFUNCTION(BlueprintCallable, Category="IK Retarget", meta = (BlueprintThreadSafe))
+	UE_API void AssignIKRigToAllOps(const ERetargetSourceOrTarget InSourceOrTarget, const UIKRigDefinition* InIKRig) const;
+
 	// Tell each op to reset any settings associated with this chain
 	UFUNCTION(BlueprintCallable, Category="IK Retarget")
 	UE_API void ResetChainSettingsInAllOps(const FName InTargetChainName) const;
