[IK Retarget] Fix static analysis issue in CopyBasePoseOp.cpp

diff --git a/Engine/Plugins/Animation/IKRig/Source/IKRig/Private/Retargeter/RetargetOps/CopyBasePoseOp.cpp b/Engine/Plugins/Animation/IKRig/Source/IKRig/Private/Retargeter/RetargetOps/CopyBasePoseOp.cpp
--- a/Engine/Plugins/Animation/IKRig/Source/IKRig/Private/Retargeter/RetargetOps/CopyBasePoseOp.cpp
+++ b/Engine/Plugins/Animation/IKRig/Source/IKRig/Private/Retargeter/RetargetOps/CopyBasePoseOp.cpp
@@ -178,7 +178,7 @@ void FIKRetargetCopyBasePoseOp::Run(
 	
 	// update input local pose to reflect the newly copied pose
 	// (this local pose is used when updating intermediate bones in subsequent retargeting ops)
-	const int32 TargetRootBoneToCopyIndex = Settings.CopyFromStart.BoneName == NAME_None ? TargetSkeleton.BoneNames.Num() : TargetSkeleton.FindBoneIndexByName(Settings.CopyFromStart.BoneName);
+	const int32 TargetRootBoneToCopyIndex = TargetSkeleton.FindBoneIndexByName(Settings.CopyFromStart.BoneName);
 	const int32 ParentOfCopyRootIndex = TargetSkeleton.GetParentIndex(TargetRootBoneToCopyIndex);
 	TargetSkeleton.UpdateLocalTransformsBelowBone(ParentOfCopyRootIndex, TargetSkeleton.InputLocalPose, TargetSkeleton.OutputGlobalPose);
 }
