[IK Rig] Integrates IK Rig into Control Rig

diff --git a/Engine/Plugins/Animation/IKRig/Source/IKRig/IKRig.Build.cs b/Engine/Plugins/Animation/IKRig/Source/IKRig/IKRig.Build.cs
--- a/Engine/Plugins/Animation/IKRig/Source/IKRig/IKRig.Build.cs
+++ b/Engine/Plugins/Animation/IKRig/Source/IKRig/IKRig.Build.cs
@@ -8,7 +8,9 @@ public class IKRig : ModuleRules
 	{
 		PCHUsage = ModuleRules.PCHUsageMode.UseExplicitOrSharedPCHs;
 
-		CppCompileWarningSettings.UnsafeTypeCastWarningLevel = WarningLevel.Error;
+		// NOTE: this was disabled when we added a Rig Unit for IK Rig due to compiler warnings in Control Rig
+		// We should renable this once those are addressed.
+		//CppCompileWarningSettings.UnsafeTypeCastWarningLevel = WarningLevel.Error;
 		
 		PrivateDependencyModuleNames.AddRange(
 			new string[]
diff --git a/Engine/Plugins/Animation/IKRig/Source/IKRig/Private/AnimNodes/RigUnit_IKRig.cpp b/Engine/Plugins/Animation/IKRig/Source/IKRig/Private/AnimNodes/RigUnit_IKRig.cpp
new file mode 100644
--- /dev/null
+++ b/Engine/Plugins/Animation/IKRig/Source/IKRig/Private/AnimNodes/RigUnit_IKRig.cpp
@@ -0,0 +1,131 @@
+﻿// Copyright Epic Games, Inc. All Rights Reserved.
+
+#include "AnimNodes/RigUnit_IKRig.h"
+
+#include "ControlRigComponent.h"
+
+#include UE_INLINE_GENERATED_CPP_BY_NAME(RigUnit_IKRig)
+
+FRigUnit_IKRig_Execute()
+{
+	DECLARE_SCOPE_HIERARCHICAL_COUNTER_RIGUNIT()
+	LLM_SCOPE_BYNAME(TEXT("Animation/IKRig"));
+
+	URigHierarchy* Hierarchy = ExecuteContext.Hierarchy;
+	if (Hierarchy == nullptr)
+	{
+		return;
+	}
+
+	auto GetSkeletalMesh = [&ExecuteContext]() ->USkeletalMesh*
+	{
+		const USkeletalMeshComponent* OwningComponent = Cast<USkeletalMeshComponent>(ExecuteContext.GetOwningComponent());
+		if (!OwningComponent)
+		{
+			return nullptr;
+		}
+
+		return OwningComponent->GetSkeletalMeshAsset();
+	};
+
+	FIKRigProcessor& Processor = WorkData.IKRigProcessor;
+
+	if (!Processor.IsInitialized())
+	{
+		Processor.Initialize(IKRigAsset, GetSkeletalMesh(), FIKRigGoalContainer());
+		if (Processor.IsInitialized())
+		{
+			CacheRigElements(WorkData.CachedRigElements, Processor.GetSkeleton().BoneNames, Hierarchy);
+		}
+	}
+	if (!Processor.IsInitialized())
+	{
+		return;
+	}
+
+#if WITH_EDITOR
+	if (GIsEditor)
+	{
+		// live preview source asset settings in the retarget editor
+		// NOTE: this copies solver settings and goal.PositionAlpha and goal.RotationAlpha
+		Processor.CopyAllSettingsFromAsset(IKRigAsset);
+	}
+#endif
+
+	// set the goals
+	FIKRigGoalContainer& GoalContainer = Processor.GetGoalContainer();
+	
+	// default all goals to 0 pos/rot alpha (overridden by user supplied goals)
+	for (FIKRigGoal& Goal : GoalContainer.GetGoalArray())
+	{
+		Goal.PositionAlpha = 0.0; 
+		Goal.RotationAlpha = 0.0;
+	}
+
+	// apply input goals to the processor
+	for (const FIKRigGoalInput& GoalInput : Goals)
+	{
+		if (FIKRigGoal* Goal = GoalContainer.FindGoalByName(GoalInput.GoalName))
+		{
+			Goal->bEnabled = true;
+			Goal->Position = GoalInput.Transform.GetLocation();
+			Goal->Rotation = GoalInput.Transform.GetRotation().Rotator();
+			Goal->PositionAlpha = GoalInput.PositionAlpha;
+			Goal->RotationAlpha = GoalInput.RotationAlpha;
+			Goal->PositionSpace = EIKRigGoalSpace::Component;
+			Goal->RotationSpace = EIKRigGoalSpace::Component;
+		}
+	}
+	
+	// trigger reinitialization if the goal container was modified such that it needs it
+	if (Processor.GetGoalContainer().NeedsInitialized())
+	{
+		Processor.Initialize(IKRigAsset, GetSkeletalMesh(), FIKRigGoalContainer());
+	}
+	
+	// copy input pose to start IK solve from
+	for (int32 BoneIndex=0; BoneIndex<WorkData.CachedRigElements.Num(); BoneIndex++)
+	{
+		const FCachedRigElement& CachedRigElement = WorkData.CachedRigElements[BoneIndex];
+		if (!ensure(CachedRigElement.IsValid()))
+		{
+			continue;
+		}
+		const FTransform BoneGlobalTransform = Hierarchy->GetGlobalTransform(CachedRigElement.GetIndex());
+		Processor.SetInputBoneGlobal(BoneIndex, BoneGlobalTransform);
+	}
+	
+	// run IK solve
+	Processor.Solve();
+	
+	// copy results of solve
+	const TArray<FTransform>& OutputPose = Processor.GetSkeleton().CurrentPoseGlobal;
+	for (int32 BoneIndex=0; BoneIndex<OutputPose.Num(); BoneIndex++)
+	{
+		const FCachedRigElement& CachedRigElement = WorkData.CachedRigElements[BoneIndex];
+		if (!CachedRigElement.IsValid())
+		{
+			continue;
+		}
+
+		static bool bInitial = false;
+		static bool bAffectChildren = false;
+		static bool bSetupUndo = false;
+		Hierarchy->SetGlobalTransform(CachedRigElement.GetKey(), OutputPose[BoneIndex], bInitial, bAffectChildren, bSetupUndo);
+	}
+}
+
+void FRigUnit_IKRig::CacheRigElements(
+	TArray<FCachedRigElement>& OutMap,
+	const TArray<FName>& InBoneNames,
+	URigHierarchy* InHierarchy)
+{
+	OutMap.Reset(InBoneNames.Num());
+    
+	// cache a rig element for each bone in the skeletal mesh
+	for (const FName& BoneName : InBoneNames)
+	{
+		FRigElementKey RigElementKey = FRigElementKey(BoneName, ERigElementType::Bone);
+		OutMap.Emplace(RigElementKey, InHierarchy);
+	}
+}
diff --git a/Engine/Plugins/Animation/IKRig/Source/IKRig/Private/Rig/IKRigProcessor.cpp b/Engine/Plugins/Animation/IKRig/Source/IKRig/Private/Rig/IKRigProcessor.cpp
--- a/Engine/Plugins/Animation/IKRig/Source/IKRig/Private/Rig/IKRigProcessor.cpp
+++ b/Engine/Plugins/Animation/IKRig/Source/IKRig/Private/Rig/IKRigProcessor.cpp
@@ -316,7 +316,15 @@ void FIKRigProcessor::SetInputPoseGlobal(const TArray<FTransform>& InGlobalBoneT
 	}
 	check(InGlobalBoneTransforms.Num() == Skeleton.CurrentPoseGlobal.Num());
 	Skeleton.CurrentPoseGlobal = InGlobalBoneTransforms;
-	Skeleton.UpdateAllLocalTransformFromGlobal();
+}
+
+void FIKRigProcessor::SetInputBoneGlobal(const int32& InBoneIndex, const FTransform& InGlobalBoneTransform)
+{
+	if (!ensure(Skeleton.CurrentPoseGlobal.IsValidIndex(InBoneIndex)))
+	{
+		return;
+	}
+	Skeleton.CurrentPoseGlobal[InBoneIndex] = InGlobalBoneTransform;
 }
 
 void FIKRigProcessor::SetInputPoseToRefPose()
@@ -327,7 +335,6 @@ void FIKRigProcessor::SetInputPoseToRefPose()
 	}
 	
 	Skeleton.CurrentPoseGlobal = Skeleton.RefPoseGlobal;
-	Skeleton.UpdateAllLocalTransformFromGlobal();
 }
 
 void FIKRigProcessor::ApplyGoalsFromOtherContainer(const FIKRigGoalContainer& InGoalContainer)
@@ -364,6 +371,8 @@ void FIKRigProcessor::Solve(const FTransform& ComponentToWorld)
 	{
 		return;
 	}
+
+	Skeleton.UpdateAllLocalTransformFromGlobal();
 	
 	// convert goals into component space and blend towards input pose by alpha
 	ResolveFinalGoalTransforms(ComponentToWorld);
@@ -445,7 +454,13 @@ void FIKRigProcessor::CopyAllSettingsFromAsset(const UIKRigDefinition* SourceAss
 
 const FIKRigGoalContainer& FIKRigProcessor::GetGoalContainer() const
 {
-	check(bInitialized);
+	ensure(bInitialized);
+	return GoalContainer;
+}
+
+FIKRigGoalContainer& FIKRigProcessor::GetGoalContainer()
+{
+	ensure(bInitialized);
 	return GoalContainer;
 }
 
diff --git a/Engine/Plugins/Animation/IKRig/Source/IKRig/Public/AnimNodes/RigUnit_IKRig.h b/Engine/Plugins/Animation/IKRig/Source/IKRig/Public/AnimNodes/RigUnit_IKRig.h
new file mode 100644
--- /dev/null
+++ b/Engine/Plugins/Animation/IKRig/Source/IKRig/Public/AnimNodes/RigUnit_IKRig.h
@@ -0,0 +1,83 @@
+﻿// Copyright Epic Games, Inc. All Rights Reserved.
+
+#pragma once
+
+#include "Rig/IKRigProcessor.h"
+#include "Units/Highlevel/RigUnit_HighlevelBase.h"
+
+#include "RigUnit_IKRig.generated.h"
+
+class UIKRigDefinition;
+
+#define UE_API IKRIG_API
+
+USTRUCT()
+struct FIKRigWorkData
+{
+	GENERATED_BODY()
+	
+	FIKRigWorkData() {}
+	
+	TArray<FCachedRigElement> CachedRigElements;
+	
+	FIKRigProcessor IKRigProcessor;
+};
+
+USTRUCT(BlueprintType)
+struct FIKRigGoalInput
+{
+	GENERATED_BODY()
+	
+	FIKRigGoalInput() :
+		GoalName(NAME_None),
+		Transform(FTransform::Identity),
+		PositionAlpha(1.0),
+		RotationAlpha(1.0)
+	{
+	}
+
+	/** The name of the goal to affect. Must match the name in the IK Rig. */
+	UPROPERTY()
+	FName GoalName;
+
+	/** The position and rotation target for the IK Goal. */
+	UPROPERTY()
+	FTransform Transform;
+
+	/** Range 0-1, default 1. Blends the Goal position from the input pose position to the Goal position. */
+	UPROPERTY()
+	double PositionAlpha;
+
+	/** Range 0-1, default 1. Blends the Goal rotation from the input pose rotation to the Goal rotation. */
+	UPROPERTY()
+	double RotationAlpha;
+};
+
+/* Supply an IK Rig asset and provide goal transforms to run IK on the skeleton. */
+USTRUCT(meta=(DisplayName="IK Rig", Category="Hierarchy", Keywords="IK Rig, IK", NodeColor="0 1 1"))
+struct FRigUnit_IKRig : public FRigUnit_HighlevelBaseMutable
+{
+	GENERATED_BODY()
+
+	FRigUnit_IKRig() {}
+
+	RIGVM_METHOD()
+	UE_API virtual void Execute() override;
+
+	/** An IK Rig asset to be evaluated. */
+	UPROPERTY(meta = (Input))
+	TObjectPtr<UIKRigDefinition> IKRigAsset;
+
+	/** A list of Goals to solve. These must match what is in the IK Rig, any missing Goals will have their alphas set to zero. */
+	UPROPERTY(meta = (Input))
+	TArray<FIKRigGoalInput> Goals;
+
+	UPROPERTY(transient)
+	FIKRigWorkData WorkData;
+
+private:
+	
+	static void CacheRigElements(TArray<FCachedRigElement>& OutMap, const TArray<FName>& InBoneNames, URigHierarchy* InHierarchy);
+};
+
+#undef UE_API
\ No newline at end of file
diff --git a/Engine/Plugins/Animation/IKRig/Source/IKRig/Public/Rig/IKRigProcessor.h b/Engine/Plugins/Animation/IKRig/Source/IKRig/Public/Rig/IKRigProcessor.h
--- a/Engine/Plugins/Animation/IKRig/Source/IKRig/Public/Rig/IKRigProcessor.h
+++ b/Engine/Plugins/Animation/IKRig/Source/IKRig/Public/Rig/IKRigProcessor.h
@@ -59,6 +59,9 @@ public:
 	/** Set all bone transforms in global space. This is the pose the IK solve will start from */
 	UE_API void SetInputPoseGlobal(const TArray<FTransform>& InGlobalBoneTransforms);
 
+	/** Set input pose bone transforms individually */
+	UE_API void SetInputBoneGlobal(const int32& InBoneIndex, const FTransform& InGlobalBoneTransform);
+
 	/** Optionally can be called before Solve() to use the reference pose as start pose */
 	UE_API void SetInputPoseToRefPose();
 
@@ -87,6 +90,7 @@ public:
 
 	/** Get access to the internal goal data (read only) */
 	UE_API const FIKRigGoalContainer& GetGoalContainer() const;
+	UE_API FIKRigGoalContainer& GetGoalContainer();
 	/** Get the bone associated with a goal */
 	UE_API const FGoalBone* GetGoalBone(const FName& GoalName) const;
 	
