[IK Retarget] Added support for retargeting anim composite assets

diff --git a/Engine/Plugins/Animation/IKRig/Source/IKRigEditor/Private/RetargetEditor/SIKRetargetAssetBrowser.cpp b/Engine/Plugins/Animation/IKRig/Source/IKRigEditor/Private/RetargetEditor/SIKRetargetAssetBrowser.cpp
--- a/Engine/Plugins/Animation/IKRig/Source/IKRigEditor/Private/RetargetEditor/SIKRetargetAssetBrowser.cpp
+++ b/Engine/Plugins/Animation/IKRig/Source/IKRigEditor/Private/RetargetEditor/SIKRetargetAssetBrowser.cpp
@@ -6,6 +6,7 @@
 #include "ContentBrowserDataSource.h"
 #include "ContentBrowserModule.h"
 #include "IContentBrowserSingleton.h"
+#include "Animation/AnimComposite.h"
 #include "Animation/AnimMontage.h"
 #include "Animation/AnimSequence.h"
 #include "Animation/PoseAsset.h"
@@ -68,6 +69,7 @@ void SIKRetargetAssetBrowser::RefreshView()
 	AssetPickerConfig.Filter.ClassPaths.Add(UAnimSequence::StaticClass()->GetClassPathName());
 	AssetPickerConfig.Filter.ClassPaths.Add(UAnimMontage::StaticClass()->GetClassPathName());
 	AssetPickerConfig.Filter.ClassPaths.Add(UPoseAsset::StaticClass()->GetClassPathName());
+	AssetPickerConfig.Filter.ClassPaths.Add(UAnimComposite::StaticClass()->GetClassPathName());
 	AssetPickerConfig.InitialAssetViewType = EAssetViewType::Column;
 	AssetPickerConfig.bAddFilterUI = true;
 	AssetPickerConfig.bShowPathInColumnView = true;
diff --git a/Engine/Plugins/Animation/IKRig/Source/IKRigEditor/Private/RigEditor/SIKRigAssetBrowser.cpp b/Engine/Plugins/Animation/IKRig/Source/IKRigEditor/Private/RigEditor/SIKRigAssetBrowser.cpp
--- a/Engine/Plugins/Animation/IKRig/Source/IKRigEditor/Private/RigEditor/SIKRigAssetBrowser.cpp
+++ b/Engine/Plugins/Animation/IKRig/Source/IKRigEditor/Private/RigEditor/SIKRigAssetBrowser.cpp
@@ -6,6 +6,7 @@
 #include "ContentBrowserDataSource.h"
 #include "ContentBrowserModule.h"
 #include "IContentBrowserSingleton.h"
+#include "Animation/AnimComposite.h"
 #include "Animation/AnimMontage.h"
 #include "Animation/AnimSequence.h"
 #include "Animation/PoseAsset.h"
@@ -42,6 +43,7 @@ void SIKRigAssetBrowser::RefreshView()
 	AssetPickerConfig.Filter.ClassPaths.Add(UAnimSequence::StaticClass()->GetClassPathName());
 	AssetPickerConfig.Filter.ClassPaths.Add(UAnimMontage::StaticClass()->GetClassPathName());
 	AssetPickerConfig.Filter.ClassPaths.Add(UPoseAsset::StaticClass()->GetClassPathName());
+	AssetPickerConfig.Filter.ClassPaths.Add(UAnimComposite::StaticClass()->GetClassPathName());
 	AssetPickerConfig.InitialAssetViewType = EAssetViewType::Column;
 	AssetPickerConfig.bAddFilterUI = true;
 	AssetPickerConfig.bShowPathInColumnView = true;
