[IK Retarget]Cleaning up Copy Base Pose getters and setters

diff --git a/Engine/Plugins/Animation/IKRig/Source/IKRig/Private/Retargeter/RetargetOps/CopyBasePoseOp.cpp b/Engine/Plugins/Animation/IKRig/Source/IKRig/Private/Retargeter/RetargetOps/CopyBasePoseOp.cpp
--- a/Engine/Plugins/Animation/IKRig/Source/IKRig/Private/Retargeter/RetargetOps/CopyBasePoseOp.cpp
+++ b/Engine/Plugins/Animation/IKRig/Source/IKRig/Private/Retargeter/RetargetOps/CopyBasePoseOp.cpp
@@ -210,7 +210,7 @@ void UIKRetargetCopyBasePoseController::SetSettings(FIKRetargetCopyBasePoseOpSet
 	OpSettingsToControl->CopySettingsAtRuntime(&InSettings);
 }
 
-void UIKRetargetCopyBasePoseController::SetCopyFromRootBone(const FName InBoneName) const
+void UIKRetargetCopyBasePoseController::SetCopyFromStart(const FName InBoneName) const
 {
 	FIKRetargetCopyBasePoseOpSettings* Settings = reinterpret_cast<FIKRetargetCopyBasePoseOpSettings*>(OpSettingsToControl);
 	if (!Settings)
@@ -221,7 +221,7 @@ void UIKRetargetCopyBasePoseController::SetCopyFromRootBone(const FName InBoneNa
 	Settings->CopyFromStart.BoneName = InBoneName;
 }
 
-FName UIKRetargetCopyBasePoseController::GetCopyFromRootBone(const FName InBoneName) const
+FName UIKRetargetCopyBasePoseController::GetCopyFromStart() const
 {
 	FIKRetargetCopyBasePoseOpSettings* Settings = reinterpret_cast<FIKRetargetCopyBasePoseOpSettings*>(OpSettingsToControl);
 	if (!Settings)
diff --git a/Engine/Plugins/Animation/IKRig/Source/IKRig/Public/Retargeter/RetargetOps/CopyBasePoseOp.h b/Engine/Plugins/Animation/IKRig/Source/IKRig/Public/Retargeter/RetargetOps/CopyBasePoseOp.h
--- a/Engine/Plugins/Animation/IKRig/Source/IKRig/Public/Retargeter/RetargetOps/CopyBasePoseOp.h
+++ b/Engine/Plugins/Animation/IKRig/Source/IKRig/Public/Retargeter/RetargetOps/CopyBasePoseOp.h
@@ -32,7 +32,7 @@ struct FIKRetargetCopyBasePoseOpSettings : public FIKRetargetOpSettingsBase
 
 	UE_API virtual void CopySettingsAtRuntime(const FIKRetargetOpSettingsBase* InSettingsToCopyFrom) override;
 
-	/** Custom serialization to handle conversion from deprecated CopyBasePoseRoot to CopyFromRoot */
+	/** Custom serialization to handle conversion from deprecated CopyBasePoseRoot to CopyFromStart */
 	void PostSerialize(const FArchive& Ar);
 
 private:
@@ -84,7 +84,7 @@ struct FIKRetargetCopyBasePoseOp : public FIKRetargetOpBase
 
 private:
 	// mapping from source bones to target bone indices
-	// NOTE: this map will only contain bones with the same name in both skeletons that are below the RootBoneToCopy (if specified)
+	// NOTE: this map will only contain bones with the same name in both skeletons that are below the CopyFromStart (if specified)
 	TMap<int32, int32> SourceToTargetBoneIndexMap;
 
 	// children of bones that are copied that need manually updated
@@ -108,13 +108,13 @@ public:
 	UFUNCTION(BlueprintCallable, Category = Settings)
 	UE_API void SetSettings(FIKRetargetCopyBasePoseOpSettings InSettings);
 
-	/* Set the name of the bone to use as the "Copy From Root" in the op */
-	UFUNCTION(BlueprintCallable, Category = CopyFromRoot)
-	UE_API void SetCopyFromRootBone(const FName InBoneName) const;
+	/* Set the name of the bone to use as the "Copy From Start" in the op */
+	UFUNCTION(BlueprintCallable, Category = CopyFromStart)
+	UE_API void SetCopyFromStart(const FName InBoneName) const;
 
-	/* Get the name of the bone to use as the "Copy From Root" in the op */
-	UFUNCTION(BlueprintCallable, Category = CopyFromRoot)
-	UE_API FName GetCopyFromRootBone(const FName InBoneName) const;
+	/* Get the name of the bone to use as the "Copy From Start" in the op */
+	UFUNCTION(BlueprintCallable, Category = CopyFromStart)
+	UE_API FName GetCopyFromStart() const;
 
 	/* Add a bone to the "BonesToExclude" array */
 	UFUNCTION(BlueprintCallable, Category = ExcludedBones)
