[IK Retarget] Fixing refresh of chain settings UI in details panel.

diff --git a/Engine/Plugins/Animation/IKRig/Source/IKRig/Private/Retargeter/IKRetargetChainMapping.cpp b/Engine/Plugins/Animation/IKRig/Source/IKRig/Private/Retargeter/IKRetargetChainMapping.cpp
--- a/Engine/Plugins/Animation/IKRig/Source/IKRig/Private/Retargeter/IKRetargetChainMapping.cpp
+++ b/Engine/Plugins/Animation/IKRig/Source/IKRig/Private/Retargeter/IKRetargetChainMapping.cpp
@@ -14,12 +14,18 @@ void FRetargetChainMapping::ReinitializeWithIKRigs(const UIKRigDefinition* InSou
 	SourceIKRig = InSourceIKRig;
 	TargetIKRig = InTargetIKRig;
 	
-	if (!(IsValid(SourceIKRig) && IsValid(TargetIKRig)))
+	if (!ChainMap.IsEmpty() && !(IsValid(SourceIKRig) && IsValid(TargetIKRig)))
 	{
 		// don't clean data until we have a replacement
 		// we want to preserve existing mappings in case they are still relevant once the op is configured with both rigs
 		return;
 	}
+
+	if (!IsValid(TargetIKRig))
+	{
+		// need at least a target before we can add target chains
+		return;
+	}
 	
 	// remove null target chains
 	ChainMap.RemoveAll([](const FRetargetChainPair& ChainPair)
@@ -44,7 +50,7 @@ void FRetargetChainMapping::ReinitializeWithIKRigs(const UIKRigDefinition* InSou
 	}
 	
 	// reset any sources that are no longer present to "None"
-	TArray<FName> SourceChainNames = SourceIKRig->GetRetargetChainNames();
+	TArray<FName> SourceChainNames = SourceIKRig ? SourceIKRig->GetRetargetChainNames() : TArray<FName>();
 	for (FRetargetChainPair& Pair : ChainMap)
 	{
 		if (!SourceChainNames.Contains(Pair.SourceChainName))
@@ -62,6 +68,11 @@ bool FRetargetChainMapping::IsReady() const
 	return SourceIKRig && TargetIKRig && !ChainMap.IsEmpty();
 }
 
+bool FRetargetChainMapping::HasAnyChains() const
+{
+	return !ChainMap.IsEmpty();
+}
+
 const UIKRigDefinition* FRetargetChainMapping::GetIKRig(ERetargetSourceOrTarget SourceOrTarget)
 {
 	return SourceOrTarget == ERetargetSourceOrTarget::Source ? SourceIKRig : TargetIKRig;
diff --git a/Engine/Plugins/Animation/IKRig/Source/IKRig/Private/Retargeter/RetargetOps/PelvisMotionOp.cpp b/Engine/Plugins/Animation/IKRig/Source/IKRig/Private/Retargeter/RetargetOps/PelvisMotionOp.cpp
--- a/Engine/Plugins/Animation/IKRig/Source/IKRig/Private/Retargeter/RetargetOps/PelvisMotionOp.cpp
+++ b/Engine/Plugins/Animation/IKRig/Source/IKRig/Private/Retargeter/RetargetOps/PelvisMotionOp.cpp
@@ -130,6 +130,11 @@ void FIKRetargetPelvisMotionOp::OnAssignIKRig(
 	const UIKRigDefinition* InIKRig,
 	const FIKRetargetOpBase* InParentOp)
 {
+	if (!InIKRig)
+	{
+		return;
+	}
+	
 	FName& PelvisNameSetting = SourceOrTarget == ERetargetSourceOrTarget::Source ? Settings.SourcePelvisBone.BoneName : Settings.TargetPelvisBone.BoneName;
 	PelvisNameSetting = InIKRig->GetPelvis();
 }
diff --git a/Engine/Plugins/Animation/IKRig/Source/IKRig/Public/Retargeter/IKRetargetChainMapping.h b/Engine/Plugins/Animation/IKRig/Source/IKRig/Public/Retargeter/IKRetargetChainMapping.h
--- a/Engine/Plugins/Animation/IKRig/Source/IKRig/Public/Retargeter/IKRetargetChainMapping.h
+++ b/Engine/Plugins/Animation/IKRig/Source/IKRig/Public/Retargeter/IKRetargetChainMapping.h
@@ -47,6 +47,9 @@ public:
 	/** returns true if both a source and target IK Rig have been loaded */
 	UE_API bool IsReady() const;
 
+	/** returns true if any chains have been loaded, regardless of if they are mapped */
+	UE_API bool HasAnyChains() const;
+
 	/* get one of the IK Rigs used in this mapping */
 	UE_API const UIKRigDefinition* GetIKRig(ERetargetSourceOrTarget SourceOrTarget);
 
diff --git a/Engine/Plugins/Animation/IKRig/Source/IKRig/Public/Retargeter/IKRetargetOps.h b/Engine/Plugins/Animation/IKRig/Source/IKRig/Public/Retargeter/IKRetargetOps.h
--- a/Engine/Plugins/Animation/IKRig/Source/IKRig/Public/Retargeter/IKRetargetOps.h
+++ b/Engine/Plugins/Animation/IKRig/Source/IKRig/Public/Retargeter/IKRetargetOps.h
@@ -122,7 +122,7 @@ public:
 
 	/** (optional) override to react after user assigns a new IK Rig to all ops
 	 * @param SourceOrTarget: whether the source or target IK Rig was changed
-	 * @param InIKRig: the IK rig asset to assign to this op
+	 * @param InIKRig: the IK rig asset to assign to this op (will be NULL if user cleared an IK Rig)
 	 * @param InParentOp: the parent for this op (may be null if unparented) */
 	virtual void OnAssignIKRig(const ERetargetSourceOrTarget SourceOrTarget, const UIKRigDefinition* InIKRig, const FIKRetargetOpBase* InParentOp) {};
 	
diff --git a/Engine/Plugins/Animation/IKRig/Source/IKRig/Public/Retargeter/RetargetOps/FKChainsOp.h b/Engine/Plugins/Animation/IKRig/Source/IKRig/Public/Retargeter/RetargetOps/FKChainsOp.h
--- a/Engine/Plugins/Animation/IKRig/Source/IKRig/Public/Retargeter/RetargetOps/FKChainsOp.h
+++ b/Engine/Plugins/Animation/IKRig/Source/IKRig/Public/Retargeter/RetargetOps/FKChainsOp.h
@@ -235,7 +235,7 @@ struct FFKChainDebugData
 };
 #endif
 
-USTRUCT(BlueprintType, meta = (DisplayName = "Retarget FK Chains"))
+USTRUCT(BlueprintType, meta = (DisplayName = "FK Chains"))
 struct FIKRetargetFKChainsOp : public FIKRetargetOpBase
 {
 	GENERATED_BODY()
diff --git a/Engine/Plugins/Animation/IKRig/Source/IKRigEditor/Private/RetargetEditor/SIKRetargetChainMapList.cpp b/Engine/Plugins/Animation/IKRig/Source/IKRigEditor/Private/RetargetEditor/SIKRetargetChainMapList.cpp
--- a/Engine/Plugins/Animation/IKRig/Source/IKRigEditor/Private/RetargetEditor/SIKRetargetChainMapList.cpp
+++ b/Engine/Plugins/Animation/IKRig/Source/IKRigEditor/Private/RetargetEditor/SIKRetargetChainMapList.cpp
@@ -233,6 +233,16 @@ void SIKRetargetChainMapList::Construct(const FArguments& InArgs)
 		return; // chain map editing must be associated with a particular op
 	}
 	
+	// register callback to refresh when the retargeter is initialized (happens whenever IK Rigs are swapped or edited)
+	TWeakPtr<SIKRetargetChainMapList> WeakSelf = StaticCastWeakPtr<SIKRetargetChainMapList>(AsWeak());
+	Config.Controller->OnRetargeterNeedsInitialized().AddLambda([WeakSelf]()
+	{
+		if (WeakSelf.IsValid())
+		{
+			WeakSelf.Pin()->RefreshView();
+		}
+	});
+	
 	TextFilter = MakeShareable(new FTextFilterExpressionEvaluator(ETextFilterExpressionEvaluatorMode::BasicString));
 	
 	// create property view for the chain properties
@@ -465,10 +475,14 @@ void SIKRetargetChainMapList::RefreshView()
 
 	// get the chain mapping for the op
 	const FRetargetChainMapping* ChainMapping = Config.Controller->GetChainMapping(Config.OpWithChainMapping);
-	if (!ChainMapping->IsReady())
+	if (!ChainMapping)
+	{
+		return; // op deleted, being cleaned up
+	}
+	
+	if (!ChainMapping->HasAnyChains())
 	{
-		// invalid chain mapping means we don't display anything
-		// user must have valid source and target IK rigs to make a mapping and store settings
+		// must have at least a target IK rig assigned with some chains loaded
 		return;
 	}
 
diff --git a/Engine/Plugins/Animation/IKRig/Source/IKRigEditor/Private/RetargetEditor/SRetargetOpStack.cpp b/Engine/Plugins/Animation/IKRig/Source/IKRigEditor/Private/RetargetEditor/SRetargetOpStack.cpp
--- a/Engine/Plugins/Animation/IKRig/Source/IKRigEditor/Private/RetargetEditor/SRetargetOpStack.cpp
+++ b/Engine/Plugins/Animation/IKRig/Source/IKRigEditor/Private/RetargetEditor/SRetargetOpStack.cpp
@@ -508,7 +508,7 @@ TSharedRef<SWidget> SRetargetOpList::CreateAddNewOpMenu()
 	constexpr bool bCloseMenuAfterSelection = true;
 	FMenuBuilder MenuBuilder(bCloseMenuAfterSelection, nullptr);
 
-	MenuBuilder.BeginSection("AddNewRetargetOp", LOCTEXT("AddOperations", "Add New Retarget Op"));
+	MenuBuilder.BeginSection("AddNewRetargetOp", LOCTEXT("AddOperations", "Add New Op"));
 
 	// add menu option to create each retarget op type
 	const TObjectPtr<UIKRetargeterController> AssetController = EditorController.Pin()->AssetController;
@@ -831,14 +831,14 @@ void SRetargetOpStack::Construct(
 				[
 					SNew(SHorizontalBox)
 					+ SHorizontalBox::Slot()
+					.AutoWidth()
 					.VAlign(VAlign_Center)
 					.HAlign(HAlign_Left)
-			
-					.Padding(6.0f, 4.0f)
+					.Padding(2.0f, 4.0f)
 					[
 						SNew(SPositiveActionButton)
 						.Icon(FAppStyle::Get().GetBrush("Icons.Plus"))
-						.Text(LOCTEXT("AddNewRetargetOpLabel", "Add New Retarget Op"))
+						.Text(LOCTEXT("AddNewRetargetOpLabel", "Add New Op"))
 						.ToolTipText(LOCTEXT("AddNewToolTip", "Add a new operation to run as part of the retargeter."))
 						.OnGetMenuContent_Lambda([this]()
 						{
@@ -847,10 +847,10 @@ void SRetargetOpStack::Construct(
 					]
 
 					+ SHorizontalBox::Slot()
+					.AutoWidth()
 					.VAlign(VAlign_Center)
 					.HAlign(HAlign_Left)
-				
-					.Padding(6.0f, 4.0f)
+					.Padding(2.0f, 4.0f)
 					[
 						SNew(SPositiveActionButton)
 						.Icon(FIKRetargetEditorStyle::Get().GetBrush("IKRetarget.PostSettingsSmall"))
