[IK Retarget] Adding floor constraints to pelvis and IK goal ops.

diff --git a/Engine/Plugins/Animation/IKRig/Source/IKRig/Private/Retargeter/RetargetOps/IKChainsOp.cpp b/Engine/Plugins/Animation/IKRig/Source/IKRig/Private/Retargeter/RetargetOps/IKChainsOp.cpp
--- a/Engine/Plugins/Animation/IKRig/Source/IKRig/Private/Retargeter/RetargetOps/IKChainsOp.cpp
+++ b/Engine/Plugins/Animation/IKRig/Source/IKRig/Private/Retargeter/RetargetOps/IKChainsOp.cpp
@@ -138,6 +138,20 @@ FVector FIKChainRetargeter::CalculateGoalPosition(
 	const FVector CurrentTargetStartLocation = InTargetGlobalPose[TargetStartBoneIndex].GetTranslation() - InvRootModification;
 	FVector GoalPosition = CurrentTargetStartLocation + (CurrentSourceEndDirectionNormalized * InitialTargetChainLength);
 
+	// apply floor constraint
+	if (!FMath::IsNearlyZero(Settings->FloorConstraintWeight))
+	{
+		const float SourceHeight = FMath::Abs(CurrentSourceEndGlobal.GetLocation().Z);
+		if (SourceHeight <= Settings->FloorHeightFalloffEnd)
+		{
+			const float Start = Settings->FloorHeightFalloffStart;
+			float End = Settings->FloorHeightFalloffEnd;
+			End = End <= Start ? Start + 1.0f : End;
+			double Falloff = SourceHeight < Start ? 1.0f : 1.0f - (SourceHeight - Start) / (End - Start);
+			GoalPosition.Z = FMath::Lerp(GoalPosition.Z, SourceHeight, Falloff * Settings->FloorConstraintWeight);
+		}
+	}
+
 	// blend to source location
 	const double BlendToSourceTranslation = Settings->BlendToSource * Settings->BlendToSourceTranslation;
 	if (BlendToSourceTranslation > UE_KINDA_SMALL_NUMBER)
@@ -171,6 +185,9 @@ FVector FIKChainRetargeter::CalculateGoalPosition(
 bool FRetargetIKChainSettings::operator==(const FRetargetIKChainSettings& Other) const
 {
 	return EnableIK == Other.EnableIK
+		&& FMath::IsNearlyEqualByULP(FloorConstraintWeight, Other.FloorConstraintWeight)
+		&& FMath::IsNearlyEqualByULP(FloorHeightFalloffStart, Other.FloorHeightFalloffStart)
+		&& FMath::IsNearlyEqualByULP(FloorHeightFalloffEnd, Other.FloorHeightFalloffEnd)
 		&& FMath::IsNearlyEqualByULP(BlendToSource, Other.BlendToSource)
 		&& BlendToSourceWeights.Equals(Other.BlendToSourceWeights)
 		&& StaticOffset.Equals(Other.StaticOffset)
@@ -316,6 +333,8 @@ void FIKRetargetIKChainsOp::Run(
 	const TArray<FTransform>& InSourceGlobalPose,
 	TArray<FTransform>& OutTargetGlobalPose)
 {
+	SCOPE_CYCLE_COUNTER(STAT_IKRetargetGoals);
+	
 	if (InProcessor.IsIKForcedOff())
 	{
 		return; // skip this op when IK is off
diff --git a/Engine/Plugins/Animation/IKRig/Source/IKRig/Private/Retargeter/RetargetOps/PelvisMotionOp.cpp b/Engine/Plugins/Animation/IKRig/Source/IKRig/Private/Retargeter/RetargetOps/PelvisMotionOp.cpp
--- a/Engine/Plugins/Animation/IKRig/Source/IKRig/Private/Retargeter/RetargetOps/PelvisMotionOp.cpp
+++ b/Engine/Plugins/Animation/IKRig/Source/IKRig/Private/Retargeter/RetargetOps/PelvisMotionOp.cpp
@@ -158,6 +158,7 @@ bool FIKRetargetPelvisMotionOp::InitializeSource(
 	// record initial root data
 	const FTransform InitialTransform = SourceSkeleton.RetargetPoses.GetGlobalRetargetPose()[Source.BoneIndex]; 
 	float InitialHeight = static_cast<float>(InitialTransform.GetTranslation().Z);
+	Source.InitialPosition = InitialTransform.GetLocation();
 	Source.InitialRotation = InitialTransform.GetRotation();
 
 	// ensure root height is not at origin, this happens if user sets root to ACTUAL skeleton root and not pelvis
@@ -276,11 +277,34 @@ void FIKRetargetPelvisMotionOp::DecodePose(FTransform& OutPelvisGlobalPose)
 	FVector Position;
 	{
 		// generate basic pelvis position by scaling the normalized position by root height
-		const FVector RetargetedPosition = Source.CurrentPositionNormalized * Target.InitialHeight;
+		FVector RetargetedPosition = Source.CurrentPositionNormalized * Target.InitialHeight;
+		Position = RetargetedPosition;
+
+		// apply crotch-to-floor constraint
+		if (!FMath::IsNearlyZero(Settings.FloorConstraintWeight))
+		{
+			if (Source.CurrentPosition.Z <= Source.InitialPosition.Z)
+			{
+				const float SourceInitialCrotchHeight = Source.InitialPosition.Z - Settings.SourceCrotchOffset;
+				const float SourceCurrentCrotchHeight = Source.CurrentPosition.Z - Settings.SourceCrotchOffset;
+				const float SourceCrotchHeightPercent = SourceCurrentCrotchHeight * (1.0f / FMath::Max(1.0f, SourceInitialCrotchHeight));
+				const float TargetInitialCrotchHeight = Target.InitialPosition.Z - Settings.TargetCrotchOffset;
+				const float TargetCurrentCrotchHeight = TargetInitialCrotchHeight * SourceCrotchHeightPercent;
+				const float ConstrainedHeight = TargetCurrentCrotchHeight + Settings.TargetCrotchOffset; 
+				RetargetedPosition.Z = Position.Z = FMath::Lerp(Position.Z, ConstrainedHeight, Settings.FloorConstraintWeight);
+			}
+			else
+			{
+				const float InitialSourceTargetDelta = Target.InitialPosition.Z - Source.InitialPosition.Z;
+				const float ConstrainedHeight = Source.CurrentPosition.Z + InitialSourceTargetDelta;
+				RetargetedPosition.Z = Position.Z = FMath::Lerp(Position.Z, ConstrainedHeight, Settings.FloorConstraintWeight);
+			}
+			
+		}
 		
 		// blend the pelvis position towards the source pelvis position
 		const FVector PerAxisAlpha = Settings.BlendToSourceTranslation * Settings.BlendToSourceTranslationWeights;
-		Position = FMath::Lerp(RetargetedPosition, Source.CurrentPosition, PerAxisAlpha);
+		Position = FMath::Lerp(Position, Source.CurrentPosition, PerAxisAlpha);
 
 		// apply vertical / horizontal scaling of motion
 		FVector ScaledRetargetedPosition = Position;
@@ -350,21 +374,42 @@ void FIKRetargetPelvisMotionOp::DebugDraw(
 
 	FScopeLock ScopeLock(&DebugDataMutex);
 	
-	const FTransform PelvisTransform = CurrentPelvisTransform * InComponentTransform;
-	const FVector PelvisCircleLocation = PelvisTransform.GetLocation() * FVector(1,1,0);
+	const FTransform TargetPelvis = CurrentPelvisTransform * InComponentTransform;
+	const FVector TargetFloor = TargetPelvis.GetLocation() * FVector(1,1,0);
+	const FVector TargetCrotch = TargetPelvis.GetLocation() + FVector(0,0, -Settings.TargetCrotchOffset);
+	const FVector SourceFloor = Source.CurrentPosition * FVector(1,1,0);
+	const FVector SourceCrotch = Source.CurrentPosition + FVector(0,0, -Settings.SourceCrotchOffset);
 	const FLinearColor CircleColor = InEditorState.bIsRootSelected ? InEditorState.MainColor : InEditorState.MainColor * InEditorState.NonSelected;
+	const float Thickness = static_cast<float>(Settings.DebugDrawThickness * InComponentScale);
+	const float Size = static_cast<float>(Settings.DebugDrawSize * InComponentScale);
 
 	InPDI->SetHitProxy(new HIKRetargetEditorRootProxy());
-	DrawCircle(
+
+	// floor coordinate system
+	DrawCoordinateSystem(
 		InPDI,
-		PelvisCircleLocation,
-		FVector(1, 0, 0),
-		FVector(0, 1, 0),
-		CircleColor,
-		Settings.DebugDrawSize * InComponentTransform.GetScale3D().GetMax(),
-		30,
+		TargetFloor,
+		TargetPelvis.GetRotation().Rotator(),
+		Size * 2.0f,
 		SDPG_World,
-		static_cast<float>(Settings.DebugDrawThickness * InComponentScale));
+		Thickness * 0.25f);
+	
+	// draw pelvis locations
+	InPDI->DrawPoint(Source.CurrentPosition, FLinearColor::Black, Size, SDPG_Foreground);
+	InPDI->DrawPoint(TargetPelvis.GetLocation(), FLinearColor::Black, Size, SDPG_Foreground);
+
+	// draw crotch locations
+	InPDI->DrawPoint(SourceCrotch, FLinearColor::Green, Size, SDPG_Foreground);
+	InPDI->DrawPoint(TargetCrotch, FLinearColor::Green, Size, SDPG_Foreground);
+
+	// draw pelvis-to-crotch lines
+	DrawDashedLine(InPDI,Source.CurrentPosition,SourceCrotch, FLinearColor::Black, 1.0f,SDPG_Foreground);
+	DrawDashedLine(InPDI,TargetPelvis.GetLocation(),TargetCrotch, FLinearColor::Black, 1.0f,SDPG_Foreground);
+	
+	// draw crotch-to-floor lines
+	DrawDashedLine(InPDI,SourceCrotch,SourceFloor, CircleColor, 1.0f,SDPG_Foreground);
+	DrawDashedLine(InPDI,TargetCrotch,TargetFloor, CircleColor, 1.0f,SDPG_Foreground);
+	
 	InPDI->SetHitProxy(nullptr);
 }
 #endif
diff --git a/Engine/Plugins/Animation/IKRig/Source/IKRig/Public/Retargeter/RetargetOps/IKChainsOp.h b/Engine/Plugins/Animation/IKRig/Source/IKRig/Public/Retargeter/RetargetOps/IKChainsOp.h
--- a/Engine/Plugins/Animation/IKRig/Source/IKRig/Public/Retargeter/RetargetOps/IKChainsOp.h
+++ b/Engine/Plugins/Animation/IKRig/Source/IKRig/Public/Retargeter/RetargetOps/IKChainsOp.h
@@ -5,6 +5,8 @@
 #include "Retargeter/IKRetargetOps.h"
 #include "IKChainsOp.generated.h"
 
+DECLARE_CYCLE_STAT(TEXT("IK Retarget Goals"), STAT_IKRetargetGoals, STATGROUP_Anim);
+
 #define UE_API IKRIG_API
 
 #define LOCTEXT_NAMESPACE "IKChainOp"
@@ -78,6 +80,26 @@ struct FRetargetIKChainSettings
 	UPROPERTY(EditAnywhere, BlueprintReadWrite, Category = "Setup", meta=(ReinitializeOnEdit))
 	bool EnableIK = true;
 
+	/** Range 0 to 1. Default is 0. Set this to 1 to turn the floor constraint ON.
+	 * When ON, the floor constraint will adjust the vertical position of the IK Goal according to the following rules.
+	 * 1. When the source goal bone is LOWER than FloorHeightFalloffStart, the height of the goal is smoothly blended to the height of the source bone.
+	 * 2. When the source goal bone is HIGHER than FloorHeightFalloffStart, the height of the goal is left at its normal retargeted location.
+	 * NOTE: the floor is assumed to be the XY plane where Z = 0.
+	 */
+	UPROPERTY(EditAnywhere, BlueprintReadWrite, Category = FloorHeightConstraint, meta = (ClampMin = "0.0", ClampMax = "1.0", UIMin = "0.0", UIMax = "1.0"))
+	double FloorConstraintWeight = 0.0f;
+
+	/** Range 0 to inf. Default is 8. The height in cm from the floor below which the goal is snapped directly to the source bone height.
+	 * NOTE: if the source bone height is greater than this value, but lower than FloorHeightFalloffEnd, then the height will smoothly blend from the source
+	 * bone height, to the height of the goal in its normal retargeted position. */
+	UPROPERTY(EditAnywhere, BlueprintReadWrite, Category = FloorHeightConstraint, meta = (ClampMin = "0.1", UIMin = "0.1", UIMax = "100.0"))
+	double FloorHeightFalloffStart = 8.0f;
+	
+	/** Range 0 to inf. Default is 20. The height in cm from the floor below which the goal is gradually blended towards the source bone height.
+	 * NOTE: if the source bone is higher than this value, the height of the goal is left at its normal retargeted height. */
+	UPROPERTY(EditAnywhere, BlueprintReadWrite, Category = FloorHeightConstraint, meta = (ClampMin = "0.1", UIMin = "0.1", UIMax = "100.0"))
+	double FloorHeightFalloffEnd = 20.0f;
+
 	/** Range 0 to 1. Default 0. Blends IK goal transform from retargeted transform (0) to source bone transform (1).
 	*  At 0 the goal is placed at the retargeted location and rotation.
 	*  At 1 the goal is placed at the location and rotation of the source chain's end bone. */
diff --git a/Engine/Plugins/Animation/IKRig/Source/IKRig/Public/Retargeter/RetargetOps/PelvisMotionOp.h b/Engine/Plugins/Animation/IKRig/Source/IKRig/Public/Retargeter/RetargetOps/PelvisMotionOp.h
--- a/Engine/Plugins/Animation/IKRig/Source/IKRig/Public/Retargeter/RetargetOps/PelvisMotionOp.h
+++ b/Engine/Plugins/Animation/IKRig/Source/IKRig/Public/Retargeter/RetargetOps/PelvisMotionOp.h
@@ -24,6 +24,24 @@ struct FIKRetargetPelvisMotionOpSettings : public FIKRetargetOpSettingsBase
 	/** The Pelvis bone on the target skeleton to copy motion TO. */
 	UPROPERTY(EditAnywhere, Category=Setup, meta=(ReinitializeOnEdit))
 	FBoneReference TargetPelvisBone;
+
+
+	/** Range 0 to 1. Default is 0. Set to a value of 1 to turn the pelvis floor constraint ON.
+	 * When ON, the floor constraint will adjust the vertical Pelvis motion separately according to the following rules:
+	 * 1. When the source pelvis is LOWER than the ref pose, the target pelvis will be lowered proportional to their relative crotch heights.
+	 * 2. When the source pelvis is HIGHER than the ref pose, the target pelvis will be raised an identical amount while maintaining any vertical offset in the ref pose.*/
+	UPROPERTY(EditAnywhere, Category=FloorConstraint, meta = (ClampMin = "0.0", ClampMax = "1.0", UIMin = "0.0", UIMax = "1.0"))
+	double FloorConstraintWeight = 0.0f;
+	
+	/** A negative vertical offset in cm relative to the Pelvis bone of the SOURCE.
+	 * NOTE: Adjust this until the green dot is located roughly at the crotch of the character. */
+	UPROPERTY(EditAnywhere, BlueprintReadWrite, Category = FloorConstraint)
+	double SourceCrotchOffset = 0.0f;
+
+	/** A negative vertical offset in cm relative to the Pelvis bone of the TARGET.
+	 * NOTE: Adjust this until the green dot is located roughly at the crotch of the character. */
+	UPROPERTY(EditAnywhere, BlueprintReadWrite, Category = FloorConstraint)
+	double TargetCrotchOffset = 0.0f;
 	
 	/** Range 0 to 1. Default 1. Blends the amount of retargeted pelvis rotation to apply.
 	*  At 0 the pelvis is left at the rotation from the retarget pose.
