[IK Retarget] Added profiling to the retarget op stack UI.

diff --git a/Engine/Plugins/Animation/IKRig/Source/IKRig/Private/AnimNodes/AnimNode_RetargetPoseFromMesh.cpp b/Engine/Plugins/Animation/IKRig/Source/IKRig/Private/AnimNodes/AnimNode_RetargetPoseFromMesh.cpp
--- a/Engine/Plugins/Animation/IKRig/Source/IKRig/Private/AnimNodes/AnimNode_RetargetPoseFromMesh.cpp
+++ b/Engine/Plugins/Animation/IKRig/Source/IKRig/Private/AnimNodes/AnimNode_RetargetPoseFromMesh.cpp
@@ -74,6 +74,11 @@ void FAnimNode_RetargetPoseFromMesh::Evaluate_AnyThread(FPoseContext& Output)
 	DECLARE_SCOPE_HIERARCHICAL_COUNTER_ANIMNODE(Evaluate_AnyThread)
 	SCOPE_CYCLE_COUNTER(STAT_IKRetarget);
 
+	// editor profiling
+#ifdef WITH_EDITOR
+	double StartTime = FPlatformTime::Seconds();
+#endif
+
 	auto IsReadyToRun = [this, &Output]()
 	{
 		if (!IKRetargeterAsset)
@@ -228,6 +233,13 @@ void FAnimNode_RetargetPoseFromMesh::Evaluate_AnyThread(FPoseContext& Output)
 
 	// apply the retargeted pose to the output pin
 	ApplyRetargetedPose(RetargetedPose);
+
+	// editor profiling
+#ifdef WITH_EDITOR
+	constexpr double FrameWindow = 30.0;
+	constexpr double Alpha = 2.0 / (FrameWindow + 1.0);
+	AverageExecutionTime += Alpha * ((FPlatformTime::Seconds() - StartTime) - AverageExecutionTime);
+#endif
 }
 
 void FAnimNode_RetargetPoseFromMesh::PreUpdate(const UAnimInstance* InAnimInstance)
@@ -414,3 +426,8 @@ void FAnimNode_RetargetPoseFromMesh::PostSerialize(const FArchive& Ar)
 #endif // WITH_EDITORONLY_DATA
 }
 
+double FAnimNode_RetargetPoseFromMesh::GetAverageExecutionTime() const
+{
+	return AverageExecutionTime;
+}
+
diff --git a/Engine/Plugins/Animation/IKRig/Source/IKRig/Private/Retargeter/IKRetargetOps.cpp b/Engine/Plugins/Animation/IKRig/Source/IKRig/Private/Retargeter/IKRetargetOps.cpp
--- a/Engine/Plugins/Animation/IKRig/Source/IKRig/Private/Retargeter/IKRetargetOps.cpp
+++ b/Engine/Plugins/Animation/IKRig/Source/IKRig/Private/Retargeter/IKRetargetOps.cpp
@@ -36,6 +36,7 @@ void FIKRetargetOpBase::SetSettings(const FIKRetargetOpSettingsBase* InSettings)
 }
 
 #if WITH_EDITOR
+
 FName FIKRetargetOpBase::GetDefaultName() const
 {
 	// calls virtual GetType() to get the derived type
@@ -66,6 +67,24 @@ FText FIKRetargetOpBase::GetWarningMessage() const
 	
 	return LOCTEXT("OpNotReady", "Not initialized. See output log.");
 }
+
+void FIKRetargetOpBase::ResetExecutionTime()
+{
+	AverageExecutionTime = 0.0f;
+}
+
+void FIKRetargetOpBase::SetLatestExecutionTime(const double InSeconds)
+{
+	constexpr double FrameWindow = 30.0;
+	constexpr double Alpha = 2.0 / (FrameWindow + 1.0);
+	AverageExecutionTime += Alpha * (InSeconds - AverageExecutionTime);
+}
+
+double FIKRetargetOpBase::GetAverageExecutionTime() const
+{
+	return AverageExecutionTime;
+}
+
 #endif
 
 FName FIKRetargetOpBase::GetName() const
diff --git a/Engine/Plugins/Animation/IKRig/Source/IKRig/Private/Retargeter/IKRetargetProcessor.cpp b/Engine/Plugins/Animation/IKRig/Source/IKRig/Private/Retargeter/IKRetargetProcessor.cpp
--- a/Engine/Plugins/Animation/IKRig/Source/IKRig/Private/Retargeter/IKRetargetProcessor.cpp
+++ b/Engine/Plugins/Animation/IKRig/Source/IKRig/Private/Retargeter/IKRetargetProcessor.cpp
@@ -1198,6 +1198,8 @@ void FIKRetargetProcessor::InitialOpStackSetup(
 		FIKRetargetOpSettingsBase* AssetOpSettings = AssetOp->GetSettings();
 		// asset instance gets reference to running editor instance
 		AssetOpSettings->EditorInstance = NewOp.GetSettings();
+		AssetOp->EditorInstance = &NewOp;
+		NewOp.ResetExecutionTime();
 		// references the skeletons are used for bone selector widgets in the details panel
 		AssetOpSettings->SourceSkeletonAsset = SourceSkeleton.SkeletalMesh->GetSkeleton();
 		AssetOpSettings->TargetSkeletonAsset = TargetSkeleton.SkeletalMesh->GetSkeleton();
@@ -1343,11 +1345,19 @@ void FIKRetargetProcessor::RunRetargetOps(
 {
 	for (FInstancedStruct& OpStruct : OpStack)
 	{
+		#if WITH_EDITOR
+		double StartTime = FPlatformTime::Seconds();
+		#endif
+		
 		FIKRetargetOpBase& Op = OpStruct.GetMutable<FIKRetargetOpBase>();
 		if (Op.IsEnabled() && Op.IsInitialized())
 		{
 			Op.Run(*this, InDeltaTime, InSourceGlobalPose, OutTargetGlobalPose);
 		}
+		
+		#if WITH_EDITOR
+		Op.SetLatestExecutionTime( FPlatformTime::Seconds() - StartTime);
+		#endif
 	}
 }
 
diff --git a/Engine/Plugins/Animation/IKRig/Source/IKRig/Public/AnimNodes/AnimNode_RetargetPoseFromMesh.h b/Engine/Plugins/Animation/IKRig/Source/IKRig/Public/AnimNodes/AnimNode_RetargetPoseFromMesh.h
--- a/Engine/Plugins/Animation/IKRig/Source/IKRig/Public/AnimNodes/AnimNode_RetargetPoseFromMesh.h
+++ b/Engine/Plugins/Animation/IKRig/Source/IKRig/Public/AnimNodes/AnimNode_RetargetPoseFromMesh.h
@@ -88,6 +88,10 @@ struct FAnimNode_RetargetPoseFromMesh : public FAnimNode_Base
 	UE_API bool Serialize(FArchive& Ar);
 	UE_API void PostSerialize(const FArchive& Ar);
 
+#ifdef WITH_EDITORONLY_DATA
+	UE_API double GetAverageExecutionTime() const;
+#endif
+
 private:
 	
 	// get a retarget profile that includes settings from the asset, plus any custom profile overrides
@@ -110,6 +114,10 @@ private:
 
 	// used to determine when to look for a parent component (prevents constantly searching component hierarchy)
 	bool bSearchedForParentComponent = false;
+
+#ifdef WITH_EDITORONLY_DATA
+	double AverageExecutionTime;
+#endif
 };
 
 template<>
diff --git a/Engine/Plugins/Animation/IKRig/Source/IKRig/Public/Retargeter/IKRetargetOps.h b/Engine/Plugins/Animation/IKRig/Source/IKRig/Public/Retargeter/IKRetargetOps.h
--- a/Engine/Plugins/Animation/IKRig/Source/IKRig/Public/Retargeter/IKRetargetOps.h
+++ b/Engine/Plugins/Animation/IKRig/Source/IKRig/Public/Retargeter/IKRetargetOps.h
@@ -219,6 +219,18 @@ public:
 
 	/** (optional) override to support "Reset to Default" UI for chains */
 	virtual bool AreChainSettingsAtDefault(const FName& InChainName) { return true; };
+
+	/** the instance of this op currently running in the editor viewport */
+	mutable FIKRetargetOpBase* EditorInstance;
+
+	/** reset the running average execution time */
+	UE_API void ResetExecutionTime();
+	
+	/** set the latest execution time for profiling */
+	UE_API void SetLatestExecutionTime(const double InSeconds);
+
+	/** returns the current average execution time using a 30 frame exponential moving average */
+	UE_API double GetAverageExecutionTime() const;
 #endif
 
 	/** Get the name of this op (may be customized by user) */
@@ -249,6 +261,11 @@ public:
 		const TArray<FName>& InPropertiesToIgnore);
 
 private:
+
+#if WITH_EDITOR
+	/** profiling data; time in seconds of last execution */
+	double AverageExecutionTime = 0.0f;
+#endif
 	
 	/** when false, execution of this op is skipped */
 	UPROPERTY()
diff --git a/Engine/Plugins/Animation/IKRig/Source/IKRig/Public/Retargeter/IKRetargeter.h b/Engine/Plugins/Animation/IKRig/Source/IKRig/Public/Retargeter/IKRetargeter.h
--- a/Engine/Plugins/Animation/IKRig/Source/IKRig/Public/Retargeter/IKRetargeter.h
+++ b/Engine/Plugins/Animation/IKRig/Source/IKRig/Public/Retargeter/IKRetargeter.h
@@ -364,6 +364,10 @@ public:
 	// Toggle debug drawing for retargeting in the viewport. 
 	UPROPERTY(EditAnywhere, Category = DebugSettings)
 	bool bDebugDraw = true;
+
+	// Toggle performance profiling of the op stack. 
+	UPROPERTY(EditAnywhere, Category = DebugSettings)
+	bool bProfileOps = false;
 	
 	// The visual size of the bones in the viewport (saved between sessions). This is set from the viewport Character>Bones menu
 	UPROPERTY()
diff --git a/Engine/Plugins/Animation/IKRig/Source/IKRigEditor/Private/RetargetEditor/IKRetargetAnimInstance.cpp b/Engine/Plugins/Animation/IKRig/Source/IKRigEditor/Private/RetargetEditor/IKRetargetAnimInstance.cpp
--- a/Engine/Plugins/Animation/IKRig/Source/IKRigEditor/Private/RetargetEditor/IKRetargetAnimInstance.cpp
+++ b/Engine/Plugins/Animation/IKRig/Source/IKRigEditor/Private/RetargetEditor/IKRetargetAnimInstance.cpp
@@ -157,6 +157,11 @@ FIKRetargetProcessor* UIKRetargetAnimInstance::GetRetargetProcessor()
 	return RetargetNode.GetRetargetProcessor();
 }
 
+const FAnimNode_RetargetPoseFromMesh& UIKRetargetAnimInstance::GetRetargetAnimNode()
+{
+	return RetargetNode;
+}
+
 void UIKRetargetAnimInstance::ForceInitializeProcessor(USkeletalMeshComponent* TargetComponent)
 {
 	RetargetNode.IKRetargeterAsset->IncrementVersion();
diff --git a/Engine/Plugins/Animation/IKRig/Source/IKRigEditor/Private/RetargetEditor/SRetargetOpStack.cpp b/Engine/Plugins/Animation/IKRig/Source/IKRigEditor/Private/RetargetEditor/SRetargetOpStack.cpp
--- a/Engine/Plugins/Animation/IKRig/Source/IKRigEditor/Private/RetargetEditor/SRetargetOpStack.cpp
+++ b/Engine/Plugins/Animation/IKRig/Source/IKRigEditor/Private/RetargetEditor/SRetargetOpStack.cpp
@@ -4,6 +4,7 @@
 
 #include "UObject/UObjectIterator.h"
 #include "SPositiveActionButton.h"
+#include "RetargetEditor/IKRetargetAnimInstance.h"
 #include "RetargetEditor/IKRetargetEditorController.h"
 #include "RetargetEditor/IKRetargetEditorStyle.h"
 #include "Retargeter/IKRetargetOps.h"
@@ -144,6 +145,30 @@ void SRetargetOpItem::Construct(const FArguments& InArgs)
 			})
 		]
 
+    	// profiling
+		+ SHorizontalBox::Slot()
+		.AutoWidth()
+		.HAlign(HAlign_Right)
+		.VAlign(VAlign_Center)
+		.Padding(FMargin(OpHorizontalPadding, OpVerticalPadding))
+		[
+			SNew(STextBlock)
+			.Text_Lambda([this]() -> FText
+			{
+				FIKRetargetEditorController* EditorController = ListView.Pin()->EditorController.Pin().Get();
+				FIKRetargetOpBase* Op = EditorController->AssetController->GetRetargetOpByIndex(Element.Pin()->GetIndexInStack());
+				const double Milliseconds = (Op && Op->EditorInstance ? Op->EditorInstance->GetAverageExecutionTime() : 0.0f) * 1000.0;
+				FString Formatted = FString::Printf(TEXT("%.3f ms"), Milliseconds);
+				return FText::FromString(Formatted);
+			})
+			.IsEnabled(false)
+			.Visibility_Lambda([this]()
+			{
+				const bool Enabled = ListView.Pin()->EditorController.Pin()->AssetController->GetAsset()->bProfileOps;
+				return Enabled ? EVisibility::Visible : EVisibility::Collapsed;
+			})
+		]
+
         // spin Box
         + SHorizontalBox::Slot()
         .AutoWidth()
@@ -154,9 +179,10 @@ void SRetargetOpItem::Construct(const FArguments& InArgs)
             SNew(SSpinBox<double>)
 	        .Value_Lambda([this]() -> double
 	        {
-		        return 1.0f; // TODO add support for alpha value on op
+		        return 1.0f;							// TODO add support for alpha value on op
 	        })
-	        .IsEnabled(false) // TODO enable this once alpha is supported
+	        .IsEnabled(false)					// TODO enable this once alpha is supported
+	        .Visibility(EVisibility::Collapsed)	// TODO enable this once alpha is supported
             .MinValue(TOptional<double>())
             .MaxValue(TOptional<double>())
             .Style(&FAppStyle::Get().GetWidgetStyle<FSpinBoxStyle>("Sequencer.HyperlinkSpinBox"))
@@ -862,6 +888,38 @@ void SRetargetOpStack::Construct(
 							return FReply::Handled();
 						})
 					]
+
+					// spacer
+					+ SHorizontalBox::Slot()
+					.HAlign(HAlign_Fill)
+					.FillWidth(1.0f)
+					[
+						SNew(SSpacer).Size(FVector2D(0.0f, 1.0f))
+					]
+
+					// profiling
+					+ SHorizontalBox::Slot()
+					.AutoWidth()
+					.HAlign(HAlign_Right)
+					.VAlign(VAlign_Center)
+					.Padding(FMargin(2.0f, 4.0f))
+					[
+						SNew(STextBlock)
+						.Text_Lambda([this]() -> FText
+						{
+							UIKRetargetAnimInstance* AnimInstance = EditorController.Pin()->GetAnimInstance(ERetargetSourceOrTarget::Target);
+							const FAnimNode_RetargetPoseFromMesh& AnimNode = AnimInstance->GetRetargetAnimNode();
+							const double Milliseconds = AnimNode.GetAverageExecutionTime() * 1000.0;
+							FString Formatted = FString::Printf(TEXT("Total: %.3f ms"), Milliseconds);
+							return FText::FromString(Formatted);
+						})
+						.IsEnabled(false)
+						.Visibility_Lambda([this]()
+						{
+							const bool Enabled = EditorController.Pin()->AssetController->GetAsset()->bProfileOps;
+							return Enabled ? EVisibility::Visible : EVisibility::Collapsed;
+						})
+					]
 				]
 			]
 		]
diff --git a/Engine/Plugins/Animation/IKRig/Source/IKRigEditor/Public/RetargetEditor/IKRetargetAnimInstance.h b/Engine/Plugins/Animation/IKRig/Source/IKRigEditor/Public/RetargetEditor/IKRetargetAnimInstance.h
--- a/Engine/Plugins/Animation/IKRig/Source/IKRigEditor/Public/RetargetEditor/IKRetargetAnimInstance.h
+++ b/Engine/Plugins/Animation/IKRig/Source/IKRigEditor/Public/RetargetEditor/IKRetargetAnimInstance.h
@@ -61,6 +61,8 @@ public:
 		TWeakObjectPtr<USkeletalMeshComponent> InSourceMeshComponent);
 
 	FIKRetargetProcessor* GetRetargetProcessor();
+	
+	const FAnimNode_RetargetPoseFromMesh& GetRetargetAnimNode();
 
 	void ForceInitializeProcessor(USkeletalMeshComponent* TargetComponent);
 
