[IK Rig] Adding stretch IK solver.

diff --git a/Engine/Plugins/Animation/IKRig/Source/IKRig/Private/Rig/IKRigSkeleton.cpp b/Engine/Plugins/Animation/IKRig/Source/IKRig/Private/Rig/IKRigSkeleton.cpp
--- a/Engine/Plugins/Animation/IKRig/Source/IKRig/Private/Rig/IKRigSkeleton.cpp
+++ b/Engine/Plugins/Animation/IKRig/Source/IKRig/Private/Rig/IKRigSkeleton.cpp
@@ -159,6 +159,21 @@ int32 FIKRigSkeleton::GetChildIndices(const int32 ParentBoneIndex, TArray<int32>
 	return Children.Num();
 }
 
+void FIKRigSkeleton::GetChildrenIndicesRecursive(const int32 BoneIndex, TArray<int32>& OutChildren) const
+{
+	const int32 LastBranchIndex = GetCachedEndOfBranchIndex(BoneIndex);
+	if (LastBranchIndex == INDEX_NONE)
+	{
+		// no children (leaf bone)
+		return;
+	}
+	
+	for (int32 ChildBoneIndex = BoneIndex + 1; ChildBoneIndex <= LastBranchIndex; ChildBoneIndex++)
+	{
+		OutChildren.Add(ChildBoneIndex);
+	}
+}
+
 int32 FIKRigSkeleton::GetCachedEndOfBranchIndex(const int32 InBoneIndex) const
 {
 	if (!CachedEndOfBranchIndices.IsValidIndex(InBoneIndex))
diff --git a/Engine/Plugins/Animation/IKRig/Source/IKRig/Private/Rig/Solvers/IKRigStretchLimb.cpp b/Engine/Plugins/Animation/IKRig/Source/IKRig/Private/Rig/Solvers/IKRigStretchLimb.cpp
new file mode 100644
--- /dev/null
+++ b/Engine/Plugins/Animation/IKRig/Source/IKRig/Private/Rig/Solvers/IKRigStretchLimb.cpp
@@ -0,0 +1,366 @@
+// Copyright Epic Games, Inc. All Rights Reserved.
+
+#include "Rig/Solvers/IKRigStretchLimb.h"
+#include "Rig/IKRigDataTypes.h"
+#include "Rig/IKRigSkeleton.h"
+
+#include UE_INLINE_GENERATED_CPP_BY_NAME(IKRigStretchLimb)
+
+#define LOCTEXT_NAMESPACE "IKRigStretchLimb"
+
+void FIKRigStretchLimbSolver::Initialize(const FIKRigSkeleton& InIKRigSkeleton)
+{
+	bIsInitialized = false;
+
+	// must assign a start/end bone and a goal
+	if (Settings.Goal.IsNone() || Settings.EndBone.IsNone() || Settings.StartBone.IsNone())
+	{
+		return;
+	}
+
+	// get start and end bone
+	const int32 EndBoneIndex = InIKRigSkeleton.GetBoneIndexFromName(Settings.EndBone);
+	const int32 StartBoneIndex = InIKRigSkeleton.GetBoneIndexFromName(Settings.StartBone);
+	if (EndBoneIndex == INDEX_NONE || StartBoneIndex == INDEX_NONE)
+	{
+		return;
+	}
+
+	// populate limb indices
+	BoneIndices = {EndBoneIndex};
+	int32 BoneIndex = InIKRigSkeleton.GetParentIndex(EndBoneIndex);
+	while (BoneIndex != INDEX_NONE && BoneIndex >= StartBoneIndex)
+	{
+		BoneIndices.Add(BoneIndex);
+		BoneIndex = InIKRigSkeleton.GetParentIndex(BoneIndex);
+	};
+
+	// limb must have at least two joints
+	if (BoneIndices.Num() < 2)
+	{
+		return;
+	}
+
+	// sort the chain from root to end
+	Algo::Reverse(BoneIndices);
+
+	// get all children that need updated
+	TArray<int32> AllChildren;
+	InIKRigSkeleton.GetChildrenIndicesRecursive(StartBoneIndex, AllChildren);
+	for (const int32 ChildBoneIndex : AllChildren)
+	{
+		if (BoneIndices.Contains(ChildBoneIndex))
+		{
+			continue;
+		}
+		ChildrenToUpdate.Add(ChildBoneIndex);
+	}
+
+	bIsInitialized = true;
+}
+
+void FIKRigStretchLimbSolver::Solve(FIKRigSkeleton& InIKRigSkeleton, const FIKRigGoalContainer& InGoals)
+{
+	if (!bIsInitialized)
+	{
+		return; // can't solve until initialized successfully
+	}
+
+	if (Settings.StretchMode == EStretchLimbStretchMode::None &&
+		Settings.RotationMode == EStretchLimbRotationMode::None)
+	{
+		return; // nothing to do
+	}
+
+	// get the goal
+	const FIKRigGoal* Goal = InGoals.FindGoalByName(Settings.Goal);
+	if (!Goal)
+	{
+		// should not happen
+		ensureMsgf(false, TEXT("StretchLimbSolver unable to get assigned Goal."));
+		return;
+	}
+
+	// orient the chain bones
+	OrientChain(InIKRigSkeleton, Goal->FinalBlendedPosition);
+	
+	// translate the chain bones
+	TranslateChain(InIKRigSkeleton, Goal->FinalBlendedPosition);
+
+	// orient end bone to goal
+	FTransform& EndBoneGlobal = InIKRigSkeleton.CurrentPoseGlobal[BoneIndices.Last()];
+	const FQuat EndBoneGlobalRotation = FQuat::FastLerp(EndBoneGlobal.GetRotation(), Goal->FinalBlendedRotation, Settings.RotateEndBoneWithGoal);
+	EndBoneGlobal.SetRotation(EndBoneGlobalRotation);
+
+	// update chain local transforms
+	for (const int32 BoneIndex : BoneIndices)
+	{
+		InIKRigSkeleton.UpdateLocalTransformFromGlobal(BoneIndex);
+	}
+
+	// propagate to children
+	for (const int32 ChildIndex: ChildrenToUpdate)
+	{
+		InIKRigSkeleton.UpdateGlobalTransformFromLocal(ChildIndex);
+	}
+}
+
+void FIKRigStretchLimbSolver::OrientChain(
+	FIKRigSkeleton& InIKRigSkeleton,
+	const FVector& GoalLocation) const
+{
+	if (Settings.RotationMode != EStretchLimbRotationMode::OrientToGoal)
+	{
+		return;
+	}
+
+	// calc:
+	// 1. Current Chain Vector: start to end
+	// 2. New Chain Vector: start to goal
+	const FVector& ChainStart = InIKRigSkeleton.CurrentPoseGlobal[BoneIndices[0]].GetLocation();
+	const FVector& ChainEnd = InIKRigSkeleton.CurrentPoseGlobal[BoneIndices.Last()].GetLocation();
+	const FVector CurrentChainVector = ChainEnd - ChainStart;
+	const FVector NewChainVector = GoalLocation - ChainStart;
+	
+	// rotate the whole bone chain to aim at the goal
+	const FQuat RotationDelta = FQuat::FindBetweenVectors(CurrentChainVector, NewChainVector);
+	const FVector Origin = InIKRigSkeleton.CurrentPoseGlobal[BoneIndices[0]].GetLocation();
+	for (const int32 BoneIndex : BoneIndices)
+	{
+		FTransform& BoneTransform = InIKRigSkeleton.CurrentPoseGlobal[BoneIndex];
+		const FVector RelativePosition = BoneTransform.GetLocation() - Origin;
+		const FVector RotatedRelativePosition = RotationDelta.RotateVector(RelativePosition);
+		BoneTransform.SetLocation(RotatedRelativePosition + Origin);
+		BoneTransform.SetRotation(RotationDelta * BoneTransform.GetRotation());
+	}
+}
+
+void FIKRigStretchLimbSolver::TranslateChain(
+	FIKRigSkeleton& InIKRigSkeleton,
+	const FVector& GoalLocation) const
+{
+	if (Settings.StretchMode == EStretchLimbStretchMode::None)
+	{
+		return; // skip adjusting translations
+	}
+
+	// calc:
+	// 1. Current Chain Vector: start to end
+	// 2. New Chain Vector: start to goal
+	const FVector& ChainStart = InIKRigSkeleton.CurrentPoseGlobal[BoneIndices[0]].GetLocation();
+	const FVector& ChainEnd = InIKRigSkeleton.CurrentPoseGlobal[BoneIndices.Last()].GetLocation();
+	const FVector CurrentChainVector = ChainEnd - ChainStart;
+	const FVector NewChainVector = GoalLocation - ChainStart;
+
+	// calculate length of the chain
+	double ChainLength = 0.0f;
+	switch (Settings.LengthType)
+	{
+	case EStretchLimbLengthType::AlongBones:
+		{
+			// iterate from tip to first child of root
+			for (int32 BoneIndex = BoneIndices.Num() - 1; BoneIndex >= 1; --BoneIndex)
+			{
+				const FVector& BonePosition = InIKRigSkeleton.CurrentPoseGlobal[BoneIndices[BoneIndex]].GetLocation();
+				const FVector& ParentPosition = InIKRigSkeleton.CurrentPoseGlobal[BoneIndices[BoneIndex-1]].GetLocation();
+				ChainLength += (ParentPosition - BonePosition).Length();
+			}
+			break;
+		}
+	case EStretchLimbLengthType::StraightLine:
+		{
+			ChainLength = CurrentChainVector.Length();
+			break;
+		}
+	default:
+		ensureMsgf(false, TEXT("Unknown chain length type."));
+		return;
+	}
+
+	// compare lenghth of the chain in its current configuration with the length it would need to reach goal
+	const double NewChainLength = NewChainVector.Size();
+	if (FMath::IsNearlyEqual(ChainLength, NewChainLength))
+	{
+		return; // nothing to squash or stretch
+	}
+
+	if (Settings.StretchMode != EStretchLimbStretchMode::SquashAndStretch)
+	{
+		if (ChainLength > NewChainLength && Settings.StretchMode == EStretchLimbStretchMode::StretchOnly)
+		{
+			return; // chain needs to squash, but we only want to stretch
+		}
+
+		if (ChainLength < NewChainLength && Settings.StretchMode == EStretchLimbStretchMode::SquashOnly)
+		{
+			return; // chain needs to stretch, but we only want to squash
+		}
+	}
+
+	// record chain positions prior to translational offset
+	TArray<FVector> BonePositionsPreOffset;
+	for (const int32 BoneIndex : BoneIndices)
+	{
+		BonePositionsPreOffset.Add(InIKRigSkeleton.CurrentPoseGlobal[BoneIndex].GetLocation());
+	}
+
+	// translate the bones
+	switch (Settings.TranslationType)
+	{
+	case EStretchLimbTranslationType::BoneLengthOnly:
+		{
+			// easier to do this in local space, so update them first
+			for (const int32 BoneIndex : BoneIndices)
+			{
+				InIKRigSkeleton.UpdateLocalTransformFromGlobal(BoneIndex);
+			}
+			
+			// scale the length of each bone from root to tip
+			const double ScaleFactor = NewChainLength / ChainLength;
+			for (int32 ChainIndex = 1; ChainIndex < BoneIndices.Num(); ++ChainIndex)
+			{
+				FTransform& BoneLocalTransform = InIKRigSkeleton.CurrentPoseLocal[BoneIndices[ChainIndex]];
+				BoneLocalTransform.SetLocation(BoneLocalTransform.GetLocation() * ScaleFactor);
+			}
+
+			// update global pose of chain
+			for (const int32 BoneIndex : BoneIndices)
+			{
+				InIKRigSkeleton.UpdateGlobalTransformFromLocal(BoneIndex);
+			}
+			break;
+		}
+	case EStretchLimbTranslationType::AllAxes:
+		{
+			// apply global translational delta to the bone to move towards goal
+			const FVector NewChainDirection = NewChainVector.GetSafeNormal();
+			const FVector TotalChainDelta = NewChainDirection * (NewChainLength - ChainLength);
+			const double InvNumBones = 1.0 / static_cast<double>(BoneIndices.Num() - 1);
+			for (int32 ChainIndex = 1; ChainIndex < BoneIndices.Num(); ++ChainIndex)
+			{
+				// calculate amount of delta to apply to the bone (100% at tip, with linear falloff to root)
+				const double ScaleFactor = static_cast<double>(ChainIndex) * InvNumBones;
+				const FVector PerBoneDelta = TotalChainDelta * ScaleFactor;
+				// add delta to the bone
+				FTransform& BoneGlobalTransform = InIKRigSkeleton.CurrentPoseGlobal[BoneIndices[ChainIndex]];
+				BoneGlobalTransform.SetLocation( BoneGlobalTransform.GetLocation() + PerBoneDelta);
+			}
+			break;
+		}
+	default:
+		ensureMsgf(false, TEXT("Unknown chain length type."));
+		return;
+	}
+
+	// update orientations of bones based on translational offset
+	for (int32 ChainIndex = 0; ChainIndex < BoneIndices.Num() - 1; ++ChainIndex)
+	{
+		// get the direction of the bone prior to any translational offsets
+		const FVector BoneDirectionOrig =  BonePositionsPreOffset[ChainIndex+1] - BonePositionsPreOffset[ChainIndex];
+
+		// get the direction of the bone after translations were applied
+		const int32 BoneIndex = BoneIndices[ChainIndex];
+		const int32 ChildIndex = BoneIndices[ChainIndex + 1];
+		FTransform& BoneGlobalTransform = InIKRigSkeleton.CurrentPoseGlobal[BoneIndex];
+		const FTransform& ChildGlobalTransform = InIKRigSkeleton.CurrentPoseGlobal[ChildIndex];
+		const FVector BoneDirectionNew = ChildGlobalTransform.GetLocation() - BoneGlobalTransform.GetLocation();
+
+		// generate a swing rotation to reorient the bone towards it's new direction 
+		const FQuat RotationDelta = FQuat::FindBetweenVectors(BoneDirectionOrig, BoneDirectionNew);
+		BoneGlobalTransform.SetRotation( RotationDelta * BoneGlobalTransform.GetRotation() );
+	}
+}
+
+void FIKRigStretchLimbSolver::GetRequiredBones(TSet<FName>& OutRequiredBones) const
+{
+	OutRequiredBones.Add(Settings.StartBone);
+}
+
+void FIKRigStretchLimbSolver::GetRequiredGoals(TSet<FName>& OutRequiredGoals) const
+{
+	OutRequiredGoals.Add(Settings.Goal);
+}
+
+FIKRigSolverSettingsBase* FIKRigStretchLimbSolver::GetSolverSettings()
+{
+	return &Settings;
+}
+
+const UScriptStruct* FIKRigStretchLimbSolver::GetSolverSettingsType() const
+{
+	return FIKRigStretchLimbSettings::StaticStruct();
+}
+
+void FIKRigStretchLimbSolver::AddGoal(const UIKRigEffectorGoal& InNewGoal)
+{
+	Settings.Goal = InNewGoal.GoalName;
+	Settings.EndBone = InNewGoal.BoneName;
+}
+
+void FIKRigStretchLimbSolver::OnGoalRemoved(const FName& InGoalName)
+{
+	if (InGoalName == Settings.Goal)
+	{
+		Settings.Goal = NAME_None;
+		Settings.EndBone = NAME_None;
+	}
+}
+
+void FIKRigStretchLimbSolver::OnGoalRenamed(const FName& InOldName, const FName& InNewName)
+{
+	if (InOldName == Settings.Goal)
+	{
+		Settings.Goal = InNewName;
+	}
+}
+
+void FIKRigStretchLimbSolver::OnGoalMovedToDifferentBone(const FName& InGoalName, const FName& InNewBoneName)
+{
+	if (InGoalName == Settings.Goal)
+	{
+		Settings.StartBone = InNewBoneName;
+	}
+}
+
+void FIKRigStretchLimbSolver::SetStartBone(const FName& InBoneName)
+{
+	Settings.StartBone = InBoneName;
+}
+
+#if WITH_EDITOR
+
+UIKRigSolverControllerBase* FIKRigStretchLimbSolver::GetSolverController(UObject* Outer)
+{
+	return nullptr;
+}
+
+FText FIKRigStretchLimbSolver::GetNiceName() const
+{
+	return FText(LOCTEXT("SolverName", "Stretch Limb"));
+}
+
+bool FIKRigStretchLimbSolver::GetWarningMessage(FText& OutWarningMessage) const
+{
+	if (Settings.StartBone == NAME_None)
+	{
+		OutWarningMessage = LOCTEXT("MissingRoot", "No start bone assigned.");
+		return true;
+	}
+
+	if (Settings.Goal == NAME_None)
+	{
+		OutWarningMessage = LOCTEXT("MissingGoal", "No goal assigned.");
+		return true;
+	}
+	
+	return false;
+}
+
+bool FIKRigStretchLimbSolver::IsBoneAffectedBySolver(const FName& InBoneName, const FIKRigSkeleton& InIKRigSkeleton) const
+{
+	return false;
+}
+#endif
+
+#undef LOCTEXT_NAMESPACE
+
diff --git a/Engine/Plugins/Animation/IKRig/Source/IKRig/Public/Rig/IKRigSkeleton.h b/Engine/Plugins/Animation/IKRig/Source/IKRig/Public/Rig/IKRigSkeleton.h
--- a/Engine/Plugins/Animation/IKRig/Source/IKRig/Public/Rig/IKRigSkeleton.h
+++ b/Engine/Plugins/Animation/IKRig/Source/IKRig/Public/Rig/IKRigSkeleton.h
@@ -86,6 +86,8 @@ struct FIKRigSkeleton
 	UE_API int32 GetParentIndexThatIsNotExcluded(const int32 BoneIndex) const;
 
 	UE_API int32 GetChildIndices(const int32 ParentBoneIndex, TArray<int32>& Children) const;
+
+	UE_API void GetChildrenIndicesRecursive(const int32 BoneIndex, TArray<int32>& OutChildren) const;
 	
 	UE_API int32 GetCachedEndOfBranchIndex(const int32 InBoneIndex) const;
 
diff --git a/Engine/Plugins/Animation/IKRig/Source/IKRig/Public/Rig/Solvers/IKRigStretchLimb.h b/Engine/Plugins/Animation/IKRig/Source/IKRig/Public/Rig/Solvers/IKRigStretchLimb.h
new file mode 100644
--- /dev/null
+++ b/Engine/Plugins/Animation/IKRig/Source/IKRig/Public/Rig/Solvers/IKRigStretchLimb.h
@@ -0,0 +1,132 @@
+﻿// Copyright Epic Games, Inc. All Rights Reserved.
+
+#pragma once
+
+#include "Rig/IKRigDefinition.h"
+#include "IKRigSolverBase.h"
+
+#include "IKRigStretchLimb.generated.h"
+
+#define UE_API IKRIG_API
+
+UENUM(BlueprintType)
+enum class EStretchLimbTranslationType : uint8
+{
+	AllAxes,
+	BoneLengthOnly,
+};
+
+UENUM(BlueprintType)
+enum class EStretchLimbLengthType : uint8
+{
+	StraightLine,
+	AlongBones,
+};
+
+UENUM(BlueprintType)
+enum class EStretchLimbStretchMode : uint8
+{
+	None,
+	SquashAndStretch,
+	StretchOnly,
+	SquashOnly,
+};
+
+UENUM(BlueprintType)
+enum class EStretchLimbRotationMode : uint8
+{
+	None,
+	OrientToGoal,
+};
+
+USTRUCT(BlueprintType)
+struct FIKRigStretchLimbSettings : public FIKRigSolverSettingsBase
+{
+	GENERATED_BODY()
+
+	/** The start bone assigned to this solver (usually the thigh or shoulder). */
+	UPROPERTY(BlueprintReadOnly, VisibleAnywhere, Category = "Limb Setup")
+	FName StartBone;
+
+	/** The end bone assigned to this solver (usually the foot or ball). */
+	UPROPERTY(BlueprintReadOnly, VisibleAnywhere, Category = "Limb Setup")
+	FName EndBone;
+
+	/** The end bone assigned to this solver (usually the foot or ball). */
+	UPROPERTY(BlueprintReadOnly, VisibleAnywhere, Category = "Limb Setup")
+	FName Goal;
+	
+	/** Determines whether to squash or stretch the bones to reach the goal. */
+	UPROPERTY(BlueprintReadWrite, EditAnywhere, Category = "Translation Settings")
+	EStretchLimbStretchMode StretchMode = EStretchLimbStretchMode::StretchOnly;
+
+	/** Consider the length of the bone chain as either a straight line from start to end or the total length along the bones. */
+	UPROPERTY(BlueprintReadWrite, EditAnywhere, Category = "Translation Settings")
+	EStretchLimbLengthType LengthType = EStretchLimbLengthType::AlongBones;
+	
+	/** Determines how to affect the translation of the bones in the chain. */
+	UPROPERTY(BlueprintReadWrite, EditAnywhere, Category = "Translation Settings")
+	EStretchLimbTranslationType TranslationType = EStretchLimbTranslationType::BoneLengthOnly;
+
+	/** Determines how to affect the orientation of the bones in the chain. */
+	UPROPERTY(BlueprintReadWrite, EditAnywhere, Category = "Rotation Settings")
+	EStretchLimbRotationMode RotationMode = EStretchLimbRotationMode::OrientToGoal;
+
+	/** Allow the end bone to be reoriented (0) or match the orientation of the goal (1). Range is 0-1. Default is 1.0. */
+	UPROPERTY(BlueprintReadWrite, EditAnywhere, Category = "Global Settings", meta = (UIMin = "0.0", UIMax = "1.0"))
+	double RotateEndBoneWithGoal = 0.0f;
+};
+
+USTRUCT(BlueprintType)
+struct FIKRigStretchLimbSolver : public FIKRigSolverBase
+{
+	GENERATED_BODY()
+	
+	// FIKRigSolverBase interface
+	UE_API virtual void Initialize(const FIKRigSkeleton& InIKRigSkeleton) override;
+	UE_API virtual void Solve(FIKRigSkeleton& InIKRigSkeleton, const FIKRigGoalContainer& InGoals) override;
+	UE_API virtual void GetRequiredBones(TSet<FName>& OutRequiredBones) const override;
+	UE_API virtual void GetRequiredGoals(TSet<FName>& OutRequiredGoals) const override;
+	
+	// solver settings
+	UE_API virtual FIKRigSolverSettingsBase* GetSolverSettings() override;
+	UE_API virtual const UScriptStruct* GetSolverSettingsType() const override;
+	
+	// goals
+	UE_API virtual void AddGoal(const UIKRigEffectorGoal& InNewGoal) override;
+	UE_API virtual void OnGoalRemoved(const FName& InGoalName) override;
+	UE_API virtual void OnGoalRenamed(const FName& InOldName, const FName& InNewName) override;
+	UE_API virtual void OnGoalMovedToDifferentBone(const FName& InGoalName, const FName& InNewBoneName) override;
+
+	// goal settings
+	UE_API virtual bool UsesCustomGoalSettings() const override { return false; };
+	
+	// start bone can be set on this solver
+	UE_API virtual bool UsesStartBone() const override { return true; };
+	UE_API virtual FName GetStartBone() const override { return Settings.StartBone; };
+	UE_API virtual void SetStartBone(const FName& InRootBoneName) override;
+
+#if WITH_EDITOR
+	// return custom controller for scripting this solver
+	virtual UIKRigSolverControllerBase* GetSolverController(UObject* Outer) override;
+	// UI stuff
+	virtual FText GetNiceName() const override;
+	virtual bool GetWarningMessage(FText& OutWarningMessage) const override;
+	virtual bool IsBoneAffectedBySolver(const FName& InBoneName, const FIKRigSkeleton& InIKRigSkeleton) const override;
+#endif
+
+	UPROPERTY()
+	FIKRigStretchLimbSettings Settings;
+
+private:
+
+	void OrientChain(FIKRigSkeleton& InIKRigSkeleton, const FVector& GoalLocation) const;
+	
+	void TranslateChain(FIKRigSkeleton& InIKRigSkeleton, const FVector& GoalLocation) const;
+
+	bool bIsInitialized = false;
+	TArray<int32> BoneIndices;
+	TArray<int32> ChildrenToUpdate;
+};
+
+
