[IK Rig] Removed bulk rotation normalization after Solve()

diff --git a/Engine/Plugins/Animation/IKRig/Source/IKRig/Private/Rig/IKRigProcessor.cpp b/Engine/Plugins/Animation/IKRig/Source/IKRig/Private/Rig/IKRigProcessor.cpp
--- a/Engine/Plugins/Animation/IKRig/Source/IKRig/Private/Rig/IKRigProcessor.cpp
+++ b/Engine/Plugins/Animation/IKRig/Source/IKRig/Private/Rig/IKRigProcessor.cpp
@@ -386,9 +386,6 @@ void FIKRigProcessor::Solve(const FTransform& ComponentToWorld)
 			Solver->Solve(Skeleton, GoalContainer);
 		}
 	}
-
-	// make sure rotations are normalized coming out
-	Skeleton.NormalizeRotations(Skeleton.CurrentPoseGlobal);
 }
 
 void FIKRigProcessor::GetOutputPoseGlobal(TArray<FTransform>& OutputPoseGlobal) const
