[IK Retarget] Adjusting default FK Chain debug line thickness.

diff --git a/Engine/Plugins/Animation/IKRig/Source/IKRig/Public/Retargeter/RetargetOps/FKChainsOp.h b/Engine/Plugins/Animation/IKRig/Source/IKRig/Public/Retargeter/RetargetOps/FKChainsOp.h
--- a/Engine/Plugins/Animation/IKRig/Source/IKRig/Public/Retargeter/RetargetOps/FKChainsOp.h
+++ b/Engine/Plugins/Animation/IKRig/Source/IKRig/Public/Retargeter/RetargetOps/FKChainsOp.h
@@ -214,7 +214,7 @@ struct FIKRetargetFKChainsOpSettings : public FIKRetargetOpSettingsBase
 
 	/** Debug draw chain line thickness. */
 	UPROPERTY(EditAnywhere, Category = Debug)
-	double ChainDrawThickness = 5.0f;
+	double ChainDrawThickness = 1.0f;
 
 	/** Debug draw size used for single-bone chains. */
 	UPROPERTY(EditAnywhere, Category = Debug)
