[IK Retarget] Adding individual LODing to retarget ops

diff --git a/Engine/Plugins/Animation/IKRig/Source/IKRig/Private/AnimNodes/AnimNode_RetargetPoseFromMesh.cpp b/Engine/Plugins/Animation/IKRig/Source/IKRig/Private/AnimNodes/AnimNode_RetargetPoseFromMesh.cpp
--- a/Engine/Plugins/Animation/IKRig/Source/IKRig/Private/AnimNodes/AnimNode_RetargetPoseFromMesh.cpp
+++ b/Engine/Plugins/Animation/IKRig/Source/IKRig/Private/AnimNodes/AnimNode_RetargetPoseFromMesh.cpp
@@ -219,14 +219,19 @@ void FAnimNode_RetargetPoseFromMesh::Evaluate_AnyThread(FPoseContext& Output)
 	}
 
 	// LOD off the IK pass
-	bool bForceIKOff = LODThresholdForIK != INDEX_NONE && Output.AnimInstanceProxy->GetLODLevel() > LODThresholdForIK;
+	const int32 CurrentLODLevel = Output.AnimInstanceProxy->GetLODLevel();
+	bool bForceIKOff = LODThresholdForIK != INDEX_NONE && CurrentLODLevel > LODThresholdForIK;
 	FRetargetProfile RetargetProfileToUse = GetMergedRetargetProfile(bForceIKOff);
 
 	// give retarget ops a chance to access data from the anim graph before running
 	Processor.OnAnimGraphEvaluateAnyThread(Output);
 
 	// run the retargeter
-	const TArray<FTransform>& RetargetedPose = Processor.RunRetargeter(PoseToRetargetFromComponentSpace, RetargetProfileToUse, DeltaTime);
+	const TArray<FTransform>& RetargetedPose = Processor.RunRetargeter(
+		PoseToRetargetFromComponentSpace,
+		RetargetProfileToUse,
+		DeltaTime,
+		CurrentLODLevel);
 
 	// reset delta time (it's accumulated)
 	DeltaTime = 0.f;
diff --git a/Engine/Plugins/Animation/IKRig/Source/IKRig/Private/Retargeter/IKRetargetProcessor.cpp b/Engine/Plugins/Animation/IKRig/Source/IKRig/Private/Retargeter/IKRetargetProcessor.cpp
--- a/Engine/Plugins/Animation/IKRig/Source/IKRig/Private/Retargeter/IKRetargetProcessor.cpp
+++ b/Engine/Plugins/Animation/IKRig/Source/IKRig/Private/Retargeter/IKRetargetProcessor.cpp
@@ -1268,7 +1268,8 @@ void FIKRetargetProcessor::InitializeRetargetOps()
 TArray<FTransform>&  FIKRetargetProcessor::RunRetargeter(
 	TArray<FTransform>& InSourceGlobalPose,
 	const FRetargetProfile& InProfile,
-	const float InDeltaTime)
+	const float InDeltaTime,
+	const int32 InLOD)
 {
 	if (!ensure(bIsInitialized))
 	{
@@ -1313,7 +1314,7 @@ TArray<FTransform>&  FIKRetargetProcessor::RunRetargeter(
 	ResetIKRigGoalContainer();
 	
 	// run the stack of retargeting operations
-	RunRetargetOps(InDeltaTime, InSourceGlobalPose, TargetSkeleton.OutputGlobalPose);
+	RunRetargetOps(InSourceGlobalPose, TargetSkeleton.OutputGlobalPose, InDeltaTime, InLOD);
 	
 	return TargetSkeleton.OutputGlobalPose;
 }
@@ -1339,22 +1340,46 @@ void FIKRetargetProcessor::GenerateBasePoses(TArray<FTransform>& InSourceGlobalP
 }
 
 void FIKRetargetProcessor::RunRetargetOps(
-	const double InDeltaTime,
 	const TArray<FTransform>& InSourceGlobalPose,
-	TArray<FTransform>& OutTargetGlobalPose)
+	TArray<FTransform>& OutTargetGlobalPose,
+	const double InDeltaTime,
+	const int32 InLOD)
 {
+	auto ShouldExecuteOp = [&InLOD](FIKRetargetOpBase& InOp)
+	{
+		if (!InOp.IsEnabled())
+		{
+			return false; // disabled
+		}
+		if (!InOp.IsInitialized())
+		{
+			return false; // not initialized
+		}
+			
+		const int32 LODThreshold = InOp.GetSettings()->LODThreshold;
+		if (LODThreshold != INDEX_NONE && InLOD > LODThreshold)
+		{
+			return false; // LOD'd off
+		}
+
+		return true;
+	};
+	
 	for (FInstancedStruct& OpStruct : OpStack)
 	{
+		// collect profiling data per op
 		#if WITH_EDITOR
 		double StartTime = FPlatformTime::Seconds();
 		#endif
-		
+
+		// execute the op
 		FIKRetargetOpBase& Op = OpStruct.GetMutable<FIKRetargetOpBase>();
-		if (Op.IsEnabled() && Op.IsInitialized())
+		if (ShouldExecuteOp(Op))
 		{
-			Op.Run(*this, InDeltaTime, InSourceGlobalPose, OutTargetGlobalPose);
+			Op.Run(*this, InDeltaTime, InSourceGlobalPose, OutTargetGlobalPose);	
 		}
-		
+
+		// END op profiler scope
 		#if WITH_EDITOR
 		Op.SetLatestExecutionTime( FPlatformTime::Seconds() - StartTime);
 		#endif
diff --git a/Engine/Plugins/Animation/IKRig/Source/IKRig/Public/Retargeter/IKRetargetOps.h b/Engine/Plugins/Animation/IKRig/Source/IKRig/Public/Retargeter/IKRetargetOps.h
--- a/Engine/Plugins/Animation/IKRig/Source/IKRig/Public/Retargeter/IKRetargetOps.h
+++ b/Engine/Plugins/Animation/IKRig/Source/IKRig/Public/Retargeter/IKRetargetOps.h
@@ -44,6 +44,12 @@ struct FIKRetargetOpSettingsBase
 	
 	virtual ~FIKRetargetOpSettingsBase() {};
 
+	/** The maximum LOD that this Op is allowed to run at.
+	 * For example if you have LODThreshold of 2, the Op will run until LOD 2 (based on 0 index). When the component LOD becomes 3, it will stop running.
+	 * A value of -1 forces the Op to execute at all LOD levels. Default is -1. */
+	UPROPERTY(EditAnywhere, BlueprintReadWrite, Category = "Retarget Op")
+	int32 LODThreshold = INDEX_NONE;
+
 	/** (required) override to specify how settings should be applied in a way that will not require reinitialization (ie runtime compatible)*/
 	virtual void CopySettingsAtRuntime(const FIKRetargetOpSettingsBase* InSettingsToCopyFrom) { checkNoEntry(); };
 	
diff --git a/Engine/Plugins/Animation/IKRig/Source/IKRig/Public/Retargeter/IKRetargetProcessor.h b/Engine/Plugins/Animation/IKRig/Source/IKRig/Public/Retargeter/IKRetargetProcessor.h
--- a/Engine/Plugins/Animation/IKRig/Source/IKRig/Public/Retargeter/IKRetargetProcessor.h
+++ b/Engine/Plugins/Animation/IKRig/Source/IKRig/Public/Retargeter/IKRetargetProcessor.h
@@ -369,12 +369,14 @@ public:
 	* @param InSourceGlobalPose -  is the source mesh input pose in Component/Global space
 	* @param InProfile -  the retarget profile to use for this update
 	* @param InDeltaTime -  time since last tick in seconds (passed to ops)
+	* @param InLOD - the current LOD level for the skeletal mesh
 	* @return The retargeted Component/Global space pose for the target skeleton
 	*/
 	UE_API TArray<FTransform>& RunRetargeter(
 		TArray<FTransform>& InSourceGlobalPose,
 		const FRetargetProfile& InProfile,
-		const float InDeltaTime);
+		const float InDeltaTime,
+		const int32 InLOD = INDEX_NONE);
 
 	/** Get whether this processor is ready to call RunRetargeter() and generate new poses. */
 	UE_API bool IsInitialized() const;
@@ -555,9 +557,10 @@ private:
 
 	/** Run all post process operations on the retargeted result. */
 	void RunRetargetOps(
-		const double InDeltaTime,
 		const TArray<FTransform>& InSourceGlobalPose,
-		TArray<FTransform>& OutTargetGlobalPose);
+		TArray<FTransform>& OutTargetGlobalPose,
+		const double InDeltaTime,
+		const int32 InLOD);
 };
 
 //
diff --git a/Engine/Plugins/Animation/IKRig/Source/IKRigEditor/Private/RetargetEditor/IKRetargetEditor.cpp b/Engine/Plugins/Animation/IKRig/Source/IKRigEditor/Private/RetargetEditor/IKRetargetEditor.cpp
--- a/Engine/Plugins/Animation/IKRig/Source/IKRigEditor/Private/RetargetEditor/IKRetargetEditor.cpp
+++ b/Engine/Plugins/Animation/IKRig/Source/IKRigEditor/Private/RetargetEditor/IKRetargetEditor.cpp
@@ -348,6 +348,11 @@ void FIKRetargetEditor::Tick(float DeltaTime)
 	
 	// forces viewport to always update, even when mouse pressed down in other tabs
 	GetPersonaToolkit()->GetPreviewScene()->InvalidateViews();
+
+	// copy the LOD from the source to ensure they are always in-sync for previewing LOD'd ops
+	const int32 SourceLOD = EditorController->SourceSkelMeshComponent->GetForcedLOD();
+	EditorController->TargetSkelMeshComponent->SetDebugForcedLOD(SourceLOD);
+	EditorController->TargetSkelMeshComponent->bTrackAttachedInstanceLOD = false;
 }
 
 TStatId FIKRetargetEditor::GetStatId() const
diff --git a/Engine/Plugins/Animation/IKRig/Source/IKRigEditor/Private/RetargetEditor/IKRetargetOpDetails.cpp b/Engine/Plugins/Animation/IKRig/Source/IKRigEditor/Private/RetargetEditor/IKRetargetOpDetails.cpp
--- a/Engine/Plugins/Animation/IKRig/Source/IKRigEditor/Private/RetargetEditor/IKRetargetOpDetails.cpp
+++ b/Engine/Plugins/Animation/IKRig/Source/IKRigEditor/Private/RetargetEditor/IKRetargetOpDetails.cpp
@@ -54,18 +54,32 @@ bool FIKRetargetOpBaseSettingsCustomization::LoadAndValidateStructToCustomize(
 	return true;
 }
 
-void FIKRetargetOpBaseSettingsCustomization::AddChildPropertyToGroup(
-	IDetailGroup& InGroup,
-	const TSharedRef<IPropertyHandle>& InParentPropertyHandle,
-	const FName& InChildPropertyName)
+void FIKRetargetOpBaseSettingsCustomization::AddPropertyToGroup(
+	const TSharedRef<IPropertyHandle>& InPropertyHandle,
+	IDetailChildrenBuilder& ChildBuilder)
 {
-	TSharedPtr<IPropertyHandle> PropertyHandle = InParentPropertyHandle->GetChildHandle(InChildPropertyName);
-	if (!ensure(PropertyHandle.IsValid()))
+	// get category from metadata
+	FString CategoryName = InPropertyHandle->GetMetaData(TEXT("Category"));
+	if (CategoryName.IsEmpty())
 	{
+		// if no category specified, add directly to builder
+		ChildBuilder.AddProperty(InPropertyHandle);
 		return;
 	}
-	
-	InGroup.AddPropertyRow(PropertyHandle.ToSharedRef());
+
+	// get the group
+	FName CategoryFName(*CategoryName);
+	if (IDetailGroup* Group = ChildBuilder.GetGroup(CategoryFName))
+	{
+		Group->AddPropertyRow(InPropertyHandle);
+		return;
+	}
+
+	// didn't already exist, so add one
+	FText CategoryText = FText::FromString(CategoryName);
+	IDetailGroup& Group = ChildBuilder.AddGroup(CategoryFName, CategoryText);
+	// add the property to the new group
+	Group.AddPropertyRow(InPropertyHandle);
 }
 
 void FIKRetargetOpBaseSettingsCustomization::AddChildPropertiesToCategoryGroups(
@@ -175,6 +189,23 @@ void FIKRetargetOpBaseSettingsCustomization::AddChildPropertiesInCategory(
 	}
 }
 
+void FIKRetargetOpBaseSettingsCustomization::AddBaseOpProperties(
+	TSharedRef<IPropertyHandle> StructPropertyHandle,
+	IDetailChildrenBuilder& ChildBuilder)
+{
+	UScriptStruct* OpBaseStructType = FIKRetargetOpSettingsBase::StaticStruct();
+
+	for (TFieldIterator<FProperty> It(OpBaseStructType); It; ++It)
+	{
+		FProperty* BaseProp = *It;
+		TSharedPtr<IPropertyHandle> BasePropHandle = StructPropertyHandle->GetChildHandle(BaseProp->GetFName());
+		if (BasePropHandle.IsValid())
+		{
+			AddPropertyToGroup(BasePropHandle.ToSharedRef(), ChildBuilder);
+		}
+	}
+}
+
 void FChainsFKOpCustomization::CustomizeChildren(
 	TSharedRef<IPropertyHandle> StructPropertyHandle,
 	IDetailChildrenBuilder& ChildBuilder,
@@ -239,6 +270,9 @@ void FChainsFKOpCustomization::CustomizeChildren(
 
 	// add the debug properties
 	AddChildPropertiesInCategory(StructPropertyHandle, ChildBuilder, FName("Debug"));
+
+	// add the properties from the base op
+	AddBaseOpProperties(StructPropertyHandle, ChildBuilder);
 }
 
 void FRunIKRigOpCustomization::CustomizeChildren(
@@ -279,6 +313,9 @@ void FRunIKRigOpCustomization::CustomizeChildren(
 
 	// add the debug properties
 	AddChildPropertiesInCategory(StructPropertyHandle, ChildBuilder, FName("Debug"));
+
+	// add the properties from the base op
+	AddBaseOpProperties(StructPropertyHandle, ChildBuilder);
 }
 
 void FIKChainOpCustomization::CustomizeChildren(
@@ -345,6 +382,9 @@ void FIKChainOpCustomization::CustomizeChildren(
 
 	// add the debug properties
 	AddChildPropertiesInCategory(StructPropertyHandle, ChildBuilder, FName("Debug"));
+
+	// add the properties from the base op
+	AddBaseOpProperties(StructPropertyHandle, ChildBuilder);
 }
 
 void FStrideWarpOpCustomization::CustomizeChildren(
@@ -413,6 +453,8 @@ void FStrideWarpOpCustomization::CustomizeChildren(
 	AddChildPropertiesInCategory(StructPropertyHandle, ChildBuilder, FName("Warping"));
 	// add the debug properties
 	AddChildPropertiesInCategory(StructPropertyHandle, ChildBuilder, FName("Debug"));
+	// add the properties from the base op
+	AddBaseOpProperties(StructPropertyHandle, ChildBuilder);
 }
 
 void FSpeedPlantOpCustomization::CustomizeChildren(
@@ -479,6 +521,8 @@ void FSpeedPlantOpCustomization::CustomizeChildren(
 
 	// add speed planting properties
 	AddChildPropertiesInCategory(StructPropertyHandle, ChildBuilder, FName("Speed Planting"));
+	// add the properties from the base op
+	AddBaseOpProperties(StructPropertyHandle, ChildBuilder);
 }
 
 void FPoleVectorOpCustomization::CustomizeChildren(
@@ -541,6 +585,9 @@ void FPoleVectorOpCustomization::CustomizeChildren(
 	[
 		SNew(SIKRetargetChainMapList).InChainMapListConfig(ChainListConfig)
 	];
+
+	// add the properties from the base op
+	AddBaseOpProperties(StructPropertyHandle, ChildBuilder);
 }
 
 void FAdditivePoseOpCustomization::CustomizeChildren(
@@ -594,6 +641,9 @@ void FAdditivePoseOpCustomization::CustomizeChildren(
 	// add the alpha property
 	TSharedPtr<IPropertyHandle> AlphaProperty = StructPropertyHandle->GetChildHandle(GET_MEMBER_NAME_CHECKED(FIKRetargetAdditivePoseOpSettings, Alpha));
 	ChildBuilder.AddProperty(AlphaProperty.ToSharedRef());
+
+	// add the properties from the base op
+	AddBaseOpProperties(StructPropertyHandle, ChildBuilder);
 }
 
 void FAdditivePoseOpCustomization::UpdatePoseNameOptions()
@@ -688,6 +738,9 @@ void FStretchChainOpCustomization::CustomizeChildren(
 	[
 		SNew(SIKRetargetChainMapList).InChainMapListConfig(ChainListConfig)
 	];
+
+	// add the properties from the base op
+	AddBaseOpProperties(StructPropertyHandle, ChildBuilder);
 }
 
 #undef LOCTEXT_NAMESPACE
diff --git a/Engine/Plugins/Animation/IKRig/Source/IKRigEditor/Public/RetargetEditor/IKRetargetOpDetails.h b/Engine/Plugins/Animation/IKRig/Source/IKRigEditor/Public/RetargetEditor/IKRetargetOpDetails.h
--- a/Engine/Plugins/Animation/IKRig/Source/IKRigEditor/Public/RetargetEditor/IKRetargetOpDetails.h
+++ b/Engine/Plugins/Animation/IKRig/Source/IKRigEditor/Public/RetargetEditor/IKRetargetOpDetails.h
@@ -32,10 +32,9 @@ public:
 		const TSharedRef<IPropertyHandle>& InStructPropertyHandle,
 		const IPropertyTypeCustomizationUtils& InStructCustomizationUtils);
 
-	static void AddChildPropertyToGroup(
-		IDetailGroup& InGroup,
-		const TSharedRef<IPropertyHandle>& InParentPropertyHandle,
-		const FName& InChildPropertyName);
+	static void AddPropertyToGroup(
+		const TSharedRef<IPropertyHandle>& InPropertyHandle,
+		IDetailChildrenBuilder& ChildBuilder);
 
 	static void AddChildPropertiesToCategoryGroups(
 		const TSharedRef<IPropertyHandle>& InParentPropertyHandle,
@@ -47,6 +46,10 @@ public:
 		const FName& InCategoryName,
 		const TArray<FName>& InPropertiesToIgnore = {});
 
+	static void AddBaseOpProperties(
+		TSharedRef<IPropertyHandle> StructPropertyHandle,
+		IDetailChildrenBuilder& ChildBuilder);
+
 protected:
 
 	// the op that owns the settings being customized / edited
diff --git a/Engine/Source/Editor/PropertyEditor/Private/CustomChildBuilder.cpp b/Engine/Source/Editor/PropertyEditor/Private/CustomChildBuilder.cpp
--- a/Engine/Source/Editor/PropertyEditor/Private/CustomChildBuilder.cpp
+++ b/Engine/Source/Editor/PropertyEditor/Private/CustomChildBuilder.cpp
@@ -27,6 +27,22 @@ IDetailGroup& FCustomChildrenBuilder::AddGroup( FName GroupName, const FText& Lo
 	return *NewCustomization.DetailGroup;
 }
 
+IDetailGroup* FCustomChildrenBuilder::GetGroup(FName GroupName)
+{
+	for (const FDetailLayoutCustomization& ChildCustomization : ChildCustomizations)
+	{
+		if (FDetailGroup* DetailGroup = ChildCustomization.DetailGroup.Get())
+		{
+			if (DetailGroup->GetGroupName() == GroupName)
+			{
+				return ChildCustomization.DetailGroup.Get();	
+			}
+		}
+	}
+
+	return nullptr;
+}
+
 FDetailWidgetRow& FCustomChildrenBuilder::AddCustomRow( const FText& SearchString )
 {
 	const TSharedRef<FDetailWidgetRow> NewRow = MakeShared<FDetailWidgetRow>();
diff --git a/Engine/Source/Editor/PropertyEditor/Private/CustomChildBuilder.h b/Engine/Source/Editor/PropertyEditor/Private/CustomChildBuilder.h
--- a/Engine/Source/Editor/PropertyEditor/Private/CustomChildBuilder.h
+++ b/Engine/Source/Editor/PropertyEditor/Private/CustomChildBuilder.h
@@ -22,6 +22,7 @@ public:
 
 	virtual IDetailChildrenBuilder& AddCustomBuilder(TSharedRef<class IDetailCustomNodeBuilder> InCustomBuilder) override;
 	virtual IDetailGroup& AddGroup(FName GroupName, const FText& LocalizedDisplayName, const bool bStartExpanded = false) override;
+	virtual IDetailGroup* GetGroup(FName GroupName) override;
 	virtual FDetailWidgetRow& AddCustomRow(const FText& SearchString) override;
 	virtual IDetailPropertyRow& AddProperty(TSharedRef<IPropertyHandle> PropertyHandle) override;
 	UE_DEPRECATED(5.1, "Please use the overload that takes an FAddPropertyParams instead.")
diff --git a/Engine/Source/Editor/PropertyEditor/Public/IDetailChildrenBuilder.h b/Engine/Source/Editor/PropertyEditor/Public/IDetailChildrenBuilder.h
--- a/Engine/Source/Editor/PropertyEditor/Public/IDetailChildrenBuilder.h
+++ b/Engine/Source/Editor/PropertyEditor/Public/IDetailChildrenBuilder.h
@@ -41,6 +41,14 @@ public:
 	 */
 	virtual IDetailGroup& AddGroup(FName GroupName, const FText& LocalizedDisplayName, const bool bStartExpanded=false) = 0;
 
+	/**
+	 * Get an existing group in the category
+	 *
+	 * @param GroupName	The name of the group
+	 * @return a pointer to the group if it exists, otherwise nullptr
+	 */
+	virtual IDetailGroup* GetGroup(FName GroupName) = 0;
+
 	/**
 	 * Adds new custom content as a child to the struct
 	 *
