[IK Retarget] Changed source goal locations to not be affected by Pelvis motion, by default.

diff --git a/Engine/Plugins/Animation/IKRig/Source/IKRig/Private/Retargeter/RetargetOps/IKChainsOp.cpp b/Engine/Plugins/Animation/IKRig/Source/IKRig/Private/Retargeter/RetargetOps/IKChainsOp.cpp
--- a/Engine/Plugins/Animation/IKRig/Source/IKRig/Private/Retargeter/RetargetOps/IKChainsOp.cpp
+++ b/Engine/Plugins/Animation/IKRig/Source/IKRig/Private/Retargeter/RetargetOps/IKChainsOp.cpp
@@ -6,6 +6,7 @@
 #include "Retargeter/RetargetOps/FKChainsOp.h"
 #include "Retargeter/RetargetOps/PelvisMotionOp.h"
 #include "Retargeter/RetargetOps/RunIKRigOp.h"
+#include "IKRigObjectVersion.h"
 
 #if WITH_EDITOR
 #include "PrimitiveDrawInterface.h"
@@ -156,9 +157,12 @@ FVector FIKChainRetargeter::CalculateGoalPosition(
 	const double BlendToSourceTranslation = Settings->BlendToSource * Settings->BlendToSourceTranslation;
 	if (BlendToSourceTranslation > UE_KINDA_SMALL_NUMBER)
 	{
-		const FVector RootModification = PelvisTranslationDelta * AffectIKWeights;
 		const FVector Weight = BlendToSourceTranslation * Settings->BlendToSourceWeights;
-		const FVector SourceLocation = CurrentSourceEndGlobal.GetLocation() + RootModification;
+		FVector SourceLocation = CurrentSourceEndGlobal.GetLocation();
+		if (Settings->ApplyPelvisOffsetToSourceGoals)
+		{
+			SourceLocation += PelvisTranslationDelta * AffectIKWeights;
+		}
 		GoalPosition.X = FMath::Lerp(GoalPosition.X, SourceLocation.X, Weight.X);
 		GoalPosition.Y = FMath::Lerp(GoalPosition.Y, SourceLocation.Y, Weight.Y);
 		GoalPosition.Z = FMath::Lerp(GoalPosition.Z, SourceLocation.Z, Weight.Z);
@@ -197,6 +201,21 @@ bool FRetargetIKChainSettings::operator==(const FRetargetIKChainSettings& Other)
 		&& FMath::IsNearlyEqualByULP(Extension, Other.Extension);
 }
 
+void FRetargetIKChainSettings::PostSerialize(const FArchive& Ar)
+{
+#if WITH_EDITORONLY_DATA
+	
+	if (Ar.IsLoading())
+	{
+		if (Ar.CustomVer(FIKRigObjectVersion::GUID) < FIKRigObjectVersion::RemovedPelvisMotionFromSourceGoals)
+		{
+			ApplyPelvisOffsetToSourceGoals = true;
+		}
+	}
+	
+#endif
+}
+
 const UClass* FIKRetargetIKChainsOpSettings::GetControllerType() const
 {
 	return UIKRetargetIKChainsController::StaticClass();
@@ -493,7 +512,7 @@ void FIKRetargetIKChainsOp::DebugDraw(
 		if (Settings.bDrawSourceLocations)
 		{
 			FTransform SourceGoalTransform;
-			SourceGoalTransform.SetTranslation(ChainDebugData.SourceTransformEnd.GetLocation() + DebugRootModification);
+			SourceGoalTransform.SetTranslation(ChainDebugData.SourceTransformEnd.GetLocation());
 			SourceGoalTransform.SetRotation(ChainDebugData.SourceTransformEnd.GetRotation());
 			SourceGoalTransform *= InComponentTransform;
 
@@ -534,6 +553,13 @@ void FIKRetargetIKChainsOp::SaveDebugData(
 {
 	FScopeLock ScopeLock(&DebugDataMutex);
 	
+	// get the root modification
+	DebugRootModification = FVector::ZeroVector;
+	if (const FIKRetargetPelvisMotionOp* PelvisMotionOp = InProcessor.GetFirstRetargetOpOfType<FIKRetargetPelvisMotionOp>())
+	{
+		DebugRootModification = PelvisMotionOp->GetPelvisTranslationOffset() * PelvisMotionOp->GetAffectIKWeightAsVector();
+	}
+	
 	AllChainsDebugData.Reset();
 	for (const FIKChainRetargeter& IKChainPair : IKChainRetargeters)
 	{
@@ -545,15 +571,10 @@ void FIKRetargetIKChainsOp::SaveDebugData(
 		NewChainData.InputTransformEnd = OutTargetGlobalPose[TargetChain.BoneIndices.Last()];
 		NewChainData.OutputTransformEnd = IKChainPair.GetGoalTransform();
 		NewChainData.SourceTransformEnd = InSourceGlobalPose[SourceChain.BoneIndices.Last()];
+		const FVector RootModification = IKChainPair.GetSettings()->ApplyPelvisOffsetToSourceGoals ? DebugRootModification : FVector::ZeroVector;
+		NewChainData.SourceTransformEnd.AddToTranslation(RootModification);
 		AllChainsDebugData.Add(NewChainData);
 	}
-
-	// get the root modification
-	DebugRootModification = FVector::ZeroVector;
-	if (const FIKRetargetPelvisMotionOp* PelvisMotionOp = InProcessor.GetFirstRetargetOpOfType<FIKRetargetPelvisMotionOp>())
-	{
-		DebugRootModification = PelvisMotionOp->GetPelvisTranslationOffset() * PelvisMotionOp->GetAffectIKWeightAsVector();
-	}
 }
 
 void FIKRetargetIKChainsOp::ResetChainSettingsToDefault(const FName& InChainName)
diff --git a/Engine/Plugins/Animation/IKRig/Source/IKRig/Public/IKRigObjectVersion.h b/Engine/Plugins/Animation/IKRig/Source/IKRig/Public/IKRigObjectVersion.h
--- a/Engine/Plugins/Animation/IKRig/Source/IKRig/Public/IKRigObjectVersion.h
+++ b/Engine/Plugins/Animation/IKRig/Source/IKRig/Public/IKRigObjectVersion.h
@@ -37,6 +37,9 @@ struct FIKRigObjectVersion
 
 		// copy base pose op converted to use FBoneReference 
 		CopyBasePoseConvertedToBoneReference,
+		
+		// changed default behavior of BlendToSource on IK Goals to not be affected by Pelvis motion
+		RemovedPelvisMotionFromSourceGoals,
 
 		// -----<new versions can be added above this line>-------------------------------------------------
 		VersionPlusOne,
diff --git a/Engine/Plugins/Animation/IKRig/Source/IKRig/Public/Retargeter/RetargetOps/IKChainsOp.h b/Engine/Plugins/Animation/IKRig/Source/IKRig/Public/Retargeter/RetargetOps/IKChainsOp.h
--- a/Engine/Plugins/Animation/IKRig/Source/IKRig/Public/Retargeter/RetargetOps/IKChainsOp.h
+++ b/Engine/Plugins/Animation/IKRig/Source/IKRig/Public/Retargeter/RetargetOps/IKChainsOp.h
@@ -57,7 +57,6 @@ private:
 	const FRetargetIKChainSettings* Settings;
 	double InvInitialLengthSource = 1.0f;
 	double InitialTargetChainLength = 1.0f;
-
 	
 	FTransform InitialTargetEndRelativeToSourceEnd;
 };
@@ -123,6 +122,12 @@ struct FRetargetIKChainSettings
 	*  At 1 the goal is placed at the location of the source chain's end bone. */
 	UPROPERTY(EditAnywhere, BlueprintReadWrite, Category = "Blend to Source", meta = (ClampMin = "0.0", ClampMax = "1.0", UIMin = "0.0", UIMax = "1.0"))
 	FVector BlendToSourceWeights = FVector::OneVector;
+	
+	/** Default false. When true, the source goal locations are affected by the target pelvis motion.
+	* NOTE: this setting applies all offsets generated in the Pelvis Motion Op. 
+	If no Pelvis Motion Op is present, or if the "Affect IK" weights in the Pelvis Motion Op are zero, then this setting has no effect. */
+	UPROPERTY(EditAnywhere, BlueprintReadWrite, Category = "Blend to Source")
+	bool ApplyPelvisOffsetToSourceGoals = false;
 
 	/** Default 0, 0, 0. Apply a static global-space offset to IK goal position. */
 	UPROPERTY(EditAnywhere, BlueprintReadWrite, Category = "Static Offset")
@@ -149,6 +154,18 @@ struct FRetargetIKChainSettings
 	double Extension = 1.0f;
 
 	UE_API bool operator==(const FRetargetIKChainSettings& Other) const;
+	
+	/** Custom serialization to handle loading old assets */
+    void PostSerialize(const FArchive& Ar);
+};
+
+// Custom serialization to patch old settings
+template<> struct TStructOpsTypeTraits<FRetargetIKChainSettings> : public TStructOpsTypeTraitsBase2<FRetargetIKChainSettings>
+{
+	enum
+	{
+		WithPostSerialize = true,
+	};
 };
 
 USTRUCT(BlueprintType, meta = (DisplayName = "Retarget IK Goals Settings"))
