[IK Retarget] Fixed two minor bugs introduced in CL 44679472 (thanks BugHawk!)

diff --git a/Engine/Plugins/Animation/IKRig/Source/IKRig/Private/Retargeter/RetargetOps/IKChainsOp.cpp b/Engine/Plugins/Animation/IKRig/Source/IKRig/Private/Retargeter/RetargetOps/IKChainsOp.cpp
--- a/Engine/Plugins/Animation/IKRig/Source/IKRig/Private/Retargeter/RetargetOps/IKChainsOp.cpp
+++ b/Engine/Plugins/Animation/IKRig/Source/IKRig/Private/Retargeter/RetargetOps/IKChainsOp.cpp
@@ -194,6 +194,7 @@ bool FRetargetIKChainSettings::operator==(const FRetargetIKChainSettings& Other)
 		&& FMath::IsNearlyEqualByULP(FloorHeightFalloffEnd, Other.FloorHeightFalloffEnd)
 		&& FMath::IsNearlyEqualByULP(BlendToSource, Other.BlendToSource)
 		&& BlendToSourceWeights.Equals(Other.BlendToSourceWeights)
+		&& ApplyPelvisOffsetToSourceGoals == Other.ApplyPelvisOffsetToSourceGoals
 		&& StaticOffset.Equals(Other.StaticOffset)
 		&& StaticLocalOffset.Equals(Other.StaticLocalOffset)
 		&& StaticRotationOffset.Equals(Other.StaticRotationOffset)
@@ -203,8 +204,6 @@ bool FRetargetIKChainSettings::operator==(const FRetargetIKChainSettings& Other)
 
 void FRetargetIKChainSettings::PostSerialize(const FArchive& Ar)
 {
-#if WITH_EDITORONLY_DATA
-	
 	if (Ar.IsLoading())
 	{
 		if (Ar.CustomVer(FIKRigObjectVersion::GUID) < FIKRigObjectVersion::RemovedPelvisMotionFromSourceGoals)
@@ -212,8 +211,6 @@ void FRetargetIKChainSettings::PostSerialize(const FArchive& Ar)
 			ApplyPelvisOffsetToSourceGoals = true;
 		}
 	}
-	
-#endif
 }
 
 const UClass* FIKRetargetIKChainsOpSettings::GetControllerType() const
