[IK Rig] Fix non-normalized rotation on Stretch Limb solver end rotation.

diff --git a/Engine/Plugins/Animation/IKRig/Source/IKRig/Private/Rig/Solvers/IKRigStretchLimb.cpp b/Engine/Plugins/Animation/IKRig/Source/IKRig/Private/Rig/Solvers/IKRigStretchLimb.cpp
--- a/Engine/Plugins/Animation/IKRig/Source/IKRig/Private/Rig/Solvers/IKRigStretchLimb.cpp
+++ b/Engine/Plugins/Animation/IKRig/Source/IKRig/Private/Rig/Solvers/IKRigStretchLimb.cpp
@@ -89,7 +89,8 @@ void FIKRigStretchLimbSolver::Solve(FIKRigSkeleton& InIKRigSkeleton, const FIKRi
 
 	// orient end bone to goal
 	FTransform& EndBoneGlobal = InIKRigSkeleton.CurrentPoseGlobal[BoneIndices.Last()];
-	const FQuat EndBoneGlobalRotation = FQuat::FastLerp(EndBoneGlobal.GetRotation(), Goal->FinalBlendedRotation, Settings.RotateEndBoneWithGoal);
+	FQuat EndBoneGlobalRotation = FQuat::FastLerp(EndBoneGlobal.GetRotation(), Goal->FinalBlendedRotation, Settings.RotateEndBoneWithGoal);
+	EndBoneGlobalRotation = EndBoneGlobalRotation.GetNormalized();
 	EndBoneGlobal.SetRotation(EndBoneGlobalRotation);
 
 	// update chain local transforms
