[IK Retarget] Added mising getter to CopyBasePoseOp scripting API

diff --git a/Engine/Plugins/Animation/IKRig/Source/IKRig/Private/Retargeter/RetargetOps/CopyBasePoseOp.cpp b/Engine/Plugins/Animation/IKRig/Source/IKRig/Private/Retargeter/RetargetOps/CopyBasePoseOp.cpp
--- a/Engine/Plugins/Animation/IKRig/Source/IKRig/Private/Retargeter/RetargetOps/CopyBasePoseOp.cpp
+++ b/Engine/Plugins/Animation/IKRig/Source/IKRig/Private/Retargeter/RetargetOps/CopyBasePoseOp.cpp
@@ -243,6 +243,24 @@ void UIKRetargetCopyBasePoseController::AddBoneToExclude(const FName InBoneName)
 	Settings->BonesToExclude.Add(InBoneName);
 }
 
+TArray<FName> UIKRetargetCopyBasePoseController::GetBonesToExclude() const
+{
+	FIKRetargetCopyBasePoseOpSettings* Settings = reinterpret_cast<FIKRetargetCopyBasePoseOpSettings*>(OpSettingsToControl);
+	if (!Settings)
+	{
+		return TArray<FName>();
+	}
+
+	// convert to FNames (FBoneReference has no BP support)
+	TArray<FName> ExcludedBoneNames;
+	for (const FBoneReference& BoneToExclude : Settings->BonesToExclude)
+	{
+		ExcludedBoneNames.Add(BoneToExclude.BoneName);
+	}
+	
+	return ExcludedBoneNames;
+}
+
 void UIKRetargetCopyBasePoseController::ResetBonesToExclude() const
 {
 	FIKRetargetCopyBasePoseOpSettings* Settings = reinterpret_cast<FIKRetargetCopyBasePoseOpSettings*>(OpSettingsToControl);
diff --git a/Engine/Plugins/Animation/IKRig/Source/IKRig/Public/Retargeter/RetargetOps/CopyBasePoseOp.h b/Engine/Plugins/Animation/IKRig/Source/IKRig/Public/Retargeter/RetargetOps/CopyBasePoseOp.h
--- a/Engine/Plugins/Animation/IKRig/Source/IKRig/Public/Retargeter/RetargetOps/CopyBasePoseOp.h
+++ b/Engine/Plugins/Animation/IKRig/Source/IKRig/Public/Retargeter/RetargetOps/CopyBasePoseOp.h
@@ -120,6 +120,10 @@ public:
 	UFUNCTION(BlueprintCallable, Category = ExcludedBones)
 	UE_API void AddBoneToExclude(const FName InBoneName) const;
 
+	/* Get an array of all the excluded bones. */
+	UFUNCTION(BlueprintCallable, Category = ExcludedBones)
+	UE_API TArray<FName> GetBonesToExclude() const;
+
 	/* Reset the array of Bones to Exclude */
 	UFUNCTION(BlueprintCallable, Category = ExcludedBones)
 	UE_API void ResetBonesToExclude() const;
