[IK Retarget] Refactored FK Chains Op to make it easier to add new modes.

diff --git a/Engine/Plugins/Animation/IKRig/Source/IKRig/Private/Retargeter/RetargetOps/FKChainsOp.cpp b/Engine/Plugins/Animation/IKRig/Source/IKRig/Private/Retargeter/RetargetOps/FKChainsOp.cpp
--- a/Engine/Plugins/Animation/IKRig/Source/IKRig/Private/Retargeter/RetargetOps/FKChainsOp.cpp
+++ b/Engine/Plugins/Animation/IKRig/Source/IKRig/Private/Retargeter/RetargetOps/FKChainsOp.cpp
@@ -167,7 +167,6 @@ void FChainDecoderFK::DecodePose(
 	FTransform SourceChainParentInitialDelta = SourceChain.BoneChain->ChainParentInitialGlobalTransform.GetRelativeTransform(BoneChain->ChainParentInitialGlobalTransform);
 	FTransform TargetChainParentCurrentGlobalTransform = BoneChain->ChainParentBoneIndex == INDEX_NONE ? FTransform::Identity : InOutGlobalPose[BoneChain->ChainParentBoneIndex]; 
 	FTransform SourceChainParentTransform = SourceChainParentInitialDelta * TargetChainParentCurrentGlobalTransform;
-
 	// apply delta to the source chain's current transforms before transferring rotations to the target
 	SourceChain.TransformCurrentChainTransforms(SourceChainParentTransform);
 
@@ -187,236 +186,282 @@ void FChainDecoderFK::DecodePose(
 		return;
 	}
 
-	const int32 NumBonesInSourceChain = SourceChain.CurrentGlobalTransforms.Num();
-	const int32 NumBonesInTargetChain = TargetBoneIndices.Num();
-	const int32 TargetStartIndex = FMath::Max(0, NumBonesInTargetChain - NumBonesInSourceChain);
-	const int32 SourceStartIndex = FMath::Max(0,NumBonesInSourceChain - NumBonesInTargetChain);
-
 	// now retarget the pose of each bone in the chain, copying from source to target
 	for (int32 ChainIndex=0; ChainIndex<TargetBoneIndices.Num(); ++ChainIndex)
 	{
-		const int32 BoneIndex = TargetBoneIndices[ChainIndex];
-		const FTransform& TargetInitialTransform = BoneChain->RefPoseGlobalTransforms[ChainIndex];
 		FTransform SourceCurrentTransform;
 		FTransform SourceInitialTransform;
 
-		// get source current / initial transforms for this bone
-		switch (Settings.RotationMode)
+		// retarget rotation
+		const FQuat OutRotation = RetargetRotation(
+			ChainIndex,
+			Settings,
+			TargetBoneIndices,
+			SourceChain,
+			SourceCurrentTransform,
+			SourceInitialTransform);
+
+		// retarget position
+		const FVector OutPosition = RetargetPosition(
+			ChainIndex,
+			OutRotation,
+			PelvisMotionOp,
+			Settings,
+			TargetBoneIndices,
+			SourceChain,
+			InSourceGlobalPose,
+			SourceSkeleton,
+			TargetSkeleton,
+			InOutGlobalPose,
+			SourceCurrentTransform,
+			SourceInitialTransform);
+
+		// retarget scale
+		const FVector OutScale = RetargetScale(ChainIndex, SourceCurrentTransform, SourceInitialTransform);
+		
+		// apply output transform
+		const int32 BoneIndex = TargetBoneIndices[ChainIndex];
+		InOutGlobalPose[BoneIndex] = FTransform(OutRotation, OutPosition, OutScale);
+	}
+
+	// apply match chain operation on-top of translated and rotated result
+	if (Settings.RotationMode == EFKChainRotationMode::MatchChain ||
+		Settings.RotationMode == EFKChainRotationMode::MatchScaledChain)
+	{
+		const bool bScaleSourceChain = Settings.RotationMode == EFKChainRotationMode::MatchScaledChain;
+		FVector TargetChainOrigin = InOutGlobalPose[TargetBoneIndices[0]].GetTranslation();
+		MatchChain(bScaleSourceChain, SourceChain, TargetBoneIndices, TargetChainOrigin, InOutGlobalPose);
+	}
+
+	// apply final blending between retarget pose of chain and newly retargeted pose
+	BlendByAlpha(Settings, TargetBoneIndices, TargetSkeleton, InOutGlobalPose);
+}
+
+FQuat FChainDecoderFK::RetargetRotation(
+	const int32 InChainIndex,
+	const FRetargetFKChainSettings& Settings,
+	const TArray<int32>& TargetBoneIndices,
+	FChainEncoderFK& SourceChain,
+	FTransform& OutSourceCurrentTransform,
+	FTransform& OutSourceInitialTransform) const
+{
+	
+	// get source current / initial transforms for this bone
+	switch (Settings.RotationMode)
+	{
+		case EFKChainRotationMode::Interpolated:
+		case EFKChainRotationMode::MatchChain:
+		case EFKChainRotationMode::MatchScaledChain:
 		{
-			case EFKChainRotationMode::Interpolated:
-			case EFKChainRotationMode::MatchChain:
-			case EFKChainRotationMode::MatchScaledChain:
+			// get the initial and current transform of source chain at param
+			// this is the interpolated transform along the chain
+			const float Param = BoneChain->Params[InChainIndex];
+			OutSourceCurrentTransform = SourceChain.BoneChain->GetTransformAtChainParam(SourceChain.CurrentGlobalTransforms, Param);
+			OutSourceInitialTransform = SourceChain.BoneChain->GetTransformAtChainParam(SourceChain.BoneChain->RefPoseGlobalTransforms, Param);
+			break;
+		}
+		case EFKChainRotationMode::OneToOne:
+		{
+			const int32 NumBonesInSourceChain = SourceChain.CurrentGlobalTransforms.Num();
+			if (InChainIndex < NumBonesInSourceChain)
 			{
-				// get the initial and current transform of source chain at param
-				// this is the interpolated transform along the chain
-				const float Param = BoneChain->Params[ChainIndex];
-				SourceCurrentTransform = SourceChain.BoneChain->GetTransformAtChainParam(SourceChain.CurrentGlobalTransforms, Param);
-				SourceInitialTransform = SourceChain.BoneChain->GetTransformAtChainParam(SourceChain.BoneChain->RefPoseGlobalTransforms, Param);
-				break;
-			}
-			case EFKChainRotationMode::OneToOne:
+				OutSourceCurrentTransform = SourceChain.CurrentGlobalTransforms[InChainIndex];
+				OutSourceInitialTransform = SourceChain.BoneChain->RefPoseGlobalTransforms[InChainIndex];
+			}else
 			{
-				if (ChainIndex < NumBonesInSourceChain)
-				{
-					SourceCurrentTransform = SourceChain.CurrentGlobalTransforms[ChainIndex];
-					SourceInitialTransform = SourceChain.BoneChain->RefPoseGlobalTransforms[ChainIndex];
-				}else
-				{
-					SourceCurrentTransform = SourceChain.CurrentGlobalTransforms.Last();
-					SourceInitialTransform = SourceChain.BoneChain->RefPoseGlobalTransforms.Last();
-				}
-				break;
+				OutSourceCurrentTransform = SourceChain.CurrentGlobalTransforms.Last();
+				OutSourceInitialTransform = SourceChain.BoneChain->RefPoseGlobalTransforms.Last();
 			}
-			case EFKChainRotationMode::OneToOneReversed:
+			break;
+		}
+		case EFKChainRotationMode::OneToOneReversed:
+		{
+			const int32 NumBonesInSourceChain = SourceChain.CurrentGlobalTransforms.Num();
+			const int32 NumBonesInTargetChain = TargetBoneIndices.Num();
+			const int32 TargetStartIndex = FMath::Max(0, NumBonesInTargetChain - NumBonesInSourceChain);
+			const int32 SourceStartIndex = FMath::Max(0,NumBonesInSourceChain - NumBonesInTargetChain);
+				
+			if (InChainIndex < TargetStartIndex)
 			{
-				if (ChainIndex < TargetStartIndex)
-				{
-					SourceCurrentTransform = SourceChain.BoneChain->RefPoseGlobalTransforms[0];
-					SourceInitialTransform = SourceChain.BoneChain->RefPoseGlobalTransforms[0];
-				}
-				else
-				{
-					const int32 SourceChainIndex = SourceStartIndex + (ChainIndex - TargetStartIndex);
-					SourceCurrentTransform = SourceChain.CurrentGlobalTransforms[SourceChainIndex];
-					SourceInitialTransform = SourceChain.BoneChain->RefPoseGlobalTransforms[SourceChainIndex];
-				}
-				break;
+				OutSourceCurrentTransform = SourceChain.BoneChain->RefPoseGlobalTransforms[0];
+				OutSourceInitialTransform = SourceChain.BoneChain->RefPoseGlobalTransforms[0];
 			}
-			case EFKChainRotationMode::None:
+			else
 			{
-				// in order to induce no rotation on the FK chain, we rotate the chain rigidly from the root of the chain
-				SourceInitialTransform = SourceChain.BoneChain->RefPoseGlobalTransforms[0];
-				// use the current global space retarget pose as the "current" transform, so chain rotates with parent
-				SourceCurrentTransform = SourceChain.BoneChain->RefPoseLocalTransforms[0] * SourceChain.ChainParentCurrentGlobalTransform;
-				break;
+				const int32 SourceChainIndex = SourceStartIndex + (InChainIndex - TargetStartIndex);
+				OutSourceCurrentTransform = SourceChain.CurrentGlobalTransforms[SourceChainIndex];
+				OutSourceInitialTransform = SourceChain.BoneChain->RefPoseGlobalTransforms[SourceChainIndex];
 			}
-			default:
-				checkNoEntry();
 			break;
 		}
-		
-		// apply rotation offset to the initial target rotation
-		const FQuat SourceCurrentRotation = SourceCurrentTransform.GetRotation();
-		const FQuat SourceInitialRotation = SourceInitialTransform.GetRotation();
-		const FQuat RotationDelta = SourceCurrentRotation * SourceInitialRotation.Inverse();
-		const FQuat TargetInitialRotation = TargetInitialTransform.GetRotation();
-		const FQuat OutRotation = RotationDelta * TargetInitialRotation;
-
-		// calculate output POSITION based on translation mode setting
-		FTransform ParentGlobalTransform = FTransform::Identity;
-		const int32 ParentIndex = TargetSkeleton.ParentIndices[BoneIndex];
-		if (ParentIndex != INDEX_NONE)
+		case EFKChainRotationMode::None:
+		{
+			// in order to induce no rotation on the FK chain, we rotate the chain rigidly from the root of the chain
+			OutSourceInitialTransform = SourceChain.BoneChain->RefPoseGlobalTransforms[0];
+			// use the current global space retarget pose as the "current" transform, so chain rotates with parent
+			OutSourceCurrentTransform = SourceChain.BoneChain->RefPoseLocalTransforms[0] * SourceChain.ChainParentCurrentGlobalTransform;
+			break;
+		}
+	case EFKChainRotationMode::PlainCopy:
+		default:
+			checkNoEntry();
+		break;
+	}
+	
+	// calculate rotation delta
+	const FQuat SourceCurrentRotation = OutSourceCurrentTransform.GetRotation();
+	const FQuat SourceInitialRotation = OutSourceInitialTransform.GetRotation();
+	const FQuat RotationDelta = SourceCurrentRotation * SourceInitialRotation.Inverse();
+
+	// apply delta to the initial target rotation
+	const FTransform& TargetInitialTransform = BoneChain->RefPoseGlobalTransforms[InChainIndex];
+	const FQuat TargetInitialRotation = TargetInitialTransform.GetRotation();
+	FQuat Rotation = RotationDelta * TargetInitialRotation;
+	return MoveTemp(Rotation);
+}
+
+FVector FChainDecoderFK::RetargetPosition(
+	const int32 InChainIndex,
+	const FQuat& NewRotation,
+	const FIKRetargetPelvisMotionOp* PelvisMotionOp,
+	const FRetargetFKChainSettings& Settings,
+	const TArray<int32>& TargetBoneIndices,
+	FChainEncoderFK& SourceChain,
+	const TArray<FTransform>& InSourceGlobalPose,
+	const FRetargetSkeleton& SourceSkeleton,
+	const FTargetSkeleton& TargetSkeleton,
+	TArray<FTransform>& InOutGlobalPose,
+	FTransform& OutSourceCurrentTransform,
+	FTransform& OutSourceInitialTransform)
+{
+	// the position this function generates
+	FVector OutPosition;
+	
+	const int32 BoneIndex = TargetBoneIndices[InChainIndex];
+	
+	// calculate output POSITION based on translation mode setting
+	FTransform ParentGlobalTransform = FTransform::Identity;
+	const int32 ParentIndex = TargetSkeleton.ParentIndices[BoneIndex];
+	if (ParentIndex != INDEX_NONE)
+	{
+		ParentGlobalTransform = InOutGlobalPose[ParentIndex];
+	}
+	
+	switch (Settings.TranslationMode)
+	{
+	case EFKChainTranslationMode::None:
+		{
+			const FVector InitialLocalOffset = TargetSkeleton.RetargetPoses.GetLocalRetargetPose()[BoneIndex].GetTranslation();
+			OutPosition = ParentGlobalTransform.TransformPosition(InitialLocalOffset);
+			break;
+		}
+	case EFKChainTranslationMode::GloballyScaled:
+		{
+			FVector GlobalScale = PelvisMotionOp ? PelvisMotionOp->GetGlobalScaleVector() : FVector::OneVector;
+			OutPosition = OutSourceCurrentTransform.GetTranslation() * GlobalScale;
+			break;
+		}
+	case EFKChainTranslationMode::Absolute:
 		{
-			ParentGlobalTransform = InOutGlobalPose[ParentIndex];
+			OutPosition = OutSourceCurrentTransform.GetTranslation();
+			break;	
 		}
-		FVector OutPosition;
-		switch (Settings.TranslationMode)
+	case EFKChainTranslationMode::StretchBoneLengthUniformly:
 		{
-		case EFKChainTranslationMode::None:
+			if (InChainIndex==0)
 			{
 				const FVector InitialLocalOffset = TargetSkeleton.RetargetPoses.GetLocalRetargetPose()[BoneIndex].GetTranslation();
 				OutPosition = ParentGlobalTransform.TransformPosition(InitialLocalOffset);
 				break;
 			}
-		case EFKChainTranslationMode::GloballyScaled:
-			{
-				FVector GlobalScale = PelvisMotionOp ? PelvisMotionOp->GetGlobalScaleVector() : FVector::OneVector;
-				OutPosition = SourceCurrentTransform.GetTranslation() * GlobalScale;
-				break;
-			}
-		case EFKChainTranslationMode::Absolute:
-			{
-				OutPosition = SourceCurrentTransform.GetTranslation();
-				break;	
-			}
-		case EFKChainTranslationMode::StretchBoneLengthUniformly:
+			// initial chain length
+			double SourceChainLengthInitial = BoneChain->GetChainLength(SourceChain.BoneChain->RefPoseGlobalTransforms);
+			double SourceChainLengthCurrent = BoneChain->GetChainLength(SourceChain.CurrentGlobalTransforms);
+			double StretchRatio = SourceChainLengthInitial < UE_KINDA_SMALL_NUMBER ? 1.0f : SourceChainLengthCurrent / SourceChainLengthInitial;
+			// stretch local translation
+			const FVector InitialLocalOffset = TargetSkeleton.RetargetPoses.GetLocalRetargetPose()[BoneIndex].GetTranslation();
+			OutPosition = ParentGlobalTransform.TransformPosition(InitialLocalOffset * StretchRatio);
+			break;
+		}
+	case EFKChainTranslationMode::StretchBoneLengthNonUniformly:
+		{
+			const double Param = BoneChain->Params[InChainIndex];
+			double StretchRatio = SourceChain.BoneChain->GetStretchAtParam( SourceChain.BoneChain->RefPoseGlobalTransforms, SourceChain.CurrentGlobalTransforms, Param);
+			// stretch local translation
+			const FVector InitialLocalOffset = TargetSkeleton.RetargetPoses.GetLocalRetargetPose()[BoneIndex].GetTranslation();
+			OutPosition = ParentGlobalTransform.TransformPosition(InitialLocalOffset * StretchRatio);
+			break;
+		}
+	case EFKChainTranslationMode::OrientAndScale:
+		{
+			const double TargetRefLength = CachedData.TargetRefPoseLengths[InChainIndex];
+			const double SourceRefLength = CachedData.SourceRefPoseLengths[InChainIndex];
+
+			// calculate current direction of source bone
+			const int32 SourceBoneIndex = CachedData.SourceBoneIndices[InChainIndex];
+			const FTransform& SourceCurrentGlobal = InSourceGlobalPose[SourceBoneIndex];
+			const FTransform& SourceParentCurrentGlobal = SourceSkeleton.GetParentTransform(SourceBoneIndex, InSourceGlobalPose);
+			const FVector SourceCurrentTranslation = SourceCurrentGlobal.GetTranslation() - SourceParentCurrentGlobal.GetTranslation();
+
+			// calculate current stretch of source bone
+			double SourceCurrentLength;
+			FVector SourceCurrentNorm;
+			SourceCurrentTranslation.ToDirectionAndLength(SourceCurrentNorm, SourceCurrentLength);
+
+			// generate a translational offset for the target bone
+			FVector TargetNewGlobalOffset = FVector::ZeroVector;
+
+			// in the absence of a proper source and target bone length, we treat translational deltas as additive rather than scaled
+			constexpr double ThresholdLength = 0.1f;
+			if (TargetRefLength <= ThresholdLength || SourceRefLength <= ThresholdLength)
 			{
-				if (ChainIndex==0)
-				{
-					const FVector InitialLocalOffset = TargetSkeleton.RetargetPoses.GetLocalRetargetPose()[BoneIndex].GetTranslation();
-					OutPosition = ParentGlobalTransform.TransformPosition(InitialLocalOffset);
-					break;
-				}
-				// initial chain length
-				double SourceChainLengthInitial = BoneChain->GetChainLength(SourceChain.BoneChain->RefPoseGlobalTransforms);
-				double SourceChainLengthCurrent = BoneChain->GetChainLength(SourceChain.CurrentGlobalTransforms);
-				double StretchRatio = SourceChainLengthInitial < UE_KINDA_SMALL_NUMBER ? 1.0f : SourceChainLengthCurrent / SourceChainLengthInitial;
-				// stretch local translation
+				// additive offset
+				const FVector SourceCurrentGlobalOffset = SourceCurrentNorm * (SourceCurrentLength - SourceRefLength);
+				// add it to the current local offset
 				const FVector InitialLocalOffset = TargetSkeleton.RetargetPoses.GetLocalRetargetPose()[BoneIndex].GetTranslation();
-				OutPosition = ParentGlobalTransform.TransformPosition(InitialLocalOffset * StretchRatio);
-				break;
+				const FVector TargetCurrentGlobalOffset = ParentGlobalTransform.TransformPosition(InitialLocalOffset) - ParentGlobalTransform.GetLocation();
+				// add the source offset to the target
+				TargetNewGlobalOffset = TargetCurrentGlobalOffset + SourceCurrentGlobalOffset;
 			}
-		case EFKChainTranslationMode::StretchBoneLengthNonUniformly:
+			else
 			{
-				const double Param = BoneChain->Params[ChainIndex];
-				double StretchRatio = SourceChain.BoneChain->GetStretchAtParam( SourceChain.BoneChain->RefPoseGlobalTransforms, SourceChain.CurrentGlobalTransforms, Param);
-				// stretch local translation
-				const FVector InitialLocalOffset = TargetSkeleton.RetargetPoses.GetLocalRetargetPose()[BoneIndex].GetTranslation();
-				OutPosition = ParentGlobalTransform.TransformPosition(InitialLocalOffset * StretchRatio);
-				break;
+				// scaled offset
+				const double SourceLengthScaleFactor = SourceCurrentLength / SourceRefLength;
+				// calculate equivalently stretched target bone in direction of source
+				TargetNewGlobalOffset = SourceCurrentNorm * (TargetRefLength * SourceLengthScaleFactor);
+				// apply rotational delta from the ref pose to the translational offset
+				const FQuat& RefPoseGlobalDeltaRotation = CachedData.FromToGlobalRotationDeltas[InChainIndex];
+				const FQuat CurrentGlobalRotationPlusDelta = RefPoseGlobalDeltaRotation * NewRotation;
+				const FQuat CurrentGlobalRotationDelta = CurrentGlobalRotationPlusDelta * NewRotation.Inverse();
+				TargetNewGlobalOffset = CurrentGlobalRotationDelta.RotateVector(TargetNewGlobalOffset);
 			}
-		case EFKChainTranslationMode::OrientAndScale:
-			{
-				const double TargetRefLength = CachedData.TargetRefPoseLengths[ChainIndex];
-				const double SourceRefLength = CachedData.SourceRefPoseLengths[ChainIndex];
-
-				// calculate current direction of source bone
-				const int32 SourceBoneIndex = CachedData.SourceBoneIndices[ChainIndex];
-				const FTransform& SourceCurrentGlobal = InSourceGlobalPose[SourceBoneIndex];
-				const FTransform& SourceParentCurrentGlobal = SourceSkeleton.GetParentTransform(SourceBoneIndex, InSourceGlobalPose);
-				const FVector SourceCurrentTranslation = SourceCurrentGlobal.GetTranslation() - SourceParentCurrentGlobal.GetTranslation();
-
-				// calculate current stretch of source bone
-				double SourceCurrentLength;
-				FVector SourceCurrentNorm;
-				SourceCurrentTranslation.ToDirectionAndLength(SourceCurrentNorm, SourceCurrentLength);
-
-				// generate a translational offset for the target bone
-				FVector TargetNewGlobalOffset = FVector::ZeroVector;
-
-				// in the absence of a proper source and target bone length, we treat translational deltas as additive rather than scaled
-				constexpr double ThresholdLength = 0.1f;
-				if (TargetRefLength <= ThresholdLength || SourceRefLength <= ThresholdLength)
-				{
-					// additive offset
-					const FVector SourceCurrentGlobalOffset = SourceCurrentNorm * (SourceCurrentLength - SourceRefLength);
-					// add it to the current local offset
-					const FVector InitialLocalOffset = TargetSkeleton.RetargetPoses.GetLocalRetargetPose()[BoneIndex].GetTranslation();
-					const FVector TargetCurrentGlobalOffset = ParentGlobalTransform.TransformPosition(InitialLocalOffset) - ParentGlobalTransform.GetLocation();
-					// add the source offset to the target
-					TargetNewGlobalOffset = TargetCurrentGlobalOffset + SourceCurrentGlobalOffset;
-				}
-				else
-				{
-					// scaled offset
-					const double SourceLengthScaleFactor = SourceCurrentLength / SourceRefLength;
-					// calculate equivalently stretched target bone in direction of source
-					TargetNewGlobalOffset = SourceCurrentNorm * (TargetRefLength * SourceLengthScaleFactor);
-					// apply rotational delta from the ref pose to the translational offset
-					const FQuat& RefPoseGlobalDeltaRotation = CachedData.FromToGlobalRotationDeltas[ChainIndex];
-					const FQuat CurrentGlobalRotationPlusDelta = RefPoseGlobalDeltaRotation * OutRotation;
-					const FQuat CurrentGlobalRotationDelta = CurrentGlobalRotationPlusDelta * OutRotation.Inverse();
-					TargetNewGlobalOffset = CurrentGlobalRotationDelta.RotateVector(TargetNewGlobalOffset);
-				}
 
-				// apply target offset to the parent global position
-				OutPosition = ParentGlobalTransform.GetLocation() + TargetNewGlobalOffset;
-				break;
-			}
-		default:
-			{
-				const FVector InitialLocalOffset = TargetSkeleton.RetargetPoses.GetLocalRetargetPose()[BoneIndex].GetTranslation();
-				OutPosition = ParentGlobalTransform.TransformPosition(InitialLocalOffset);
-				break;
-			}
+			// apply target offset to the parent global position
+			OutPosition = ParentGlobalTransform.GetLocation() + TargetNewGlobalOffset;
+			break;
 		}
-
-		// calculate output SCALE
-		const FVector SourceCurrentScale = SourceCurrentTransform.GetScale3D();
-		const FVector SourceInitialScale = SourceInitialTransform.GetScale3D();
-		const FVector TargetInitialScale = TargetInitialTransform.GetScale3D();
-		const FVector OutScale = SourceCurrentScale + (TargetInitialScale - SourceInitialScale);
-		
-		// apply output transform
-		InOutGlobalPose[BoneIndex] = FTransform(OutRotation, OutPosition, OutScale);
-	}
-
-	// apply match chain operation on-top of translated and rotated result
-	if (Settings.RotationMode == EFKChainRotationMode::MatchChain ||
-		Settings.RotationMode == EFKChainRotationMode::MatchScaledChain)
-	{
-		const bool bScaleSourceChain = Settings.RotationMode == EFKChainRotationMode::MatchScaledChain;
-		FVector TargetChainOrigin = InOutGlobalPose[TargetBoneIndices[0]].GetTranslation();
-		MatchChain(bScaleSourceChain, SourceChain, TargetBoneIndices, TargetChainOrigin, InOutGlobalPose);
-	}
-
-	// apply final blending between retarget pose of chain and newly retargeted pose
-	// blend must be done in local space, so we do it in a separate loop after full chain pose is generated
-	const bool bShouldBlendRotation = !FMath::IsNearlyEqual(Settings.RotationAlpha, 1.0f);
-	const bool bShouldBlendTranslation = !FMath::IsNearlyEqual(Settings.TranslationAlpha, 1.0f);
-	if (bShouldBlendRotation || bShouldBlendTranslation) // (skipped if the alphas are not near 1.0)
-	{
-		// generate local space pose of chain
-		TArray<FTransform> NewLocalTransforms;
-		NewLocalTransforms.SetNum(BoneChain->RefPoseLocalTransforms.Num());
-		BoneChain->FillTransformsWithLocalSpaceOfChain(TargetSkeleton, InOutGlobalPose, TargetBoneIndices, NewLocalTransforms);
-
-		// blend each bone in chain with the retarget pose
-		for (int32 ChainIndex=0; ChainIndex<BoneChain->RefPoseLocalTransforms.Num(); ++ChainIndex)
+	default:
 		{
-			// blend between current local pose and initial local pose
-			FTransform& NewLocalTransform = NewLocalTransforms[ChainIndex];
-			const FTransform& RefPoseLocalTransform = BoneChain->RefPoseLocalTransforms[ChainIndex];
-			NewLocalTransform.SetTranslation(FMath::Lerp(RefPoseLocalTransform.GetTranslation(), NewLocalTransform.GetTranslation(), Settings.TranslationAlpha));
-			NewLocalTransform.SetRotation(FQuat::FastLerp(RefPoseLocalTransform.GetRotation(), NewLocalTransform.GetRotation(), Settings.RotationAlpha).GetNormalized());
-
-			// put blended transforms back in global space and store in final output pose
-			const int32 BoneIndex = TargetBoneIndices[ChainIndex];
-			const int32 ParentIndex = TargetSkeleton.ParentIndices[BoneIndex];
-			const FTransform& ParentGlobalTransform = ParentIndex == INDEX_NONE ? FTransform::Identity : InOutGlobalPose[ParentIndex];
-			InOutGlobalPose[BoneIndex] = NewLocalTransform * ParentGlobalTransform;
+			const FVector InitialLocalOffset = TargetSkeleton.RetargetPoses.GetLocalRetargetPose()[BoneIndex].GetTranslation();
+			OutPosition = ParentGlobalTransform.TransformPosition(InitialLocalOffset);
+			break;
 		}
 	}
+
+	return OutPosition;
+}
+
+FVector FChainDecoderFK::RetargetScale(	
+	const int32 InChainIndex,
+	const FTransform& InSourceCurrentTransform,
+	const FTransform& InSourceInitialTransform) const
+{
+	// calculate output SCALE
+	const FVector SourceCurrentScale = InSourceCurrentTransform.GetScale3D();
+	const FVector SourceInitialScale = InSourceInitialTransform.GetScale3D();
+	const FVector TargetInitialScale = BoneChain->RefPoseGlobalTransforms[InChainIndex].GetScale3D();
+	FVector Scale = SourceCurrentScale + (TargetInitialScale - SourceInitialScale);
+	return MoveTemp(Scale);
 }
 
 void FChainDecoderFK::MatchChain(
@@ -512,6 +557,43 @@ void FChainDecoderFK::MatchChain(
 	}
 }
 
+void FChainDecoderFK::BlendByAlpha(
+	const FRetargetFKChainSettings& Settings,
+	const TArray<int32>& TargetBoneIndices,
+	const FTargetSkeleton& TargetSkeleton,
+	TArray<FTransform>& InOutGlobalPose)
+{
+	// apply final blending between retarget pose of chain and newly retargeted pose
+	// blend must be done in local space, so we do it in a separate loop after full chain pose is generated
+	const bool bShouldBlendRotation = !FMath::IsNearlyEqual(Settings.RotationAlpha, 1.0f);
+	const bool bShouldBlendTranslation = !FMath::IsNearlyEqual(Settings.TranslationAlpha, 1.0f);
+	if (!bShouldBlendRotation && !bShouldBlendTranslation)
+	{
+		return; // skipped if the alphas are near 1.0 (no blending needed)
+	}
+
+	// generate local space pose of chain
+	TArray<FTransform> NewLocalTransforms;
+	NewLocalTransforms.SetNum(BoneChain->RefPoseLocalTransforms.Num());
+	BoneChain->FillTransformsWithLocalSpaceOfChain(TargetSkeleton, InOutGlobalPose, TargetBoneIndices, NewLocalTransforms);
+
+	// blend each bone in chain with the retarget pose
+	for (int32 ChainIndex=0; ChainIndex<BoneChain->RefPoseLocalTransforms.Num(); ++ChainIndex)
+	{
+		// blend between current local pose and initial local pose
+		FTransform& NewLocalTransform = NewLocalTransforms[ChainIndex];
+		const FTransform& RefPoseLocalTransform = BoneChain->RefPoseLocalTransforms[ChainIndex];
+		NewLocalTransform.SetTranslation(FMath::Lerp(RefPoseLocalTransform.GetTranslation(), NewLocalTransform.GetTranslation(), Settings.TranslationAlpha));
+		NewLocalTransform.SetRotation(FQuat::FastLerp(RefPoseLocalTransform.GetRotation(), NewLocalTransform.GetRotation(), Settings.RotationAlpha).GetNormalized());
+
+		// put blended transforms back in global space and store in final output pose
+		const int32 BoneIndex = TargetBoneIndices[ChainIndex];
+		const int32 ParentIndex = TargetSkeleton.ParentIndices[BoneIndex];
+		const FTransform& ParentGlobalTransform = ParentIndex == INDEX_NONE ? FTransform::Identity : InOutGlobalPose[ParentIndex];
+		InOutGlobalPose[BoneIndex] = NewLocalTransform * ParentGlobalTransform;
+	}
+}
+
 FVector FChainDecoderFK::GetPointOnSplineFromIndexAndAlpha(
 	const TArray<FVector>& InSplinePoints,
 	const int32 InPointIndex,
diff --git a/Engine/Plugins/Animation/IKRig/Source/IKRig/Public/Retargeter/RetargetOps/FKChainsOp.h b/Engine/Plugins/Animation/IKRig/Source/IKRig/Public/Retargeter/RetargetOps/FKChainsOp.h
--- a/Engine/Plugins/Animation/IKRig/Source/IKRig/Public/Retargeter/RetargetOps/FKChainsOp.h
+++ b/Engine/Plugins/Animation/IKRig/Source/IKRig/Public/Retargeter/RetargetOps/FKChainsOp.h
@@ -73,6 +73,33 @@ struct FChainDecoderFK
 	
 private:
 
+	FQuat RetargetRotation(
+		const int32 InChainIndex,
+		const FRetargetFKChainSettings& Settings,
+		const TArray<int32>& TargetBoneIndices,
+		FChainEncoderFK& SourceChain,
+		FTransform& OutSourceCurrentTransform,
+		FTransform& OutSourceInitialTransform) const;
+
+	FVector RetargetPosition(
+		const int32 InChainIndex,
+		const FQuat& NewRotation,
+		const FIKRetargetPelvisMotionOp* PelvisMotionOp,
+		const FRetargetFKChainSettings& Settings,
+		const TArray<int32>& TargetBoneIndices,
+		FChainEncoderFK& SourceChain,
+		const TArray<FTransform>& InSourceGlobalPose,
+		const FRetargetSkeleton& SourceSkeleton,
+		const FTargetSkeleton& TargetSkeleton,
+		TArray<FTransform>& InOutGlobalPose,
+		FTransform& OutSourceCurrentTransform,
+		FTransform& OutSourceInitialTransform);
+
+	FVector RetargetScale(
+		const int32 InChainIndex,
+		const FTransform& OutSourceCurrentTransform,
+		const FTransform& OutSourceInitialTransform) const;
+
 	void MatchChain(
 		const bool bScaleSourceChain,
 		const FChainEncoderFK& SourceChain,
@@ -80,6 +107,12 @@ private:
 		const FVector& TargetChainOrigin,
 		TArray<FTransform>& InOutGlobalPose);
 
+	void BlendByAlpha(
+		const FRetargetFKChainSettings& Settings,
+		const TArray<int32>& TargetBoneIndices,
+		const FTargetSkeleton& TargetSkeleton,
+		TArray<FTransform>& InOutGlobalPose);
+
 	static FVector GetPointOnSplineFromIndexAndAlpha(const TArray<FVector>& InSplinePoints, const int32 InPointIndex, const double InSegmentAlpha);
 	
 	static void GetPointOnSplineDistanceFromPoint(
@@ -131,12 +164,13 @@ enum class EFKChainTranslationMode : uint8
 UENUM(BlueprintType)
 enum class EFKChainRotationMode : uint8
 {
-	Interpolated		UMETA(DisplayName = "Interpolated"),
-	OneToOne			UMETA(DisplayName = "One to One"),
-	OneToOneReversed	UMETA(DisplayName = "One to One Reversed"),
-	MatchChain			UMETA(DisplayName = "Match Chain"),
-	MatchScaledChain	UMETA(DisplayName = "Match Scaled Chain"),
-	None				UMETA(DisplayName = "None"),
+	None				= 5 UMETA(DisplayName = "None"),
+	Interpolated		= 0 UMETA(DisplayName = "Interpolated"),
+	OneToOne			= 1 UMETA(DisplayName = "One to One"),
+	OneToOneReversed	= 2 UMETA(DisplayName = "One to One Reversed"),
+	MatchChain			= 3 UMETA(DisplayName = "Match Chain"),
+	MatchScaledChain	= 4 UMETA(DisplayName = "Match Scaled Chain"),
+	PlainCopy			= 6 UMETA(DisplayName = "Plain Copy"),
 };
 
 USTRUCT(BlueprintType)
