[IK Retarget] Added editor preview options.

diff --git a/Engine/Plugins/Animation/IKRig/Source/IKRig/Public/Retargeter/IKRetargeter.h b/Engine/Plugins/Animation/IKRig/Source/IKRig/Public/Retargeter/IKRetargeter.h
--- a/Engine/Plugins/Animation/IKRig/Source/IKRig/Public/Retargeter/IKRetargeter.h
+++ b/Engine/Plugins/Animation/IKRig/Source/IKRig/Public/Retargeter/IKRetargeter.h
@@ -344,17 +344,49 @@ public:
 #if WITH_EDITORONLY_DATA
 	
 	// The offset applied to the target mesh in the editor viewport. 
-	UPROPERTY(EditAnywhere, Category = PreviewSettings)
+	UPROPERTY(EditAnywhere, Category = PreviewOffset)
 	FVector TargetMeshOffset;
 
 	// Scale the target mesh in the viewport for easier visualization next to the source.
-	UPROPERTY(EditAnywhere, Category = PreviewSettings, meta = (UIMin = "0.01", UIMax = "10.0"))
+	UPROPERTY(EditAnywhere, Category = PreviewOffset, meta = (UIMin = "0.01", UIMax = "10.0"))
 	float TargetMeshScale = 1.0f;
 
 	// The offset applied to the source mesh in the editor viewport. 
-	UPROPERTY(EditAnywhere, Category = PreviewSettings)
+	UPROPERTY(EditAnywhere, Category = PreviewOffset)
 	FVector SourceMeshOffset;
 
+	// Show/hide the source skeletal mesh in the editor viewport.
+	UPROPERTY(EditAnywhere, Category = PreviewVisibility)
+	bool bShowSourceMesh = true;
+
+	// Show/hide the target skeletal mesh in the editor viewport.
+	UPROPERTY(EditAnywhere, Category = PreviewVisibility)
+	bool bShowTargetMesh = true;
+
+	// Show/hide the source skeleton in the editor viewport. Note: the viewport must be showing bones to see the skeleton.
+	UPROPERTY(EditAnywhere, Category = PreviewVisibility)
+	bool bShowSourceSkeleton = true;
+
+	// Show/hide the target skeleton in the editor viewport. Note: the viewport must be showing bones to see the skeleton.
+	UPROPERTY(EditAnywhere, Category = PreviewVisibility)
+	bool bShowTargetSkeleton = true;
+
+	// Override the source skeleton color in the editor viewport.
+	UPROPERTY(EditAnywhere, Category = SkeletonColorOverride)
+	bool bOverrideSourceSkeletonColor = false;
+
+	// Override the source skeleton color in the editor viewport.
+	UPROPERTY(EditAnywhere, Category = SkeletonColorOverride)
+	FLinearColor SourceOverideColor = FLinearColor::Blue;
+
+	// Override the target skeleton color in the editor viewport.
+	UPROPERTY(EditAnywhere, Category = SkeletonColorOverride)
+	bool bOverrideTargetSkeletonColor = false;
+
+	// Override the target skeleton color in the editor viewport.
+	UPROPERTY(EditAnywhere, Category = SkeletonColorOverride)
+	FLinearColor TargetOverideColor = FLinearColor::Red;
+
 	// When true, animation sequences with "Force Root Lock" turned On will act as though it is Off.
 	// This affects only the preview in the retarget editor. Use ExportRootLockMode to control exported animation behavior.
 	// This setting has no effect on runtime retargeting where root motion is copied from the source component.
diff --git a/Engine/Plugins/Animation/IKRig/Source/IKRigEditor/Private/RetargetEditor/IKRetargetEditorController.cpp b/Engine/Plugins/Animation/IKRig/Source/IKRigEditor/Private/RetargetEditor/IKRetargetEditorController.cpp
--- a/Engine/Plugins/Animation/IKRig/Source/IKRigEditor/Private/RetargetEditor/IKRetargetEditorController.cpp
+++ b/Engine/Plugins/Animation/IKRig/Source/IKRigEditor/Private/RetargetEditor/IKRetargetEditorController.cpp
@@ -376,6 +376,7 @@ void FIKRetargetEditorController::UpdateSkeletalMeshComponents() const
 			return;
 		}
 
+		// set component transform
 		const bool bIsSource  = SourceOrTarget == ERetargetSourceOrTarget::Source;
 		const FVector Position = bIsSource ? Asset->SourceMeshOffset : Asset->TargetMeshOffset;
 		const double SourceScaleFactor = GetRetargetProcessor()->GetSourceScaleFactor();
@@ -385,6 +386,10 @@ void FIKRetargetEditorController::UpdateSkeletalMeshComponents() const
 		constexpr ETeleportType Teleport = ETeleportType::TeleportPhysics;
 		Component->SetWorldLocation(Position, bSweep, OutSweepHitResult, Teleport);
 		Component->SetWorldScale3D(FVector(Scale,Scale,Scale));
+			
+		// toggle mesh visibility
+		const bool bIsVisible = bIsSource ? Asset->bShowSourceMesh : Asset->bShowTargetMesh;
+		Component->SetVisibility(bIsVisible);	
 	};
 
 	UpdateMeshComponent(ERetargetSourceOrTarget::Source);
@@ -1478,6 +1483,16 @@ void FIKRetargetEditorController::RenderSkeleton(FPrimitiveDrawInterface* PDI, E
 	{
 		return;
 	}
+
+	const UIKRetargeter* Asset = AssetController->GetAsset();
+
+	// toggle skeleton visibility
+	const bool bIsSource = InSourceOrTarget == ERetargetSourceOrTarget::Source;
+	const bool bIsSkeletonVisible = bIsSource ? Asset->bShowSourceSkeleton : Asset->bShowTargetSkeleton;
+	if (!bIsSkeletonVisible)
+	{
+		return;
+	}
 	
 	const FTransform ComponentTransform = MeshComponent->GetComponentTransform();
 	const FReferenceSkeleton& RefSkeleton = MeshComponent->GetReferenceSkeleton();
@@ -1494,7 +1509,7 @@ void FIKRetargetEditorController::RenderSkeleton(FPrimitiveDrawInterface* PDI, E
 		WorldTransforms[Index] = MeshComponent->GetBoneTransform(Index, ComponentTransform);
 	}
 	
-	const UIKRetargeter* Asset = AssetController->GetAsset();
+	
 	const float BoneDrawSize = Asset->BoneDrawSize;
 	const float MaxDrawRadius = static_cast<float>(MeshComponent->Bounds.SphereRadius * 0.01f);
 	const float BoneRadius = FMath::Min(1.0f, MaxDrawRadius) * BoneDrawSize;
@@ -1588,6 +1603,14 @@ void FIKRetargetEditorController::RenderSkeleton(FPrimitiveDrawInterface* PDI, E
 		}
 	}
 
+	// override skeleton color
+	const bool bOverrideSkeletonColor = bIsSource ? Asset->bOverrideSourceSkeletonColor : Asset->bOverrideTargetSkeletonColor;
+	if (bOverrideSkeletonColor)
+	{
+		const FLinearColor OverrideColor = bIsSource ? Asset->SourceOverideColor : Asset->TargetOverideColor;
+		BoneColors.Init(OverrideColor, RefSkeleton.GetNum());
+	}
+
 	SkeletalDebugRendering::DrawBones(
 		PDI,
 		ComponentTransform.GetLocation(),
