[IK Rig] Fixes ensure when deleting a solver while it's details panel is open.

diff --git a/Engine/Plugins/Animation/IKRig/Source/IKRigEditor/Private/RigEditor/SIKRigSolverStack.cpp b/Engine/Plugins/Animation/IKRig/Source/IKRigEditor/Private/RigEditor/SIKRigSolverStack.cpp
--- a/Engine/Plugins/Animation/IKRig/Source/IKRigEditor/Private/RigEditor/SIKRigSolverStack.cpp
+++ b/Engine/Plugins/Animation/IKRig/Source/IKRigEditor/Private/RigEditor/SIKRigSolverStack.cpp
@@ -408,6 +408,8 @@ void SIKRigSolverStack::DeleteSolver(TSharedPtr<FSolverStackElement> SolverToDel
 		return;
 	}
 
+	Controller->ShowAssetDetails();
+
 	const UIKRigController* AssetController = Controller->AssetController;
 	AssetController->RemoveSolver(SolverToDelete->IndexInStack);
 	RefreshStackView();
