[IK Rig] Fixing crash when trying to load null goals in retargeter.

diff --git a/Engine/Plugins/Animation/IKRig/Source/IKRig/Private/Rig/IKRigDataTypes.cpp b/Engine/Plugins/Animation/IKRig/Source/IKRig/Private/Rig/IKRigDataTypes.cpp
--- a/Engine/Plugins/Animation/IKRig/Source/IKRig/Private/Rig/IKRigDataTypes.cpp
+++ b/Engine/Plugins/Animation/IKRig/Source/IKRig/Private/Rig/IKRigDataTypes.cpp
@@ -26,6 +26,11 @@ void FIKRigGoalContainer::SetIKGoal(const FIKRigGoal& InGoal)
 
 void FIKRigGoalContainer::SetIKGoal(const UIKRigEffectorGoal* InEffectorGoal)
 {
+	if (!InEffectorGoal)
+	{
+		return;
+	}
+	
 	FIKRigGoal* Goal = FindGoalWriteable(InEffectorGoal->GoalName);
 	if (!Goal)
 	{
