[IK Retarget] Fixes serialization not setting new properties correctly.

diff --git a/Engine/Plugins/Animation/IKRig/Source/IKRig/Private/Retargeter/IKRetargeter.cpp b/Engine/Plugins/Animation/IKRig/Source/IKRig/Private/Retargeter/IKRetargeter.cpp
--- a/Engine/Plugins/Animation/IKRig/Source/IKRig/Private/Retargeter/IKRetargeter.cpp
+++ b/Engine/Plugins/Animation/IKRig/Source/IKRig/Private/Retargeter/IKRetargeter.cpp
@@ -132,6 +132,8 @@ void UIKRetargeter::PostLoad()
 	CleanRetargetPoses();
 	
 	CleanOpStack();
+
+	PostLoadOpStack();
 }
 
 #if WITH_EDITOR
@@ -734,6 +736,17 @@ void UIKRetargeter::CleanOpStack()
 	}
 	
 	ensureAlwaysMsgf(RetargetOps.Num() == TempOps.Num(), TEXT("Retarget ops were lost during cleaning."));
+}
+
+void UIKRetargeter::PostLoadOpStack()
+{
+	// this is required because TStructOpsTypeTraits PostSerialize() is not called if the struct data is at default values
+	const FIKRigObjectVersion::Type CustomVersion = static_cast<FIKRigObjectVersion::Type>(GetLinkerCustomVersion(FIKRigObjectVersion::GUID));
+	for (FInstancedStruct& OpStruct : RetargetOps)
+	{
+		FIKRetargetOpBase& Op = OpStruct.GetMutable<FIKRetargetOpBase>();
+		Op.PostLoad(CustomVersion);
+	}
 };
 
 #if WITH_EDITORONLY_DATA
diff --git a/Engine/Plugins/Animation/IKRig/Source/IKRig/Private/Retargeter/RetargetOps/CopyBasePoseOp.cpp b/Engine/Plugins/Animation/IKRig/Source/IKRig/Private/Retargeter/RetargetOps/CopyBasePoseOp.cpp
--- a/Engine/Plugins/Animation/IKRig/Source/IKRig/Private/Retargeter/RetargetOps/CopyBasePoseOp.cpp
+++ b/Engine/Plugins/Animation/IKRig/Source/IKRig/Private/Retargeter/RetargetOps/CopyBasePoseOp.cpp
@@ -26,23 +26,6 @@ void FIKRetargetCopyBasePoseOpSettings::CopySettingsAtRuntime(const FIKRetargetO
 		PropertiesToIgnore);
 }
 
-void FIKRetargetCopyBasePoseOpSettings::PostSerialize(const FArchive& Ar)
-{
-#if WITH_EDITORONLY_DATA
-	
-	if (Ar.IsLoading())
-	{
-		if (Ar.CustomVer(FIKRigObjectVersion::GUID) < FIKRigObjectVersion::CopyBasePoseConvertedToBoneReference)
-		{
-			PRAGMA_DISABLE_DEPRECATION_WARNINGS
-			CopyFromStart.BoneName = CopyBasePoseRoot_DEPRECATED;
-			PRAGMA_ENABLE_DEPRECATION_WARNINGS
-		}
-	}
-	
-#endif // WITH_EDITORONLY_DATA
-}
-
 bool FIKRetargetCopyBasePoseOp::Initialize(
 	const FIKRetargetProcessor& InProcessor,
 	const FRetargetSkeleton& InSourceSkeleton,
@@ -200,6 +183,17 @@ const UScriptStruct* FIKRetargetCopyBasePoseOp::GetType() const
 	return FIKRetargetCopyBasePoseOp::StaticStruct();
 }
 
+void FIKRetargetCopyBasePoseOp::PostLoad(const FIKRigObjectVersion::Type InVersion)
+{
+	// custom serialization to handle conversion from deprecated CopyBasePoseRoot to CopyFromStart
+	if (InVersion < FIKRigObjectVersion::CopyBasePoseConvertedToBoneReference)
+	{
+		PRAGMA_DISABLE_DEPRECATION_WARNINGS
+		Settings.CopyFromStart.BoneName = Settings.CopyBasePoseRoot_DEPRECATED;
+		PRAGMA_ENABLE_DEPRECATION_WARNINGS
+	}
+}
+
 FIKRetargetCopyBasePoseOpSettings UIKRetargetCopyBasePoseController::GetSettings()
 {
 	return *reinterpret_cast<FIKRetargetCopyBasePoseOpSettings*>(OpSettingsToControl);
diff --git a/Engine/Plugins/Animation/IKRig/Source/IKRig/Private/Retargeter/RetargetOps/IKChainsOp.cpp b/Engine/Plugins/Animation/IKRig/Source/IKRig/Private/Retargeter/RetargetOps/IKChainsOp.cpp
--- a/Engine/Plugins/Animation/IKRig/Source/IKRig/Private/Retargeter/RetargetOps/IKChainsOp.cpp
+++ b/Engine/Plugins/Animation/IKRig/Source/IKRig/Private/Retargeter/RetargetOps/IKChainsOp.cpp
@@ -202,17 +202,6 @@ bool FRetargetIKChainSettings::operator==(const FRetargetIKChainSettings& Other)
 		&& FMath::IsNearlyEqualByULP(Extension, Other.Extension);
 }
 
-void FRetargetIKChainSettings::PostSerialize(const FArchive& Ar)
-{
-	if (Ar.IsLoading())
-	{
-		if (Ar.CustomVer(FIKRigObjectVersion::GUID) < FIKRigObjectVersion::RemovedPelvisMotionFromSourceGoals)
-		{
-			ApplyPelvisOffsetToSourceGoals = true;
-		}
-	}
-}
-
 const UClass* FIKRetargetIKChainsOpSettings::GetControllerType() const
 {
 	return UIKRetargetIKChainsController::StaticClass();
@@ -465,6 +454,17 @@ void FIKRetargetIKChainsOp::OnParentReinitPropertyEdited(
 	RegenerateChainSettings(&InParentOp);
 }
 
+void FIKRetargetIKChainsOp::PostLoad(const FIKRigObjectVersion::Type InVersion)
+{
+	if (InVersion < FIKRigObjectVersion::RemovedPelvisMotionFromSourceGoals)
+	{
+		for (FRetargetIKChainSettings& ChainSettings : Settings.ChainsToRetarget)
+		{
+			ChainSettings.ApplyPelvisOffsetToSourceGoals = true;
+		}
+	}
+}
+
 #if WITH_EDITOR
 
 FCriticalSection FIKRetargetIKChainsOp::DebugDataMutex;
diff --git a/Engine/Plugins/Animation/IKRig/Source/IKRig/Private/Retargeter/RetargetOps/PinBoneOp.cpp b/Engine/Plugins/Animation/IKRig/Source/IKRig/Private/Retargeter/RetargetOps/PinBoneOp.cpp
--- a/Engine/Plugins/Animation/IKRig/Source/IKRig/Private/Retargeter/RetargetOps/PinBoneOp.cpp
+++ b/Engine/Plugins/Animation/IKRig/Source/IKRig/Private/Retargeter/RetargetOps/PinBoneOp.cpp
@@ -11,24 +11,6 @@
 #define LOCTEXT_NAMESPACE "PinBoneOp"
 
 
-void FPinBoneData::PostSerialize(const FArchive& Ar)
-{
-#if WITH_EDITORONLY_DATA
-	
-	if (Ar.IsLoading())
-	{
-		if (Ar.CustomVer(FIKRigObjectVersion::GUID) < FIKRigObjectVersion::PinBoneTypeAndOffsetsUpgraded)
-		{
-			PRAGMA_DISABLE_DEPRECATION_WARNINGS
-			BoneToCopyTo.BoneName = BoneToPin_DEPRECATED;
-			BoneToCopyFrom.BoneName = BoneToPinTo_DEPRECATED;
-			PRAGMA_ENABLE_DEPRECATION_WARNINGS
-		}
-	}
-	
-#endif // WITH_EDITORONLY_DATA
-}
-
 void FPinBoneData::CachePinData(ERetargetSourceOrTarget InSkeletonToCopyFrom, const FIKRetargetProcessor& Processor)
 {
 	SourceScaleFactor = Processor.GetSourceScaleFactor();
@@ -214,6 +196,21 @@ const UScriptStruct* FIKRetargetPinBoneOp::GetType() const
 	return FIKRetargetPinBoneOp::StaticStruct();
 }
 
+void FIKRetargetPinBoneOp::PostLoad(const FIKRigObjectVersion::Type InVersion)
+{
+	// patch old versions to use FBoneReference instead of FName
+	if (InVersion < FIKRigObjectVersion::PinBoneTypeAndOffsetsUpgraded)
+	{
+		PRAGMA_DISABLE_DEPRECATION_WARNINGS
+		for (FPinBoneData& PinBoneData : Settings.BonesToPin)
+		{
+			PinBoneData.BoneToCopyTo.BoneName = PinBoneData.BoneToPin_DEPRECATED;
+			PinBoneData.BoneToCopyFrom.BoneName = PinBoneData.BoneToPinTo_DEPRECATED;
+		}
+		PRAGMA_ENABLE_DEPRECATION_WARNINGS
+	}
+}
+
 #if WITH_EDITOR
 FText FIKRetargetPinBoneOp::GetWarningMessage() const
 {
diff --git a/Engine/Plugins/Animation/IKRig/Source/IKRig/Private/Retargeter/RetargetOps/RootMotionGeneratorOp.cpp b/Engine/Plugins/Animation/IKRig/Source/IKRig/Private/Retargeter/RetargetOps/RootMotionGeneratorOp.cpp
--- a/Engine/Plugins/Animation/IKRig/Source/IKRig/Private/Retargeter/RetargetOps/RootMotionGeneratorOp.cpp
+++ b/Engine/Plugins/Animation/IKRig/Source/IKRig/Private/Retargeter/RetargetOps/RootMotionGeneratorOp.cpp
@@ -56,25 +56,6 @@ USkeleton* FIKRetargetRootMotionOpSettings::GetSkeleton(const FName InPropertyNa
 }
 #endif
 
-void FIKRetargetRootMotionOpSettings::PostSerialize(const FArchive& Ar)
-{
-#if WITH_EDITORONLY_DATA
-	
-	if (Ar.IsLoading())
-	{
-		if (Ar.CustomVer(FIKRigObjectVersion::GUID) < FIKRigObjectVersion::ModularRetargeterOps)
-		{
-			PRAGMA_DISABLE_DEPRECATION_WARNINGS
-			SourceRoot.BoneName = SourceRootBone_DEPRECATED;
-			TargetRoot.BoneName = TargetRootBone_DEPRECATED;
-			TargetPelvis.BoneName = TargetPelvisBone_DEPRECATED;
-			PRAGMA_ENABLE_DEPRECATION_WARNINGS
-		}
-	}
-	
-#endif // WITH_EDITORONLY_DATA
-}
-
 bool FIKRetargetRootMotionOp::Initialize(
 	const FIKRetargetProcessor& Processor,
 	const FRetargetSkeleton& SourceSkeleton,
@@ -231,6 +212,19 @@ const UScriptStruct* FIKRetargetRootMotionOp::GetType() const
 	return FIKRetargetRootMotionOp::StaticStruct();
 }
 
+void FIKRetargetRootMotionOp::PostLoad(const FIKRigObjectVersion::Type InVersion)
+{
+	// load old FNames into new FBoneReference
+	if (InVersion < FIKRigObjectVersion::ModularRetargeterOps)
+	{
+		PRAGMA_DISABLE_DEPRECATION_WARNINGS
+		Settings.SourceRoot.BoneName = Settings.SourceRootBone_DEPRECATED;
+		Settings.TargetRoot.BoneName = Settings.TargetRootBone_DEPRECATED;
+		Settings.TargetPelvis.BoneName = Settings.TargetPelvisBone_DEPRECATED;
+		PRAGMA_ENABLE_DEPRECATION_WARNINGS
+	}
+}
+
 void FIKRetargetRootMotionOp::Reset()
 {
 	SourceRootIndex = INDEX_NONE;
diff --git a/Engine/Plugins/Animation/IKRig/Source/IKRig/Public/Retargeter/IKRetargetOps.h b/Engine/Plugins/Animation/IKRig/Source/IKRig/Public/Retargeter/IKRetargetOps.h
--- a/Engine/Plugins/Animation/IKRig/Source/IKRig/Public/Retargeter/IKRetargetOps.h
+++ b/Engine/Plugins/Animation/IKRig/Source/IKRig/Public/Retargeter/IKRetargetOps.h
@@ -1,6 +1,7 @@
 // Copyright Epic Games, Inc. All Rights Reserved.
 
 #pragma once
+#include "IKRigObjectVersion.h"
 
 #include "IKRetargetOps.generated.h"
 
@@ -122,6 +123,9 @@ public:
 		const FRetargetSkeleton& InSourceSkeleton,
 		const FTargetSkeleton& InTargetSkeleton,
 		FIKRigLogger& InLog) {};
+
+	/** (optional) called during the PostLoad of the owning asset. Can be used to apply patches to the data for backwards compatibility. */
+	virtual void PostLoad(const FIKRigObjectVersion::Type InVersion) {};
 	
 	/** (optional) override to automate initial setup after being added to the stack */
 	virtual void OnAddedToStack(const UIKRetargeter* InRetargetAsset, const FIKRetargetOpBase* InParentOp) {};
diff --git a/Engine/Plugins/Animation/IKRig/Source/IKRig/Public/Retargeter/IKRetargeter.h b/Engine/Plugins/Animation/IKRig/Source/IKRig/Public/Retargeter/IKRetargeter.h
--- a/Engine/Plugins/Animation/IKRig/Source/IKRig/Public/Retargeter/IKRetargeter.h
+++ b/Engine/Plugins/Animation/IKRig/Source/IKRig/Public/Retargeter/IKRetargeter.h
@@ -310,6 +310,7 @@ private:
 #endif
 	UE_API void CleanRetargetPoses();
 	UE_API void CleanOpStack();
+	UE_API void PostLoadOpStack();
 	
 	// make the provided op name valid and unique
 	// NOTE: if InIndexOfOp is INDEX_None, then it assumes this is a new op
diff --git a/Engine/Plugins/Animation/IKRig/Source/IKRig/Public/Retargeter/RetargetOps/CopyBasePoseOp.h b/Engine/Plugins/Animation/IKRig/Source/IKRig/Public/Retargeter/RetargetOps/CopyBasePoseOp.h
--- a/Engine/Plugins/Animation/IKRig/Source/IKRig/Public/Retargeter/RetargetOps/CopyBasePoseOp.h
+++ b/Engine/Plugins/Animation/IKRig/Source/IKRig/Public/Retargeter/RetargetOps/CopyBasePoseOp.h
@@ -28,26 +28,13 @@ struct FIKRetargetCopyBasePoseOpSettings : public FIKRetargetOpSettingsBase
 	UPROPERTY(EditAnywhere, Category = RetargetPhases, meta = (ReinitializeOnEdit, EditCondition = "bCopyBasePose"))
 	TArray<FBoneReference> BonesToExclude;
 
-	UE_API virtual const UClass* GetControllerType() const override;
-
-	UE_API virtual void CopySettingsAtRuntime(const FIKRetargetOpSettingsBase* InSettingsToCopyFrom) override;
-
-	/** Custom serialization to handle conversion from deprecated CopyBasePoseRoot to CopyFromStart */
-	void PostSerialize(const FArchive& Ar);
-
-private:
-	/** converted on PoseSerialize */
+	/** converted on PostLoad */
 	UPROPERTY(meta = (DeprecatedProperty))
 	FName CopyBasePoseRoot_DEPRECATED = NAME_None;
-};
 
-// Custom serialization to load deprecated settings
-template<> struct TStructOpsTypeTraits<FIKRetargetCopyBasePoseOpSettings> : public TStructOpsTypeTraitsBase2<FIKRetargetCopyBasePoseOpSettings>
-{
-	enum
-	{
-		WithPostSerialize = true,
-	};
+	UE_API virtual const UClass* GetControllerType() const override;
+
+	UE_API virtual void CopySettingsAtRuntime(const FIKRetargetOpSettingsBase* InSettingsToCopyFrom) override;
 };
 
 USTRUCT(BlueprintType, meta = (DisplayName = "Copy Base Pose"))
@@ -79,6 +66,8 @@ struct FIKRetargetCopyBasePoseOp : public FIKRetargetOpBase
 
 	virtual bool IsSingleton() const override { return true; };
 
+	UE_API virtual void PostLoad(const FIKRigObjectVersion::Type InVersion) override;
+
 	UPROPERTY()
 	FIKRetargetCopyBasePoseOpSettings Settings;
 
diff --git a/Engine/Plugins/Animation/IKRig/Source/IKRig/Public/Retargeter/RetargetOps/IKChainsOp.h b/Engine/Plugins/Animation/IKRig/Source/IKRig/Public/Retargeter/RetargetOps/IKChainsOp.h
--- a/Engine/Plugins/Animation/IKRig/Source/IKRig/Public/Retargeter/RetargetOps/IKChainsOp.h
+++ b/Engine/Plugins/Animation/IKRig/Source/IKRig/Public/Retargeter/RetargetOps/IKChainsOp.h
@@ -154,18 +154,6 @@ struct FRetargetIKChainSettings
 	double Extension = 1.0f;
 
 	UE_API bool operator==(const FRetargetIKChainSettings& Other) const;
-	
-	/** Custom serialization to handle loading old assets */
-    void PostSerialize(const FArchive& Ar);
-};
-
-// Custom serialization to patch old settings
-template<> struct TStructOpsTypeTraits<FRetargetIKChainSettings> : public TStructOpsTypeTraitsBase2<FRetargetIKChainSettings>
-{
-	enum
-	{
-		WithPostSerialize = true,
-	};
 };
 
 USTRUCT(BlueprintType, meta = (DisplayName = "Retarget IK Goals Settings"))
@@ -244,6 +232,8 @@ struct FIKRetargetIKChainsOp : public FIKRetargetOpBase
 
 	UE_API virtual void OnParentReinitPropertyEdited(const FIKRetargetOpBase& InParentOp, const FPropertyChangedEvent* InPropertyChangedEvent) override;
 
+	UE_API virtual void PostLoad(const FIKRigObjectVersion::Type InVersion) override;
+
 	UPROPERTY()
 	FIKRetargetIKChainsOpSettings Settings;
 
diff --git a/Engine/Plugins/Animation/IKRig/Source/IKRig/Public/Retargeter/RetargetOps/PinBoneOp.h b/Engine/Plugins/Animation/IKRig/Source/IKRig/Public/Retargeter/RetargetOps/PinBoneOp.h
--- a/Engine/Plugins/Animation/IKRig/Source/IKRig/Public/Retargeter/RetargetOps/PinBoneOp.h
+++ b/Engine/Plugins/Animation/IKRig/Source/IKRig/Public/Retargeter/RetargetOps/PinBoneOp.h
@@ -24,8 +24,6 @@ USTRUCT(BlueprintType)
 struct FPinBoneData
 {
 	GENERATED_BODY()
-
-	void PostSerialize(const FArchive& Ar);
 	
 	void CachePinData(ERetargetSourceOrTarget InSkeletonToCopyFrom, const FIKRetargetProcessor& Processor);
 
@@ -52,15 +50,6 @@ struct FPinBoneData
 	FName BoneToPinTo_DEPRECATED;
 };
 
-template<>
-struct TStructOpsTypeTraits<FPinBoneData> : public TStructOpsTypeTraitsBase2<FPinBoneData>
-{
-	enum
-	{
-		WithPostSerialize = true,
-	};
-};
-
 UENUM()
 enum class EPinBoneTranslationMode : uint8
 {
@@ -163,6 +152,8 @@ struct FIKRetargetPinBoneOp : public FIKRetargetOpBase
 	
 	UE_API virtual const UScriptStruct* GetType() const override;
 
+	UE_API virtual void PostLoad(const FIKRigObjectVersion::Type InVersion) override;
+
 	UPROPERTY()
 	FIKRetargetPinBoneOpSettings Settings;
 	
diff --git a/Engine/Plugins/Animation/IKRig/Source/IKRig/Public/Retargeter/RetargetOps/RootMotionGeneratorOp.h b/Engine/Plugins/Animation/IKRig/Source/IKRig/Public/Retargeter/RetargetOps/RootMotionGeneratorOp.h
--- a/Engine/Plugins/Animation/IKRig/Source/IKRig/Public/Retargeter/RetargetOps/RootMotionGeneratorOp.h
+++ b/Engine/Plugins/Animation/IKRig/Source/IKRig/Public/Retargeter/RetargetOps/RootMotionGeneratorOp.h
@@ -84,10 +84,7 @@ struct FIKRetargetRootMotionOpSettings : public FIKRetargetOpSettingsBase
 #if WITH_EDITOR
 	UE_API virtual USkeleton* GetSkeleton(const FName InPropertyName) override;
 #endif
-
-	// patch to load old bone names into new FBoneReference
-	UE_API void PostSerialize(const FArchive& Ar);
-#if WITH_EDITORONLY_DATA
+	
 	// Changed to FBoneReference to support bone selector UI
 	UPROPERTY(meta=(DeprecatedProperty))
 	FName SourceRootBone_DEPRECATED;
@@ -95,19 +92,8 @@ struct FIKRetargetRootMotionOpSettings : public FIKRetargetOpSettingsBase
 	FName TargetRootBone_DEPRECATED;
 	UPROPERTY(meta=(DeprecatedProperty))
 	FName TargetPelvisBone_DEPRECATED;
-#endif
-};
-
-template<>
-struct TStructOpsTypeTraits<FIKRetargetRootMotionOpSettings> : public TStructOpsTypeTraitsBase2<FIKRetargetRootMotionOpSettings>
-{
-	enum
-	{
-		WithPostSerialize = true,
-	};
 };
 
-
 USTRUCT(BlueprintType, meta = (DisplayName = "Root Motion"))
 struct FIKRetargetRootMotionOp : public FIKRetargetOpBase
 {
@@ -140,6 +126,8 @@ struct FIKRetargetRootMotionOp : public FIKRetargetOpBase
 	
 	UE_API virtual const UScriptStruct* GetType() const override;
 
+	UE_API virtual void PostLoad(const FIKRigObjectVersion::Type InVersion) override;
+
 	UPROPERTY()
 	FIKRetargetRootMotionOpSettings Settings;
 
@@ -161,9 +149,9 @@ private:
 		const FIKRetargetProcessor& Processor,
 		const TArray<FTransform>& InSourceGlobalPose) const;
 	
-	int32 SourceRootIndex;
-	int32 TargetRootIndex;
-	int32 TargetPelvisIndex;
+	int32 SourceRootIndex = INDEX_NONE;
+	int32 TargetRootIndex = INDEX_NONE;
+	int32 TargetPelvisIndex = INDEX_NONE;
 
 	FTransform TargetPelvisRelativeToTargetRootRefPose;
 	FTransform TargetPelvisInRefPose;
