[IK Retarget] Adding runtime API access to FK chain mapping.

diff --git a/Engine/Plugins/Animation/IKRig/Source/IKRig/Private/Retargeter/IKRetargeter.cpp b/Engine/Plugins/Animation/IKRig/Source/IKRig/Private/Retargeter/IKRetargeter.cpp
--- a/Engine/Plugins/Animation/IKRig/Source/IKRig/Private/Retargeter/IKRetargeter.cpp
+++ b/Engine/Plugins/Animation/IKRig/Source/IKRig/Private/Retargeter/IKRetargeter.cpp
@@ -470,7 +470,7 @@ void UIKRetargeter::PostLoadPutChainMappingInOps()
 	for (FIKRetargetFKChainsOp* Op : FKChainOps)
 	{
 		Op->Settings.IKRigAsset = TargetIKRigAsset;
-		Op->ChainMapping = ChainMap_DEPRECATED;
+		Op->Settings.ChainMapping = ChainMap_DEPRECATED;
 	}
 
 	TArray<FIKRetargetAlignPoleVectorOp*> PoleVectorOps = GetAllRetargetOpsOfType<FIKRetargetAlignPoleVectorOp>();
diff --git a/Engine/Plugins/Animation/IKRig/Source/IKRig/Private/Retargeter/RetargetOps/FKChainsOp.cpp b/Engine/Plugins/Animation/IKRig/Source/IKRig/Private/Retargeter/RetargetOps/FKChainsOp.cpp
--- a/Engine/Plugins/Animation/IKRig/Source/IKRig/Private/Retargeter/RetargetOps/FKChainsOp.cpp
+++ b/Engine/Plugins/Animation/IKRig/Source/IKRig/Private/Retargeter/RetargetOps/FKChainsOp.cpp
@@ -2,6 +2,7 @@
 
 #include "Retargeter/RetargetOps/FKChainsOp.h"
 
+#include "IKRigObjectVersion.h"
 #include "Retargeter/RetargetOps/PelvisMotionOp.h"
 
 #include UE_INLINE_GENERATED_CPP_BY_NAME(FKChainsOp)
@@ -802,7 +803,7 @@ bool FIKRetargetFKChainsOp::Initialize(
 		}
 
 		// which source chain was this target chain mapped to?
-		const FName SourceChainName = ChainMapping.GetChainMappedTo(ChainSettings.TargetChainName, ERetargetSourceOrTarget::Target);
+		const FName SourceChainName = Settings.ChainMapping.GetChainMappedTo(ChainSettings.TargetChainName, ERetargetSourceOrTarget::Target);
 		const FResolvedBoneChain* SourceBoneChain = BoneChains.GetResolvedBoneChainByName(SourceChainName, ERetargetSourceOrTarget::Source);
 		if (SourceBoneChain == nullptr)
 		{
@@ -970,7 +971,7 @@ void FIKRetargetFKChainsOp::OnAddedToStack(const UIKRetargeter* InRetargetAsset,
 
 	// auto map
 	static bool bForceRemap = true;
-	ChainMapping.AutoMapChains(EAutoMapChainType::Fuzzy, bForceRemap);
+	Settings.ChainMapping.AutoMapChains(EAutoMapChainType::Fuzzy, bForceRemap);
 }
 
 void FIKRetargetFKChainsOp::OnAssignIKRig(
@@ -980,14 +981,14 @@ void FIKRetargetFKChainsOp::OnAssignIKRig(
 {
 	// get the newly assigned IK Rig
 	const bool bAssignedSource = SourceOrTarget == ERetargetSourceOrTarget::Source;
-	const UIKRigDefinition* SourceIKRig = bAssignedSource ? InIKRig : ChainMapping.GetIKRig(ERetargetSourceOrTarget::Source);
+	const UIKRigDefinition* SourceIKRig = bAssignedSource ? InIKRig : Settings.ChainMapping.GetIKRig(ERetargetSourceOrTarget::Source);
 	const bool bAssignedTarget = SourceOrTarget == ERetargetSourceOrTarget::Target;
-	const UIKRigDefinition* TargetIKRig = bAssignedTarget ? InIKRig : ChainMapping.GetIKRig(ERetargetSourceOrTarget::Target);
+	const UIKRigDefinition* TargetIKRig = bAssignedTarget ? InIKRig : Settings.ChainMapping.GetIKRig(ERetargetSourceOrTarget::Target);
 	ApplyIKRigs(SourceIKRig, TargetIKRig);
 
 	// auto map, but don't force it (leave existing mappings alone)
 	static bool bForceRemap = false;
-	ChainMapping.AutoMapChains(EAutoMapChainType::Fuzzy, bForceRemap);
+	Settings.ChainMapping.AutoMapChains(EAutoMapChainType::Fuzzy, bForceRemap);
 }
 
 void FIKRetargetFKChainsOp::CollectRetargetedBones(TSet<int32>& OutRetargetedBones) const
@@ -1006,7 +1007,7 @@ const UIKRigDefinition* FIKRetargetFKChainsOp::GetCustomTargetIKRig() const
 
 FRetargetChainMapping* FIKRetargetFKChainsOp::GetChainMapping()
 {
-	return &ChainMapping;
+	return &Settings.ChainMapping;
 }
 
 void FIKRetargetFKChainsOp::OnTargetChainRenamed(const FName InOldChainName, const FName InNewChainName)
@@ -1022,7 +1023,7 @@ void FIKRetargetFKChainsOp::OnTargetChainRenamed(const FName InOldChainName, con
 
 void FIKRetargetFKChainsOp::OnReinitPropertyEdited(const FPropertyChangedEvent* InPropertyChangedEvent)
 {
-	const UIKRigDefinition* SourceIKRig = ChainMapping.GetIKRig(ERetargetSourceOrTarget::Source);
+	const UIKRigDefinition* SourceIKRig = Settings.ChainMapping.GetIKRig(ERetargetSourceOrTarget::Source);
 	const UIKRigDefinition* TargetIKRig = Settings.IKRigAsset;
 	ApplyIKRigs(SourceIKRig, TargetIKRig);
 }
@@ -1042,6 +1043,16 @@ const UScriptStruct* FIKRetargetFKChainsOp::GetType() const
 	return FIKRetargetFKChainsOp::StaticStruct();
 }
 
+void FIKRetargetFKChainsOp::PostLoad(const FIKRigObjectVersion::Type InVersion)
+{
+	if (InVersion < FIKRigObjectVersion::MovedChainMappingToFKChainOpSettings)
+	{
+		PRAGMA_DISABLE_DEPRECATION_WARNINGS
+		Settings.ChainMapping = ChainMapping_DEPRECATED;
+		PRAGMA_ENABLE_DEPRECATION_WARNINGS
+	}
+}
+
 #if WITH_EDITOR
 
 FCriticalSection FIKRetargetFKChainsOp::DebugDataMutex;
@@ -1149,10 +1160,10 @@ void FIKRetargetFKChainsOp::ApplyIKRigs(const UIKRigDefinition* InSourceIKRig, c
 	Settings.IKRigAsset = InTargetIKRig;
 	
 	// update chain mapping
-	ChainMapping.ReinitializeWithIKRigs(InSourceIKRig, InTargetIKRig);
+	Settings.ChainMapping.ReinitializeWithIKRigs(InSourceIKRig, InTargetIKRig);
 
 	// update settings only if we have a valid mapping
-	if (!ChainMapping.IsReady())
+	if (!Settings.ChainMapping.IsReady())
 	{
 		// don't remove settings, instead we want to preserve existing settings at least until the next valid rig is loaded
 		return;
diff --git a/Engine/Plugins/Animation/IKRig/Source/IKRig/Public/IKRigObjectVersion.h b/Engine/Plugins/Animation/IKRig/Source/IKRig/Public/IKRigObjectVersion.h
--- a/Engine/Plugins/Animation/IKRig/Source/IKRig/Public/IKRigObjectVersion.h
+++ b/Engine/Plugins/Animation/IKRig/Source/IKRig/Public/IKRigObjectVersion.h
@@ -41,6 +41,9 @@ struct FIKRigObjectVersion
 		// changed default behavior of BlendToSource on IK Goals to not be affected by Pelvis motion
 		RemovedPelvisMotionFromSourceGoals,
 
+		// moved chain mapping in FKChainsOp from Op to Op-settings for read-only API access
+		MovedChainMappingToFKChainOpSettings,
+
 		// -----<new versions can be added above this line>-------------------------------------------------
 		VersionPlusOne,
 		LatestVersion = VersionPlusOne - 1
diff --git a/Engine/Plugins/Animation/IKRig/Source/IKRig/Public/Retargeter/RetargetOps/FKChainsOp.h b/Engine/Plugins/Animation/IKRig/Source/IKRig/Public/Retargeter/RetargetOps/FKChainsOp.h
--- a/Engine/Plugins/Animation/IKRig/Source/IKRig/Public/Retargeter/RetargetOps/FKChainsOp.h
+++ b/Engine/Plugins/Animation/IKRig/Source/IKRig/Public/Retargeter/RetargetOps/FKChainsOp.h
@@ -254,9 +254,16 @@ struct FIKRetargetFKChainsOpSettings : public FIKRetargetOpSettingsBase
 	UPROPERTY(EditAnywhere, Category = Debug)
 	double ChainDrawSize = 1.0f;
 
+	/* This op maintains its own chain mapping table to allow per-op mapping.
+	 * NOTE: this is only editable through a custom UI in the details panel or the C++ API. */
+	UPROPERTY()
+	FRetargetChainMapping ChainMapping;
+
 	IKRIG_API virtual const UClass* GetControllerType() const override;
 
 	IKRIG_API virtual void CopySettingsAtRuntime(const FIKRetargetOpSettingsBase* InSettingsToCopyFrom) override;
+
+	friend FIKRetargetFKChainsOp;
 };
 
 #if WITH_EDITOR
@@ -312,12 +319,10 @@ struct FIKRetargetFKChainsOp : public FIKRetargetOpBase
 
 	IKRIG_API virtual void OnReinitPropertyEdited(const FPropertyChangedEvent* InPropertyChangedEvent) override;
 
-	UPROPERTY()
-	FIKRetargetFKChainsOpSettings Settings;
+	IKRIG_API virtual void PostLoad(const FIKRigObjectVersion::Type InVersion) override;
 
-	/* This op maintains its own chain mapping table to allow per-op mapping */
 	UPROPERTY()
-	FRetargetChainMapping ChainMapping;
+	FIKRetargetFKChainsOpSettings Settings;
 
 #if WITH_EDITOR
 public:
@@ -347,6 +352,10 @@ private:
 	/** Children of the FK chains that need updated after the chains are retargeted */
 	TArray<int32> NonRetargetedChildrenToUpdate;
 	TArray<FTransform> ChildrenToUpdateLocalTransforms;
+	
+	/* Deprecated storage for chain mapping (now on settings for API access)*/
+	UPROPERTY(meta = (DeprecatedProperty))
+	FRetargetChainMapping ChainMapping_DEPRECATED;
 };
 
 /* The blueprint/python API for editing a FK Chains Op */
diff --git a/Engine/Plugins/Animation/IKRig/Source/IKRigEditor/Private/RetargetEditor/SRetargetAnimAssetsWindow.cpp b/Engine/Plugins/Animation/IKRig/Source/IKRigEditor/Private/RetargetEditor/SRetargetAnimAssetsWindow.cpp
--- a/Engine/Plugins/Animation/IKRig/Source/IKRigEditor/Private/RetargetEditor/SRetargetAnimAssetsWindow.cpp
+++ b/Engine/Plugins/Animation/IKRig/Source/IKRigEditor/Private/RetargetEditor/SRetargetAnimAssetsWindow.cpp
@@ -681,7 +681,7 @@ void FProceduralRetargetAssets::AutoGenerateIKRetargetAsset()
 	// use only exact chain mapping in the auto-retarget (default is fuzzy which can map toes to fingers and other chaos)
 	FIKRetargetFKChainsOp* FKChainOp = RetargetController->GetFirstRetargetOpOfType<FIKRetargetFKChainsOp>();
 	constexpr bool bForceRemap = true;
-	FKChainOp->ChainMapping.AutoMapChains(EAutoMapChainType::Exact, bForceRemap);
+	FKChainOp->Settings.ChainMapping.AutoMapChains(EAutoMapChainType::Exact, bForceRemap);
 
 	// reset and regenerate the target retarget pose
 	const FName TargetRetargetPoseName = RetargetController->GetCurrentRetargetPoseName(ERetargetSourceOrTarget::Target);
