[IK Retarget] Added "Copy Local" rotation mode

diff --git a/Engine/Plugins/Animation/IKRig/Source/IKRig/Private/Retargeter/RetargetOps/FKChainsOp.cpp b/Engine/Plugins/Animation/IKRig/Source/IKRig/Private/Retargeter/RetargetOps/FKChainsOp.cpp
--- a/Engine/Plugins/Animation/IKRig/Source/IKRig/Private/Retargeter/RetargetOps/FKChainsOp.cpp
+++ b/Engine/Plugins/Animation/IKRig/Source/IKRig/Private/Retargeter/RetargetOps/FKChainsOp.cpp
@@ -22,7 +22,7 @@ void FChainEncoderFK::Initialize(const FResolvedBoneChain* InBoneChain)
 	check(InBoneChain);
 	BoneChain = InBoneChain;
 	CurrentGlobalTransforms = BoneChain->RefPoseGlobalTransforms;
-	CurrentLocalTransforms = BoneChain->RefPoseLocalTransforms;
+	InBoneChain->CurrentLocalTransforms = BoneChain->RefPoseLocalTransforms;
 }
 
 void FChainEncoderFK::EncodePose(
@@ -38,9 +38,8 @@ void FChainEncoderFK::EncodePose(
 		const int32 BoneIndex = SourceBoneIndices[ChainIndex];
 		CurrentGlobalTransforms[ChainIndex] = InSourceGlobalPose[BoneIndex];
 	}
-
-	CurrentLocalTransforms.SetNum(SourceBoneIndices.Num());
-	FResolvedBoneChain::FillTransformsWithLocalSpaceOfChain(SourceSkeleton, InSourceGlobalPose, SourceBoneIndices, CurrentLocalTransforms);
+	
+	FResolvedBoneChain::FillTransformsWithLocalSpaceOfChain(SourceSkeleton, InSourceGlobalPose, SourceBoneIndices, BoneChain->CurrentLocalTransforms);
 	
 	ChainParentCurrentGlobalTransform = BoneChain->ChainParentBoneIndex != INDEX_NONE ? InSourceGlobalPose[BoneChain->ChainParentBoneIndex] : FTransform::Identity;
 }
@@ -51,11 +50,11 @@ void FChainEncoderFK::TransformCurrentChainTransforms(const FTransform& NewParen
 	{
 		if (ChainIndex == 0)
 		{
-			CurrentGlobalTransforms[ChainIndex] = CurrentLocalTransforms[ChainIndex] * NewParentTransform;
+			CurrentGlobalTransforms[ChainIndex] = BoneChain->CurrentLocalTransforms[ChainIndex] * NewParentTransform;
 		}
 		else
 		{
-			CurrentGlobalTransforms[ChainIndex] = CurrentLocalTransforms[ChainIndex] * CurrentGlobalTransforms[ChainIndex-1];
+			CurrentGlobalTransforms[ChainIndex] = BoneChain->CurrentLocalTransforms[ChainIndex] * CurrentGlobalTransforms[ChainIndex-1];
 		}
 	}
 }
@@ -198,6 +197,8 @@ void FChainDecoderFK::DecodePose(
 			ChainIndex,
 			Settings,
 			TargetBoneIndices,
+			TargetSkeleton,
+			InOutGlobalPose,
 			SourceChain,
 			SourceCurrentTransform,
 			SourceInitialTransform);
@@ -214,8 +215,7 @@ void FChainDecoderFK::DecodePose(
 			SourceSkeleton,
 			TargetSkeleton,
 			InOutGlobalPose,
-			SourceCurrentTransform,
-			SourceInitialTransform);
+			SourceCurrentTransform);
 
 		// retarget scale
 		const FVector OutScale = RetargetScale(ChainIndex, SourceCurrentTransform, SourceInitialTransform);
@@ -242,10 +242,33 @@ FQuat FChainDecoderFK::RetargetRotation(
 	const int32 InChainIndex,
 	const FRetargetFKChainSettings& Settings,
 	const TArray<int32>& TargetBoneIndices,
+	const FTargetSkeleton& TargetSkeleton,
+	const TArray<FTransform>& InOutGlobalPose,
 	FChainEncoderFK& SourceChain,
 	FTransform& OutSourceCurrentTransform,
 	FTransform& OutSourceInitialTransform) const
 {
+	// special case, copy in local space (assumes user knows the local axes are desirable)
+	if (Settings.RotationMode == EFKChainRotationMode::CopyLocal)
+	{
+		// use the closest bone transform, cannot reliably interpolate local rotations
+		const float Param = BoneChain->Params[InChainIndex];
+		int32 ClosestBoneChainIndex = SourceChain.BoneChain->GetBoneClosestToParam(Param);
+		OutSourceCurrentTransform = SourceChain.CurrentGlobalTransforms[ClosestBoneChainIndex];
+		OutSourceInitialTransform = SourceChain.BoneChain->RefPoseGlobalTransforms[ClosestBoneChainIndex];
+
+		// apply the local rotation relative to the target parent's global pose
+		const int32 BoneIndex = TargetBoneIndices[InChainIndex];
+		FQuat ParentGlobalRotation = FQuat::Identity;
+		const int32 ParentIndex = TargetSkeleton.ParentIndices[BoneIndex];
+		if (ParentIndex != INDEX_NONE)
+		{
+			ParentGlobalRotation = InOutGlobalPose[ParentIndex].GetRotation();
+		}
+		
+		FQuat OutRotation = ParentGlobalRotation * SourceChain.BoneChain->CurrentLocalTransforms[ClosestBoneChainIndex].GetRotation();
+		return MoveTemp(OutRotation);
+	}
 	
 	// get source current / initial transforms for this bone
 	switch (Settings.RotationMode)
@@ -303,7 +326,7 @@ FQuat FChainDecoderFK::RetargetRotation(
 			OutSourceCurrentTransform = SourceChain.BoneChain->RefPoseLocalTransforms[0] * SourceChain.ChainParentCurrentGlobalTransform;
 			break;
 		}
-	case EFKChainRotationMode::PlainCopy:
+		case EFKChainRotationMode::CopyLocal:
 		default:
 			checkNoEntry();
 		break;
@@ -332,8 +355,7 @@ FVector FChainDecoderFK::RetargetPosition(
 	const FRetargetSkeleton& SourceSkeleton,
 	const FTargetSkeleton& TargetSkeleton,
 	TArray<FTransform>& InOutGlobalPose,
-	FTransform& OutSourceCurrentTransform,
-	FTransform& OutSourceInitialTransform)
+	const FTransform& InCurrentSourceGlobal)
 {
 	// the position this function generates
 	FVector OutPosition;
@@ -359,12 +381,12 @@ FVector FChainDecoderFK::RetargetPosition(
 	case EFKChainTranslationMode::GloballyScaled:
 		{
 			FVector GlobalScale = PelvisMotionOp ? PelvisMotionOp->GetGlobalScaleVector() : FVector::OneVector;
-			OutPosition = OutSourceCurrentTransform.GetTranslation() * GlobalScale;
+			OutPosition = InCurrentSourceGlobal.GetTranslation() * GlobalScale;
 			break;
 		}
 	case EFKChainTranslationMode::Absolute:
 		{
-			OutPosition = OutSourceCurrentTransform.GetTranslation();
+			OutPosition = InCurrentSourceGlobal.GetTranslation();
 			break;	
 		}
 	case EFKChainTranslationMode::StretchBoneLengthUniformly:
@@ -574,15 +596,13 @@ void FChainDecoderFK::BlendByAlpha(
 	}
 
 	// generate local space pose of chain
-	TArray<FTransform> NewLocalTransforms;
-	NewLocalTransforms.SetNum(BoneChain->RefPoseLocalTransforms.Num());
-	BoneChain->FillTransformsWithLocalSpaceOfChain(TargetSkeleton, InOutGlobalPose, TargetBoneIndices, NewLocalTransforms);
+	BoneChain->FillTransformsWithLocalSpaceOfChain(TargetSkeleton, InOutGlobalPose, TargetBoneIndices, BoneChain->CurrentLocalTransforms);
 
 	// blend each bone in chain with the retarget pose
 	for (int32 ChainIndex=0; ChainIndex<BoneChain->RefPoseLocalTransforms.Num(); ++ChainIndex)
 	{
 		// blend between current local pose and initial local pose
-		FTransform& NewLocalTransform = NewLocalTransforms[ChainIndex];
+		FTransform& NewLocalTransform = BoneChain->CurrentLocalTransforms[ChainIndex];
 		const FTransform& RefPoseLocalTransform = BoneChain->RefPoseLocalTransforms[ChainIndex];
 		NewLocalTransform.SetTranslation(FMath::Lerp(RefPoseLocalTransform.GetTranslation(), NewLocalTransform.GetTranslation(), Settings.TranslationAlpha));
 		NewLocalTransform.SetRotation(FQuat::FastLerp(RefPoseLocalTransform.GetRotation(), NewLocalTransform.GetRotation(), Settings.RotationAlpha).GetNormalized());
diff --git a/Engine/Plugins/Animation/IKRig/Source/IKRig/Public/Retargeter/IKRetargetProcessor.h b/Engine/Plugins/Animation/IKRig/Source/IKRig/Public/Retargeter/IKRetargetProcessor.h
--- a/Engine/Plugins/Animation/IKRig/Source/IKRig/Public/Retargeter/IKRetargetProcessor.h
+++ b/Engine/Plugins/Animation/IKRig/Source/IKRig/Public/Retargeter/IKRetargetProcessor.h
@@ -245,7 +245,7 @@ struct FResolvedBoneChain
 	
 	TArray<FTransform> RefPoseGlobalTransforms;
 	TArray<FTransform> RefPoseLocalTransforms;
-	TArray<FTransform> CurrentLocalTransforms;
+	mutable TArray<FTransform> CurrentLocalTransforms; // reusable scratch space when op needs it
 
 	TArray<float> Params;
 	TArray<int32> BoneIndices;
diff --git a/Engine/Plugins/Animation/IKRig/Source/IKRig/Public/Retargeter/RetargetOps/FKChainsOp.h b/Engine/Plugins/Animation/IKRig/Source/IKRig/Public/Retargeter/RetargetOps/FKChainsOp.h
--- a/Engine/Plugins/Animation/IKRig/Source/IKRig/Public/Retargeter/RetargetOps/FKChainsOp.h
+++ b/Engine/Plugins/Animation/IKRig/Source/IKRig/Public/Retargeter/RetargetOps/FKChainsOp.h
@@ -25,7 +25,6 @@ struct FChainEncoderFK
 	void TransformCurrentChainTransforms(const FTransform& NewParentTransform);
 	
 	TArray<FTransform> CurrentGlobalTransforms;
-	TArray<FTransform> CurrentLocalTransforms;
 	FTransform ChainParentCurrentGlobalTransform;
 	const FResolvedBoneChain* BoneChain;
 };
@@ -77,6 +76,8 @@ private:
 		const int32 InChainIndex,
 		const FRetargetFKChainSettings& Settings,
 		const TArray<int32>& TargetBoneIndices,
+		const FTargetSkeleton& TargetSkeleton,
+		const TArray<FTransform>& InOutGlobalPose,
 		FChainEncoderFK& SourceChain,
 		FTransform& OutSourceCurrentTransform,
 		FTransform& OutSourceInitialTransform) const;
@@ -92,8 +93,7 @@ private:
 		const FRetargetSkeleton& SourceSkeleton,
 		const FTargetSkeleton& TargetSkeleton,
 		TArray<FTransform>& InOutGlobalPose,
-		FTransform& OutSourceCurrentTransform,
-		FTransform& OutSourceInitialTransform);
+		const FTransform& InCurrentSourceGlobal);
 
 	FVector RetargetScale(
 		const int32 InChainIndex,
@@ -170,7 +170,7 @@ enum class EFKChainRotationMode : uint8
 	OneToOneReversed	= 2 UMETA(DisplayName = "One to One Reversed"),
 	MatchChain			= 3 UMETA(DisplayName = "Match Chain"),
 	MatchScaledChain	= 4 UMETA(DisplayName = "Match Scaled Chain"),
-	PlainCopy			= 6 UMETA(DisplayName = "Plain Copy"),
+	CopyLocal			= 6 UMETA(DisplayName = "Copy Local"),
 };
 
 USTRUCT(BlueprintType)
@@ -191,12 +191,14 @@ struct FRetargetFKChainSettings
 	bool EnableFK = true;
 	
 	/** Determines how rotation is copied from the source chain to the target chain. Default is Interpolated.
+	* None: The rotation of each target bone in the chain is left at the reference pose. 
 	* Interpolated: Source and target chains are normalized by their length, then each target bone rotation is generated by finding the rotation at the same normalized distance on the source chain and interpolating between the neighboring bones.
 	* One to One: Each target bone rotation is copied from the equivalent bone in the source chain, based on the order in the chain, starting at the root of the chain. If the target chain has more bones than the source, the extra bones will remain at their reference pose.
 	* One to One Reversed: Same as One-to-One, but starting from the tip of the chain.
 	* Match Chain: Uses a Spline-IK type method to orient the target joint chain to exactly match the shape of the source chain.
 	* Match Scaled Chain: Same as Match Chain, but scales the source chain to the same length as the target chain before running the match procedure.
-	* None: The rotation of each target bone in the chain is left at the reference pose. */
+	* Copy Local: The local space rotation of the closest source bone is copied to the target. No retargeting is applied. Retarget pose is ignored.
+	*/
 	UPROPERTY(EditAnywhere, BlueprintReadWrite, Category = "FK Chain Retarget")
 	EFKChainRotationMode RotationMode = EFKChainRotationMode::Interpolated;
 
