[IK Retarget] Fixed bug, Chain Mapping can be lost for very old assets due to op PostLoad overwriting the retargeter PostLoad patch.

diff --git a/Engine/Plugins/Animation/IKRig/Source/IKRig/Private/Retargeter/RetargetOps/FKChainsOp.cpp b/Engine/Plugins/Animation/IKRig/Source/IKRig/Private/Retargeter/RetargetOps/FKChainsOp.cpp
--- a/Engine/Plugins/Animation/IKRig/Source/IKRig/Private/Retargeter/RetargetOps/FKChainsOp.cpp
+++ b/Engine/Plugins/Animation/IKRig/Source/IKRig/Private/Retargeter/RetargetOps/FKChainsOp.cpp
@@ -1065,7 +1065,8 @@ const UScriptStruct* FIKRetargetFKChainsOp::GetType() const
 
 void FIKRetargetFKChainsOp::PostLoad(const FIKRigObjectVersion::Type InVersion)
 {
-	if (InVersion < FIKRigObjectVersion::MovedChainMappingToFKChainOpSettings)
+	if (InVersion >= FIKRigObjectVersion::OpsOwnChainMapping
+		&& InVersion < FIKRigObjectVersion::MovedChainMappingToFKChainOpSettings)
 	{
 		PRAGMA_DISABLE_DEPRECATION_WARNINGS
 		Settings.ChainMapping = ChainMapping_DEPRECATED;
