[IK Retarget] Bunch of minor documentation issues caught by [REDACTED_PERSON].

diff --git a/Engine/Plugins/Animation/IKRig/Source/IKRig/Private/Retargeter/RetargetOps/CopyBasePoseOp.cpp b/Engine/Plugins/Animation/IKRig/Source/IKRig/Private/Retargeter/RetargetOps/CopyBasePoseOp.cpp
--- a/Engine/Plugins/Animation/IKRig/Source/IKRig/Private/Retargeter/RetargetOps/CopyBasePoseOp.cpp
+++ b/Engine/Plugins/Animation/IKRig/Source/IKRig/Private/Retargeter/RetargetOps/CopyBasePoseOp.cpp
@@ -165,7 +165,7 @@ void FIKRetargetCopyBasePoseOp::Run(
 	// (this local pose is used when updating intermediate bones in subsequent retargeting ops)
 	const int32 TargetRootBoneToCopyIndex = TargetSkeleton.FindBoneIndexByName(Settings.CopyFromStart.BoneName);
 	const int32 ParentOfCopyRootIndex = TargetSkeleton.GetParentIndex(TargetRootBoneToCopyIndex);
-	TargetSkeleton.UpdateLocalTransformsBelowBone(ParentOfCopyRootIndex, TargetSkeleton.InputLocalPose, TargetSkeleton.OutputGlobalPose);
+	TargetSkeleton.UpdateLocalTransformsBelowBone(ParentOfCopyRootIndex, TargetSkeleton.InputLocalPose, OutTargetGlobalPose);
 }
 
 FIKRetargetOpSettingsBase* FIKRetargetCopyBasePoseOp::GetSettings()
diff --git a/Engine/Plugins/Animation/IKRig/Source/IKRig/Private/Retargeter/RetargetOps/PinBoneOp.cpp b/Engine/Plugins/Animation/IKRig/Source/IKRig/Private/Retargeter/RetargetOps/PinBoneOp.cpp
--- a/Engine/Plugins/Animation/IKRig/Source/IKRig/Private/Retargeter/RetargetOps/PinBoneOp.cpp
+++ b/Engine/Plugins/Animation/IKRig/Source/IKRig/Private/Retargeter/RetargetOps/PinBoneOp.cpp
@@ -222,7 +222,7 @@ FText FIKRetargetPinBoneOp::GetWarningMessage() const
 		}
 		else
 		{
-			return FText::Format(LOCTEXT("MissingBones", "Running, but missing bones. See log."), FText::AsNumber(Settings.BonesToPin.Num()));
+			return FText::Format(LOCTEXT("MissingBones", "Running, but missing {0} bones. See log."), FText::AsNumber(Settings.BonesToPin.Num()));
 		}
 	}
 	
diff --git a/Engine/Plugins/Animation/IKRig/Source/IKRig/Private/Retargeter/RetargetOps/RootMotionGeneratorOp.cpp b/Engine/Plugins/Animation/IKRig/Source/IKRig/Private/Retargeter/RetargetOps/RootMotionGeneratorOp.cpp
--- a/Engine/Plugins/Animation/IKRig/Source/IKRig/Private/Retargeter/RetargetOps/RootMotionGeneratorOp.cpp
+++ b/Engine/Plugins/Animation/IKRig/Source/IKRig/Private/Retargeter/RetargetOps/RootMotionGeneratorOp.cpp
@@ -313,7 +313,7 @@ FText FIKRetargetRootMotionOp::GetWarningMessage() const
 {
 	if (SourceRootIndex == INDEX_NONE || TargetRootIndex == INDEX_NONE)
 	{
-		return FText::Format(LOCTEXT("MissingAPelvisBone", "Missing a pelvis bone."), FText::FromName(Settings.TargetRoot.BoneName));
+		return FText::Format(LOCTEXT("MissingARootBone", "Missing a root bone."), FText::FromName(Settings.TargetRoot.BoneName));
 	}
 
 	if (Settings.RootMotionSource == ERootMotionSource::GenerateFromTargetPelvis && TargetPelvisIndex == INDEX_NONE)
diff --git a/Engine/Plugins/Animation/IKRig/Source/IKRig/Public/Retargeter/RetargetOps/PinBoneOp.h b/Engine/Plugins/Animation/IKRig/Source/IKRig/Public/Retargeter/RetargetOps/PinBoneOp.h
--- a/Engine/Plugins/Animation/IKRig/Source/IKRig/Public/Retargeter/RetargetOps/PinBoneOp.h
+++ b/Engine/Plugins/Animation/IKRig/Source/IKRig/Public/Retargeter/RetargetOps/PinBoneOp.h
@@ -38,7 +38,7 @@ struct FPinBoneData
 	double SourceScaleFactor;
 	// the local transform of the BoneToCopyTo in the ref pose
 	FTransform LocalRefPoseBoneToCopyTo;
-	// the local transform of the BoneToCopyTo in the ref pose
+	// the local transform of the BoneToCopyFrom in the ref pose
 	FTransform LocalRefPoseBoneToCopyFrom;
 	// the relative transform between the BoneToCopyFrom and BoneToCopyTo in the reference pose
 	FTransform OffsetFromBoneToCopyFromInRefPose;
@@ -193,7 +193,7 @@ public:
 	UFUNCTION(BlueprintCallable, Category = BonePairs)
 	UE_API void ClearAllBonePairs();
 
-	/* Clear all the bone pairs
+	/* Add a pair of bones to copy animation between.
 	 * @param InBoneToCopyFrom the name of the bone to copy animation from
 	 * @param InBoneToCopyTo the name of the bone to apply animation to
 	 * NOTE: if the bone to copy to is already present in the op, it will be updated with the new bone to copy from. */
@@ -210,7 +210,7 @@ public:
 // BEGIN LEGACY UOBJECT-based OP
 //
 
-// NOTE: This type has been replaced with FIKRetargetCurveRemapOp.
+// NOTE: This type has been replaced with FIKRetargetPinBoneOp.
 UCLASS(MinimalAPI)
 class UPinBoneOp : public URetargetOpBase
 {
diff --git a/Engine/Plugins/Animation/IKRig/Source/IKRig/Public/Retargeter/RetargetOps/RootMotionGeneratorOp.h b/Engine/Plugins/Animation/IKRig/Source/IKRig/Public/Retargeter/RetargetOps/RootMotionGeneratorOp.h
--- a/Engine/Plugins/Animation/IKRig/Source/IKRig/Public/Retargeter/RetargetOps/RootMotionGeneratorOp.h
+++ b/Engine/Plugins/Animation/IKRig/Source/IKRig/Public/Retargeter/RetargetOps/RootMotionGeneratorOp.h
@@ -217,7 +217,7 @@ private:
 // BEGIN DEPRECATED UOBJECT-based OP
 //
 
-// NOTE: This type has been replaced with FIKRetargetCurveRemapOp.
+// NOTE: This type has been replaced with FIKRetargetRootMotionOp.
 UCLASS(MinimalAPI)
 class URootMotionGeneratorOp : public URetargetOpBase
 {
