[IK Retarget] Treating new "Copy Local" rotation mode as additive instead of absolute.

diff --git a/Engine/Plugins/Animation/IKRig/Source/IKRig/Private/Retargeter/RetargetOps/FKChainsOp.cpp b/Engine/Plugins/Animation/IKRig/Source/IKRig/Private/Retargeter/RetargetOps/FKChainsOp.cpp
--- a/Engine/Plugins/Animation/IKRig/Source/IKRig/Private/Retargeter/RetargetOps/FKChainsOp.cpp
+++ b/Engine/Plugins/Animation/IKRig/Source/IKRig/Private/Retargeter/RetargetOps/FKChainsOp.cpp
@@ -265,8 +265,16 @@ FQuat FChainDecoderFK::RetargetRotation(
 		{
 			ParentGlobalRotation = InOutGlobalPose[ParentIndex].GetRotation();
 		}
+
+		// calculate the local delta on the source
+		const FQuat SourceCurrentLocalRotation = SourceChain.BoneChain->CurrentLocalTransforms[ClosestBoneChainIndex].GetRotation();
+		const FQuat SourceInitialLocalRotation = SourceChain.BoneChain->RefPoseLocalTransforms[ClosestBoneChainIndex].GetRotation();
+		const FQuat SourceLocalDelta = SourceCurrentLocalRotation * SourceInitialLocalRotation.Inverse();
+		// add the delta to the target local ref rotation
+		const FQuat TargetLocalRefRotation = BoneChain->RefPoseLocalTransforms[InChainIndex].GetRotation();
 		
-		FQuat OutRotation = ParentGlobalRotation * SourceChain.BoneChain->CurrentLocalTransforms[ClosestBoneChainIndex].GetRotation();
+		// apply local space offset to global space
+		FQuat OutRotation = ParentGlobalRotation * (SourceLocalDelta * TargetLocalRefRotation);
 		return MoveTemp(OutRotation);
 	}
 	
