[IK Retarget] More robust auto bone-mapping for Orient and Scale translation mode in FKChainOp

diff --git a/Engine/Plugins/Animation/IKRig/Source/IKRig/Private/Retargeter/IKRetargetProcessor.cpp b/Engine/Plugins/Animation/IKRig/Source/IKRig/Private/Retargeter/IKRetargetProcessor.cpp
--- a/Engine/Plugins/Animation/IKRig/Source/IKRig/Private/Retargeter/IKRetargetProcessor.cpp
+++ b/Engine/Plugins/Animation/IKRig/Source/IKRig/Private/Retargeter/IKRetargetProcessor.cpp
@@ -802,6 +802,18 @@ int32 FResolvedBoneChain::GetBoneAtParam(const double& Param) const
 	return Params.Num() - 1;
 }
 
+int32 FResolvedBoneChain::GetEquivalentBoneInOtherChain(const int32 InChainIndex, const FResolvedBoneChain& InOtherBoneChain) const
+{
+	// use 1:1 mapping if chains have same number of bones
+	if (BoneIndices.Num() == InOtherBoneChain.BoneIndices.Num())
+	{
+		return InChainIndex;
+	}
+
+	// otherwise, use the closest bone
+	return InOtherBoneChain.GetBoneClosestToParam(Params[InChainIndex]);
+}
+
 double FResolvedBoneChain::GetChainLength(const TArray<FTransform>& Transforms)
 {
 	double TotalLength = 0.0f;
diff --git a/Engine/Plugins/Animation/IKRig/Source/IKRig/Private/Retargeter/RetargetOps/FKChainsOp.cpp b/Engine/Plugins/Animation/IKRig/Source/IKRig/Private/Retargeter/RetargetOps/FKChainsOp.cpp
--- a/Engine/Plugins/Animation/IKRig/Source/IKRig/Private/Retargeter/RetargetOps/FKChainsOp.cpp
+++ b/Engine/Plugins/Animation/IKRig/Source/IKRig/Private/Retargeter/RetargetOps/FKChainsOp.cpp
@@ -72,8 +72,7 @@ void FChainDecoderCachedData::Initialize(
 	for (int32 TargetChainIndex=0; TargetChainIndex<InTargetBoneChain->BoneIndices.Num(); ++TargetChainIndex)
 	{
 		// cache the source bone associated with each target bone in the chain
-		const double Param = InTargetBoneChain->Params[TargetChainIndex];
-		const int32 SourceBoneChainIndex = InSourceBoneChain->GetBoneAtParam(Param);
+		const int32 SourceBoneChainIndex = InTargetBoneChain->GetEquivalentBoneInOtherChain(TargetChainIndex, *InSourceBoneChain);
 		const int32 SourceBoneIndex = InSourceBoneChain->BoneIndices[SourceBoneChainIndex];
 		SourceBoneIndices.Add(SourceBoneIndex);
 
diff --git a/Engine/Plugins/Animation/IKRig/Source/IKRig/Public/Retargeter/IKRetargetProcessor.h b/Engine/Plugins/Animation/IKRig/Source/IKRig/Public/Retargeter/IKRetargetProcessor.h
--- a/Engine/Plugins/Animation/IKRig/Source/IKRig/Public/Retargeter/IKRetargetProcessor.h
+++ b/Engine/Plugins/Animation/IKRig/Source/IKRig/Public/Retargeter/IKRetargetProcessor.h
@@ -221,6 +221,8 @@ struct FResolvedBoneChain
 	int32 GetBoneClosestToParam(const double& Param) const;
 
 	int32 GetBoneAtParam(const double& Param) const;
+
+	int32 GetEquivalentBoneInOtherChain(const int32 InChainIndex, const FResolvedBoneChain& InOtherBoneChain) const;
 	
 	static UE_API double GetChainLength(const TArray<FTransform>& Transforms);
 
