[IK Retarget] Toggle debug draw in stack view.

diff --git a/Engine/Plugins/Animation/IKRig/Source/IKRig/Private/Retargeter/IKRetargetProcessor.cpp b/Engine/Plugins/Animation/IKRig/Source/IKRig/Private/Retargeter/IKRetargetProcessor.cpp
--- a/Engine/Plugins/Animation/IKRig/Source/IKRig/Private/Retargeter/IKRetargetProcessor.cpp
+++ b/Engine/Plugins/Animation/IKRig/Source/IKRig/Private/Retargeter/IKRetargetProcessor.cpp
@@ -1720,7 +1720,7 @@ void FIKRetargetProcessor::DebugDrawAllOps(
 	{
 		// filter disabled or uninitialized ops
 		const FIKRetargetOpBase* Op = OpStack[OpIndex].GetPtr<FIKRetargetOpBase>();
-		if (!Op->IsEnabled() || !Op->IsInitialized())
+		if (!Op->IsEnabled() || !Op->IsInitialized() || !Op->GetSettingsConst()->bDebugDraw)
 		{
 			continue;
 		}
diff --git a/Engine/Plugins/Animation/IKRig/Source/IKRig/Private/Retargeter/RetargetOps/PelvisMotionOp.cpp b/Engine/Plugins/Animation/IKRig/Source/IKRig/Private/Retargeter/RetargetOps/PelvisMotionOp.cpp
--- a/Engine/Plugins/Animation/IKRig/Source/IKRig/Private/Retargeter/RetargetOps/PelvisMotionOp.cpp
+++ b/Engine/Plugins/Animation/IKRig/Source/IKRig/Private/Retargeter/RetargetOps/PelvisMotionOp.cpp
@@ -252,6 +252,14 @@ FVector FIKRetargetPelvisMotionOp::GetPelvisTranslationOffset() const
 	return Target.PelvisTranslationDelta;
 }
 
+void FIKRetargetPelvisMotionOp::PostLoad(const FIKRigObjectVersion::Type InVersion)
+{
+	if (InVersion < FIKRigObjectVersion::AddedDebugDrawTogglePerOp)
+	{
+		Settings.bDebugDraw = Settings.bEnableDebugDraw_DEPRECATED;
+	}
+}
+
 int32 FIKRetargetPelvisMotionOp::GetPelvisBoneIndex(ERetargetSourceOrTarget SourceOrTarget) const
 {
 	return SourceOrTarget == ERetargetSourceOrTarget::Source ? Source.BoneIndex : Target.BoneIndex;
@@ -367,18 +375,14 @@ void FIKRetargetPelvisMotionOp::DebugDraw(
 	const double InComponentScale,
 	const FIKRetargetDebugDrawState& InEditorState) const
 {
-	if (!Settings.bEnableDebugDraw)
-	{
-		return;
-	}
-
 	FScopeLock ScopeLock(&DebugDataMutex);
-	
+
+	const FTransform SourcePelvis = FTransform(Source.CurrentRotation, Source.CurrentPosition) * InSourceTransform;
 	const FTransform TargetPelvis = CurrentPelvisTransform * InComponentTransform;
 	const FVector TargetFloor = TargetPelvis.GetLocation() * FVector(1,1,0);
 	const FVector TargetCrotch = TargetPelvis.GetLocation() + FVector(0,0, -Settings.TargetCrotchOffset);
-	const FVector SourceFloor = Source.CurrentPosition * FVector(1,1,0);
-	const FVector SourceCrotch = Source.CurrentPosition + FVector(0,0, -Settings.SourceCrotchOffset);
+	const FVector SourceFloor = SourcePelvis.GetLocation() * FVector(1,1,0);
+	const FVector SourceCrotch = SourcePelvis.GetLocation() + FVector(0,0, -Settings.SourceCrotchOffset);
 	const FLinearColor CircleColor = InEditorState.bIsRootSelected ? InEditorState.MainColor : InEditorState.MainColor * InEditorState.NonSelected;
 	const float Thickness = static_cast<float>(Settings.DebugDrawThickness * InComponentScale);
 	const float Size = static_cast<float>(Settings.DebugDrawSize * InComponentScale);
@@ -395,7 +399,7 @@ void FIKRetargetPelvisMotionOp::DebugDraw(
 		Thickness * 0.25f);
 	
 	// draw pelvis locations
-	InPDI->DrawPoint(Source.CurrentPosition, FLinearColor::Black, Size, SDPG_Foreground);
+	InPDI->DrawPoint(SourcePelvis.GetLocation(), FLinearColor::Black, Size, SDPG_Foreground);
 	InPDI->DrawPoint(TargetPelvis.GetLocation(), FLinearColor::Black, Size, SDPG_Foreground);
 
 	// draw crotch locations
@@ -403,7 +407,7 @@ void FIKRetargetPelvisMotionOp::DebugDraw(
 	InPDI->DrawPoint(TargetCrotch, FLinearColor::Green, Size, SDPG_Foreground);
 
 	// draw pelvis-to-crotch lines
-	DrawDashedLine(InPDI,Source.CurrentPosition,SourceCrotch, FLinearColor::Black, 1.0f,SDPG_Foreground);
+	DrawDashedLine(InPDI,SourcePelvis.GetLocation(),SourceCrotch, FLinearColor::Black, 1.0f,SDPG_Foreground);
 	DrawDashedLine(InPDI,TargetPelvis.GetLocation(),TargetCrotch, FLinearColor::Black, 1.0f,SDPG_Foreground);
 	
 	// draw crotch-to-floor lines
diff --git a/Engine/Plugins/Animation/IKRig/Source/IKRig/Private/Retargeter/RetargetOps/StrideWarpingOp.cpp b/Engine/Plugins/Animation/IKRig/Source/IKRig/Private/Retargeter/RetargetOps/StrideWarpingOp.cpp
--- a/Engine/Plugins/Animation/IKRig/Source/IKRig/Private/Retargeter/RetargetOps/StrideWarpingOp.cpp
+++ b/Engine/Plugins/Animation/IKRig/Source/IKRig/Private/Retargeter/RetargetOps/StrideWarpingOp.cpp
@@ -393,6 +393,14 @@ void FIKRetargetStrideWarpingOp::OnParentReinitPropertyEdited(
 	RegenerateChainSettings(&InParentOp);
 }
 
+void FIKRetargetStrideWarpingOp::PostLoad(const FIKRigObjectVersion::Type InVersion)
+{
+	if (InVersion < FIKRigObjectVersion::AddedDebugDrawTogglePerOp)
+	{
+		Settings.bDebugDraw = Settings.bEnableDebugDraw_DEPRECATED;
+	}
+}
+
 #if WITH_EDITOR
 
 FCriticalSection FIKRetargetStrideWarpingOp::DebugDataMutex;
@@ -404,11 +412,6 @@ void FIKRetargetStrideWarpingOp::DebugDraw(
 	const double InComponentScale,
 	const FIKRetargetDebugDrawState& InEditorState) const
 {
-	if (!Settings.bEnableDebugDraw)
-	{
-		return;
-	}
-
 	FScopeLock ScopeLock(&DebugDataMutex);
 	
 	FTransform WarpingFrame = DebugStrideWarpingFrame * InComponentTransform;
diff --git a/Engine/Plugins/Animation/IKRig/Source/IKRig/Public/IKRigObjectVersion.h b/Engine/Plugins/Animation/IKRig/Source/IKRig/Public/IKRigObjectVersion.h
--- a/Engine/Plugins/Animation/IKRig/Source/IKRig/Public/IKRigObjectVersion.h
+++ b/Engine/Plugins/Animation/IKRig/Source/IKRig/Public/IKRigObjectVersion.h
@@ -44,6 +44,9 @@ struct FIKRigObjectVersion
 		// moved chain mapping in FKChainsOp from Op to Op-settings for read-only API access
 		MovedChainMappingToFKChainOpSettings,
 
+		// added "bEnableDebugDraw" to base op settings
+		AddedDebugDrawTogglePerOp,
+
 		// -----<new versions can be added above this line>-------------------------------------------------
 		VersionPlusOne,
 		LatestVersion = VersionPlusOne - 1
diff --git a/Engine/Plugins/Animation/IKRig/Source/IKRig/Public/Retargeter/IKRetargetOps.h b/Engine/Plugins/Animation/IKRig/Source/IKRig/Public/Retargeter/IKRetargetOps.h
--- a/Engine/Plugins/Animation/IKRig/Source/IKRig/Public/Retargeter/IKRetargetOps.h
+++ b/Engine/Plugins/Animation/IKRig/Source/IKRig/Public/Retargeter/IKRetargetOps.h
@@ -51,6 +51,10 @@ struct FIKRetargetOpSettingsBase
 	UPROPERTY(EditAnywhere, BlueprintReadWrite, Category = "Retarget Op")
 	int32 LODThreshold = INDEX_NONE;
 
+	/** toggle on/off all debug drawing on the op. */
+	UPROPERTY()
+	bool bDebugDraw = true;
+
 	/** (required) override to specify how settings should be applied in a way that will not require reinitialization (ie runtime compatible)*/
 	virtual void CopySettingsAtRuntime(const FIKRetargetOpSettingsBase* InSettingsToCopyFrom) { checkNoEntry(); };
 	
@@ -145,8 +149,12 @@ public:
 	/** return true if this op is initialized and ready to run */
 	bool IsInitialized() const { return bIsInitialized; };
 
-	/** (required) override and return a pointer to the settings struct used by this operation */ 
+	/** (required) override and return a pointer to the settings struct used by this operation
+	 * NOTE: it is not permitted to modify the settings from inside this getter. Treat as a const function. */ 
 	virtual FIKRetargetOpSettingsBase* GetSettings() { checkNoEntry(); return nullptr; };
+
+	/** const access to the settings */
+	const FIKRetargetOpSettingsBase* GetSettingsConst() const { return const_cast<FIKRetargetOpBase*>(this)->GetSettings(); }
 	
 	/** (optional) override to react when settings are applied at runtime
 	 * NOTE:
@@ -224,6 +232,9 @@ public:
 		const double InComponentScale,
 		const FIKRetargetDebugDrawState& InEditorState) const {};
 
+	/** (required if using DebugDraw()) return true to display the debug draw toggle in the UI */
+	virtual bool HasDebugDrawing() { return false; }
+
 	/** (optional) override to support resetting all settings for a given chain */
 	virtual void ResetChainSettingsToDefault(const FName& InChainName) {};
 
@@ -275,6 +286,8 @@ private:
 #if WITH_EDITOR
 	/** profiling data; time in seconds of last execution */
 	double AverageExecutionTime = 0.0f;
+
+	
 #endif
 	
 	/** when false, execution of this op is skipped */
diff --git a/Engine/Plugins/Animation/IKRig/Source/IKRig/Public/Retargeter/RetargetOps/FKChainsOp.h b/Engine/Plugins/Animation/IKRig/Source/IKRig/Public/Retargeter/RetargetOps/FKChainsOp.h
--- a/Engine/Plugins/Animation/IKRig/Source/IKRig/Public/Retargeter/RetargetOps/FKChainsOp.h
+++ b/Engine/Plugins/Animation/IKRig/Source/IKRig/Public/Retargeter/RetargetOps/FKChainsOp.h
@@ -335,6 +335,8 @@ public:
 		const FTransform& InComponentTransform,
 		const double InComponentScale,
 		const FIKRetargetDebugDrawState& InEditorState) const override;
+
+	IKRIG_API virtual bool HasDebugDrawing() override { return true; };
 	
 	IKRIG_API void SaveDebugData(const TArray<FTransform>& OutTargetGlobalPose);
 	TArray<FFKChainDebugData> AllChainsDebugData;
diff --git a/Engine/Plugins/Animation/IKRig/Source/IKRig/Public/Retargeter/RetargetOps/IKChainsOp.h b/Engine/Plugins/Animation/IKRig/Source/IKRig/Public/Retargeter/RetargetOps/IKChainsOp.h
--- a/Engine/Plugins/Animation/IKRig/Source/IKRig/Public/Retargeter/RetargetOps/IKChainsOp.h
+++ b/Engine/Plugins/Animation/IKRig/Source/IKRig/Public/Retargeter/RetargetOps/IKChainsOp.h
@@ -245,10 +245,14 @@ public:
 		const FTransform& InComponentTransform,
 		const double InComponentScale,
 		const FIKRetargetDebugDrawState& InEditorState) const override;
+
+	UE_API virtual bool HasDebugDrawing() override { return true; };
+	
 	UE_API void SaveDebugData(
 		const FIKRetargetProcessor& InProcessor,
 		const TArray<FTransform>& InSourceGlobalPose,
 		const TArray<FTransform>& OutTargetGlobalPose);
+	
 	TArray<FChainDebugData> AllChainsDebugData;
 	FVector DebugRootModification;
 	static UE_API FCriticalSection DebugDataMutex;
diff --git a/Engine/Plugins/Animation/IKRig/Source/IKRig/Public/Retargeter/RetargetOps/PelvisMotionOp.h b/Engine/Plugins/Animation/IKRig/Source/IKRig/Public/Retargeter/RetargetOps/PelvisMotionOp.h
--- a/Engine/Plugins/Animation/IKRig/Source/IKRig/Public/Retargeter/RetargetOps/PelvisMotionOp.h
+++ b/Engine/Plugins/Animation/IKRig/Source/IKRig/Public/Retargeter/RetargetOps/PelvisMotionOp.h
@@ -93,10 +93,6 @@ struct FIKRetargetPelvisMotionOpSettings : public FIKRetargetOpSettingsBase
 	UPROPERTY(EditAnywhere, BlueprintReadWrite, Category = "Affect IK Settings", meta = (ClampMin = "0.0", ClampMax = "1.0", UIMin = "0.0", UIMax = "1.0", DisplayName = "Affect IK Vertical"))
 	double AffectIKVertical = 0.0f;
 
-	/** Toggle debug drawing on/off in the viewport */
-	UPROPERTY(EditAnywhere, Category=Debug)
-	bool bEnableDebugDraw = true;
-
 	/** Adjust size of the debug drawing */
 	UPROPERTY(EditAnywhere, Category=Debug)
 	double DebugDrawSize = 5.0f;
@@ -113,6 +109,9 @@ struct FIKRetargetPelvisMotionOpSettings : public FIKRetargetOpSettingsBase
 #if WITH_EDITOR
 	IKRIG_API virtual USkeleton* GetSkeleton(const FName InPropertyName) override;
 #endif
+
+	UPROPERTY(meta = (DeprecatedProperty))
+	bool bEnableDebugDraw_DEPRECATED = true;
 };
 
 struct FPelvisSource
@@ -177,6 +176,8 @@ struct FIKRetargetPelvisMotionOp : public FIKRetargetOpBase
 
 	IKRIG_API FVector GetPelvisTranslationOffset() const;
 
+	IKRIG_API virtual void PostLoad(const FIKRigObjectVersion::Type InVersion) override;
+
 	int32 GetPelvisBoneIndex(ERetargetSourceOrTarget SourceOrTarget) const;
 
 	UPROPERTY()
@@ -207,13 +208,18 @@ private:
 	
 #if WITH_EDITOR
 public:
+	
 	IKRIG_API virtual FText GetWarningMessage() const override;
+	
 	IKRIG_API virtual void DebugDraw(
 		FPrimitiveDrawInterface* InPDI,
 		const FTransform& InSourceTransform,
 		const FTransform& InComponentTransform,
 		const double InComponentScale,
 		const FIKRetargetDebugDrawState& InEditorState) const override;
+
+	IKRIG_API virtual bool HasDebugDrawing() override { return true; };
+	
 	FTransform CurrentPelvisTransform;
 	static IKRIG_API FCriticalSection DebugDataMutex;
 #endif
diff --git a/Engine/Plugins/Animation/IKRig/Source/IKRig/Public/Retargeter/RetargetOps/RunIKRigOp.h b/Engine/Plugins/Animation/IKRig/Source/IKRig/Public/Retargeter/RetargetOps/RunIKRigOp.h
--- a/Engine/Plugins/Animation/IKRig/Source/IKRig/Public/Retargeter/RetargetOps/RunIKRigOp.h
+++ b/Engine/Plugins/Animation/IKRig/Source/IKRig/Public/Retargeter/RetargetOps/RunIKRigOp.h
@@ -98,14 +98,20 @@ struct FIKRetargetRunIKRigOp : public FIKRetargetOpBase
 	FRetargetChainMapping ChainMapping;
 
 #if WITH_EDITOR
+	
 	UE_API virtual void DebugDraw(
 		FPrimitiveDrawInterface* InPDI,
 		const FTransform& InSourceTransform,
 		const FTransform& InComponentTransform,
 		const double InComponentScale,
 		const FIKRetargetDebugDrawState& InEditorState) const override;
+
+	IKRIG_API virtual bool HasDebugDrawing() override { return true; };
+	
 	UE_API void SaveInitialGoalTransformsIntoDebugData(const FIKRetargetProcessor& InProcessor, const TArray<FTransform>& InTargetGlobalPose);
+	
 	UE_API void SaveCurrentGoalTransformsIntoDebugData();
+	
 	TArray<FRunIKRigOpGoalDebugData> GoalDebugData;
 	static UE_API FCriticalSection DebugDataMutex;
 #endif
diff --git a/Engine/Plugins/Animation/IKRig/Source/IKRig/Public/Retargeter/RetargetOps/StrideWarpingOp.h b/Engine/Plugins/Animation/IKRig/Source/IKRig/Public/Retargeter/RetargetOps/StrideWarpingOp.h
--- a/Engine/Plugins/Animation/IKRig/Source/IKRig/Public/Retargeter/RetargetOps/StrideWarpingOp.h
+++ b/Engine/Plugins/Animation/IKRig/Source/IKRig/Public/Retargeter/RetargetOps/StrideWarpingOp.h
@@ -83,10 +83,6 @@ struct FIKRetargetStrideWarpingOpSettings : public FIKRetargetOpSettingsBase
 	UPROPERTY(EditAnywhere, BlueprintReadWrite, Category = Warping, meta = (UIMin = "0.0", UIMax = "2.0", ClampMin = "0.0"))
 	double WarpSplay = 1.0f;
 
-	/** Enable debug drawing for this op. */
-	UPROPERTY(EditAnywhere, BlueprintReadWrite, Category = Debug)
-	bool bEnableDebugDraw = true;
-
 	/** Adjust the size of the debug drawing */
 	UPROPERTY(EditAnywhere, BlueprintReadWrite, Category = Debug)
 	double DebugDrawSize = 5.0;
@@ -102,6 +98,9 @@ struct FIKRetargetStrideWarpingOpSettings : public FIKRetargetOpSettingsBase
 	static UE_API FVector GetAxisVector(const EBasicAxis& Axis);
 
 	UE_API bool operator==(const FIKRetargetStrideWarpingOpSettings& Other) const;
+
+	UPROPERTY(meta = (DeprecatedProperty))
+	bool bEnableDebugDraw_DEPRECATED = true;
 };
 
 struct FStrideWarpGoalData
@@ -148,16 +147,22 @@ struct FIKRetargetStrideWarpingOp : public FIKRetargetOpBase
 
 	UE_API virtual void OnParentReinitPropertyEdited(const FIKRetargetOpBase& InParentOp, const FPropertyChangedEvent* InPropertyChangedEvent) override;
 
+	IKRIG_API virtual void PostLoad(const FIKRigObjectVersion::Type InVersion) override;
+
 	UPROPERTY()
 	FIKRetargetStrideWarpingOpSettings Settings;
 
 #if WITH_EDITOR
+	
 	UE_API virtual void DebugDraw(
 		FPrimitiveDrawInterface* InPDI,
 		const FTransform& InSourceTransform,
 		const FTransform& InComponentTransform,
 		const double InComponentScale,
 		const FIKRetargetDebugDrawState& InEditorState) const override;
+
+	IKRIG_API virtual bool HasDebugDrawing() override { return true; };
+	
 	FTransform DebugStrideWarpingFrame;
 	static UE_API FCriticalSection DebugDataMutex;
 
diff --git a/Engine/Plugins/Animation/IKRig/Source/IKRigEditor/Private/RetargetEditor/IKRetargetDefaultMode.cpp b/Engine/Plugins/Animation/IKRig/Source/IKRigEditor/Private/RetargetEditor/IKRetargetDefaultMode.cpp
--- a/Engine/Plugins/Animation/IKRig/Source/IKRigEditor/Private/RetargetEditor/IKRetargetDefaultMode.cpp
+++ b/Engine/Plugins/Animation/IKRig/Source/IKRigEditor/Private/RetargetEditor/IKRetargetDefaultMode.cpp
@@ -62,7 +62,7 @@ void FIKRetargetDefaultMode::Render(const FSceneView* View, FViewport* Viewport,
 	RenderDebugProxies(PDI, Controller);
 }
 
-void FIKRetargetDefaultMode::RenderDebugProxies(FPrimitiveDrawInterface* PDI, const FIKRetargetEditorController* Controller) const
+void FIKRetargetDefaultMode::RenderDebugProxies(FPrimitiveDrawInterface* PDI, const FIKRetargetEditorController* Controller)
 {
 	// asset settings can disable debug drawing
 	const UIKRetargeter* Asset = Controller->AssetController->GetAsset();
diff --git a/Engine/Plugins/Animation/IKRig/Source/IKRigEditor/Private/RetargetEditor/SRetargetOpStack.cpp b/Engine/Plugins/Animation/IKRig/Source/IKRigEditor/Private/RetargetEditor/SRetargetOpStack.cpp
--- a/Engine/Plugins/Animation/IKRig/Source/IKRigEditor/Private/RetargetEditor/SRetargetOpStack.cpp
+++ b/Engine/Plugins/Animation/IKRig/Source/IKRigEditor/Private/RetargetEditor/SRetargetOpStack.cpp
@@ -13,8 +13,8 @@
 
 #define LOCTEXT_NAMESPACE "SRetargetOpStack"
 
-float SRetargetOpItem::OpHorizontalPadding = 6.0f;;
-float SRetargetOpItem::OpVerticalPadding = 3.0f;
+float SRetargetOpItem::OpHorizontalPadding = 3.0f;;
+float SRetargetOpItem::OpVerticalPadding = 2.0f;
 
 TSharedRef<ITableRow> FRetargetOpStackElement::MakeListRowWidget(
 	const TSharedRef<STableViewBase>& InOwnerTable,
@@ -54,6 +54,13 @@ void SRetargetOpItem::Construct(const FArguments& InArgs)
     ChildSlot
     [
         SNew(SHorizontalBox)
+    	.ToolTipText(TAttribute<FText>::CreateLambda([this]()
+		{
+    		return FText::Format(
+			FText::FromString(TEXT("Op Type: {0} \nIndex in Stack: {1}\n")),
+			FText::FromString(Element.Pin()->GetType()->GetStructCPPName()),
+			FText::AsNumber(Element.Pin()->GetIndexInStack()) );
+		}))
 	    
         // drag Icon
         + SHorizontalBox::Slot()
@@ -90,15 +97,54 @@ void SRetargetOpItem::Construct(const FArguments& InArgs)
             })
         ]
 
-    	// display index
-		+ SHorizontalBox::Slot()
+	    // toggle debug draw button
+    	+ SHorizontalBox::Slot()
 		.AutoWidth()
 		.HAlign(HAlign_Left)
 		.VAlign(VAlign_Center)
-		.Padding(0, OpVerticalPadding)
+		.Padding(FMargin(OpHorizontalPadding, OpVerticalPadding))
 		[
-			SNew(STextBlock)
-			.Text(this, &SRetargetOpItem::GetIndexAsText)
+			SNew(SButton)
+			.NormalPaddingOverride(FMargin(2, 2))
+			.PressedPaddingOverride(FMargin(2, 2))
+			.ToolTipText(LOCTEXT("ToggleOpDebugDraw", "Toggle this Ops debug drawing in the viewport."))
+			.Visibility_Lambda([this]()
+			{
+				FIKRetargetEditorController* EditorController = ListView.Pin()->EditorController.Pin().Get();
+				FIKRetargetOpBase* Op = EditorController->AssetController->GetRetargetOpByIndex(Element.Pin()->GetIndexInStack());
+				if (!ensure(Op))
+				{
+					return EVisibility::Collapsed;
+				}
+				return Op->HasDebugDrawing() ? EVisibility::Visible : EVisibility::Collapsed;
+			})
+			.OnClicked_Lambda([this]() -> FReply
+			{
+				FIKRetargetEditorController* EditorController = ListView.Pin()->EditorController.Pin().Get();
+				FIKRetargetOpBase* Op = EditorController->AssetController->GetRetargetOpByIndex(Element.Pin()->GetIndexInStack());
+				if (!ensure(Op))
+				{
+					return FReply::Handled();
+				}
+
+				Op->GetSettings()->bDebugDraw = !Op->GetSettings()->bDebugDraw;
+				return FReply::Handled();
+			})
+			.Content()
+			[
+				SNew(SImage)
+				.Image(FIKRetargetEditorStyle::Get().GetBrush("IKRetarget.Debug"))
+				.ColorAndOpacity_Lambda([this]()
+				{
+					FIKRetargetEditorController* EditorController = ListView.Pin()->EditorController.Pin().Get();
+					FIKRetargetOpBase* Op = EditorController->AssetController->GetRetargetOpByIndex(Element.Pin()->GetIndexInStack());
+					if (!Op)
+					{
+						return FLinearColor::White;
+					}
+					return Op->GetSettings()->bDebugDraw ?FLinearColor(0,1,0,1.0) : FLinearColor::White;
+				})
+			]
 		]
 
         // display Name
@@ -128,23 +174,6 @@ void SRetargetOpItem::Construct(const FArguments& InArgs)
             .Size(FVector2D(0.0f, 1.0f))
         ]
 
-    	// add sub-op button
-		+ SHorizontalBox::Slot()
-		.AutoWidth()
-		.HAlign(HAlign_Right)
-		.VAlign(VAlign_Center)
-		.Padding(FMargin(OpHorizontalPadding, OpVerticalPadding))
-		[
-			SNew(SPositiveActionButton)
-			.Visibility(Element.Pin()->GetCanHaveChildren() ? EVisibility::Visible : EVisibility::Collapsed)
-			.Icon(FAppStyle::Get().GetBrush("Icons.Plus"))
-			.ToolTipText(LOCTEXT("AddChildToolTip", "Add a child op to run before this op."))
-			.OnGetMenuContent_Lambda([this]()
-			{
-				return ListView.Pin()->CreateAddNewOpMenu();	
-			})
-		]
-
     	// profiling
 		+ SHorizontalBox::Slot()
 		.AutoWidth()
@@ -169,6 +198,23 @@ void SRetargetOpItem::Construct(const FArguments& InArgs)
 			})
 		]
 
+    	// add sub-op button
+		+ SHorizontalBox::Slot()
+		.AutoWidth()
+		.HAlign(HAlign_Right)
+		.VAlign(VAlign_Center)
+		.Padding(FMargin(OpHorizontalPadding, OpVerticalPadding))
+		[
+			SNew(SPositiveActionButton)
+			.Visibility(Element.Pin()->GetCanHaveChildren() ? EVisibility::Visible : EVisibility::Collapsed)
+			.Icon(FAppStyle::Get().GetBrush("Icons.Plus"))
+			.ToolTipText(LOCTEXT("AddChildToolTip", "Add a child op to run before this op."))
+			.OnGetMenuContent_Lambda([this]()
+			{
+				return ListView.Pin()->CreateAddNewOpMenu();	
+			})
+		]
+
         // spin Box
         + SHorizontalBox::Slot()
         .AutoWidth()
@@ -196,7 +242,8 @@ void SRetargetOpItem::Construct(const FArguments& InArgs)
         .Padding(FMargin(OpHorizontalPadding, OpVerticalPadding))
         [
             SNew(SButton)
-            .ContentPadding(FMargin(0,0))
+        	.NormalPaddingOverride(FMargin(2, 2))
+			.PressedPaddingOverride(FMargin(2, 2))
             .ToolTipText(LOCTEXT("DeleteOp", "Delete retarget op and remove from stack."))
             .OnClicked_Lambda([this]() -> FReply
             {
@@ -205,9 +252,9 @@ void SRetargetOpItem::Construct(const FArguments& InArgs)
             })
             .Content()
             [
-                SNew(SImage)
-                .Image(FAppStyle::Get().GetBrush("Icons.Delete"))
-                .ColorAndOpacity(FSlateColor::UseForeground())
+            	SNew(SImage)
+					.Image(FAppStyle::Get().GetBrush("Icons.Delete"))
+					.ColorAndOpacity(FSlateColor::UseForeground())
             ]
         ]
     ];
@@ -229,11 +276,6 @@ FText SRetargetOpItem::GetName() const
 	return FText::FromName(OpName);
 }
 
-FText SRetargetOpItem::GetIndexAsText() const
-{
-	return FText::Format(FText::FromString(TEXT("{0}: ")), Element.Pin()->GetIndexInStack() + 1);
-}
-
 bool SRetargetOpItem::IsOpEnabled() const
 {
 	if (const FIKRetargetOpBase* Op = GetRetargetOp())
@@ -853,9 +895,10 @@ void SRetargetOpStack::Construct(
 				+ SVerticalBox::Slot()
 				.AutoHeight()
 				.VAlign(VAlign_Top)
-				.HAlign(HAlign_Left)
+				.HAlign(HAlign_Fill)
 				[
 					SNew(SHorizontalBox)
+					
 					+ SHorizontalBox::Slot()
 					.AutoWidth()
 					.VAlign(VAlign_Center)
diff --git a/Engine/Plugins/Animation/IKRig/Source/IKRigEditor/Public/RetargetEditor/IKRetargetDefaultMode.h b/Engine/Plugins/Animation/IKRig/Source/IKRigEditor/Public/RetargetEditor/IKRetargetDefaultMode.h
--- a/Engine/Plugins/Animation/IKRig/Source/IKRigEditor/Public/RetargetEditor/IKRetargetDefaultMode.h
+++ b/Engine/Plugins/Animation/IKRig/Source/IKRigEditor/Public/RetargetEditor/IKRetargetDefaultMode.h
@@ -39,7 +39,7 @@ public:
 	/** END FEdMode interface */
 
 private:
-	void RenderDebugProxies(FPrimitiveDrawInterface* PDI, const FIKRetargetEditorController* Controller) const;
+	void RenderDebugProxies(FPrimitiveDrawInterface* PDI, const FIKRetargetEditorController* Controller);
 	
 	// the skeleton currently being edited
 	ERetargetSourceOrTarget SkeletonMode;
diff --git a/Engine/Plugins/Animation/IKRig/Source/IKRigEditor/Public/RetargetEditor/IKRetargetEditorStyle.h b/Engine/Plugins/Animation/IKRig/Source/IKRigEditor/Public/RetargetEditor/IKRetargetEditorStyle.h
--- a/Engine/Plugins/Animation/IKRig/Source/IKRigEditor/Public/RetargetEditor/IKRetargetEditorStyle.h
+++ b/Engine/Plugins/Animation/IKRig/Source/IKRigEditor/Public/RetargetEditor/IKRetargetEditorStyle.h
@@ -33,6 +33,7 @@ public:
 		Set("IKRetarget.EditRetargetPose", new IMAGE_BRUSH_SVG("Slate/EditRetargetPose", Icon64x64));
 		Set("IKRetarget.ShowRetargetPose", new IMAGE_BRUSH_SVG("Slate/ShowRetargetPose", Icon64x64));
 		Set("IKRetarget.AutoAlign", new IMAGE_BRUSH_SVG("Slate/AutoRetargetPose", Icon16x16));
+		Set("IKRetarget.Debug", new IMAGE_BRUSH_SVG("Slate/Debug", Icon16x16));
 
 		FSlateColor OpOutlineColor = FSlateColor(FLinearColor(0.1843f, 0.1843f, 0.1843f, 0.5f));
 		FSlateColor OpOutlineColorSelected = FSlateColor(FLinearColor(0.1843f, 0.1843f, 0.1843f, 1.0f));
diff --git a/Engine/Plugins/Animation/IKRig/Source/IKRigEditor/Public/RetargetEditor/SRetargetOpStack.h b/Engine/Plugins/Animation/IKRig/Source/IKRigEditor/Public/RetargetEditor/SRetargetOpStack.h
--- a/Engine/Plugins/Animation/IKRig/Source/IKRigEditor/Public/RetargetEditor/SRetargetOpStack.h
+++ b/Engine/Plugins/Animation/IKRig/Source/IKRigEditor/Public/RetargetEditor/SRetargetOpStack.h
@@ -178,7 +178,6 @@ public:
 	// RENAMING op
 	void OnNameCommitted(const FText& InText, ETextCommit::Type InCommitType) const;
 	FText GetName() const;
-	FText GetIndexAsText() const;
 	// ENd renaming op
 
 	static float OpHorizontalPadding;
