[IK Retarget] Bughawk minor issues

diff --git a/Engine/Plugins/Animation/IKRig/Source/IKRig/Private/Retargeter/RetargetOps/FKChainsOp.cpp b/Engine/Plugins/Animation/IKRig/Source/IKRig/Private/Retargeter/RetargetOps/FKChainsOp.cpp
--- a/Engine/Plugins/Animation/IKRig/Source/IKRig/Private/Retargeter/RetargetOps/FKChainsOp.cpp
+++ b/Engine/Plugins/Animation/IKRig/Source/IKRig/Private/Retargeter/RetargetOps/FKChainsOp.cpp
@@ -508,6 +508,12 @@ void FChainDecoderFK::MatchChain(
 	//
 	// Typically, spline-ik is not "twist aware" because aligning a joint chain with swing rotations does not twist,
 	// but by using the interpolated rotations/translation as a starting place we have a full twist/bend/stretch.
+
+	// match chain feature requires at least 3 bones to work
+	if (TargetBoneIndices.Num() < 3 || SourceChain.BoneChain->BoneIndices.Num() < 3)
+	{
+		return;
+	}
 	
 	// convert source chain into a linear spline
 	TArray<FVector> SourceSplinePoints;
@@ -749,7 +755,7 @@ bool FChainPairFK::Initialize(
 	{
 		InLog.LogWarning( FText::Format(
 			LOCTEXT("FKChainInvalid", "FK Chain Op unable to retarget source chain, {0} to target chain, {1}.'"),
-			FText::FromName(SourceBoneChain->ChainName), FText::FromName(TargetBoneChain->ChainName)));
+			FText::FromName(InSourceBoneChain.ChainName), FText::FromName(InTargetBoneChain.ChainName)));
 		return false;
 	}
 	
