[IK Rig] Adding FABRIK to Stretch Limb solver.

diff --git a/Engine/Plugins/Animation/IKRig/Source/IKRig/Private/Rig/Solvers/IKRigStretchLimb.cpp b/Engine/Plugins/Animation/IKRig/Source/IKRig/Private/Rig/Solvers/IKRigStretchLimb.cpp
--- a/Engine/Plugins/Animation/IKRig/Source/IKRig/Private/Rig/Solvers/IKRigStretchLimb.cpp
+++ b/Engine/Plugins/Animation/IKRig/Source/IKRig/Private/Rig/Solvers/IKRigStretchLimb.cpp
@@ -4,6 +4,8 @@
 #include "Rig/IKRigDataTypes.h"
 #include "Rig/IKRigSkeleton.h"
 
+#include "Algo/Accumulate.h"
+
 #include UE_INLINE_GENERATED_CPP_BY_NAME(IKRigStretchLimb)
 
 #define LOCTEXT_NAMESPACE "IKRigStretchLimb"
@@ -28,11 +30,11 @@ void FIKRigStretchLimbSolver::Initialize(const FIKRigSkeleton& InIKRigSkeleton)
 
 	// populate limb indices
 	BoneIndices = {EndBoneIndex};
-	int32 BoneIndex = InIKRigSkeleton.GetParentIndex(EndBoneIndex);
-	while (BoneIndex != INDEX_NONE && BoneIndex >= StartBoneIndex)
+	int32 CurrentBoneIndex = InIKRigSkeleton.GetParentIndex(EndBoneIndex);
+	while (CurrentBoneIndex != INDEX_NONE && CurrentBoneIndex >= StartBoneIndex)
 	{
-		BoneIndices.Add(BoneIndex);
-		BoneIndex = InIKRigSkeleton.GetParentIndex(BoneIndex);
+		BoneIndices.Add(CurrentBoneIndex);
+		CurrentBoneIndex = InIKRigSkeleton.GetParentIndex(CurrentBoneIndex);
 	};
 
 	// limb must have at least two joints
@@ -56,6 +58,23 @@ void FIKRigStretchLimbSolver::Initialize(const FIKRigSkeleton& InIKRigSkeleton)
 		ChildrenToUpdate.Add(ChildBoneIndex);
 	}
 
+	// initialize storage for bone positions prior to any translational offset
+	BonePositionsPreOffset.Init(FVector::ZeroVector, BoneIndices.Num());
+
+	// initialize pole vector params
+	PoleVectorParams.Reset(BoneIndices.Num());
+	const FVector Start = InIKRigSkeleton.RefPoseGlobal[StartBoneIndex].GetTranslation();
+	const FVector End = InIKRigSkeleton.RefPoseGlobal[EndBoneIndex].GetTranslation();
+	const FVector PoleVector = End - Start;
+	InitialPoleVectorLenSq = FMath::Max(UE_DOUBLE_KINDA_SMALL_NUMBER,PoleVector.SizeSquared());
+	for (const int32 BoneIndex : BoneIndices)
+	{
+		const FVector BonePosition = InIKRigSkeleton.RefPoseGlobal[BoneIndex].GetTranslation();
+		const FVector StartToBone = BonePosition - Start;
+		double T = FVector::DotProduct(StartToBone, PoleVector) / InitialPoleVectorLenSq;
+		PoleVectorParams.Add(T);
+	}
+
 	bIsInitialized = true;
 }
 
@@ -66,12 +85,6 @@ void FIKRigStretchLimbSolver::Solve(FIKRigSkeleton& InIKRigSkeleton, const FIKRi
 		return; // can't solve until initialized successfully
 	}
 
-	if (Settings.StretchMode == EStretchLimbStretchMode::None &&
-		Settings.RotationMode == EStretchLimbRotationMode::None)
-	{
-		return; // nothing to do
-	}
-
 	// get the goal
 	const FIKRigGoal* Goal = InGoals.FindGoalByName(Settings.Goal);
 	if (!Goal)
@@ -81,32 +94,38 @@ void FIKRigStretchLimbSolver::Solve(FIKRigSkeleton& InIKRigSkeleton, const FIKRi
 		return;
 	}
 
+	// update rest lengths with incoming animation
+	UpdateRestLengths(InIKRigSkeleton);
+
 	// orient the chain bones
-	OrientChain(InIKRigSkeleton, Goal->FinalBlendedPosition);
+	AimChainAtGoal(InIKRigSkeleton, Goal->FinalBlendedPosition);
+
+	// record chain positions prior to translational offsets
+	StoreBonePositions(InIKRigSkeleton);
 	
-	// translate the chain bones
-	TranslateChain(InIKRigSkeleton, Goal->FinalBlendedPosition);
+	// translate the bones parallel to the pole vector
+	SquashOrStretchChainParallel(InIKRigSkeleton, Goal->FinalBlendedPosition);
 
-	// orient end bone to goal
-	FTransform& EndBoneGlobal = InIKRigSkeleton.CurrentPoseGlobal[BoneIndices.Last()];
-	FQuat EndBoneGlobalRotation = FQuat::FastLerp(EndBoneGlobal.GetRotation(), Goal->FinalBlendedRotation, Settings.RotateEndBoneWithGoal);
-	EndBoneGlobalRotation = EndBoneGlobalRotation.GetNormalized();
-	EndBoneGlobal.SetRotation(EndBoneGlobalRotation);
+	// translate the bones perpendicular to the pole vector
+	SquashChainPerpendicular(InIKRigSkeleton);
 
-	// update chain local transforms
-	for (const int32 BoneIndex : BoneIndices)
-	{
-		InIKRigSkeleton.UpdateLocalTransformFromGlobal(BoneIndex);
-	}
+	// solve distance constraints
+	RunFABRIK(InIKRigSkeleton, Goal->FinalBlendedPosition);
 
-	// propagate to children
-	for (const int32 ChildIndex: ChildrenToUpdate)
-	{
-		InIKRigSkeleton.UpdateGlobalTransformFromLocal(ChildIndex);
-	}
+	// allow final stretching to hit goal
+	StretchChainFinal(InIKRigSkeleton, Goal->FinalBlendedPosition);
+
+	// done moving bones, update their orientations
+	UpdateBoneOrientations(InIKRigSkeleton);
+
+	// orient end bone to goal
+	RotateEndBoneWithGoal(InIKRigSkeleton, Goal);
+
+	// update all the children bones not in the chain
+	UpdateFKChildren(InIKRigSkeleton);
 }
 
-void FIKRigStretchLimbSolver::OrientChain(
+void FIKRigStretchLimbSolver::AimChainAtGoal(
 	FIKRigSkeleton& InIKRigSkeleton,
 	const FVector& GoalLocation) const
 {
@@ -136,123 +155,186 @@ void FIKRigStretchLimbSolver::OrientChain(
 	}
 }
 
-void FIKRigStretchLimbSolver::TranslateChain(
+void FIKRigStretchLimbSolver::SquashOrStretchChainParallel(
 	FIKRigSkeleton& InIKRigSkeleton,
 	const FVector& GoalLocation) const
 {
-	if (Settings.StretchMode == EStretchLimbStretchMode::None)
+	// calc current and new pole vector (based on goal)
+	const FVector& ChainStart = InIKRigSkeleton.CurrentPoseGlobal[BoneIndices[0]].GetLocation();
+	const FVector& ChainEnd = InIKRigSkeleton.CurrentPoseGlobal[BoneIndices.Last()].GetLocation();
+	const FVector CurrentPoleVector = ChainEnd - ChainStart;
+
+	// clamp goal to percent distance along chain (optionally prevents overextension)
+	FVector ClampedGoal = GoalLocation;
+	if (Settings.bEnableStretching)
 	{
-		return; // skip adjusting translations
+		const FVector UnclampedPoleVector = GoalLocation - ChainStart;
+		const FVector ClampedPoleVector = UnclampedPoleVector.GetClampedToMaxSize(TotalRestLength * Settings.StretchStartPercent);
+		ClampedGoal = ChainStart + ClampedPoleVector;
 	}
 
-	// calc:
-	// 1. Current Chain Vector: start to end
-	// 2. New Chain Vector: start to goal
+	// new pole vector is from chain origin to clamped goal location
+	const FVector NewPoleVector = ClampedGoal - ChainStart;
+
+	// squash/stretch chain in it's input pose to reach goal
+	ApplyParallelOffsetToChain(InIKRigSkeleton, ChainStart, NewPoleVector, CurrentPoleVector);
+}
+
+void FIKRigStretchLimbSolver::ApplyParallelOffsetToChain(
+	FIKRigSkeleton& InIKRigSkeleton,
+	const FVector& ChainStart,
+	const FVector& InNewPoleVector,
+	const FVector& CurrentPoleVector) const
+{
+	// compare length of the chain in its current configuration with the length it would need to reach goal
+	const double NewChainLength = InNewPoleVector.Size();
+	const double CurrentChainLength = CurrentPoleVector.Size();
+	if (FMath::IsNearlyEqual(CurrentChainLength, NewChainLength))
+	{
+		return; // nothing to squash or stretch
+	}
+
+	for (int32 ChainIndex = 1; ChainIndex < BoneIndices.Num(); ++ChainIndex)
+	{
+		FTransform& BoneGlobalTransform = InIKRigSkeleton.CurrentPoseGlobal[BoneIndices[ChainIndex]];
+		const FVector CurrentProjOnPoleVector = ChainStart + CurrentPoleVector * PoleVectorParams[ChainIndex];
+		const FVector CurrentProjToBone = BoneGlobalTransform.GetLocation() - CurrentProjOnPoleVector;
+		const FVector NewProjOnPoleVector = ChainStart + InNewPoleVector * PoleVectorParams[ChainIndex];
+		BoneGlobalTransform.SetLocation(NewProjOnPoleVector + CurrentProjToBone);
+	}
+}
+
+void FIKRigStretchLimbSolver::SquashChainPerpendicular(FIKRigSkeleton& InIKRigSkeleton)
+{
+	if (BoneIndices.Num() < 3)
+	{
+		return;
+	}
+	
+	// calc pole vector
 	const FVector& ChainStart = InIKRigSkeleton.CurrentPoseGlobal[BoneIndices[0]].GetLocation();
 	const FVector& ChainEnd = InIKRigSkeleton.CurrentPoseGlobal[BoneIndices.Last()].GetLocation();
-	const FVector CurrentChainVector = ChainEnd - ChainStart;
-	const FVector NewChainVector = GoalLocation - ChainStart;
+	const FVector PoleVector = ChainEnd - ChainStart;
+	const double NewPoleVectorLenSq = FMath::Max(1.e-2,PoleVector.SizeSquared());
+	if (NewPoleVectorLenSq >= InitialPoleVectorLenSq)
+	{
+		return; // limb is stretched
+	}
+	const FVector PoleNorm = PoleVector.GetSafeNormal();
 
-	// calculate length of the chain
-	double ChainLength = 0.0f;
-	switch (Settings.LengthType)
+	// calc bending directions for each bone (skipping start/end)
+	TArray<FVector> SquashDirections;
+	SquashDirections.Init(FVector::UpVector, BoneIndices.Num());
+	for (int32 ChainIndex = 1; ChainIndex < BoneIndices.Num() - 1 ; ++ChainIndex)
 	{
-	case EStretchLimbLengthType::AlongBones:
-		{
-			// iterate from tip to first child of root
-			for (int32 BoneIndex = BoneIndices.Num() - 1; BoneIndex >= 1; --BoneIndex)
-			{
-				const FVector& BonePosition = InIKRigSkeleton.CurrentPoseGlobal[BoneIndices[BoneIndex]].GetLocation();
-				const FVector& ParentPosition = InIKRigSkeleton.CurrentPoseGlobal[BoneIndices[BoneIndex-1]].GetLocation();
-				ChainLength += (ParentPosition - BonePosition).Length();
-			}
-			break;
-		}
-	case EStretchLimbLengthType::StraightLine:
+		const int32 BoneIndex = BoneIndices[ChainIndex];
+		const FVector& BonePosition = InIKRigSkeleton.CurrentPoseGlobal[BoneIndex].GetLocation();
+		const FVector StartToBone = BonePosition - ChainStart;
+
+		// signed distance along the line direction
+		const double T = FVector::DotProduct(StartToBone, PoleNorm);
+
+		// closest point on the pole vector
+		const FVector& ClosestPoint = ChainStart + T * PoleNorm;
+		const FVector BendDirection = (BonePosition - ClosestPoint).GetSafeNormal();
+		SquashDirections[ChainIndex] = BendDirection;
+	}
+
+	// push bones along squash directions
+	const double SquashRatio = 1.0 - NewPoleVectorLenSq / InitialPoleVectorLenSq;
+	for (int32 ChainIndex = 1; ChainIndex < BoneIndices.Num() - 1 ; ++ChainIndex)
+	{
+		auto SmoothTriangleFalloff = [](double InT) -> double
 		{
-			ChainLength = CurrentChainVector.Length();
-			break;
-		}
-	default:
-		ensureMsgf(false, TEXT("Unknown chain length type."));
+			InT = FMath::Clamp(InT, 0.0f, 1.0f);
+			return 2.0f * FMath::Sqrt(InT * (1.0f - InT)); // 0 at 0, 0 at 1, 1 at 0.5
+		};
+
+		const int32 BoneIndex = BoneIndices[ChainIndex];
+		const FVector& BendDirection = SquashDirections[ChainIndex];
+		const double Falloff = SmoothTriangleFalloff(PoleVectorParams[ChainIndex]);
+		InIKRigSkeleton.CurrentPoseGlobal[BoneIndex].AddToTranslation(Settings.SquashStrength * BendDirection * Falloff * SquashRatio);
+	}
+}
+
+void FIKRigStretchLimbSolver::StretchChainFinal(FIKRigSkeleton& InIKRigSkeleton, const FVector& GoalLocation) const
+{
+	if (!Settings.bEnableStretching)
+	{
 		return;
 	}
+	
+	// calc current and new pole vector (based on goal)
+	const FVector& ChainStart = InIKRigSkeleton.CurrentPoseGlobal[BoneIndices[0]].GetLocation();
+	const FVector& ChainEnd = InIKRigSkeleton.CurrentPoseGlobal[BoneIndices.Last()].GetLocation();
+	const FVector CurrentPoleVector = ChainEnd - ChainStart;
+	FVector NewPoleVector = GoalLocation - ChainStart;
+	
+	// clamp to max stretch distance
+	if (Settings.MaximumStretchDistance > 0.0)
+	{
+		NewPoleVector = NewPoleVector.GetClampedToMaxSize(TotalRestLength + Settings.MaximumStretchDistance);
+	}
+	
+	// squash/stretch chain in it's input pose to reach goal
+	ApplyParallelOffsetToChain(InIKRigSkeleton, ChainStart, NewPoleVector, CurrentPoleVector);
+}
 
-	// compare lenghth of the chain in its current configuration with the length it would need to reach goal
-	const double NewChainLength = NewChainVector.Size();
-	if (FMath::IsNearlyEqual(ChainLength, NewChainLength))
+void FIKRigStretchLimbSolver::RunFABRIK(FIKRigSkeleton& InIKRigSkeleton, const FVector& GoalLocation)
+{
+	if (BoneIndices.Num() < 3)
 	{
-		return; // nothing to squash or stretch
+		return;
 	}
 
-	if (Settings.StretchMode != EStretchLimbStretchMode::SquashAndStretch)
+	// clamp goal to percent distance along chain (optionally prevents overextension)
+	const FVector FABRIKStart = InIKRigSkeleton.CurrentPoseGlobal[BoneIndices[0]].GetLocation();
+	FVector FABRIKEnd = GoalLocation;
+	if (Settings.bEnableStretching)
 	{
-		if (ChainLength > NewChainLength && Settings.StretchMode == EStretchLimbStretchMode::StretchOnly)
-		{
-			return; // chain needs to squash, but we only want to stretch
-		}
+		const FVector CurrentPoleVector = FABRIKEnd - FABRIKStart;
+		const FVector ClampedPoleVector = CurrentPoleVector.GetClampedToMaxSize(TotalRestLength * Settings.StretchStartPercent);
+		FABRIKEnd = FABRIKStart + ClampedPoleVector;
+	}
 
-		if (ChainLength < NewChainLength && Settings.StretchMode == EStretchLimbStretchMode::SquashOnly)
-		{
-			return; // chain needs to stretch, but we only want to squash
-		}
+	// get positions for FABRIK
+	TArray<FVector> BonePositions;
+	BonePositions.Init(FVector::UpVector, BoneIndices.Num());
+	for (int32 ChainIndex = 0; ChainIndex < BoneIndices.Num(); ++ChainIndex)
+	{
+		const int32 BoneIndex = BoneIndices[ChainIndex];
+		BonePositions[ChainIndex] = InIKRigSkeleton.CurrentPoseGlobal[BoneIndex].GetLocation();
 	}
 
-	// record chain positions prior to translational offset
-	TArray<FVector> BonePositionsPreOffset;
-	for (const int32 BoneIndex : BoneIndices)
+	// run FABRIK iterations
+	for (int32 Iteration=0; Iteration<Settings.Iterations; ++Iteration)
 	{
-		BonePositionsPreOffset.Add(InIKRigSkeleton.CurrentPoseGlobal[BoneIndex].GetLocation());
+		RunFABRIKForwardPass(FABRIKStart, BonePositions);
+		RunFABRIKBackwardPass(FABRIKEnd, BonePositions);
 	}
 
-	// translate the bones
-	switch (Settings.TranslationType)
+	// always finish with a forward pass
+	RunFABRIKForwardPass(FABRIKStart, BonePositions);
+	
+	// apply FABRIK positions 
+	for (int32 ChainIndex = 1; ChainIndex < BoneIndices.Num(); ++ChainIndex)
 	{
-	case EStretchLimbTranslationType::BoneLengthOnly:
-		{
-			// easier to do this in local space, so update them first
-			for (const int32 BoneIndex : BoneIndices)
-			{
-				InIKRigSkeleton.UpdateLocalTransformFromGlobal(BoneIndex);
-			}
-			
-			// scale the length of each bone from root to tip
-			const double ScaleFactor = NewChainLength / ChainLength;
-			for (int32 ChainIndex = 1; ChainIndex < BoneIndices.Num(); ++ChainIndex)
-			{
-				FTransform& BoneLocalTransform = InIKRigSkeleton.CurrentPoseLocal[BoneIndices[ChainIndex]];
-				BoneLocalTransform.SetLocation(BoneLocalTransform.GetLocation() * ScaleFactor);
-			}
-
-			// update global pose of chain
-			for (const int32 BoneIndex : BoneIndices)
-			{
-				InIKRigSkeleton.UpdateGlobalTransformFromLocal(BoneIndex);
-			}
-			break;
-		}
-	case EStretchLimbTranslationType::AllAxes:
-		{
-			// apply global translational delta to the bone to move towards goal
-			const FVector NewChainDirection = NewChainVector.GetSafeNormal();
-			const FVector TotalChainDelta = NewChainDirection * (NewChainLength - ChainLength);
-			const double InvNumBones = 1.0 / static_cast<double>(BoneIndices.Num() - 1);
-			for (int32 ChainIndex = 1; ChainIndex < BoneIndices.Num(); ++ChainIndex)
-			{
-				// calculate amount of delta to apply to the bone (100% at tip, with linear falloff to root)
-				const double ScaleFactor = static_cast<double>(ChainIndex) * InvNumBones;
-				const FVector PerBoneDelta = TotalChainDelta * ScaleFactor;
-				// add delta to the bone
-				FTransform& BoneGlobalTransform = InIKRigSkeleton.CurrentPoseGlobal[BoneIndices[ChainIndex]];
-				BoneGlobalTransform.SetLocation( BoneGlobalTransform.GetLocation() + PerBoneDelta);
-			}
-			break;
-		}
-	default:
-		ensureMsgf(false, TEXT("Unknown chain length type."));
-		return;
+		const int32 BoneIndex = BoneIndices[ChainIndex];
+		InIKRigSkeleton.CurrentPoseGlobal[BoneIndex].SetLocation(BonePositions[ChainIndex]);
+	}
+}
+
+void FIKRigStretchLimbSolver::StoreBonePositions(FIKRigSkeleton& InIKRigSkeleton)
+{
+	for (int32 ChainIndex = 0; ChainIndex < BoneIndices.Num(); ++ChainIndex)
+	{
+		const int32 BoneIndex = BoneIndices[ChainIndex];
+		BonePositionsPreOffset[ChainIndex] = InIKRigSkeleton.CurrentPoseGlobal[BoneIndex].GetLocation();
 	}
+}
 
+void FIKRigStretchLimbSolver::UpdateBoneOrientations(FIKRigSkeleton& InIKRigSkeleton)
+{
 	// update orientations of bones based on translational offset
 	for (int32 ChainIndex = 0; ChainIndex < BoneIndices.Num() - 1; ++ChainIndex)
 	{
@@ -272,6 +354,113 @@ void FIKRigStretchLimbSolver::TranslateChain(
 	}
 }
 
+void FIKRigStretchLimbSolver::RunFABRIKBackwardPass(
+	const FVector& Goal,
+	TArray<FVector>& BonePositions)
+{
+	check(BonePositions.Num() >= 2);
+	check(RestLengths.Num() == BonePositions.Num());
+
+	// backward pass: leaf to root (keep leaf fixed)
+	BonePositions.Last() = Goal; // reset leaf
+	for (int32 i = BonePositions.Num() - 2; i >= 0; --i)
+	{
+		const double RestLength = RestLengths[i+1];
+		const FVector ToParent = BonePositions[i] - BonePositions[i+1];
+		const double CurrentLength = ToParent.Size();
+		if (CurrentLength < KINDA_SMALL_NUMBER)
+		{
+			continue;
+		}
+
+		// move bone towards child to satisfy rest length
+		BonePositions[i] = BonePositions[i+1] + (RestLength / CurrentLength) * ToParent;
+	}
+}
+
+void FIKRigStretchLimbSolver::RunFABRIKForwardPass(
+	const FVector& Start,
+	TArray<FVector>& BonePositions)
+{
+	check(BonePositions.Num() >= 2);
+	check(RestLengths.Num() == BonePositions.Num());
+
+	// forward pass: root to leaf (keep root fixed)
+	BonePositions[0] = Start; // reset root
+	for (int32 i = 1; i < BonePositions.Num(); ++i)
+	{
+		const double RestLength = RestLengths[i];
+		const FVector ToChild = BonePositions[i] - BonePositions[i-1];
+		const double CurrentLength = ToChild.Size();
+		if (CurrentLength < KINDA_SMALL_NUMBER)
+		{
+			continue;
+		}
+
+		// move bone towards parent to satisfy rest length
+		BonePositions[i] = BonePositions[i-1] + (RestLength / CurrentLength) * ToChild;
+	}	
+}
+
+double FIKRigStretchLimbSolver::CalculateCurrentChainLength(FIKRigSkeleton& InIKRigSkeleton) const
+{
+	double ChainLength = 0;
+	
+	// iterate from tip to first child of root
+	for (int32 BoneIndex = BoneIndices.Num() - 1; BoneIndex >= 1; --BoneIndex)
+	{
+		const FVector& BonePosition = InIKRigSkeleton.CurrentPoseGlobal[BoneIndices[BoneIndex]].GetLocation();
+		const FVector& ParentPosition = InIKRigSkeleton.CurrentPoseGlobal[BoneIndices[BoneIndex-1]].GetLocation();
+		ChainLength += (ParentPosition - BonePosition).Length();
+	}
+	
+	return ChainLength;
+}
+
+void FIKRigStretchLimbSolver::UpdateRestLengths(FIKRigSkeleton& InIKRigSkeleton)
+{
+	// reset/update rest lengths of bones (may be animated)
+	RestLengths.Reset(BoneIndices.Num());
+	for (const int32 BoneIndex : BoneIndices)
+	{
+		const int32 ParentIndex = InIKRigSkeleton.GetParentIndex(BoneIndex);
+		if (ParentIndex == INDEX_NONE)
+		{
+			RestLengths.Add( InIKRigSkeleton.RefPoseGlobal[BoneIndex].GetTranslation().Size() );
+			continue;
+		}
+		const FVector ParentPosition = InIKRigSkeleton.RefPoseGlobal[ParentIndex].GetTranslation();
+		const FVector BonePosition = InIKRigSkeleton.RefPoseGlobal[BoneIndex].GetTranslation();
+		RestLengths.Add((BonePosition - ParentPosition).Size());
+	}
+
+	TotalRestLength = Algo::Accumulate(RestLengths, 0.0);
+}
+
+void FIKRigStretchLimbSolver::UpdateFKChildren(FIKRigSkeleton& InIKRigSkeleton)
+{
+	// update chain local transforms
+	for (const int32 BoneIndex : BoneIndices)
+	{
+		InIKRigSkeleton.UpdateLocalTransformFromGlobal(BoneIndex);
+	}
+
+	// propagate to children
+	for (const int32 ChildIndex: ChildrenToUpdate)
+	{
+		InIKRigSkeleton.UpdateGlobalTransformFromLocal(ChildIndex);
+	}
+}
+
+void FIKRigStretchLimbSolver::RotateEndBoneWithGoal(FIKRigSkeleton& InIKRigSkeleton, const FIKRigGoal* Goal)
+{
+	// orient end bone to goal
+	FTransform& EndBoneGlobal = InIKRigSkeleton.CurrentPoseGlobal[BoneIndices.Last()];
+	FQuat EndBoneGlobalRotation = FQuat::FastLerp(EndBoneGlobal.GetRotation(), Goal->FinalBlendedRotation, Settings.RotateEndBoneWithGoal);
+	EndBoneGlobalRotation = EndBoneGlobalRotation.GetNormalized();
+	EndBoneGlobal.SetRotation(EndBoneGlobalRotation);
+}
+
 void FIKRigStretchLimbSolver::GetRequiredBones(TSet<FName>& OutRequiredBones) const
 {
 	OutRequiredBones.Add(Settings.StartBone);
@@ -359,7 +548,7 @@ bool FIKRigStretchLimbSolver::GetWarningMessage(FText& OutWarningMessage) const
 
 bool FIKRigStretchLimbSolver::IsBoneAffectedBySolver(const FName& InBoneName, const FIKRigSkeleton& InIKRigSkeleton) const
 {
-	return false;
+	return InIKRigSkeleton.IsBoneInDirectLineage(InBoneName, Settings.StartBone);
 }
 #endif
 
diff --git a/Engine/Plugins/Animation/IKRig/Source/IKRig/Public/Rig/Solvers/IKRigStretchLimb.h b/Engine/Plugins/Animation/IKRig/Source/IKRig/Public/Rig/Solvers/IKRigStretchLimb.h
--- a/Engine/Plugins/Animation/IKRig/Source/IKRig/Public/Rig/Solvers/IKRigStretchLimb.h
+++ b/Engine/Plugins/Animation/IKRig/Source/IKRig/Public/Rig/Solvers/IKRigStretchLimb.h
@@ -7,30 +7,9 @@
 
 #include "IKRigStretchLimb.generated.h"
 
-#define UE_API IKRIG_API
+struct FIKRigGoal;
 
-UENUM(BlueprintType)
-enum class EStretchLimbTranslationType : uint8
-{
-	AllAxes,
-	BoneLengthOnly,
-};
-
-UENUM(BlueprintType)
-enum class EStretchLimbLengthType : uint8
-{
-	StraightLine,
-	AlongBones,
-};
-
-UENUM(BlueprintType)
-enum class EStretchLimbStretchMode : uint8
-{
-	None,
-	SquashAndStretch,
-	StretchOnly,
-	SquashOnly,
-};
+#define UE_API IKRIG_API
 
 UENUM(BlueprintType)
 enum class EStretchLimbRotationMode : uint8
@@ -57,24 +36,34 @@ struct FIKRigStretchLimbSettings : public FIKRigSolverSettingsBase
 	FName Goal;
 	
 	/** Determines whether to squash or stretch the bones to reach the goal. */
-	UPROPERTY(BlueprintReadWrite, EditAnywhere, Category = "Translation Settings")
-	EStretchLimbStretchMode StretchMode = EStretchLimbStretchMode::StretchOnly;
+	UPROPERTY(BlueprintReadWrite, EditAnywhere, Category = "Stretch Settings")
+	bool bEnableStretching = true;
 
-	/** Consider the length of the bone chain as either a straight line from start to end or the total length along the bones. */
-	UPROPERTY(BlueprintReadWrite, EditAnywhere, Category = "Translation Settings")
-	EStretchLimbLengthType LengthType = EStretchLimbLengthType::AlongBones;
-	
-	/** Determines how to affect the translation of the bones in the chain. */
-	UPROPERTY(BlueprintReadWrite, EditAnywhere, Category = "Translation Settings")
-	EStretchLimbTranslationType TranslationType = EStretchLimbTranslationType::BoneLengthOnly;
+	/** The maximum distance the limb is allowed to stretch (beyond its rest length). Default is -1.0 (no limit).*/
+	UPROPERTY(BlueprintReadWrite, EditAnywhere, Category = "Stretch Settings")
+	double MaximumStretchDistance = -1.0;
+
+	/** The percentage of the bone chain length to straighten before stretching is applied (can prevent over extension).
+	 * NOTE: At 1.0, the limb will not stretch until the goal is 100% of the length of the limb away from the root of the chain. */
+	UPROPERTY(BlueprintReadWrite, EditAnywhere, Category = "Stretch Settings", meta = (ClampMin = "0.1", ClampMax = "1.0", UIMin = "0.1", UIMax = "1.0"))
+	double StretchStartPercent = 0.95;
 
 	/** Determines how to affect the orientation of the bones in the chain. */
 	UPROPERTY(BlueprintReadWrite, EditAnywhere, Category = "Rotation Settings")
 	EStretchLimbRotationMode RotationMode = EStretchLimbRotationMode::OrientToGoal;
 
 	/** Allow the end bone to be reoriented (0) or match the orientation of the goal (1). Range is 0-1. Default is 1.0. */
-	UPROPERTY(BlueprintReadWrite, EditAnywhere, Category = "Global Settings", meta = (UIMin = "0.0", UIMax = "1.0"))
-	double RotateEndBoneWithGoal = 0.0f;
+	UPROPERTY(BlueprintReadWrite, EditAnywhere, Category = "Rotation Settings", meta = (UIMin = "0.0", UIMax = "1.0"))
+	double RotateEndBoneWithGoal = 1.0f;
+
+	/** Number of FABRIK iterations to correct bones lengths. Adjust until the end bone converges on the goal. Default is 8. */
+	UPROPERTY(BlueprintReadWrite, EditAnywhere, Category = "Constraint Solving", meta = (UIMin = "0", UIMax = "50.0"))
+	int32 Iterations = 8;
+
+	/** The distance to squash the bones perpendicular to the pole vector when the limb is fully compressed.
+	 * NOTE: tune this to help dense chains achieve a reasonable pose. Must be used with constraint iterations to fix bones lengths. */
+	UPROPERTY(BlueprintReadWrite, EditAnywhere, Category = "Constraint Solving", meta = (UIMin = "0", UIMax = "5000.0"))
+	double SquashStrength = 50;
 };
 
 USTRUCT(BlueprintType)
@@ -120,13 +109,37 @@ struct FIKRigStretchLimbSolver : public FIKRigSolverBase
 
 private:
 
-	void OrientChain(FIKRigSkeleton& InIKRigSkeleton, const FVector& GoalLocation) const;
+	void AimChainAtGoal(FIKRigSkeleton& InIKRigSkeleton, const FVector& GoalLocation) const;
+
+	void StoreBonePositions(FIKRigSkeleton& InIKRigSkeleton);
+	void UpdateBoneOrientations(FIKRigSkeleton& InIKRigSkeleton);
 	
-	void TranslateChain(FIKRigSkeleton& InIKRigSkeleton, const FVector& GoalLocation) const;
+	void SquashOrStretchChainParallel(FIKRigSkeleton& InIKRigSkeleton, const FVector& GoalLocation) const;
+	void SquashChainPerpendicular(FIKRigSkeleton& InIKRigSkeleton);
+	void StretchChainFinal(FIKRigSkeleton& InIKRigSkeleton, const FVector& GoalLocation) const;
+	void ApplyParallelOffsetToChain(FIKRigSkeleton& InIKRigSkeleton, const FVector& ChainStart, const FVector& NewPoleVector, const FVector& CurrentPoleVector) const;
+
+	void RunFABRIK(FIKRigSkeleton& InIKRigSkeleton, const FVector& GoalLocation);
+	void RunFABRIKBackwardPass(const FVector& Goal, TArray<FVector>& BonePositions);
+	void RunFABRIKForwardPass(const FVector& Start, TArray<FVector>& BonePositions);
+
+	void UpdateFKChildren(FIKRigSkeleton& InIKRigSkeleton);
+
+	void RotateEndBoneWithGoal(FIKRigSkeleton& InIKRigSkeleton, const FIKRigGoal* Goal);
+	
+	double CalculateCurrentChainLength(FIKRigSkeleton& InIKRigSkeleton) const;
+	
+	void UpdateRestLengths(FIKRigSkeleton& InIKRigSkeleton);
 
 	bool bIsInitialized = false;
 	TArray<int32> BoneIndices;
 	TArray<int32> ChildrenToUpdate;
+
+	TArray<FVector> BonePositionsPreOffset;
+	TArray<double> RestLengths;
+	double TotalRestLength;
+	TArray<double> PoleVectorParams;
+	double InitialPoleVectorLenSq;
 };
 
 
