[IK Rig] Fixed bug in stretch limb not turning off iterations when at 0.

diff --git a/Engine/Plugins/Animation/IKRig/Source/IKRig/Private/Rig/Solvers/IKRigStretchLimb.cpp b/Engine/Plugins/Animation/IKRig/Source/IKRig/Private/Rig/Solvers/IKRigStretchLimb.cpp
--- a/Engine/Plugins/Animation/IKRig/Source/IKRig/Private/Rig/Solvers/IKRigStretchLimb.cpp
+++ b/Engine/Plugins/Animation/IKRig/Source/IKRig/Private/Rig/Solvers/IKRigStretchLimb.cpp
@@ -287,6 +287,11 @@ void FIKRigStretchLimbSolver::RunFABRIK(FIKRigSkeleton& InIKRigSkeleton, const F
 		return;
 	}
 
+	if (Settings.Iterations <= 0)
+	{
+		return;
+	}
+
 	// clamp goal to percent distance along chain (optionally prevents overextension)
 	const FVector FABRIKStart = InIKRigSkeleton.CurrentPoseGlobal[BoneIndices[0]].GetLocation();
 	FVector FABRIKEnd = GoalLocation;
