[IK Retarget] Fix non-normalized rotation from PelvisMotionOp rendering before first Run()

diff --git a/Engine/Plugins/Animation/IKRig/Source/IKRig/Private/Retargeter/RetargetOps/PelvisMotionOp.cpp b/Engine/Plugins/Animation/IKRig/Source/IKRig/Private/Retargeter/RetargetOps/PelvisMotionOp.cpp
--- a/Engine/Plugins/Animation/IKRig/Source/IKRig/Private/Retargeter/RetargetOps/PelvisMotionOp.cpp
+++ b/Engine/Plugins/Animation/IKRig/Source/IKRig/Private/Retargeter/RetargetOps/PelvisMotionOp.cpp
@@ -158,8 +158,8 @@ bool FIKRetargetPelvisMotionOp::InitializeSource(
 	// record initial root data
 	const FTransform InitialTransform = SourceSkeleton.RetargetPoses.GetGlobalRetargetPose()[Source.BoneIndex]; 
 	float InitialHeight = static_cast<float>(InitialTransform.GetTranslation().Z);
-	Source.InitialPosition = InitialTransform.GetLocation();
-	Source.InitialRotation = InitialTransform.GetRotation();
+	Source.InitialPosition = Source.CurrentPosition = InitialTransform.GetLocation();
+	Source.InitialRotation = Source.CurrentRotation = InitialTransform.GetRotation();
 
 	// ensure root height is not at origin, this happens if user sets root to ACTUAL skeleton root and not pelvis
 	if (InitialHeight < UE_KINDA_SMALL_NUMBER)
