[IK Retarget] Fixed Pelvis Motion floor constraint debug drawing

diff --git a/Engine/Plugins/Animation/IKRig/Source/IKRig/Private/Retargeter/RetargetOps/PelvisMotionOp.cpp b/Engine/Plugins/Animation/IKRig/Source/IKRig/Private/Retargeter/RetargetOps/PelvisMotionOp.cpp
--- a/Engine/Plugins/Animation/IKRig/Source/IKRig/Private/Retargeter/RetargetOps/PelvisMotionOp.cpp
+++ b/Engine/Plugins/Animation/IKRig/Source/IKRig/Private/Retargeter/RetargetOps/PelvisMotionOp.cpp
@@ -377,7 +377,10 @@ void FIKRetargetPelvisMotionOp::DebugDraw(
 {
 	FScopeLock ScopeLock(&DebugDataMutex);
 
-	const FTransform SourcePelvis = FTransform(Source.CurrentRotation, Source.CurrentPosition) * InSourceTransform;
+	// strip scale out of the source transform (already applied to pose)
+	const FTransform SourceComponentTransform(InSourceTransform.GetRotation(), InSourceTransform.GetLocation(), FVector::OneVector);
+	
+	const FTransform SourcePelvis = FTransform(Source.CurrentRotation, Source.CurrentPosition) * SourceComponentTransform;
 	const FTransform TargetPelvis = CurrentPelvisTransform * InComponentTransform;
 	const FVector TargetFloor = TargetPelvis.GetLocation() * FVector(1,1,0);
 	const FVector TargetCrotch = TargetPelvis.GetLocation() + FVector(0,0, -Settings.TargetCrotchOffset);
