[IKRetargeter] added RigMapper Op to IKRetargeter Op Stack

diff --git a/Engine/Plugins/Animation/IKRig/IKRig.uplugin b/Engine/Plugins/Animation/IKRig/IKRig.uplugin
--- a/Engine/Plugins/Animation/IKRig/IKRig.uplugin
+++ b/Engine/Plugins/Animation/IKRig/IKRig.uplugin
@@ -23,6 +23,10 @@
 		{
 			"Name": "FullBodyIK",
 			"Enabled": true
+		},
+		{
+			"Name": "RigMapper",
+			"Enabled": true
 		}
 	],
 	"Modules": [
diff --git a/Engine/Plugins/Animation/IKRig/Source/IKRig/IKRig.Build.cs b/Engine/Plugins/Animation/IKRig/Source/IKRig/IKRig.Build.cs
--- a/Engine/Plugins/Animation/IKRig/Source/IKRig/IKRig.Build.cs
+++ b/Engine/Plugins/Animation/IKRig/Source/IKRig/IKRig.Build.cs
@@ -18,6 +18,7 @@ public class IKRig : ModuleRules
 				"CoreUObject",
 				"Engine",
 				"AnimationCore",
+				"RigMapper",
 			});
 
 		PublicDependencyModuleNames.AddRange(
diff --git a/Engine/Plugins/Animation/IKRig/Source/IKRig/Private/Retargeter/IKRetargetProcessor.cpp b/Engine/Plugins/Animation/IKRig/Source/IKRig/Private/Retargeter/IKRetargetProcessor.cpp
--- a/Engine/Plugins/Animation/IKRig/Source/IKRig/Private/Retargeter/IKRetargetProcessor.cpp
+++ b/Engine/Plugins/Animation/IKRig/Source/IKRig/Private/Retargeter/IKRetargetProcessor.cpp
@@ -1449,10 +1449,18 @@ void FIKRetargetProcessor::OnAnimGraphPreUpdateMainThread(
 
 void FIKRetargetProcessor::OnAnimGraphEvaluateAnyThread(FPoseContext& Output)
 {
+	// work out which (if any) op is the first enabled curve processing op, so that in this case, we know to copy the inputs from the source anim instance
+	// after that, we take input from the curves output by the previous node
+	bool bFirstCurveOp = true;
 	for (FInstancedStruct& OpStruct : OpStack)
 	{
 		FIKRetargetOpBase& Op = OpStruct.GetMutable<FIKRetargetOpBase>();
+		Op.SetTakeInputCurvesFromSourceAnimInstance(bFirstCurveOp);
 		Op.AnimGraphEvaluateAnyThread(Output);
+		if (Op.IsEnabled() && Op.HasCurveProcessing())
+		{
+			bFirstCurveOp = false;
+		}
 	}
 }
 
diff --git a/Engine/Plugins/Animation/IKRig/Source/IKRig/Private/Retargeter/RetargetOps/CurveRemapOp.cpp b/Engine/Plugins/Animation/IKRig/Source/IKRig/Private/Retargeter/RetargetOps/CurveRemapOp.cpp
--- a/Engine/Plugins/Animation/IKRig/Source/IKRig/Private/Retargeter/RetargetOps/CurveRemapOp.cpp
+++ b/Engine/Plugins/Animation/IKRig/Source/IKRig/Private/Retargeter/RetargetOps/CurveRemapOp.cpp
@@ -59,23 +59,64 @@ void FIKRetargetCurveRemapOp::AnimGraphEvaluateAnyThread(FPoseContext& Output)
 	
 	FBlendedCurve& OutputCurves = Output.Curve;
 
-	// copy curves over with same name (if it exists)
+	TMap<FName, float> CurvesToRemapValues;
 	if (Settings.bCopyAllSourceCurves)
 	{
-		OutputCurves.CopyFrom(SourceCurves);
+		// copy curves over from anim instance
+		if (GetTakeInputCurvesFromSourceAnimInstance())
+		{
+			OutputCurves.CopyFrom(SourceCurves);
+		}
 	}
 
-	// copy curves over with different names (remap)
-	if (Settings.bRemapCurves)
+	if (!GetTakeInputCurvesFromSourceAnimInstance())
 	{
-		FBlendedCurve RemapedCurves;
+		// take a copy of the source curves we need to remap
 		for (const FCurveRemapPair& CurveToRemap : Settings.CurvesToRemap)
 		{
 			bool bOutIsValid = false;
-			const float SourceValue = SourceCurves.Get(CurveToRemap.SourceCurve, bOutIsValid);
+			const float SourceValue = OutputCurves.Get(CurveToRemap.SourceCurve, bOutIsValid);
 			if (bOutIsValid)
 			{
-				RemapedCurves.Add(CurveToRemap.TargetCurve, SourceValue);
+				CurvesToRemapValues.Add(CurveToRemap.SourceCurve, SourceValue);
+			}
+		}
+
+		// clear the outputs if needed
+		if (!Settings.bCopyAllSourceCurves)
+		{
+			OutputCurves.ForEachElement([&OutputCurves](const UE::Anim::FCurveElement& InCurveElement)
+				{
+					OutputCurves.Set(InCurveElement.Name, 0.0f);
+				});
+		}
+	}
+
+	// copy curves over with different names (remap)
+	if (Settings.bRemapCurves)
+	{
+		FBlendedCurve RemapedCurves;
+		if (GetTakeInputCurvesFromSourceAnimInstance())
+		{
+			for (const FCurveRemapPair& CurveToRemap : Settings.CurvesToRemap)
+			{
+				bool bOutIsValid = false;
+				const float SourceValue = SourceCurves.Get(CurveToRemap.SourceCurve, bOutIsValid);
+				if (bOutIsValid)
+				{
+					RemapedCurves.Add(CurveToRemap.TargetCurve, SourceValue);
+				}
+			}
+		}
+		else
+		{
+			for (const FCurveRemapPair& CurveToRemap : Settings.CurvesToRemap)
+			{
+				float* SourceValuePtr = CurvesToRemapValues.Find(CurveToRemap.SourceCurve);
+				if (SourceValuePtr)
+				{
+					RemapedCurves.Add(CurveToRemap.TargetCurve, *SourceValuePtr);
+				}
 			}
 		}
 		
@@ -83,6 +124,99 @@ void FIKRetargetCurveRemapOp::AnimGraphEvaluateAnyThread(FPoseContext& Output)
 	}
 }
 
+void FIKRetargetCurveRemapOp::ProcessAnimSequenceCurves(FIKRetargetOpBase::FCurveData InCurveMetaData, FIKRetargetOpBase::FFrameValues InCurveFrameValues,
+	FIKRetargetOpBase::FCurveData& OutCurveMetaData, FIKRetargetOpBase::FFrameValues& OutCurveFrameValues) const
+{
+	FIKRetargetOpBase::FCurveData SourceCurveMetaData;
+	FIKRetargetOpBase::FFrameValues SourceCurveFrameValues;
+	TArray<FName> ValidSourceTargetNames;
+
+	if (Settings.bRemapCurves)
+	{
+		// first copy the source curve data
+		SourceCurveFrameValues.SetNum(InCurveFrameValues.Num());
+		for (int32 Frame = 0; Frame < InCurveFrameValues.Num(); ++Frame)
+		{
+			SourceCurveFrameValues[Frame].Reserve(Settings.CurvesToRemap.Num());
+		}
+		SourceCurveMetaData.Names.Reserve(Settings.CurvesToRemap.Num());
+		SourceCurveMetaData.Flags.Reserve(Settings.CurvesToRemap.Num());
+		SourceCurveMetaData.Colors.Reserve(Settings.CurvesToRemap.Num());
+		ValidSourceTargetNames.Reserve(Settings.CurvesToRemap.Num());
+		int32 RemapCurveCounter = 0;
+
+		for (const FCurveRemapPair& CurveToRemap : Settings.CurvesToRemap)
+		{
+			int32 FoundSourceIndex = InCurveMetaData.Names.IndexOfByKey(CurveToRemap.SourceCurve);
+			if (FoundSourceIndex != INDEX_NONE)
+			{
+				SourceCurveMetaData.Names.Add(InCurveMetaData.Names[FoundSourceIndex]);
+				SourceCurveMetaData.Flags.Add(InCurveMetaData.Flags[FoundSourceIndex]);
+				SourceCurveMetaData.Colors.Add(InCurveMetaData.Colors[FoundSourceIndex]);
+				for (int32 Frame = 0; Frame < InCurveFrameValues.Num(); ++Frame)
+				{
+					SourceCurveFrameValues[Frame].Add(InCurveFrameValues[Frame][FoundSourceIndex]);
+				}
+				ValidSourceTargetNames.Add(Settings.CurvesToRemap[RemapCurveCounter].TargetCurve);
+			}
+			RemapCurveCounter++;
+		}
+	}
+
+	if (Settings.bCopyAllSourceCurves)
+	{
+		OutCurveFrameValues = MoveTemp(InCurveFrameValues);
+		OutCurveMetaData = MoveTemp(InCurveMetaData);
+	}
+	else
+	{
+		OutCurveFrameValues = {};
+		OutCurveFrameValues.SetNum(InCurveFrameValues.Num());
+		for (int32 Frame = 0; Frame < InCurveFrameValues.Num(); ++Frame)
+		{
+			OutCurveFrameValues[Frame].Reserve(ValidSourceTargetNames.Num());
+		}
+		OutCurveMetaData = {};
+		OutCurveMetaData.Names.Reserve(ValidSourceTargetNames.Num());
+		OutCurveMetaData.Flags.Reserve(ValidSourceTargetNames.Num());
+		OutCurveMetaData.Colors.Reserve(ValidSourceTargetNames.Num());
+	}
+
+	if (Settings.bRemapCurves)
+	{
+		int32 RemapCurveCounter = 0;
+		for (const FName& TargetCurve : ValidSourceTargetNames)
+		{
+			int32 FoundTargetIndex = OutCurveMetaData.Names.IndexOfByKey(TargetCurve);
+			if (FoundTargetIndex != INDEX_NONE)
+			{
+				// overwrite the target data with the source
+				for (int32 Frame = 0; Frame < OutCurveFrameValues.Num(); ++Frame)
+				{
+					OutCurveFrameValues[Frame][FoundTargetIndex] = SourceCurveFrameValues[Frame][RemapCurveCounter];
+				}
+				OutCurveMetaData.Names[FoundTargetIndex] = TargetCurve;
+				OutCurveMetaData.Flags[FoundTargetIndex] = SourceCurveMetaData.Flags[RemapCurveCounter];
+				OutCurveMetaData.Colors[FoundTargetIndex] = SourceCurveMetaData.Colors[RemapCurveCounter];
+			}
+			else
+			{
+				// add the target data from the source
+				for (int32 Frame = 0; Frame < OutCurveFrameValues.Num(); ++Frame)
+				{
+					OutCurveFrameValues[Frame].Add(SourceCurveFrameValues[Frame][RemapCurveCounter]);
+				}
+				OutCurveMetaData.Names.Add(TargetCurve);
+				OutCurveMetaData.Flags.Add(SourceCurveMetaData.Flags[RemapCurveCounter]);
+				OutCurveMetaData.Colors.Add(SourceCurveMetaData.Colors[RemapCurveCounter]);
+			}
+
+			RemapCurveCounter++;
+		}
+	}
+}
+
+
 FIKRetargetCurveRemapOpSettings UIKRetargetCurveRemapController::GetSettings()
 {
 	return *reinterpret_cast<FIKRetargetCurveRemapOpSettings*>(OpSettingsToControl);
diff --git a/Engine/Plugins/Animation/IKRig/Source/IKRig/Private/Retargeter/RetargetOps/RigMapperOp.cpp b/Engine/Plugins/Animation/IKRig/Source/IKRig/Private/Retargeter/RetargetOps/RigMapperOp.cpp
new file mode 100644
--- /dev/null
+++ b/Engine/Plugins/Animation/IKRig/Source/IKRig/Private/Retargeter/RetargetOps/RigMapperOp.cpp
@@ -0,0 +1,305 @@
+// Copyright Epic Games, Inc. All Rights Reserved.
+
+#include "Retargeter/RetargetOps/RigMapperOp.h"
+
+#include "Animation/AnimCurveUtils.h"
+#include "Engine/SkeletalMesh.h"
+#include "Animation/AnimInstance.h"
+#include "Components/SkeletalMeshComponent.h"
+#include "Animation/AnimNodeBase.h"
+#include "RigMapperDefinition.h"
+#include "Retargeter/IKRetargetProcessor.h"
+
+#include UE_INLINE_GENERATED_CPP_BY_NAME(RigMapperOp)
+
+const UClass* FIKRetargetRigMapperOpSettings::GetControllerType() const
+{
+	return UIKRetargetRigMapperOpController::StaticClass();
+}
+
+void FIKRetargetRigMapperOpSettings::CopySettingsAtRuntime(const FIKRetargetOpSettingsBase* InSettingsToCopyFrom)
+{
+	// copies all properties
+	const TArray<FName> PropertiesToIgnore = {};
+	FIKRetargetOpBase::CopyStructProperties(
+		FIKRetargetRigMapperOpSettings::StaticStruct(),
+		InSettingsToCopyFrom,
+		this,
+		PropertiesToIgnore);
+}
+
+void FIKRetargetRigMapperOp::AnimGraphPreUpdateMainThread(
+	USkeletalMeshComponent& SourceMeshComponent,
+    USkeletalMeshComponent& TargetMeshComponent)
+{
+	if (!IsEnabled())
+	{
+		return;
+	}
+	
+	SourceCurves.Empty();
+	
+	// get the source curves out of the source anim instance
+	const UAnimInstance* SourceAnimInstance = SourceMeshComponent.GetAnimInstance();
+	if (!SourceAnimInstance)
+	{
+		return;
+	}
+	
+	// Potential optimization/tradeoff: If we stored the curve results on the mesh component in non-editor scenarios, this would be
+	// much faster (but take more memory). As it is, we need to translate the map stored on the anim instance.
+	const TMap<FName, float>& AnimCurveList = SourceAnimInstance->GetAnimationCurveList(EAnimCurveType::AttributeCurve);
+	UE::Anim::FCurveUtils::BuildUnsorted(SourceCurves, AnimCurveList);
+
+
+	// Need reinit if the definitions in use have changed.
+	bool bReInit = false;
+
+	// If definitions were loaded from the SKM asset user data, they take priority and are the ones we will check against
+	TArray<TObjectPtr<URigMapperDefinition>>* DefinitionsToCheck = &Settings.Definitions;
+	if (LoadedUserData)
+	{
+		DefinitionsToCheck = &LoadedUserData->Definitions;
+	}
+
+	check(DefinitionsToCheck);
+
+	// Otherwise, reinit if the definitions are not the same
+	if (DefinitionsToCheck->Num() != LoadedDefinitions.Num())
+	{
+		bReInit = true;
+	}
+	else
+	{
+		for (int32 DefinitionIndex = 0; DefinitionIndex < DefinitionsToCheck->Num(); DefinitionIndex++)
+		{
+			const bool bDefinitionEdittedAndNotValidated = (*DefinitionsToCheck)[DefinitionIndex] != nullptr && !(*DefinitionsToCheck)[DefinitionIndex]->WasDefinitionValidated();
+			const bool bDefinitionHasChanged = (*DefinitionsToCheck)[DefinitionIndex] != LoadedDefinitions[DefinitionIndex];
+			if (bDefinitionHasChanged || bDefinitionEdittedAndNotValidated)
+			{
+				bReInit = true;
+				break;
+			}
+		}
+	}
+	if (bReInit)
+	{
+		USkeletalMesh* TargetMesh = TargetMeshComponent.GetSkeletalMeshAsset();
+		bIsInitialized = InitializeRigMapping(TargetMesh);
+	}
+
+}
+
+void FIKRetargetRigMapperOp::AnimGraphEvaluateAnyThread(FPoseContext& Output)
+{
+	if (!IsEnabled())
+	{
+		return;
+	}
+
+	FBlendedCurve& OutputCurves = Output.Curve;
+
+	FBlendedCurve InputCurves;
+	if (GetTakeInputCurvesFromSourceAnimInstance())
+	{
+		InputCurves.CopyFrom(SourceCurves);
+		if (Settings.bCopyAllSourceCurves)
+		{
+			OutputCurves.CopyFrom(SourceCurves);
+		}
+	}
+	else
+	{
+		InputCurves.CopyFrom(OutputCurves);
+		// clear the outputs if needed
+		if (!Settings.bCopyAllSourceCurves)
+		{
+			OutputCurves.ForEachElement([&OutputCurves](const UE::Anim::FCurveElement& InCurveElement)
+				{
+					OutputCurves.Set(InCurveElement.Name, 0.0f);
+				});
+		}
+
+	}
+
+	EvaluateRigMapping(InputCurves, OutputCurves);
+
+}
+
+void FIKRetargetRigMapperOp::EvaluateRigMapping(const FBlendedCurve& InCurve, FBlendedCurve& OutCurve)
+{
+	if (!RigMapperProcessor.IsValid())
+	{
+		return;
+	}
+	
+	// Retrieve inputs
+	CachedInputValues.Reset(InputCurveMappings.Num());
+	CachedInputValues.AddDefaulted(InputCurveMappings.Num());
+	UE::Anim::FCurveUtils::BulkGet(InCurve, InputCurveMappings,
+		[&Inputs = CachedInputValues](const FRigMapperCurveMapping& InBulkElement, const float InValue)
+		{
+			Inputs[InBulkElement.CurveIndex] = InValue;
+		}
+	);
+
+	{
+		// Evaluate frame
+		RigMapperProcessor.EvaluateFrame(RigMapperProcessor.GetInputNames(), CachedInputValues, CachedOutputValues);
+	}
+
+	{
+		// Set all output curves from the given output values.
+		UE::Anim::FCurveUtils::BulkSet(OutCurve, OutputCurveMappings,
+			[&Inputs = CachedInputValues, &Outputs = CachedOutputValues](const FRigMapperOutputCurveMapping& InBulkElement)
+			{
+				TOptional<float> Value = Outputs[InBulkElement.CurveIndex];
+	
+				// If the value was not set by the rig mapper default to 0
+				return Value.Get(0.f);
+			}
+		);
+	}
+}
+
+
+bool FIKRetargetRigMapperOp::InitializeRigMapping(const USkeletalMesh* InTargetMesh)
+{
+	TArray<URigMapperDefinition*> DefinitionsToLoad;
+
+	// Retrieve the definitions to load and use (either the ones set in the op, or the ones overriden by the SKM asset user data)
+	if (InTargetMesh)
+	{
+		// Note that we cast away constness here as we need to get at the Asset User Data. This is not changed anywhere so is a reasonable pragmatic solution
+		LoadedUserData = Cast<URigMapperDefinitionUserData>(const_cast<USkeletalMesh*>(InTargetMesh)->GetAssetUserDataOfClass(URigMapperDefinitionUserData::StaticClass()));
+		if (LoadedUserData)
+		{
+			DefinitionsToLoad = LoadedUserData->Definitions;
+		}
+	}
+	if (DefinitionsToLoad.IsEmpty())
+	{
+		DefinitionsToLoad = Settings.Definitions;
+	}
+
+	RigMapperProcessor = FRigMapperProcessor(DefinitionsToLoad);
+	LoadedDefinitions = MoveTemp(DefinitionsToLoad);
+
+	if (!RigMapperProcessor.IsValid())
+	{
+		return false;
+	}
+
+	// Cache a map of curve indices to bulk get current curve values for the input
+	const TArray<FName>& InputNames = RigMapperProcessor.GetInputNames();
+	InputCurveMappings.Empty();
+	InputCurveMappings.Reserve(InputNames.Num());
+	for (int32 InputIndex = 0; InputIndex < InputNames.Num(); InputIndex++)
+	{
+		InputCurveMappings.Add(InputNames[InputIndex], InputIndex);
+	}
+
+	// Cache a map of curve indices to bulk set the new curve values for the output
+	const TArray<FName>& OutputNames = RigMapperProcessor.GetOutputNames();
+	OutputCurveMappings.Empty();
+	OutputCurveMappings.Reserve(OutputNames.Num());
+	for (int32 OutputIndex = 0; OutputIndex < OutputNames.Num(); OutputIndex++)
+	{
+		const FName& CurveName = OutputNames[OutputIndex];
+
+		int32 InputIndex = InputNames.Find(CurveName);
+		OutputCurveMappings.Add(OutputNames[OutputIndex], OutputIndex, InputIndex);
+	}
+
+	return true;
+}
+
+static FLinearColor MakeCurveColor(const FName& InCurveName)
+{
+	// Create a color based on the hash of the name
+	FRandomStream Stream(GetTypeHash(InCurveName));
+	const uint8 Hue = (uint8)(Stream.FRand() * 255.0f);
+	return FLinearColor::MakeFromHSV8(Hue, 196, 196);
+}
+
+void FIKRetargetRigMapperOp::ProcessAnimSequenceCurves(FIKRetargetOpBase::FCurveData InCurveMetaData, FIKRetargetOpBase::FFrameValues InCurveFrameValues,
+	FIKRetargetOpBase::FCurveData& OutCurveMetaData, FIKRetargetOpBase::FFrameValues& OutCurveFrameValues) const
+{
+	// We need to copy the RigMapperProcessor to keep this method const. Revisit this for speed at some point.
+	FRigMapperProcessor RigMapperProcessorTemp = RigMapperProcessor;
+	TSet<FName> OutputNames;
+
+	if (RigMapperProcessorTemp.IsValid())
+	{
+		RigMapperProcessorTemp.EvaluateFrames(InCurveMetaData.Names, InCurveFrameValues, OutCurveMetaData.Names, OutCurveFrameValues);
+		OutCurveMetaData.Flags.Init(4, OutCurveMetaData.Names.Num()); // a value of 4 means Editable curve
+		OutCurveMetaData.Colors.SetNum(OutCurveMetaData.Names.Num());
+		for (int32 Curve = 0; Curve < OutCurveMetaData.Names.Num(); ++Curve)
+		{
+			OutCurveMetaData.Colors[Curve] = MakeCurveColor(OutCurveMetaData.Names[Curve]);
+			OutputNames.Add(OutCurveMetaData.Names[Curve]);
+		}
+	}
+
+	// 'pass-through' any input curves which have not been modified, if bCopyAllSourceCurves is true
+	if (Settings.bCopyAllSourceCurves)
+	{
+		TArray<int32> InputCurveIndicesToAdd;
+		InputCurveIndicesToAdd.Reserve(InCurveMetaData.Names.Num());
+		for (int32 InputCurve = 0; InputCurve < InCurveMetaData.Names.Num(); ++InputCurve)
+		{
+			if (!OutputNames.Contains(InCurveMetaData.Names[InputCurve]))
+			{
+				InputCurveIndicesToAdd.Add(InputCurve);
+			}
+		}
+
+		// now add these curves
+		OutCurveMetaData.Names.Reserve(OutCurveMetaData.Names.Num() + InputCurveIndicesToAdd.Num());
+		OutCurveMetaData.Flags.Reserve(OutCurveMetaData.Flags.Num() + InputCurveIndicesToAdd.Num());
+		OutCurveMetaData.Colors.Reserve(OutCurveMetaData.Colors.Num() + InputCurveIndicesToAdd.Num());
+		int32 Frame = 0;
+		for (TArray<TOptional<float>>& OutputFrame : OutCurveFrameValues)
+		{
+			OutputFrame.Reserve(OutputFrame.Num() + InputCurveIndicesToAdd.Num());
+			for (int32 CurveIndex : InputCurveIndicesToAdd)
+			{
+				OutputFrame.Add(InCurveFrameValues[Frame][CurveIndex]);
+			}
+
+			Frame++;
+		}
+
+		for (int32 CurveIndex : InputCurveIndicesToAdd)
+		{
+			OutCurveMetaData.Names.Add(InCurveMetaData.Names[CurveIndex]);
+			OutCurveMetaData.Flags.Add(InCurveMetaData.Flags[CurveIndex]);
+			OutCurveMetaData.Colors.Add(InCurveMetaData.Colors[CurveIndex]);
+		}
+	}
+}
+
+bool FIKRetargetRigMapperOp::Initialize(
+	const FIKRetargetProcessor& InProcessor,
+	const FRetargetSkeleton& /*InSourceSkeleton*/,
+	const FTargetSkeleton& /*InTargetSkeleton*/,
+	const FIKRetargetOpBase* /*InParentOp*/,
+	FIKRigLogger& /*Log*/)
+{
+	const FRetargetSkeleton& TargetSkeleton = InProcessor.GetSkeleton(ERetargetSourceOrTarget::Target);
+	bIsInitialized = InitializeRigMapping(TargetSkeleton.SkeletalMesh);
+
+	return bIsInitialized;
+};
+
+
+
+FIKRetargetRigMapperOpSettings UIKRetargetRigMapperOpController::GetSettings()
+{
+	return *reinterpret_cast<FIKRetargetRigMapperOpSettings*>(OpSettingsToControl);
+}
+
+void UIKRetargetRigMapperOpController::SetSettings(FIKRetargetRigMapperOpSettings InSettings)
+{
+	OpSettingsToControl->CopySettingsAtRuntime(&InSettings);
+}
diff --git a/Engine/Plugins/Animation/IKRig/Source/IKRig/Public/Retargeter/IKRetargetOps.h b/Engine/Plugins/Animation/IKRig/Source/IKRig/Public/Retargeter/IKRetargetOps.h
--- a/Engine/Plugins/Animation/IKRig/Source/IKRig/Public/Retargeter/IKRetargetOps.h
+++ b/Engine/Plugins/Animation/IKRig/Source/IKRig/Public/Retargeter/IKRetargetOps.h
@@ -130,6 +130,41 @@ public:
 
 	/** (optional) called during the PostLoad of the owning asset. Can be used to apply patches to the data for backwards compatibility. */
 	virtual void PostLoad(const FIKRigObjectVersion::Type InVersion) {};
+
+	using FFrameValues = TArray<TArray<TOptional<float>>>;
+
+	struct FCurveData
+	{
+		TArray<FName> Names;
+		TArray<int32> Flags;
+		TArray<FLinearColor> Colors;
+	};
+
+	/** (optional) a function which can be used during batch processing to apply any operations to curve data for each frame of an anim sequence. By default, this just moves the input data to the output data*/
+	virtual void ProcessAnimSequenceCurves(FIKRetargetOpBase::FCurveData InCurveMetaData, FIKRetargetOpBase::FFrameValues InCurveFrameValues,
+		FIKRetargetOpBase::FCurveData& OutCurveMetaData, FIKRetargetOpBase::FFrameValues& OutCurveFrameValues) const
+	{
+		OutCurveMetaData = MoveTemp(InCurveMetaData);
+		OutCurveFrameValues = MoveTemp(InCurveFrameValues);
+	}
+
+	/** (optional) a function which returns true if the operation performs curve processing, false otherwise (default is false) */
+	virtual bool HasCurveProcessing() const 
+	{
+		return false; 
+	}
+
+	/** a function which allows the user to manually indicate whether to take source curves from the source anim instance, or from the current pose context, since only the first enabled curve processing node must take inputs from the source anim instance*/
+	void SetTakeInputCurvesFromSourceAnimInstance(bool bInTakeInputCurvesFromSourceAnimInstance)
+	{
+		bTakeInputCurvesFromSourceAnimInstance = bInTakeInputCurvesFromSourceAnimInstance;
+	}
+
+	/** a function which allows the user to manually indicate whether to take source curves from the source anim instance, or from the current pose context, since only the first enabled curve processing node must take inputs from the source anim instance */
+	bool GetTakeInputCurvesFromSourceAnimInstance() const
+	{
+		return bTakeInputCurvesFromSourceAnimInstance;
+	}
 	
 	/** (optional) override to automate initial setup after being added to the stack */
 	virtual void OnAddedToStack(const UIKRetargeter* InRetargetAsset, const FIKRetargetOpBase* InParentOp) {};
@@ -302,6 +337,10 @@ private:
 	UPROPERTY()
 	FName ParentOpName;
 
+	/** (optional) determines whether input curves are taken from the source anim instance or from the current pose context (for example when we have multiple nodes which process curves). Can be set on the fly */
+	UPROPERTY()
+	bool bTakeInputCurvesFromSourceAnimInstance = true;
+
 protected:
 	
 	/** set to true when the op is ready to run */
diff --git a/Engine/Plugins/Animation/IKRig/Source/IKRig/Public/Retargeter/RetargetOps/CurveRemapOp.h b/Engine/Plugins/Animation/IKRig/Source/IKRig/Public/Retargeter/RetargetOps/CurveRemapOp.h
--- a/Engine/Plugins/Animation/IKRig/Source/IKRig/Public/Retargeter/RetargetOps/CurveRemapOp.h
+++ b/Engine/Plugins/Animation/IKRig/Source/IKRig/Public/Retargeter/RetargetOps/CurveRemapOp.h
@@ -101,6 +101,14 @@ struct FIKRetargetCurveRemapOp : public FIKRetargetOpBase
 
 	virtual bool IsSingleton() const override { return true; };
 
+	virtual bool HasCurveProcessing() const override 
+	{ 
+		return true; 
+	}
+
+	virtual void ProcessAnimSequenceCurves(FIKRetargetOpBase::FCurveData InCurveMetaData, FIKRetargetOpBase::FFrameValues InCurveFrameValues,
+		FIKRetargetOpBase::FCurveData& OutCurveMetaData, FIKRetargetOpBase::FFrameValues& OutCurveFrameValues)const override;
+
 	UE_API virtual void AnimGraphPreUpdateMainThread(USkeletalMeshComponent& SourceMeshComponent, USkeletalMeshComponent& TargetMeshComponent) override;
 
 	UE_API virtual void AnimGraphEvaluateAnyThread(FPoseContext& Output) override;
diff --git a/Engine/Plugins/Animation/IKRig/Source/IKRig/Public/Retargeter/RetargetOps/RigMapperOp.h b/Engine/Plugins/Animation/IKRig/Source/IKRig/Public/Retargeter/RetargetOps/RigMapperOp.h
new file mode 100644
--- /dev/null
+++ b/Engine/Plugins/Animation/IKRig/Source/IKRig/Public/Retargeter/RetargetOps/RigMapperOp.h
@@ -0,0 +1,183 @@
+// Copyright Epic Games, Inc. All Rights Reserved.
+#pragma once
+
+#include "Retargeter/IKRetargetOps.h"
+#include "StructUtils/InstancedStruct.h"
+#include "Animation/AnimCurveTypes.h"
+#include "RigMapperProcessor.h"
+
+#include "RigMapperOp.generated.h"
+
+#define UE_API IKRIG_API
+
+#define LOCTEXT_NAMESPACE "RigMapperOp"
+
+class URigMapperDefinitionUserData;
+class URigMapperDefinition;
+class USkeletalMesh;
+class USkeletalMeshComponent;
+
+USTRUCT(BlueprintType, meta = (DisplayName = "RigMapper Settings"))
+struct FIKRetargetRigMapperOpSettings : public FIKRetargetOpSettingsBase
+{
+	GENERATED_BODY()
+
+	// Whether to copy all curves over to the target animation instance
+	// NOTE: This setting also applies when exporting retargeted animations.
+	// True: all source curves are copied to the target animation instance/asset
+	// False: only remapped curves are copied on the target animation instance/asset
+	// In general, we should set this to true if the source and target rig are the same,
+	// and the RigMapper covers only a subset of the controls, but false otherwise.
+	UPROPERTY(EditAnywhere, BlueprintReadWrite, Category = "Copy Curves")
+	bool bCopyAllSourceCurves = false;
+
+	UPROPERTY(EditAnywhere, BlueprintReadWrite, Category = RigMapper)
+	TArray<TObjectPtr<URigMapperDefinition>> Definitions;
+
+	UE_API virtual const UClass* GetControllerType() const override;
+	
+	UE_API virtual void CopySettingsAtRuntime(const FIKRetargetOpSettingsBase* InSettingsToCopyFrom) override;
+};
+
+USTRUCT(BlueprintType, meta = (DisplayName = "RigMapper"))
+struct FIKRetargetRigMapperOp : public FIKRetargetOpBase
+{
+	GENERATED_BODY()
+	
+	// NOTE: this op does not do anything in Initialize() or Run().
+	
+	virtual bool Initialize(
+		const FIKRetargetProcessor& InProcessor,
+		const FRetargetSkeleton& InSourceSkeleton,
+		const FTargetSkeleton& InTargetSkeleton,
+		const FIKRetargetOpBase* InParentOp,
+		FIKRigLogger& Log) override;
+	
+	virtual void Run(
+		FIKRetargetProcessor& InProcessor,
+		const double InDeltaTime,
+		const TArray<FTransform>& InSourceGlobalPose,
+		TArray<FTransform>& OutTargetGlobalPose) override {};
+
+	virtual FIKRetargetOpSettingsBase* GetSettings() override
+	{
+		return &Settings;
+	};
+
+	virtual void SetSettings(const FIKRetargetOpSettingsBase* InSettings) override
+	{
+		Settings = *reinterpret_cast<const FIKRetargetRigMapperOpSettings*>(InSettings);
+	};
+	
+	virtual const UScriptStruct* GetSettingsType() const override
+	{
+		return FIKRetargetRigMapperOpSettings::StaticStruct();
+	}
+	
+	virtual const UScriptStruct* GetType() const override
+	{
+		return FIKRetargetRigMapperOp::StaticStruct();
+	}
+
+	/* RigMapperOp is a singleton because it can override definitions from the Target SkelMesh user data and therefore we should only do that once*/
+	virtual bool IsSingleton() const override { return true; };
+
+	virtual bool HasCurveProcessing() const override 
+	{ 
+		return true; 
+	}
+
+	virtual void ProcessAnimSequenceCurves(FIKRetargetOpBase::FCurveData InCurveMetaData, FIKRetargetOpBase::FFrameValues InCurveFrameValues,
+		FIKRetargetOpBase::FCurveData& OutCurveMetaData, FIKRetargetOpBase::FFrameValues& OutCurveFrameValues)const override;
+
+
+	UE_API virtual void AnimGraphPreUpdateMainThread(USkeletalMeshComponent& SourceMeshComponent, USkeletalMeshComponent& TargetMeshComponent) override;
+
+	UE_API virtual void AnimGraphEvaluateAnyThread(FPoseContext& Output) override;
+
+
+	UPROPERTY()
+	FIKRetargetRigMapperOpSettings Settings;
+	
+	// cached curves, copied on the game thread in PreUpdate()
+	FBlendedHeapCurve SourceCurves;
+
+
+private:
+	bool InitializeRigMapping(const USkeletalMesh* InTargetMesh);
+	void EvaluateRigMapping(const FBlendedCurve& InCurve, FBlendedCurve& OutCurve);
+
+	// The definitions that we have loaded. Cached to check against changes and reinit if need be
+	TArray<TObjectPtr<URigMapperDefinition>> LoadedDefinitions;
+
+	// The asset user data currently used to override definitions, if any was set on the skeletal mesh
+	TObjectPtr<URigMapperDefinitionUserData> LoadedUserData;
+
+	// The processor to evaluate the rig mapping
+	FRigMapperProcessor RigMapperProcessor;
+
+	// The cached input values passed to the rig mapper processor to avoid reallocations
+	FRigMapperProcessor::FPoseValues CachedInputValues;
+
+	// The cached output values passed to the rig mapper processor to avoid reallocations
+	FRigMapperProcessor::FPoseValues CachedOutputValues;
+
+	// Base curve mapping for bulk get/set of the linked pose curves
+	struct FRigMapperCurveMapping
+	{
+		FRigMapperCurveMapping() = default;
+
+		FRigMapperCurveMapping(FName InName, int32 InCurveIndex)
+			: Name(InName)
+			, CurveIndex(InCurveIndex)
+		{}
+
+		FName Name = NAME_None;
+		int32 CurveIndex = INDEX_NONE;
+	};
+
+	// Curve mapping for bulk set of the curves with a mapping to the matching input
+	struct FRigMapperOutputCurveMapping : FRigMapperCurveMapping
+	{
+		FRigMapperOutputCurveMapping() = default;
+
+		FRigMapperOutputCurveMapping(FName InName, int32 InOutputCurveIndex, int32 InInputCurveIndex)
+			: FRigMapperCurveMapping(InName, InOutputCurveIndex)
+			, InputCurveIndex(InInputCurveIndex)
+		{}
+
+		int32 InputCurveIndex = INDEX_NONE;
+	};
+
+	// Input curve mapping for bulk get of the curve values
+	using FInputCurveMappings = UE::Anim::TNamedValueArray<FDefaultAllocator, FRigMapperCurveMapping>;
+	FInputCurveMappings InputCurveMappings;
+
+	// Output curve mapping for the bulk get of the curve values
+	using FOutputCurveMappings = UE::Anim::TNamedValueArray<FDefaultAllocator, FRigMapperOutputCurveMapping>;
+	FOutputCurveMappings OutputCurveMappings;
+
+};
+
+/* The blueprint/python API for editing a RigMapper Op */
+UCLASS(MinimalAPI, BlueprintType)
+class UIKRetargetRigMapperOpController : public UIKRetargetOpControllerBase
+{
+	GENERATED_BODY()
+	
+public:
+	/* Get the current op settings as a struct.
+	 * @return FIKRetargetRigMapperOpSettings struct with the current settings used by the op. */
+	UFUNCTION(BlueprintCallable, Category = Settings)
+	UE_API FIKRetargetRigMapperOpSettings GetSettings();
+
+	/* Set the op settings. Input is a custom struct type for this op.
+	 * @param InSettings a FIKRetargetRigMapperOpSettings struct containing all the settings to apply to this op */
+	UFUNCTION(BlueprintCallable, Category = Settings)
+	UE_API void SetSettings(FIKRetargetRigMapperOpSettings InSettings);
+};
+
+
+#undef LOCTEXT_NAMESPACE
+
+#undef UE_API
diff --git a/Engine/Plugins/Animation/IKRig/Source/IKRigEditor/IKRigEditor.Build.cs b/Engine/Plugins/Animation/IKRig/Source/IKRigEditor/IKRigEditor.Build.cs
--- a/Engine/Plugins/Animation/IKRig/Source/IKRigEditor/IKRigEditor.Build.cs
+++ b/Engine/Plugins/Animation/IKRig/Source/IKRigEditor/IKRigEditor.Build.cs
@@ -51,7 +51,8 @@ namespace UnrealBuildTool.Rules
                     "AnimationWidgets",
                     "ApplicationCore",
 					"MeshDescription",
-					"SkeletalMeshDescription"
+					"SkeletalMeshDescription",
+					"RigMapperEditor"
                 }
             );
         }
diff --git a/Engine/Plugins/Animation/IKRig/Source/IKRigEditor/Private/RetargetEditor/IKRetargetBatchOperation.cpp b/Engine/Plugins/Animation/IKRig/Source/IKRigEditor/Private/RetargetEditor/IKRetargetBatchOperation.cpp
--- a/Engine/Plugins/Animation/IKRig/Source/IKRigEditor/Private/RetargetEditor/IKRetargetBatchOperation.cpp
+++ b/Engine/Plugins/Animation/IKRig/Source/IKRigEditor/Private/RetargetEditor/IKRetargetBatchOperation.cpp
@@ -7,7 +7,6 @@
 #include "AnimationBlueprintLibrary.h"
 #include "AnimPose.h"
 #include "AssetToolsModule.h"
-#include "Animation/AnimSequence.h"
 #include "ContentBrowserModule.h"
 #include "EditorReimportHandler.h"
 #include "IContentBrowserSingleton.h"
@@ -22,10 +21,10 @@
 #include "Widgets/Notifications/SNotificationList.h"
 #include "Animation/AnimMontage.h"
 #include "Retargeter/IKRetargetOps.h"
-#include "Retargeter/RetargetOps/CurveRemapOp.h"
 #include "ObjectTools.h"
 #include "AssetRegistry/AssetRegistryModule.h"
 #include "Retargeter/RetargetOps/SpeedPlantingOp.h"
+#include "RigMapperEditorSubsystem.h"
 
 #include UE_INLINE_GENERATED_CPP_BY_NAME(IKRetargetBatchOperation)
 
@@ -212,6 +211,7 @@ void UIKRetargetBatchOperation::RetargetAssets(
 	USkeleton* NewSkeleton = Context.TargetMesh->GetSkeleton();
 
 	TArray<FAdditiveRetargetSettings> SettingsToRestoreAfterRetarget;
+
 	
 	for (UAnimationAsset* AssetToRetarget : AnimationAssetsToRetarget)
 	{
@@ -225,9 +225,6 @@ void UIKRetargetBatchOperation::RetargetAssets(
 		{
 			FString AssetName = AnimSequenceToRetarget->GetName();
 			Progress.EnterProgressFrame(1.f, FText::Format(LOCTEXT("PreparingAsset", "Preparing asset: {0}"), FText::FromString(AssetName)));
-			
-			// copy curve data from source asset, preserving data in the target if present.
-			UAnimationBlueprintLibrary::CopyAnimationCurveNamesToSkeleton(OldSkeleton, NewSkeleton, AnimSequenceToRetarget, ERawCurveTrackTypes::RCT_Float);	
 
 			// clear transform curves since those curves won't work in new skeleton
 			IAnimationDataController& Controller = AnimSequenceToRetarget->GetController();
@@ -277,8 +274,20 @@ void UIKRetargetBatchOperation::RetargetAssets(
 		AssetToRetarget->MarkPackageDirty();
 	}
 
-	// convert the animation using the IK retargeter
-	ConvertAnimation(Context,Progress);
+	// convert the animation using the IK retargeter	
+	// initialize the retarget processor
+	FIKRetargetProcessor Processor;
+	FRetargetProfile RetargetProfile;
+	RetargetProfile.FillProfileWithAssetSettings(Context.IKRetargetAsset);
+	Processor.Initialize(Context.SourceMesh, Context.TargetMesh, Context.IKRetargetAsset, RetargetProfile);
+	if (!Processor.IsInitialized())
+	{
+		UE_LOG(LogTemp, Warning, TEXT("Unable to initialize the IK Retargeter. Newly created animations were not retargeted!"));
+		return;
+	}
+
+	ConvertAnimation(Context, Processor, Progress);
+
 
 	// optionally restore the additive flags
 	if (Context.bRetainAdditiveFlags)
@@ -324,27 +333,17 @@ void UIKRetargetBatchOperation::RetargetAssets(
 		AnimBlueprint->MarkPackageDirty();
 	}
 
-	// copy/remap curves to duplicate sequences
-	RemapCurves(Context, Progress);
+	// apply any curve operations to duplicate sequences
+	ApplyCurveOps(Context, Processor, Progress);
 }
 
 void UIKRetargetBatchOperation::ConvertAnimation(
 	const FIKRetargetBatchOperationContext& Context,
+	FIKRetargetProcessor& OutProcessor,
 	FScopedSlowTask& Progress)
 {
-	// initialize the retargeter
-	FIKRetargetProcessor Processor;
-	FRetargetProfile RetargetProfile;
-	RetargetProfile.FillProfileWithAssetSettings(Context.IKRetargetAsset);
-	Processor.Initialize(Context.SourceMesh, Context.TargetMesh, Context.IKRetargetAsset, RetargetProfile);
-	if (!Processor.IsInitialized())
-	{
-		UE_LOG(LogTemp, Warning, TEXT("Unable to initialize the IK Retargeter. Newly created animations were not retargeted!"));
-		return;
-	}
-
 	// target skeleton data
-	const FRetargetSkeleton& TargetSkeleton = Processor.GetSkeleton(ERetargetSourceOrTarget::Target);
+	const FRetargetSkeleton& TargetSkeleton = OutProcessor.GetSkeleton(ERetargetSourceOrTarget::Target);
 	const TArray<FName>& TargetBoneNames = TargetSkeleton.BoneNames;
 	const int32 NumTargetBones = TargetBoneNames.Num();
 
@@ -353,7 +352,7 @@ void UIKRetargetBatchOperation::ConvertAnimation(
 	BoneTracks.SetNumZeroed(NumTargetBones);
 
 	// source skeleton data
-	const FRetargetSkeleton& SourceSkeleton = Processor.GetSkeleton(ERetargetSourceOrTarget::Source);
+	const FRetargetSkeleton& SourceSkeleton = OutProcessor.GetSkeleton(ERetargetSourceOrTarget::Source);
 	const TArray<FName>& SourceBoneNames = SourceSkeleton.BoneNames;
 	const int32 NumSourceBones = SourceBoneNames.Num();
 
@@ -418,7 +417,7 @@ void UIKRetargetBatchOperation::ConvertAnimation(
 		EvaluationOptions.bIncorporateRootMotionIntoPose = true;
 
 		// reset playback of ops
-		Processor.OnPlaybackReset();
+		OutProcessor.OnPlaybackReset();
 		
 		// retarget each frame's pose from source to target
 		for (int32 FrameIndex=0; FrameIndex<NumFrames; ++FrameIndex)
@@ -476,10 +475,10 @@ void UIKRetargetBatchOperation::ConvertAnimation(
 			SettingsProfile.FillProfileWithAssetSettings(Context.IKRetargetAsset);
 
 			// give processor the pose to apply source scaling (if needed)
-			Processor.ScaleSourcePose(SourceComponentPose);
+			OutProcessor.ScaleSourcePose(SourceComponentPose);
 
 			// run the retargeter
-			const TArray<FTransform>& TargetComponentPose = Processor.RunRetargeter(SourceComponentPose, SettingsProfile, DeltaTime);
+			const TArray<FTransform>& TargetComponentPose = OutProcessor.RunRetargeter(SourceComponentPose, SettingsProfile, DeltaTime);
 
 			// convert to a local-space pose
 			TArray<FTransform> TargetLocalPose = TargetComponentPose;
@@ -522,47 +521,37 @@ void UIKRetargetBatchOperation::ConvertAnimation(
 	}
 }
 
-void UIKRetargetBatchOperation::RemapCurves(const FIKRetargetBatchOperationContext& Context, FScopedSlowTask& Progress)
+void UIKRetargetBatchOperation::ApplyCurveOps(const FIKRetargetBatchOperationContext& Context, const FIKRetargetProcessor& InProcessor, FScopedSlowTask& Progress)
 {
 	USkeleton* SourceSkeleton = Context.SourceMesh->GetSkeleton();
 	USkeleton* TargetSkeleton = Context.TargetMesh->GetSkeleton();
-	
-	// get map of curves to remap (Source:Target)
-	bool bCopyAllSourceCurves = true;
-	TMap<FAnimationCurveIdentifier, FAnimationCurveIdentifier> CurvesToRemap;
-	const TArray<FInstancedStruct>& OpStack = Context.IKRetargetAsset->GetRetargetOps();
-	for (const FInstancedStruct& OpStruct : OpStack)
-	{
-		if (OpStruct.GetScriptStruct() != FIKRetargetCurveRemapOp::StaticStruct())
-		{
-			continue;
-		}
-		
-		const FIKRetargetCurveRemapOp* CurveRemapOp = OpStruct.GetPtr<FIKRetargetCurveRemapOp>();
-		if (!CurveRemapOp)
-		{
-			continue;
-		}
 
-		if (!CurveRemapOp->IsEnabled())
-		{
-			continue;
-		}
+	// update progress bar
+	Progress.EnterProgressFrame(1.f, FText(LOCTEXT("ApplyCurveOps", "Applying Curve Operations")));
 
-		for (const FCurveRemapPair& CurveToRemap : CurveRemapOp->Settings.CurvesToRemap)
+	auto NoCurveOps = [](const TArray<FInstancedStruct>&InOpStack)
 		{
-			const FAnimationCurveIdentifier SourceCurveId = UAnimationCurveIdentifierExtensions::FindCurveIdentifier(SourceSkeleton, CurveToRemap.SourceCurve, ERawCurveTrackTypes::RCT_Float);
-			const FAnimationCurveIdentifier TargetCurveId = UAnimationCurveIdentifierExtensions::FindCurveIdentifier(TargetSkeleton, CurveToRemap.TargetCurve, ERawCurveTrackTypes::RCT_Float);
-			CurvesToRemap.Add(SourceCurveId, TargetCurveId);
-		}
+			for (const FInstancedStruct& OpStruct : InOpStack)
+			{
+				const UStruct* StructType = OpStruct.GetScriptStruct();
+				if (StructType && StructType->IsChildOf(FIKRetargetOpBase::StaticStruct()) &&
+					OpStruct.Get<FIKRetargetOpBase>().HasCurveProcessing())
+				{
+					return false;
+				}
+			}
 
-		bCopyAllSourceCurves &= CurveRemapOp->Settings.bCopyAllSourceCurves;
-	}
+			return true;
+		};
 
-	// update progress bar
-	Progress.EnterProgressFrame(1.f, FText::Format(LOCTEXT("RemappingCurves", "Remapping {0} curves on animation assets..."), FText::AsNumber(CurvesToRemap.Num())));
+	const TArray<FInstancedStruct>& OpStack = Context.IKRetargetAsset->GetRetargetOps();
 
-	// for each exported animation, remap curves from source to target anim
+	// for each exported animation, apply any required operations to curves, but if no curve ops, skip entirely for efficiency
+	if (NoCurveOps(OpStack))
+	{
+		return;
+	}
+	
 	for (TPair<UAnimationAsset*, UAnimationAsset*>& Pair : DuplicatedAnimAssets)
 	{
 		UAnimSequence* SourceSequence = Cast<UAnimSequence>(Pair.Key);
@@ -578,68 +567,97 @@ void UIKRetargetBatchOperation::RemapCurves(const FIKRetargetBatchOperationConte
 			return;
 		}
 
-		// all curves were copied when we duplicated the animation sequence, so now we have to rename curves
-		// based on the remapping defined in the curve remap op(s)
+		// all curves were copied when we duplicated the animation sequence, so now we have to modify curves based upon the operations in the Op Stack
 		IAnimationDataController& TargetSeqController = TargetSequence->GetController();
 		constexpr bool bShouldTransact = false;
-		TargetSeqController.OpenBracket(FText::FromString("Remapping Curve Data"), bShouldTransact);
+		TargetSeqController.OpenBracket(FText::FromString("Applying Curve Operations"), bShouldTransact);
 
-		const IAnimationDataModel* SourceDataModel = SourceSequence->GetDataModel();
-		
-		for (const TTuple<FAnimationCurveIdentifier, FAnimationCurveIdentifier>& CurveToRemap : CurvesToRemap)
+		FIKRetargetOpBase::FFrameValues InputCurveFrameValues;
+		FIKRetargetOpBase::FCurveData InputCurveMetaData;
+
+		TArray<FFrameTime> InputFrameTimes;
+		bool bLoadedData = URigMapperEditorSubsystem::LoadCurveValuesFromAnimSequence(SourceSequence, SourceSkeleton, InputCurveMetaData.Names, InputFrameTimes, InputCurveFrameValues, InputCurveMetaData.Flags, InputCurveMetaData.Colors);
+		if (bLoadedData)
 		{
-			// get the source curve to copy from
-			const FFloatCurve* SourceCurve =  SourceDataModel->FindFloatCurve(CurveToRemap.Key);
-			if (!SourceCurve)
+			FIKRetargetOpBase::FFrameValues OutputCurveFrameValues;
+			FIKRetargetOpBase::FCurveData OutputCurveMetaData;
+
+			for (const FInstancedStruct& OpStruct : InProcessor.GetRetargetOps())
 			{
-				continue; // missing source curve to remap
+				const UStruct* StructType = OpStruct.GetScriptStruct();
+				if (StructType)
+				{
+					const FIKRetargetOpBase& Op = OpStruct.Get<FIKRetargetOpBase>();
+					if (Op.IsEnabled())
+					{
+						Op.ProcessAnimSequenceCurves(InputCurveMetaData, InputCurveFrameValues, OutputCurveMetaData, OutputCurveFrameValues);
+
+						// Move everything so that it is input for the next stack op
+						InputCurveMetaData = MoveTemp(OutputCurveMetaData);
+						InputCurveFrameValues = MoveTemp(OutputCurveFrameValues);
+					}
+				}
 			}
 
-			// add a curve to the target to house the keys
-			FAnimationCurveIdentifier TargetCurveID = CurveToRemap.Value;
-			if (!TargetCurveID.IsValid())
+			// set the output values and metadata into the anim sequence
+			// note that outputs are contained in the Input data variables since they are moved there after each operation
+			const IAnimationDataModel* SourceDataModel = SourceSequence->GetDataModel();
+			if (SourceDataModel)
+			{
+				const FFrameRate InputFrameRate = URigMapperEditorSubsystem::GetAnimSequenceRate(SourceSequence);
+				AddCurveValuesToAnimSequence(SourceDataModel, SourceSkeleton, InputCurveMetaData, InputCurveFrameValues, InputFrameRate, InputFrameTimes, bShouldTransact, TargetSeqController);
+			}
+			else
 			{
-				continue; // must provide a valid name for the target curve
+				bLoadedData = false;
 			}
-			TargetSeqController.AddCurve(TargetCurveID, SourceCurve->GetCurveTypeFlags(), bShouldTransact);
+		}
 			
-			// copy data into target curve
-			TargetSeqController.SetCurveKeys(TargetCurveID, SourceCurve->FloatCurve.GetConstRefOfKeys(), bShouldTransact);
-			TargetSeqController.SetCurveColor(TargetCurveID, SourceCurve->GetColor(), bShouldTransact);
+		if(!bLoadedData)
+		{
+			UE_LOG(LogTemp, Warning, TEXT("Batch retarget aborted. Failed to get data from Anim Sequence."));
 		}
 
-		// optionally remove all source curves from the target asset
-		// (remove all curves that were copied when the source sequence was duplicated UNLESS they are remapped)
-		if (!bCopyAllSourceCurves)
+		TargetSeqController.CloseBracket(bShouldTransact);
+	}
+}
+
+
+void UIKRetargetBatchOperation::AddCurveValuesToAnimSequence(const IAnimationDataModel* InSourceDataModel, USkeleton* InTargetSkeleton, const FIKRetargetOpBase::FCurveData& InCurveMetaData, const FIKRetargetOpBase::FFrameValues& InCurveValuesPerFrame, 
+	const FFrameRate& InFrameRate, const TArray<FFrameTime> & InFrameTimes, bool bInShouldTransact, IAnimationDataController& OutTargetSeqController) const
+{
+
+	// remove existing float curves, as we will replace them all
+	OutTargetSeqController.RemoveAllCurvesOfType(ERawCurveTrackTypes::RCT_Float, bInShouldTransact);
+
+	// AnimSequence should start on frame 0
+	const double TimeOffset = InFrameRate.AsSeconds(InFrameTimes[0]);
+
+	for (int32 CurveIndex = 0; CurveIndex < InCurveMetaData.Names.Num(); CurveIndex++)
+	{
+		FAnimationCurveIdentifier CurveId = UAnimationCurveIdentifierExtensions::GetCurveIdentifier(InTargetSkeleton, InCurveMetaData.Names[CurveIndex], ERawCurveTrackTypes::RCT_Float);
+
+		if (CurveId.IsValid())
 		{
-			// get list of target curves to keep
-			TArray<FAnimationCurveIdentifier> TargetCurvesToKeep;
-			CurvesToRemap.GenerateValueArray(TargetCurvesToKeep);
+			OutTargetSeqController.AddCurve(CurveId, InCurveMetaData.Flags[CurveIndex], bInShouldTransact);
 
-			// get list of target curves to remove
-			const TArray<FFloatCurve>& AllTargetCurves = TargetSeqController.GetModel()->GetFloatCurves();
-			TArray<FAnimationCurveIdentifier> CurvesToRemove;
-			for (const FFloatCurve& TargetCurve : AllTargetCurves)
+			for (int32 FrameIndex = 0; FrameIndex < InFrameTimes.Num(); FrameIndex++)
 			{
-				const FAnimationCurveIdentifier TargetCurveId = UAnimationCurveIdentifierExtensions::FindCurveIdentifier(TargetSkeleton, TargetCurve.GetName(), ERawCurveTrackTypes::RCT_Float);
-				if (TargetCurvesToKeep.Contains(TargetCurveId))
+				const double Time = InFrameRate.AsSeconds(InFrameTimes[FrameIndex]) - TimeOffset;
+				const TOptional<float> Value = InCurveValuesPerFrame[FrameIndex][CurveIndex];
+
+				if (Value.IsSet())
 				{
-					continue;
+					OutTargetSeqController.SetCurveKey(CurveId, FRichCurveKey(static_cast<float>(Time), Value.GetValue()), bInShouldTransact);
 				}
-				CurvesToRemove.Add(TargetCurveId);
-			}
-
-			// remove the curves
-			for (const FAnimationCurveIdentifier& CurveToRemove : CurvesToRemove)
-			{
-				TargetSeqController.RemoveCurve(CurveToRemove, bShouldTransact);
 			}
 		}
 
-		TargetSeqController.CloseBracket(bShouldTransact);
+		OutTargetSeqController.SetCurveColor(CurveId, InCurveMetaData.Colors[CurveIndex], bInShouldTransact);
 	}
 }
 
+
 void UIKRetargetBatchOperation::OverwriteExistingAssets(const FIKRetargetBatchOperationContext& Context, FScopedSlowTask& Progress)
 {
 	if (!Context.bOverwriteExistingFiles)
diff --git a/Engine/Plugins/Animation/IKRig/Source/IKRigEditor/Public/RetargetEditor/IKRetargetBatchOperation.h b/Engine/Plugins/Animation/IKRig/Source/IKRigEditor/Public/RetargetEditor/IKRetargetBatchOperation.h
--- a/Engine/Plugins/Animation/IKRig/Source/IKRigEditor/Public/RetargetEditor/IKRetargetBatchOperation.h
+++ b/Engine/Plugins/Animation/IKRig/Source/IKRigEditor/Public/RetargetEditor/IKRetargetBatchOperation.h
@@ -16,7 +16,9 @@ class USkeletalMesh;
 class UAnimationAsset;
 class UAnimBlueprint;
 class UAnimSequence;
-
+class USkeleton;
+class IAnimationDataModel;
+class IAnimationDataController;
 
 struct FAdditiveRetargetSettings
 {
@@ -158,10 +160,10 @@ private:
 	UE_API void RetargetAssets(const FIKRetargetBatchOperationContext& Context, FScopedSlowTask& Progress);
 
 	// Convert animation on all the duplicates
-	UE_API void ConvertAnimation(const FIKRetargetBatchOperationContext& Context, FScopedSlowTask& Progress);
+	UE_API void ConvertAnimation(const FIKRetargetBatchOperationContext& Context, FIKRetargetProcessor& OutProcessor, FScopedSlowTask& Progress);
 
-	// Copy/remap curves on all the duplicates
-	UE_API void RemapCurves(const FIKRetargetBatchOperationContext& Context, FScopedSlowTask& Progress);
+	// Apply any curve operations on all the duplicates
+	UE_API void ApplyCurveOps(const FIKRetargetBatchOperationContext& Context, const FIKRetargetProcessor& InProcessor, FScopedSlowTask& Progress);
 
 	// Replace existing assets (optional)
 	UE_API void OverwriteExistingAssets(const FIKRetargetBatchOperationContext& Context, FScopedSlowTask& Progress);
@@ -174,6 +176,11 @@ private:
 
 	// If user cancelled half way, cleanup all the duplicated assets
 	UE_API void CleanupIfCancelled(const FScopedSlowTask& Progress) const;
+
+	// Set the processed curve data into the target animation sequence
+	UE_API void AddCurveValuesToAnimSequence(const IAnimationDataModel* InSourceDataModel, USkeleton* InTargetSkeleton, const FIKRetargetOpBase::FCurveData& InCurveMetaData,
+		const FIKRetargetOpBase::FFrameValues& InCurveValuesPerFrame, const FFrameRate& InFrameRate, const TArray<FFrameTime>& InFrameTimes, bool bInShouldTransact, IAnimationDataController& OutTargetSeqController) const;
+
 	
 	// Lists of assets to retarget. Populated from selection during init
 	TArray<UAnimationAsset*>	AnimationAssetsToRetarget;
diff --git a/Engine/Plugins/Experimental/Animation/RigMapper/Source/RigMapperEditor/Private/RigMapperEditorSubsystem.cpp b/Engine/Plugins/Experimental/Animation/RigMapper/Source/RigMapperEditor/Private/RigMapperEditorSubsystem.cpp
--- a/Engine/Plugins/Experimental/Animation/RigMapper/Source/RigMapperEditor/Private/RigMapperEditorSubsystem.cpp
+++ b/Engine/Plugins/Experimental/Animation/RigMapper/Source/RigMapperEditor/Private/RigMapperEditorSubsystem.cpp
@@ -343,6 +343,47 @@ bool URigMapperEditorSubsystem::AddCurveValuesToAnimSequence(UAnimSequence* Targ
 	return true;
 }
 
+
+bool URigMapperEditorSubsystem::LoadCurveValuesFromAnimSequence(const UAnimSequence* InSource, const USkeleton* InSourceSkeleton, TArray<FName>& OutCurveNames, TArray<FFrameTime>& OutFrameTimes, FFrameValues& OutCurveValuesPerFrame,
+	TArray<int32>& OutCurveFlags, TArray<FLinearColor>& OutCurveColors)
+{
+	TArray<FString> CurveNames;
+	if (LoadCurveValuesFromAnimSequence(InSource, CurveNames, OutFrameTimes, OutCurveValuesPerFrame))
+	{
+		OutCurveNames.Reserve(CurveNames.Num());
+		OutCurveFlags.Reserve(CurveNames.Num());
+		OutCurveColors.Reserve(CurveNames.Num());
+
+		const IAnimationDataModel* SourceDataModel = InSource->GetDataModel();
+
+		for (const FString& CurveName : CurveNames)
+		{
+			const FName CurCurveName(*CurveName);
+			const FAnimationCurveIdentifier CurCurveId = UAnimationCurveIdentifierExtensions::FindCurveIdentifier(InSourceSkeleton, CurCurveName, ERawCurveTrackTypes::RCT_Float);
+
+			const FFloatCurve* SourceCurve = SourceDataModel->FindFloatCurve(CurCurveId);
+			if (SourceCurve)
+			{
+				OutCurveFlags.Add(SourceCurve->GetCurveTypeFlags());
+				OutCurveColors.Add(SourceCurve->Color);
+			}
+			else
+			{
+				return false;
+			}
+
+			OutCurveNames.Add(CurCurveName);
+		}
+	}
+	else
+	{
+		return false;
+	}
+
+	return true;
+}
+
+
 bool URigMapperEditorSubsystem::LoadCurveValuesFromAnimSequence(const UAnimSequence* Source, TArray<FString>& CurveNames, TArray<FFrameTime>& FrameTimes, FFrameValues& CurveValuesPerFrame)
 {
 	FFramePoses Poses;
diff --git a/Engine/Plugins/Experimental/Animation/RigMapper/Source/RigMapperEditor/Public/RigMapperEditorSubsystem.h b/Engine/Plugins/Experimental/Animation/RigMapper/Source/RigMapperEditor/Public/RigMapperEditorSubsystem.h
--- a/Engine/Plugins/Experimental/Animation/RigMapper/Source/RigMapperEditor/Public/RigMapperEditorSubsystem.h
+++ b/Engine/Plugins/Experimental/Animation/RigMapper/Source/RigMapperEditor/Public/RigMapperEditorSubsystem.h
@@ -19,6 +19,7 @@ class UAnimSequence;
 class ULevelSequence;
 class UControlRig;
 class UMovieSceneControlRigParameterSection;
+class USkeleton;
 
 /**
  * URigMapperEditorSubsystem
@@ -148,9 +149,15 @@ public:
 
 	UFUNCTION(BlueprintCallable, Category = "Editor Scripting|Animation|Rig Mapper")
 	static UE_API void SetAnimSequenceRate(UAnimSequence* AnimSequence, FFrameRate FrameRate, bool bSetImportProperties);
-	
-private:
+
 	using FFrameValues = TArray<FRigMapperProcessor::FPoseValues>;
+	/**
+	 * Load curve values from animation sequence (not blueprint callable)
+	 */
+	static UE_API bool LoadCurveValuesFromAnimSequence(const UAnimSequence* InSource, const USkeleton* InSourceSkeleton, TArray<FName>& OutCurveNames, TArray<FFrameTime>& OutFrameTimes, FFrameValues& OutCurveValuesPerFrame, TArray<int32> & OutCurveFlags, TArray<FLinearColor> & OutCurveColors);
+
+
+private:
 	using FFramePoses = TArray<FRigMapperProcessor::FPose>;
 	
 	static UE_API bool ConvertCurveValuesToCsv(TArray<FString>& CurveNames, const TArray<FFrameTime>& FrameTimes, FFrameValues& CurveValuesPerFrame, const FFilePath& OutputFile, const TArray<URigMapperDefinition*>& Definitions, bool bOutputIntermediateCsvFiles);
diff --git a/Engine/Plugins/Experimental/Animation/RigMapper/Source/RigMapperEditor/RigMapperEditor.Build.cs b/Engine/Plugins/Experimental/Animation/RigMapper/Source/RigMapperEditor/RigMapperEditor.Build.cs
--- a/Engine/Plugins/Experimental/Animation/RigMapper/Source/RigMapperEditor/RigMapperEditor.Build.cs
+++ b/Engine/Plugins/Experimental/Animation/RigMapper/Source/RigMapperEditor/RigMapperEditor.Build.cs
@@ -12,7 +12,8 @@ public class RigMapperEditor : ModuleRules
 				"EditorSubsystem",
 				"UnrealEd",
 				"GraphEditor",
-				"Engine"
+				"Engine",
+				"RigMapper"
 			});
 
 		PrivateDependencyModuleNames.AddRange(
@@ -34,8 +35,7 @@ public class RigMapperEditor : ModuleRules
 				"ContentBrowser",
 				"MessageLog",
 				"InputCore",
-				"Sequencer",
-				"RigMapper"
+				"Sequencer"
 			});
 	}
 }
