[RigMapper] fix issues in 45364220 CL found by BugHawk

diff --git a/Engine/Plugins/Animation/IKRig/Source/IKRig/Private/Retargeter/RetargetOps/RigMapperOp.cpp b/Engine/Plugins/Animation/IKRig/Source/IKRig/Private/Retargeter/RetargetOps/RigMapperOp.cpp
--- a/Engine/Plugins/Animation/IKRig/Source/IKRig/Private/Retargeter/RetargetOps/RigMapperOp.cpp
+++ b/Engine/Plugins/Animation/IKRig/Source/IKRig/Private/Retargeter/RetargetOps/RigMapperOp.cpp
@@ -239,42 +239,42 @@ void FIKRetargetRigMapperOp::ProcessAnimSequenceCurves(FIKRetargetOpBase::FCurve
 			OutCurveMetaData.Colors[Curve] = MakeCurveColor(OutCurveMetaData.Names[Curve]);
 			OutputNames.Add(OutCurveMetaData.Names[Curve]);
 		}
-	}
 
-	// 'pass-through' any input curves which have not been modified, if bCopyAllSourceCurves is true
-	if (Settings.bCopyAllSourceCurves)
-	{
-		TArray<int32> InputCurveIndicesToAdd;
-		InputCurveIndicesToAdd.Reserve(InCurveMetaData.Names.Num());
-		for (int32 InputCurve = 0; InputCurve < InCurveMetaData.Names.Num(); ++InputCurve)
+		// 'pass-through' any input curves which have not been modified, if bCopyAllSourceCurves is true
+		if (Settings.bCopyAllSourceCurves)
 		{
-			if (!OutputNames.Contains(InCurveMetaData.Names[InputCurve]))
+			TArray<int32> InputCurveIndicesToAdd;
+			InputCurveIndicesToAdd.Reserve(InCurveMetaData.Names.Num());
+			for (int32 InputCurve = 0; InputCurve < InCurveMetaData.Names.Num(); ++InputCurve)
 			{
-				InputCurveIndicesToAdd.Add(InputCurve);
+				if (!OutputNames.Contains(InCurveMetaData.Names[InputCurve]))
+				{
+					InputCurveIndicesToAdd.Add(InputCurve);
+				}
 			}
-		}
 
-		// now add these curves
-		OutCurveMetaData.Names.Reserve(OutCurveMetaData.Names.Num() + InputCurveIndicesToAdd.Num());
-		OutCurveMetaData.Flags.Reserve(OutCurveMetaData.Flags.Num() + InputCurveIndicesToAdd.Num());
-		OutCurveMetaData.Colors.Reserve(OutCurveMetaData.Colors.Num() + InputCurveIndicesToAdd.Num());
-		int32 Frame = 0;
-		for (TArray<TOptional<float>>& OutputFrame : OutCurveFrameValues)
-		{
-			OutputFrame.Reserve(OutputFrame.Num() + InputCurveIndicesToAdd.Num());
-			for (int32 CurveIndex : InputCurveIndicesToAdd)
+			// now add these curves
+			OutCurveMetaData.Names.Reserve(OutCurveMetaData.Names.Num() + InputCurveIndicesToAdd.Num());
+			OutCurveMetaData.Flags.Reserve(OutCurveMetaData.Flags.Num() + InputCurveIndicesToAdd.Num());
+			OutCurveMetaData.Colors.Reserve(OutCurveMetaData.Colors.Num() + InputCurveIndicesToAdd.Num());
+			int32 Frame = 0;
+			for (TArray<TOptional<float>>& OutputFrame : OutCurveFrameValues)
 			{
-				OutputFrame.Add(InCurveFrameValues[Frame][CurveIndex]);
-			}
+				OutputFrame.Reserve(OutputFrame.Num() + InputCurveIndicesToAdd.Num());
+				for (int32 CurveIndex : InputCurveIndicesToAdd)
+				{
+					OutputFrame.Add(InCurveFrameValues[Frame][CurveIndex]);
+				}
 
-			Frame++;
-		}
+				Frame++;
+			}
 
-		for (int32 CurveIndex : InputCurveIndicesToAdd)
-		{
-			OutCurveMetaData.Names.Add(InCurveMetaData.Names[CurveIndex]);
-			OutCurveMetaData.Flags.Add(InCurveMetaData.Flags[CurveIndex]);
-			OutCurveMetaData.Colors.Add(InCurveMetaData.Colors[CurveIndex]);
+			for (int32 CurveIndex : InputCurveIndicesToAdd)
+			{
+				OutCurveMetaData.Names.Add(InCurveMetaData.Names[CurveIndex]);
+				OutCurveMetaData.Flags.Add(InCurveMetaData.Flags[CurveIndex]);
+				OutCurveMetaData.Colors.Add(InCurveMetaData.Colors[CurveIndex]);
+			}
 		}
 	}
 }
diff --git a/Engine/Plugins/Animation/IKRig/Source/IKRig/Public/Retargeter/IKRetargetOps.h b/Engine/Plugins/Animation/IKRig/Source/IKRig/Public/Retargeter/IKRetargetOps.h
--- a/Engine/Plugins/Animation/IKRig/Source/IKRig/Public/Retargeter/IKRetargetOps.h
+++ b/Engine/Plugins/Animation/IKRig/Source/IKRig/Public/Retargeter/IKRetargetOps.h
@@ -338,7 +338,7 @@ private:
 	FName ParentOpName;
 
 	/** (optional) determines whether input curves are taken from the source anim instance or from the current pose context (for example when we have multiple nodes which process curves). Can be set on the fly */
-	UPROPERTY()
+	UPROPERTY(Transient)
 	bool bTakeInputCurvesFromSourceAnimInstance = true;
 
 protected:
diff --git a/Engine/Plugins/Animation/IKRig/Source/IKRig/Public/Retargeter/RetargetOps/RigMapperOp.h b/Engine/Plugins/Animation/IKRig/Source/IKRig/Public/Retargeter/RetargetOps/RigMapperOp.h
--- a/Engine/Plugins/Animation/IKRig/Source/IKRig/Public/Retargeter/RetargetOps/RigMapperOp.h
+++ b/Engine/Plugins/Animation/IKRig/Source/IKRig/Public/Retargeter/RetargetOps/RigMapperOp.h
@@ -44,7 +44,6 @@ struct FIKRetargetRigMapperOp : public FIKRetargetOpBase
 {
 	GENERATED_BODY()
 	
-	// NOTE: this op does not do anything in Initialize() or Run().
 	
 	virtual bool Initialize(
 		const FIKRetargetProcessor& InProcessor,
@@ -53,6 +52,7 @@ struct FIKRetargetRigMapperOp : public FIKRetargetOpBase
 		const FIKRetargetOpBase* InParentOp,
 		FIKRigLogger& Log) override;
 	
+	// NOTE: this op does not do anything in Run().
 	virtual void Run(
 		FIKRetargetProcessor& InProcessor,
 		const double InDeltaTime,
diff --git a/Engine/Plugins/Animation/IKRig/Source/IKRigEditor/Private/RetargetEditor/IKRetargetBatchOperation.cpp b/Engine/Plugins/Animation/IKRig/Source/IKRigEditor/Private/RetargetEditor/IKRetargetBatchOperation.cpp
--- a/Engine/Plugins/Animation/IKRig/Source/IKRigEditor/Private/RetargetEditor/IKRetargetBatchOperation.cpp
+++ b/Engine/Plugins/Animation/IKRig/Source/IKRigEditor/Private/RetargetEditor/IKRetargetBatchOperation.cpp
@@ -601,16 +601,8 @@ void UIKRetargetBatchOperation::ApplyCurveOps(const FIKRetargetBatchOperationCon
 
 			// set the output values and metadata into the anim sequence
 			// note that outputs are contained in the Input data variables since they are moved there after each operation
-			const IAnimationDataModel* SourceDataModel = SourceSequence->GetDataModel();
-			if (SourceDataModel)
-			{
-				const FFrameRate InputFrameRate = URigMapperEditorSubsystem::GetAnimSequenceRate(SourceSequence);
-				AddCurveValuesToAnimSequence(SourceDataModel, SourceSkeleton, InputCurveMetaData, InputCurveFrameValues, InputFrameRate, InputFrameTimes, bShouldTransact, TargetSeqController);
-			}
-			else
-			{
-				bLoadedData = false;
-			}
+			const FFrameRate InputFrameRate = URigMapperEditorSubsystem::GetAnimSequenceRate(SourceSequence);
+			AddCurveValuesToAnimSequence(SourceSkeleton, InputCurveMetaData, InputCurveFrameValues, InputFrameRate, InputFrameTimes, bShouldTransact, TargetSeqController);
 		}
 			
 		if(!bLoadedData)
@@ -623,15 +615,18 @@ void UIKRetargetBatchOperation::ApplyCurveOps(const FIKRetargetBatchOperationCon
 }
 
 
-void UIKRetargetBatchOperation::AddCurveValuesToAnimSequence(const IAnimationDataModel* InSourceDataModel, USkeleton* InTargetSkeleton, const FIKRetargetOpBase::FCurveData& InCurveMetaData, const FIKRetargetOpBase::FFrameValues& InCurveValuesPerFrame, 
+void UIKRetargetBatchOperation::AddCurveValuesToAnimSequence(USkeleton* InTargetSkeleton, const FIKRetargetOpBase::FCurveData& InCurveMetaData, const FIKRetargetOpBase::FFrameValues& InCurveValuesPerFrame, 
 	const FFrameRate& InFrameRate, const TArray<FFrameTime> & InFrameTimes, bool bInShouldTransact, IAnimationDataController& OutTargetSeqController) const
 {
-
 	// remove existing float curves, as we will replace them all
 	OutTargetSeqController.RemoveAllCurvesOfType(ERawCurveTrackTypes::RCT_Float, bInShouldTransact);
 
 	// AnimSequence should start on frame 0
-	const double TimeOffset = InFrameRate.AsSeconds(InFrameTimes[0]);
+	double TimeOffset = 0;
+	if (InFrameTimes.Num() != 0)
+	{
+		TimeOffset = InFrameRate.AsSeconds(InFrameTimes[0]);
+	}
 
 	for (int32 CurveIndex = 0; CurveIndex < InCurveMetaData.Names.Num(); CurveIndex++)
 	{
diff --git a/Engine/Plugins/Animation/IKRig/Source/IKRigEditor/Public/RetargetEditor/IKRetargetBatchOperation.h b/Engine/Plugins/Animation/IKRig/Source/IKRigEditor/Public/RetargetEditor/IKRetargetBatchOperation.h
--- a/Engine/Plugins/Animation/IKRig/Source/IKRigEditor/Public/RetargetEditor/IKRetargetBatchOperation.h
+++ b/Engine/Plugins/Animation/IKRig/Source/IKRigEditor/Public/RetargetEditor/IKRetargetBatchOperation.h
@@ -178,7 +178,7 @@ private:
 	UE_API void CleanupIfCancelled(const FScopedSlowTask& Progress) const;
 
 	// Set the processed curve data into the target animation sequence
-	UE_API void AddCurveValuesToAnimSequence(const IAnimationDataModel* InSourceDataModel, USkeleton* InTargetSkeleton, const FIKRetargetOpBase::FCurveData& InCurveMetaData,
+	UE_API void AddCurveValuesToAnimSequence(USkeleton* InTargetSkeleton, const FIKRetargetOpBase::FCurveData& InCurveMetaData,
 		const FIKRetargetOpBase::FFrameValues& InCurveValuesPerFrame, const FFrameRate& InFrameRate, const TArray<FFrameTime>& InFrameTimes, bool bInShouldTransact, IAnimationDataController& OutTargetSeqController) const;
 
 	
