[IK Retarget] Adding ability to scale source from optional pivot.

diff --git a/Engine/Plugins/Animation/IKRig/Source/IKRig/Private/AnimNodes/AnimNode_RetargetPoseFromMesh.cpp b/Engine/Plugins/Animation/IKRig/Source/IKRig/Private/AnimNodes/AnimNode_RetargetPoseFromMesh.cpp
--- a/Engine/Plugins/Animation/IKRig/Source/IKRig/Private/AnimNodes/AnimNode_RetargetPoseFromMesh.cpp
+++ b/Engine/Plugins/Animation/IKRig/Source/IKRig/Private/AnimNodes/AnimNode_RetargetPoseFromMesh.cpp
@@ -3,12 +3,10 @@
 #include "AnimNodes/AnimNode_RetargetPoseFromMesh.h"
 
 #include "IKRigObjectVersion.h"
-#include "Animation/AnimCurveUtils.h"
 #include "Animation/AnimInstance.h"
 #include "Animation/AnimInstanceProxy.h"
 #include "Components/SkeletalMeshComponent.h"
 #include "Engine/SkeletalMesh.h"
-#include "Retargeter/IKRetargetOps.h"
 #include "Retargeter/RetargetOps/CurveRemapOp.h"
 
 #include UE_INLINE_GENERATED_CPP_BY_NAME(AnimNode_RetargetPoseFromMesh)
@@ -168,8 +166,8 @@ void FAnimNode_RetargetPoseFromMesh::Evaluate_AnyThread(FPoseContext& Output)
 		// convert to component space
 		FAnimationRuntime::FillUpComponentSpaceTransforms(TargetRefSkeleton, InputLocalTransforms, PoseToRetargetFromComponentSpace);
 
-		// let the retargeter scale the input pose
-		Processor.ScaleSourcePose(PoseToRetargetFromComponentSpace);
+		// scale it accorrding to the scale source op (if used)
+		Processor.ApplySourceScaleToPose(PoseToRetargetFromComponentSpace);
 	};
 	
 	auto ApplyRetargetedPose = [this, &Output](const TArray<FTransform>& RetargetedPose)
@@ -240,7 +238,7 @@ void FAnimNode_RetargetPoseFromMesh::Evaluate_AnyThread(FPoseContext& Output)
 	ApplyRetargetedPose(RetargetedPose);
 
 	// editor profiling
-#ifdef WITH_EDITOR
+#if WITH_EDITOR
 	constexpr double FrameWindow = 30.0;
 	constexpr double Alpha = 2.0 / (FrameWindow + 1.0);
 	AverageExecutionTime += Alpha * ((FPlatformTime::Seconds() - StartTime) - AverageExecutionTime);
@@ -305,9 +303,8 @@ void FAnimNode_RetargetPoseFromMesh::PreUpdate(const UAnimInstance* InAnimInstan
 		{
 			Transform.SetScale3D(FVector::OneVector);
 		}
-		
-		// let the retargeter scale the input pose
-		Processor.ScaleSourcePose(PoseToRetargetFromComponentSpace);
+
+		Processor.ApplySourceScaleToPose(PoseToRetargetFromComponentSpace);
 	};
 
 	// copy all the data from the source component
@@ -431,8 +428,10 @@ void FAnimNode_RetargetPoseFromMesh::PostSerialize(const FArchive& Ar)
 #endif // WITH_EDITORONLY_DATA
 }
 
+#if WITH_EDITOR
 double FAnimNode_RetargetPoseFromMesh::GetAverageExecutionTime() const
 {
 	return AverageExecutionTime;
 }
+#endif
 
diff --git a/Engine/Plugins/Animation/IKRig/Source/IKRig/Private/Retargeter/IKRetargetProcessor.cpp b/Engine/Plugins/Animation/IKRig/Source/IKRig/Private/Retargeter/IKRetargetProcessor.cpp
--- a/Engine/Plugins/Animation/IKRig/Source/IKRig/Private/Retargeter/IKRetargetProcessor.cpp
+++ b/Engine/Plugins/Animation/IKRig/Source/IKRig/Private/Retargeter/IKRetargetProcessor.cpp
@@ -26,14 +26,14 @@ FResolvedRetargetPose& FResolvedRetargetPoseSet::AddOrUpdateRetargetPose(
 	const FName InRetargetPoseName,
 	const FIKRetargetPose* InRetargetPose,
 	const FName PelvisBoneName,
-	const double ScaleFactor)
+	const FRetargetPoseScaleWithPivot& InSourceScale)
 {
 	// add retarget pose if it doesn't already exist
 	FResolvedRetargetPose& RetargetPose = FindOrAddRetargetPose(InRetargetPoseName);
 	
 	// record the version of the retarget pose (prevents re-initialization if profile swaps it)
 	RetargetPose.Version = InRetargetPose->GetVersion();
-	RetargetPose.Scale = ScaleFactor;
+	RetargetPose.PoseScale = InSourceScale;
 	
 	// initialize retarget pose to the skeletal mesh reference pose
 	RetargetPose.LocalPose = InSkeleton.SkeletalMesh->GetRefSkeleton().GetRefBonePose();
@@ -77,15 +77,9 @@ FResolvedRetargetPose& FResolvedRetargetPoseSet::AddOrUpdateRetargetPose(
 	// update global transforms based on local pose modified by the retarget pose offsets
 	InSkeleton.UpdateGlobalTransformsBelowBone(INDEX_NONE, RetargetPose.LocalPose, RetargetPose.GlobalPose);
 
-	// scale the retarget pose
-	if (!FMath::IsNearlyEqual(ScaleFactor, 1.0f))
+	// scale the global retarget pose
+	if (RetargetPose.PoseScale.ScalePose(RetargetPose.GlobalPose))
 	{
-		// scale translations of global transforms
-		for (FTransform& BoneTransform : RetargetPose.GlobalPose)
-		{
-			BoneTransform.SetTranslation(BoneTransform.GetTranslation() * ScaleFactor);
-		}
-
 		// update the local transforms
 		InSkeleton.UpdateLocalTransformsBelowBone(INDEX_NONE, RetargetPose.LocalPose, RetargetPose.GlobalPose);
 	}
@@ -166,7 +160,7 @@ void FRetargetSkeleton::Initialize(
 	const ERetargetSourceOrTarget InSourceOrTarget,
 	const UIKRetargeter* InRetargetAsset,
 	const FName PelvisBoneName,
-	const double ScaleFactor)
+	const FRetargetPoseScaleWithPivot& InPoseScale)
 {
 	// reset all skeleton data
 	Reset();
@@ -188,7 +182,7 @@ void FRetargetSkeleton::Initialize(
 	// add default retarget pose for this skeleton
 	static FName DefaultRetargetPoseName = UIKRetargeter::GetDefaultPoseName();
 	const FIKRetargetPose* DefaultRetargetPose = InRetargetAsset->GetRetargetPoseByName(InSourceOrTarget, UIKRetargeter::GetDefaultPoseName());
-	RetargetPoses.AddOrUpdateRetargetPose(*this,DefaultRetargetPoseName, DefaultRetargetPose, PelvisBoneName, ScaleFactor);
+	RetargetPoses.AddOrUpdateRetargetPose(*this,DefaultRetargetPoseName, DefaultRetargetPose, PelvisBoneName, InPoseScale);
 	RetargetPoses.CurrentRetargetPoseName = DefaultRetargetPoseName;
 
 	// add current retarget pose
@@ -196,7 +190,7 @@ void FRetargetSkeleton::Initialize(
 	const FIKRetargetPose* CurrentRetargetPose = InRetargetAsset->GetRetargetPoseByName(InSourceOrTarget, CurrentRetargetPoseName);
 	if (CurrentRetargetPose)
 	{
-		RetargetPoses.AddOrUpdateRetargetPose(*this,CurrentRetargetPoseName, CurrentRetargetPose, PelvisBoneName, ScaleFactor);
+		RetargetPoses.AddOrUpdateRetargetPose(*this,CurrentRetargetPoseName, CurrentRetargetPose, PelvisBoneName, InPoseScale);
 		// set this as the current pose to use (may be overridden by ops)
 		RetargetPoses.CurrentRetargetPoseName = CurrentRetargetPoseName;
 	}
@@ -488,12 +482,12 @@ void FTargetSkeleton::Initialize(
 const USkeletalMesh* InSkeletalMesh,
 	const ERetargetSourceOrTarget InSourceOrTarget,
 	const UIKRetargeter* InRetargetAsset,
-	const FName RetargetRootBone)
+	const FName RetargetRootBone,
+	const FRetargetPoseScaleWithPivot& InPoseScale)
 {
 	Reset();
-
-	constexpr double TargetSkeletonScaleFactor = 1.0;
-	FRetargetSkeleton::Initialize(InSkeletalMesh, InSourceOrTarget, InRetargetAsset, RetargetRootBone, TargetSkeletonScaleFactor);
+	
+	FRetargetSkeleton::Initialize(InSkeletalMesh, InSourceOrTarget, InRetargetAsset, RetargetRootBone, InPoseScale);
 
 	// make storage for per-bone "Is Retargeted" flag (used for hierarchy updates)
 	// these are bones that are in a target chain that is mapped to a source chain (ie, will actually be retargeted)
@@ -1134,14 +1128,15 @@ void FIKRetargetProcessor::Initialize(
 		ERetargetSourceOrTarget::Source,
 		RetargeterAsset,
 		GetPelvisBone(ERetargetSourceOrTarget::Source, ERetargetOpsToSearch::AssetOps),
-		GetSourceScaleFactor());
+		GetPoseScale(ERetargetSourceOrTarget::Source));
 	
 	// initialize target skeleton
 	TargetSkeleton.Initialize(
 		TargetSkeletalMesh,
 		ERetargetSourceOrTarget::Target,
 		RetargeterAsset,
-		GetPelvisBone(ERetargetSourceOrTarget::Target, ERetargetOpsToSearch::AssetOps));
+		GetPelvisBone(ERetargetSourceOrTarget::Target, ERetargetOpsToSearch::AssetOps),
+		GetPoseScale(ERetargetSourceOrTarget::Target));
 
 	// setup the goal container
 	ResetIKRigGoalContainer();
@@ -1331,19 +1326,6 @@ TArray<FTransform>&  FIKRetargetProcessor::RunRetargeter(
 	return TargetSkeleton.OutputGlobalPose;
 }
 
-void FIKRetargetProcessor::ScaleSourcePose(TArray<FTransform>& InOutSourceGlobalPose) const
-{
-	// scale the input pose
-	const double SourceScaleFactor = GetSourceScaleFactor();
-	if (!FMath::IsNearlyEqual(SourceScaleFactor, 1.0f))
-	{
-		for (FTransform& SourceBoneTransform : InOutSourceGlobalPose)
-		{
-			SourceBoneTransform.SetTranslation(SourceBoneTransform.GetTranslation() * SourceScaleFactor);
-		}
-	}
-}
-
 void FIKRetargetProcessor::GenerateBasePoses(TArray<FTransform>& InSourceGlobalPose)
 {
 	// start from the retarget pose
@@ -1794,15 +1776,15 @@ void FIKRetargetProcessor::UpdateRetargetPoseAtRuntime(
 	{
 		// load the pose and resolve it onto the current skeleton
 		const FName PelvisBoneName = GetPelvisBone(SourceOrTarget, ERetargetOpsToSearch::ProcessorOps);
-		const double ScaleFactor = SourceOrTarget == ERetargetSourceOrTarget::Source ? GetSourceScaleFactor() : 1.0;
 		const bool bIsSource = SourceOrTarget == ERetargetSourceOrTarget::Source;
+		const FRetargetPoseScaleWithPivot PoseScale = GetPoseScale(SourceOrTarget);
 		FRetargetSkeleton& RetargetSkeleton = bIsSource ? SourceSkeleton : TargetSkeleton;
 		FResolvedRetargetPose& ResolvedPose = RetargetSkeleton.RetargetPoses.AddOrUpdateRetargetPose(
 			RetargetSkeleton,
 			RetargetPoseToUseName,
 			PoseToLoad,
 			PelvisBoneName,
-			ScaleFactor);
+			PoseScale);
 
 		// set this as the current pose to use
 		RetargetSkeleton.RetargetPoses.CurrentRetargetPoseName = RetargetPoseToUseName;
@@ -1838,8 +1820,7 @@ void FIKRetargetProcessor::UpdateRetargetPoseAtRuntime(
 	const FResolvedRetargetPose& CurrentResolvedPose = RetargetSkeleton.RetargetPoses.FindOrAddRetargetPose(CurrentRetargetPoseName);
 	const bool bSameRetargetPose = CurrentResolvedPose.Name == ResolvedPoseToUse->Name;
 	const bool bSameVersion = CurrentResolvedPose.Version == UnresolvedPoseFromAsset->GetVersion();
-	const double SourceScaleFactor = GetSourceScaleFactor();
-	const bool bSameGlobalScale = bIsSource ? FMath::IsNearlyEqual(CurrentResolvedPose.Scale, SourceScaleFactor) : true;
+	const bool bSameGlobalScale = GetPoseScale(SourceOrTarget) == CurrentResolvedPose.PoseScale;
 	if (bSameRetargetPose && bSameVersion && bSameGlobalScale)
 	{
 		return; // retarget pose has not changed since it was initialized
@@ -1872,6 +1853,48 @@ FName FIKRetargetProcessor::GetPelvisBone(ERetargetSourceOrTarget SourceOrTarget
 	return NAME_None;
 }
 
+void FIKRetargetProcessor::ApplySourceScaleToPose(TArray<FTransform>& InOutSourceGlobalPose) const
+{
+	const FIKRetargetScaleSourceOp* ScaleSourceOp = GetFirstRetargetOpOfType<FIKRetargetScaleSourceOp>();
+	if (!ScaleSourceOp)
+	{
+		return;
+	}
+	
+	auto GetScalePivot = [this, &InOutSourceGlobalPose, ScaleSourceOp]() -> FVector
+		{
+			switch (ScaleSourceOp->Settings.ScalePivot)
+			{
+			case EScaleSourcePivot::ComponentOrigin:
+				{
+					return FVector::ZeroVector;
+				}
+			case EScaleSourcePivot::Bone:
+				{
+					const int32 ScalePivotBoneIndex = ScaleSourceOp->GetScalePivotBoneIndex(SourceSkeleton);
+					if (!InOutSourceGlobalPose.IsValidIndex(ScalePivotBoneIndex))
+					{
+						return FVector::ZeroVector;
+					}
+
+					// optionally project to floor
+					const FVector Projection = ScaleSourceOp->Settings.bProjectScalePivotToFloor ? FVector(1,1,0) : FVector::OneVector;
+					return InOutSourceGlobalPose[ScalePivotBoneIndex].GetTranslation() * Projection;
+				}
+			default:
+				checkNoEntry();
+				return FVector::ZeroVector;
+			}
+		};
+
+	// record latest scale pivot and factor
+	CurrentSourceScale.Factor = ScaleSourceOp->Settings.SourceScaleFactor;
+	CurrentSourceScale.Pivot = GetScalePivot();
+	
+	// scale the input pose
+	CurrentSourceScale.ScalePose(InOutSourceGlobalPose);
+}
+
 double FIKRetargetProcessor::GetSourceScaleFactor() const
 {
 	// first try get it from the processor
@@ -1887,5 +1910,11 @@ double FIKRetargetProcessor::GetSourceScaleFactor() const
 	return 1.0;
 }
 
+const FRetargetPoseScaleWithPivot& FIKRetargetProcessor::GetPoseScale(const ERetargetSourceOrTarget InSourceOrTarget) const
+{
+	static FRetargetPoseScaleWithPivot TargetPoseScale;
+	return InSourceOrTarget == ERetargetSourceOrTarget::Source ? CurrentSourceScale : TargetPoseScale;
+}
+
 #undef LOCTEXT_NAMESPACE
 
diff --git a/Engine/Plugins/Animation/IKRig/Source/IKRig/Private/Retargeter/RetargetOps/PinBoneOp.cpp b/Engine/Plugins/Animation/IKRig/Source/IKRig/Private/Retargeter/RetargetOps/PinBoneOp.cpp
--- a/Engine/Plugins/Animation/IKRig/Source/IKRig/Private/Retargeter/RetargetOps/PinBoneOp.cpp
+++ b/Engine/Plugins/Animation/IKRig/Source/IKRig/Private/Retargeter/RetargetOps/PinBoneOp.cpp
@@ -13,7 +13,7 @@
 
 void FPinBoneData::CachePinData(ERetargetSourceOrTarget InSkeletonToCopyFrom, const FIKRetargetProcessor& Processor)
 {
-	SourceScaleFactor = Processor.GetSourceScaleFactor();
+	LastUsedSourceScale = Processor.GetPoseScale(ERetargetSourceOrTarget::Source);
 
     // get skeletons we are copying from/to
     const FRetargetSkeleton& SkeletonToCopyFrom = Processor.GetSkeleton(InSkeletonToCopyFrom);
@@ -115,7 +115,7 @@ bool FIKRetargetPinBoneOp::Initialize(
 		}
 
 		// force reinitialization
-		PinData.SourceScaleFactor = -1.0f;
+		PinData.LastUsedSourceScale.Factor = -1.0f;
 	}
 	
 	// always treat this op as "initialized", individual pins will only execute if their prerequisites are met
@@ -141,7 +141,7 @@ void FIKRetargetPinBoneOp::Run(
 		}
 
 		// recache the offsets if the source is scaled differently
-		if (!FMath::IsNearlyEqual(PinData.SourceScaleFactor, Processor.GetSourceScaleFactor()))
+		if (PinData.LastUsedSourceScale != Processor.GetPoseScale(ERetargetSourceOrTarget::Source))
 		{
 			PinData.CachePinData(Settings.SkeletonToCopyFrom, Processor);
 		}
diff --git a/Engine/Plugins/Animation/IKRig/Source/IKRig/Private/Retargeter/RetargetOps/ScaleSourceOp.cpp b/Engine/Plugins/Animation/IKRig/Source/IKRig/Private/Retargeter/RetargetOps/ScaleSourceOp.cpp
--- a/Engine/Plugins/Animation/IKRig/Source/IKRig/Private/Retargeter/RetargetOps/ScaleSourceOp.cpp
+++ b/Engine/Plugins/Animation/IKRig/Source/IKRig/Private/Retargeter/RetargetOps/ScaleSourceOp.cpp
@@ -2,8 +2,27 @@
 
 #include "Retargeter/RetargetOps/ScaleSourceOp.h"
 
+#include "Retargeter/IKRetargetProcessor.h"
+
 #include UE_INLINE_GENERATED_CPP_BY_NAME(ScaleSourceOp)
 
+bool FRetargetPoseScaleWithPivot::ScalePose(TArray<FTransform>& InOutGlobalPose) const
+{
+	if (FMath::IsNearlyEqual(Factor, 1.0))
+	{
+		return false;
+	}
+	
+	// scale the input pose
+	for (FTransform& SourceBoneTransform : InOutGlobalPose)
+	{
+		const FVector ScaleBoneOffset = (SourceBoneTransform.GetTranslation() - Pivot) * Factor;
+		SourceBoneTransform.SetTranslation(Pivot + ScaleBoneOffset);
+	}
+
+	return true;
+}
+
 const UClass* FIKRetargetScaleSourceOpSettings::GetControllerType() const
 {
 	return UIKRetargetScaleSourceController::StaticClass();
@@ -20,6 +39,13 @@ void FIKRetargetScaleSourceOpSettings::CopySettingsAtRuntime(const FIKRetargetOp
 		PropertiesToIgnore);
 }
 
+#if WITH_EDITORONLY_DATA
+USkeleton* FIKRetargetScaleSourceOpSettings::GetSkeleton(const FName InPropertyName)
+{
+	return const_cast<USkeleton*>(SourceSkeletonAsset);
+}
+#endif
+
 bool FIKRetargetScaleSourceOp::Initialize(
 	const FIKRetargetProcessor& InProcessor,
 	const FRetargetSkeleton& InSourceSkeleton,
@@ -46,6 +72,17 @@ const UScriptStruct* FIKRetargetScaleSourceOp::GetType() const
 	return FIKRetargetScaleSourceOp::StaticStruct();
 }
 
+int32 FIKRetargetScaleSourceOp::GetScalePivotBoneIndex(const FRetargetSkeleton& InSourceSkeleton) const
+{
+	if (CachedScalePivotBoneName != Settings.ScalePivotBone.BoneName)
+	{
+		CachedScalePivotBoneIndex = InSourceSkeleton.FindBoneIndexByName(Settings.ScalePivotBone.BoneName);
+		CachedScalePivotBoneName = Settings.ScalePivotBone.BoneName;
+	}
+
+	return CachedScalePivotBoneIndex;
+}
+
 FIKRetargetScaleSourceOpSettings UIKRetargetScaleSourceController::GetSettings()
 {
 	return *reinterpret_cast<FIKRetargetScaleSourceOpSettings*>(OpSettingsToControl);
diff --git a/Engine/Plugins/Animation/IKRig/Source/IKRig/Public/AnimNodes/AnimNode_RetargetPoseFromMesh.h b/Engine/Plugins/Animation/IKRig/Source/IKRig/Public/AnimNodes/AnimNode_RetargetPoseFromMesh.h
--- a/Engine/Plugins/Animation/IKRig/Source/IKRig/Public/AnimNodes/AnimNode_RetargetPoseFromMesh.h
+++ b/Engine/Plugins/Animation/IKRig/Source/IKRig/Public/AnimNodes/AnimNode_RetargetPoseFromMesh.h
@@ -88,7 +88,7 @@ struct FAnimNode_RetargetPoseFromMesh : public FAnimNode_Base
 	UE_API bool Serialize(FArchive& Ar);
 	UE_API void PostSerialize(const FArchive& Ar);
 
-#ifdef WITH_EDITORONLY_DATA
+#if WITH_EDITOR
 	UE_API double GetAverageExecutionTime() const;
 #endif
 
@@ -115,7 +115,7 @@ private:
 	// used to determine when to look for a parent component (prevents constantly searching component hierarchy)
 	bool bSearchedForParentComponent = false;
 
-#ifdef WITH_EDITORONLY_DATA
+#if WITH_EDITOR
 	double AverageExecutionTime;
 #endif
 };
diff --git a/Engine/Plugins/Animation/IKRig/Source/IKRig/Public/Retargeter/IKRetargetProcessor.h b/Engine/Plugins/Animation/IKRig/Source/IKRig/Public/Retargeter/IKRetargetProcessor.h
--- a/Engine/Plugins/Animation/IKRig/Source/IKRig/Public/Retargeter/IKRetargetProcessor.h
+++ b/Engine/Plugins/Animation/IKRig/Source/IKRig/Public/Retargeter/IKRetargetProcessor.h
@@ -31,6 +31,22 @@ struct FResolvedBoneChain;
 struct FIKRetargetDebugDrawState;
 
 
+struct FRetargetPoseScaleWithPivot
+{
+	FRetargetPoseScaleWithPivot() = default;
+	FRetargetPoseScaleWithPivot(const double InFactor, const FVector& InPivot) : Factor(InFactor), Pivot(InPivot) {};
+
+	bool operator==(const FRetargetPoseScaleWithPivot& Other) const
+	{
+		return FMath::IsNearlyEqual(Factor, Other.Factor) && Pivot.Equals(Other.Pivot);
+	}
+
+	bool ScalePose(TArray<FTransform>& InOutGlobalPose) const;
+	
+	double Factor = 1.0;
+	FVector Pivot = FVector::ZeroVector;
+};
+
 // cached retarget pose transforms (both local and global) after it has been applied/resolved on a specific skeletal mesh
 struct FResolvedRetargetPose
 {
@@ -38,7 +54,7 @@ struct FResolvedRetargetPose
 	
 	FName Name;						// the name of the retarget pose this was initialized with
 	int32 Version;					// the version of the retarget pose this was initialized with (transient)
-	double Scale;					// the scale of the retarget pose this was initialized with (transient)
+	FRetargetPoseScaleWithPivot PoseScale;			// the scale of the retarget pose this was initialized with (transient)
 	
 	TArray<FTransform> LocalPose;	// local space retarget pose
 	TArray<FTransform> GlobalPose;	// global space retarget pose
@@ -55,7 +71,7 @@ struct FResolvedRetargetPoseSet
 		const FName InRetargetPoseName,
 		const FIKRetargetPose* InRetargetPose,
 		const FName PelvisBoneName,
-		const double ScaleFactor);
+		const FRetargetPoseScaleWithPivot& InSourceScale);
 
 	UE_API const FResolvedRetargetPose* FindRetargetPoseByName(const FName InRetargetPoseName) const;
 	
@@ -93,7 +109,7 @@ struct FRetargetSkeleton
 		const ERetargetSourceOrTarget InSourceOrTarget,
 		const UIKRetargeter* InRetargetAsset,
 		const FName PelvisBoneName,
-		const double ScaleFactor);
+		const FRetargetPoseScaleWithPivot& InPoseScale);
 	
 	UE_API void Reset();
 
@@ -169,7 +185,8 @@ struct FTargetSkeleton : public FRetargetSkeleton
 		const USkeletalMesh* InSkeletalMesh,
 		const ERetargetSourceOrTarget InSourceOrTarget,
 		const UIKRetargeter* InRetargetAsset,
-		const FName RetargetRootBone);
+		const FName RetargetRootBone,
+		const FRetargetPoseScaleWithPivot& InPoseScale);
 
 	void Reset();
 
@@ -358,14 +375,6 @@ public:
 		const FRetargetProfile& InRetargetProfile,
 		const bool bSuppressWarnings=false);
 
-	/**
-	 * This scales the input source pose according to the source scale factor
-	 * This must be called outside of the retargeter itself because otherwise the pose may be scaled multiple times.
-	 * This is because we don't copy the source pose and we don't control when it's updated.
-	 * By explicitly scaling outside of RunRetargeter(), outside systems can scale it whenever they update it.
-	 * @param InOutSourceGlobalPose: an array of global space transforms for each bone in the source input pose */ 
-	UE_API void ScaleSourcePose(TArray<FTransform>& InOutSourceGlobalPose) const;
-
 	/**
 	* Run the retarget to generate a new pose.
 	* @param InSourceGlobalPose -  is the source mesh input pose in Component/Global space
@@ -471,9 +480,21 @@ public:
 	/** Get name of the pelvis bone for either the source or target skeleton. */
 	UE_API FName GetPelvisBone(ERetargetSourceOrTarget SourceOrTarget, ERetargetOpsToSearch InOpsToSearch) const;
 
+	/**
+	 * This scales the input source pose according to the source scale factor
+	 * This must be called outside of the retargeter itself because otherwise the pose may be scaled multiple times.
+	 * This is because we don't copy the source pose and we don't control when it's updated.
+	 * By explicitly scaling outside of RunRetargeter(), outside systems can scale it whenever they update it.
+	 * @param InOutSourceGlobalPose: an array of global space transforms for each bone in the source input pose
+	 */
+	UE_API void ApplySourceScaleToPose(TArray<FTransform>& InOutSourceGlobalPose) const;
+
 	/** Get the scale factor for the source pose (comes from presence of SourceScaleOp)*/
 	UE_API double GetSourceScaleFactor() const;
 
+	/** Get the scale factor and pivot for the retarget pose (comes from presence of SourceScaleOp)*/
+	UE_API const FRetargetPoseScaleWithPivot& GetPoseScale(const ERetargetSourceOrTarget SourceOrTarget) const;
+
 	/** The logging system */
 	FIKRigLogger Log;
 	
@@ -542,6 +563,9 @@ private:
 	/** The collection of operations to run to transfer animation from source to target. */
 	TArray<FInstancedStruct> OpStack;
 
+	/** cached source scaling */
+	mutable FRetargetPoseScaleWithPivot CurrentSourceScale;
+
 	/** Apply the settings stored in a retarget profile. Called inside RunRetargeter(). */
 	void ApplyProfileAtRuntime(const FRetargetProfile& Profile);
 
diff --git a/Engine/Plugins/Animation/IKRig/Source/IKRig/Public/Retargeter/RetargetOps/PinBoneOp.h b/Engine/Plugins/Animation/IKRig/Source/IKRig/Public/Retargeter/RetargetOps/PinBoneOp.h
--- a/Engine/Plugins/Animation/IKRig/Source/IKRig/Public/Retargeter/RetargetOps/PinBoneOp.h
+++ b/Engine/Plugins/Animation/IKRig/Source/IKRig/Public/Retargeter/RetargetOps/PinBoneOp.h
@@ -4,6 +4,7 @@
 
 #include "Retargeter/IKRetargeter.h"
 #include "Retargeter/IKRetargetOps.h"
+#include "Retargeter/IKRetargetProcessor.h"
 
 #include "PinBoneOp.generated.h"
 
@@ -35,7 +36,7 @@ struct FPinBoneData
 	FBoneReference BoneToCopyTo;
 
 	// the scale factor used to calculate all the cached offset (lazy evaluated when scale factor is modified)
-	double SourceScaleFactor;
+	FRetargetPoseScaleWithPivot LastUsedSourceScale;
 	// the local transform of the BoneToCopyTo in the ref pose
 	FTransform LocalRefPoseBoneToCopyTo;
 	// the local transform of the BoneToCopyFrom in the ref pose
diff --git a/Engine/Plugins/Animation/IKRig/Source/IKRig/Public/Retargeter/RetargetOps/ScaleSourceOp.h b/Engine/Plugins/Animation/IKRig/Source/IKRig/Public/Retargeter/RetargetOps/ScaleSourceOp.h
--- a/Engine/Plugins/Animation/IKRig/Source/IKRig/Public/Retargeter/RetargetOps/ScaleSourceOp.h
+++ b/Engine/Plugins/Animation/IKRig/Source/IKRig/Public/Retargeter/RetargetOps/ScaleSourceOp.h
@@ -2,6 +2,7 @@
 #pragma once
 
 #include "Retargeter/IKRetargetOps.h"
+#include "Animation/BoneReference.h"
 
 #include "ScaleSourceOp.generated.h"
 
@@ -9,6 +10,13 @@
 
 #define LOCTEXT_NAMESPACE "ScaleSourceOp"
 
+UENUM(BlueprintType)
+enum class EScaleSourcePivot : uint8
+{
+	ComponentOrigin		UMETA(DisplayName = "Component Origin"),
+	Bone				UMETA(DisplayName = "Bone"),
+};
+
 USTRUCT(BlueprintType, meta = (DisplayName = "Scale Source Settings"))
 struct FIKRetargetScaleSourceOpSettings : public FIKRetargetOpSettingsBase
 {
@@ -18,9 +26,22 @@ struct FIKRetargetScaleSourceOpSettings : public FIKRetargetOpSettingsBase
 	UPROPERTY(EditAnywhere, BlueprintReadWrite, Category = Scale, meta = (ReinitializeOnEdit, ClampMin = "0.01", UIMin = "0.01", UIMax = "10.0"))
 	double SourceScaleFactor = 1.0f;
 
+	UPROPERTY(EditAnywhere, BlueprintReadWrite, Category = ScalePivot)
+	EScaleSourcePivot ScalePivot;
+
+	UPROPERTY(EditAnywhere, Category = ScalePivot)
+	FBoneReference ScalePivotBone;
+
+	UPROPERTY(EditAnywhere, BlueprintReadWrite, Category = ScalePivot)
+	bool bProjectScalePivotToFloor = false;
+
 	UE_API virtual const UClass* GetControllerType() const override;
 
 	UE_API virtual void CopySettingsAtRuntime(const FIKRetargetOpSettingsBase* InSettingsToCopyFrom) override;
+
+#if WITH_EDITORONLY_DATA
+	virtual USkeleton* GetSkeleton(const FName InPropertyName) override;
+#endif
 };
 
 USTRUCT(BlueprintType, meta = (DisplayName = "Scale Source"))
@@ -52,8 +73,16 @@ struct FIKRetargetScaleSourceOp : public FIKRetargetOpBase
 
 	virtual bool IsSingleton() const override { return true; };
 
+	int32 GetScalePivotBoneIndex(const FRetargetSkeleton& InSourceSkeleton) const;
+
 	UPROPERTY()
 	FIKRetargetScaleSourceOpSettings Settings;
+
+private:
+	
+	// caching the bone index to scale the source from
+	mutable int32 CachedScalePivotBoneIndex = INDEX_NONE;
+	mutable FName CachedScalePivotBoneName = NAME_None;
 };
 
 /* The blueprint/python API for editing a Scale Source Op */
diff --git a/Engine/Plugins/Animation/IKRig/Source/IKRigEditor/Private/RetargetEditor/IKRetargetBatchOperation.cpp b/Engine/Plugins/Animation/IKRig/Source/IKRigEditor/Private/RetargetEditor/IKRetargetBatchOperation.cpp
--- a/Engine/Plugins/Animation/IKRig/Source/IKRigEditor/Private/RetargetEditor/IKRetargetBatchOperation.cpp
+++ b/Engine/Plugins/Animation/IKRig/Source/IKRigEditor/Private/RetargetEditor/IKRetargetBatchOperation.cpp
@@ -473,9 +473,9 @@ void UIKRetargetBatchOperation::ConvertAnimation(
 			// get the settings profile
 			FRetargetProfile SettingsProfile;
 			SettingsProfile.FillProfileWithAssetSettings(Context.IKRetargetAsset);
-
-			// give processor the pose to apply source scaling (if needed)
-			OutProcessor.ScaleSourcePose(SourceComponentPose);
+			
+			// let the retargeter scale the input pose
+			OutProcessor.ApplySourceScaleToPose(SourceComponentPose);
 
 			// run the retargeter
 			const TArray<FTransform>& TargetComponentPose = OutProcessor.RunRetargeter(SourceComponentPose, SettingsProfile, DeltaTime);
diff --git a/Engine/Plugins/Animation/IKRig/Source/IKRigEditor/Private/RetargetEditor/IKRetargetEditorController.cpp b/Engine/Plugins/Animation/IKRig/Source/IKRigEditor/Private/RetargetEditor/IKRetargetEditorController.cpp
--- a/Engine/Plugins/Animation/IKRig/Source/IKRigEditor/Private/RetargetEditor/IKRetargetEditorController.cpp
+++ b/Engine/Plugins/Animation/IKRig/Source/IKRigEditor/Private/RetargetEditor/IKRetargetEditorController.cpp
@@ -370,26 +370,43 @@ void FIKRetargetEditorController::UpdateSkeletalMeshComponents() const
 			return;
 		}
 		
-		USceneComponent* Component = GetSkeletalMeshComponent(SourceOrTarget);
-		if (!Component)
+		UDebugSkelMeshComponent* SkelMeshComponent = GetSkeletalMeshComponent(SourceOrTarget);
+		if (!SkelMeshComponent)
 		{
 			return;
 		}
 
-		// set component transform
-		const bool bIsSource  = SourceOrTarget == ERetargetSourceOrTarget::Source;
-		const FVector Position = bIsSource ? Asset->SourceMeshOffset : Asset->TargetMeshOffset;
-		const double SourceScaleFactor = GetRetargetProcessor()->GetSourceScaleFactor();
-		const double Scale = bIsSource ? SourceScaleFactor : Asset->TargetMeshScale;
+		// generate component transform
+		FTransform ComponentTransform = FTransform::Identity;
+		if (SourceOrTarget == ERetargetSourceOrTarget::Source)
+		{
+			// compose a delta FTransform that scales the source component about the scale pivot
+			FIKRetargetProcessor* Processor = GetRetargetProcessor();
+			const FRetargetPoseScaleWithPivot& PoseScale = Processor->GetPoseScale(ERetargetSourceOrTarget::Source);
+			const FVector TranslationDelta = (1.0 - PoseScale.Factor) * PoseScale.Pivot;
+			const FVector Scale(PoseScale.Factor,PoseScale.Factor,PoseScale.Factor);
+			FTransform Delta(FQuat::Identity, TranslationDelta, Scale);
+
+			// apply delta to the source component
+			ComponentTransform = FTransform(FQuat::Identity, Asset->SourceMeshOffset, FVector::OneVector);
+			ComponentTransform = Delta * ComponentTransform;
+		}
+		else
+		{
+			const double Scale = Asset->TargetMeshScale;
+			ComponentTransform.SetTranslation(Asset->TargetMeshOffset);
+			ComponentTransform.SetScale3D(FVector(Scale,Scale,Scale));
+		}
+			
+		// apply component transform
 		constexpr bool bSweep = false;
 		constexpr FHitResult* OutSweepHitResult = nullptr;
 		constexpr ETeleportType Teleport = ETeleportType::TeleportPhysics;
-		Component->SetWorldLocation(Position, bSweep, OutSweepHitResult, Teleport);
-		Component->SetWorldScale3D(FVector(Scale,Scale,Scale));
+		SkelMeshComponent->SetWorldTransform(ComponentTransform, bSweep, OutSweepHitResult, Teleport);
 			
 		// toggle mesh visibility
-		const bool bIsVisible = bIsSource ? Asset->bShowSourceMesh : Asset->bShowTargetMesh;
-		Component->SetVisibility(bIsVisible);	
+		const bool bIsVisible = SourceOrTarget == ERetargetSourceOrTarget::Source ? Asset->bShowSourceMesh : Asset->bShowTargetMesh;
+		SkelMeshComponent->SetVisibility(bIsVisible);	
 	};
 
 	UpdateMeshComponent(ERetargetSourceOrTarget::Source);
@@ -1683,15 +1700,14 @@ void FIKRetargetEditorController::FixZeroHeightRetargetRoot(ERetargetSourceOrTar
 		return;
 	}
 
-	const double SourceScaleFactor = GetRetargetProcessor()->GetSourceScaleFactor();
-	const double ScaleFactor = SourceOrTarget == ERetargetSourceOrTarget::Source ? SourceScaleFactor : 1.0;
+	const FRetargetPoseScaleWithPivot& PoseScale = GetRetargetProcessor()->GetPoseScale(SourceOrTarget);
 	FRetargetSkeleton DummySkeleton;
 	DummySkeleton.Initialize(
 		SkeletalMesh,
 		SourceOrTarget,
 		AssetController->GetAsset(),
 		RetargetRootBoneName,
-		ScaleFactor);
+		PoseScale);
 
 	const int32 RootBoneIndex = DummySkeleton.FindBoneIndexByName(RetargetRootBoneName);
 	if (RootBoneIndex == INDEX_NONE)
