[IK Retarget] New Floor Constraint op.

diff --git a/Engine/Plugins/Animation/IKRig/Source/IKRig/Private/Retargeter/IKRetargetProcessor.cpp b/Engine/Plugins/Animation/IKRig/Source/IKRig/Private/Retargeter/IKRetargetProcessor.cpp
--- a/Engine/Plugins/Animation/IKRig/Source/IKRig/Private/Retargeter/IKRetargetProcessor.cpp
+++ b/Engine/Plugins/Animation/IKRig/Source/IKRig/Private/Retargeter/IKRetargetProcessor.cpp
@@ -1138,9 +1138,6 @@ void FIKRetargetProcessor::Initialize(
 		GetPelvisBone(ERetargetSourceOrTarget::Target, ERetargetOpsToSearch::AssetOps),
 		GetPoseScale(ERetargetSourceOrTarget::Target));
 
-	// setup the goal container
-	ResetIKRigGoalContainer();
-
 	// copy ops from asset 
 	InitialOpStackSetup(RetargeterAsset->GetRetargetOps(), InRetargetProfile);
 	
@@ -1159,26 +1156,6 @@ void FIKRetargetProcessor::Initialize(
 #endif
 }
 
-void FIKRetargetProcessor::ResetIKRigGoalContainer()
-{
-	// fill the goal container with all the goals in the IK Rig
-	GoalContainer.Empty();
-	const UIKRigDefinition* TargetIKRig = RetargeterAsset->GetIKRig(ERetargetSourceOrTarget::Target);
-	if (TargetIKRig)
-	{
-		GoalContainer.FillWithGoalArray(TargetIKRig->GetGoalArray());
-	}
-	
-	// set all goals to be additive (ops may adjust them)
-	for (FIKRigGoal& Goal : GoalContainer.GetGoalArray())
-	{
-		Goal.PositionSpace = EIKRigGoalSpace::Additive;
-		Goal.RotationSpace = EIKRigGoalSpace::Additive;
-		Goal.Position = FVector::ZeroVector;
-		Goal.Rotation = FRotator::ZeroRotator;
-	}
-}
-
 void FIKRetargetProcessor::InitialOpStackSetup(
 	const TArray<FInstancedStruct>& OpsFromAsset,
 	const FRetargetProfile& InRetargetProfile)
@@ -1243,7 +1220,7 @@ void FIKRetargetProcessor::InitializeRetargetOps()
 	for (FInstancedStruct& OpStruct : OpStack)
 	{
 		FIKRetargetOpBase& Op = OpStruct.GetMutable<FIKRetargetOpBase>();
-		const FIKRetargetOpBase* ParentOp = GetRetargetOpByName(Op.GetParentOpName());
+		FIKRetargetOpBase* ParentOp = GetRetargetOpByName(Op.GetParentOpName());
 		Op.Initialize(*this, SourceSkeleton, TargetSkeleton, ParentOp, Log);
 	}
 	
@@ -1316,9 +1293,6 @@ TArray<FTransform>&  FIKRetargetProcessor::RunRetargeter(
 	
 	// generate the pose to start retargeting from
 	GenerateBasePoses(InSourceGlobalPose);
-
-	// reset all goals (ops may pose / adjust them)
-	ResetIKRigGoalContainer();
 	
 	// run the stack of retargeting operations
 	RunRetargetOps(InSourceGlobalPose, TargetSkeleton.OutputGlobalPose, InDeltaTime, InLOD);
@@ -1358,6 +1332,8 @@ void FIKRetargetProcessor::RunRetargetOps(
 
 		return true;
 	};
+
+	TArray<FIKRetargetOpBase*> OpsRunBeforeChildrenThisTick;
 	
 	for (FInstancedStruct& OpStruct : OpStack)
 	{
@@ -1368,8 +1344,18 @@ void FIKRetargetProcessor::RunRetargetOps(
 
 		// execute the op
 		FIKRetargetOpBase& Op = OpStruct.GetMutable<FIKRetargetOpBase>();
+
+		// give parent ops a chance to run right before their children
+		FIKRetargetOpBase* ParentOp = GetRetargetOpByName(Op.GetParentOpName());
+		if (ParentOp && !OpsRunBeforeChildrenThisTick.Contains(ParentOp))
+		{
+			ParentOp->RunBeforeChildren(*this, InDeltaTime, InSourceGlobalPose, OutTargetGlobalPose);
+			OpsRunBeforeChildrenThisTick.Add(ParentOp);
+		}
+		
 		if (ShouldExecuteOp(Op))
 		{
+			// run the op
 			Op.Run(*this, InDeltaTime, InSourceGlobalPose, OutTargetGlobalPose);	
 		}
 
@@ -1396,11 +1382,11 @@ TArray<const FIKRetargetOpBase*> FIKRetargetProcessor::GetRetargetOpsByType(cons
 	return MoveTemp(Ops);
 }
 
-const FIKRetargetOpBase* FIKRetargetProcessor::GetRetargetOpByName(const FName InOpName) const
+FIKRetargetOpBase* FIKRetargetProcessor::GetRetargetOpByName(const FName InOpName)
 {
-	for (const FInstancedStruct& OpStruct : OpStack)
+	for (FInstancedStruct& OpStruct : OpStack)
 	{
-		const FIKRetargetOpBase* Op = OpStruct.GetPtr<FIKRetargetOpBase>();
+		FIKRetargetOpBase* Op = OpStruct.GetMutablePtr<FIKRetargetOpBase>();
 		if (Op->GetName() == InOpName)
 		{
 			return Op;
diff --git a/Engine/Plugins/Animation/IKRig/Source/IKRig/Private/Retargeter/RetargetOps/FloorConstraintOp.cpp b/Engine/Plugins/Animation/IKRig/Source/IKRig/Private/Retargeter/RetargetOps/FloorConstraintOp.cpp
new file mode 100644
--- /dev/null
+++ b/Engine/Plugins/Animation/IKRig/Source/IKRig/Private/Retargeter/RetargetOps/FloorConstraintOp.cpp
@@ -0,0 +1,268 @@
+// Copyright Epic Games, Inc. All Rights Reserved.
+
+#include "Retargeter/RetargetOps/FloorConstraintOp.h"
+
+#include "Retargeter/IKRetargetOpUtils.h"
+#include "Retargeter/IKRetargetProcessor.h"
+#include "Retargeter/RetargetOps/IKChainsOp.h"
+#include "Retargeter/RetargetOps/RunIKRigOp.h"
+
+#define LOCTEXT_NAMESPACE "FloorConstraintOp"
+
+bool FFloorConstraintChainSettings::operator==(const FFloorConstraintChainSettings& Other) const
+{
+	return EnableFloorConstraint == Other.EnableFloorConstraint
+		&& FMath::IsNearlyEqualByULP(Alpha,Other.Alpha);
+}
+
+const UClass* FIKRetargetFloorConstraintOpSettings::GetControllerType() const
+{
+	return UIKRetargetFloorGoalsController::StaticClass();
+}
+
+void FIKRetargetFloorConstraintOpSettings::CopySettingsAtRuntime(const FIKRetargetOpSettingsBase* InSettingsToCopyFrom)
+{
+	// copies everything except the ChainsToAffect array (those are copied below, only for already existing chains)
+	const TArray<FName> PropertiesToIgnore = {"ChainsToAffect"};
+	FIKRetargetOpBase::CopyStructProperties(
+		FIKRetargetFloorConstraintOpSettings::StaticStruct(),
+		InSettingsToCopyFrom,
+		this,
+		PropertiesToIgnore);
+	
+	// copy settings only for chains that the op has initialized
+	const FIKRetargetFloorConstraintOpSettings* NewSettings = reinterpret_cast<const FIKRetargetFloorConstraintOpSettings*>(InSettingsToCopyFrom);
+	for (const FFloorConstraintChainSettings& NewChainSettings : NewSettings->ChainsToAffect)
+	{
+		for (FFloorConstraintChainSettings& ChainSettings : ChainsToAffect)
+		{
+			if (ChainSettings.TargetChainName == NewChainSettings.TargetChainName)
+			{
+				ChainSettings = NewChainSettings;
+			}
+		}
+	}
+}
+
+bool FIKRetargetFloorConstraintOp::Initialize(
+    const FIKRetargetProcessor& InProcessor,
+    const FRetargetSkeleton& InSourceSkeleton,
+    const FTargetSkeleton& InTargetSkeleton,
+    const FIKRetargetOpBase* InParentOp,
+    FIKRigLogger& InLog)
+{
+	bIsInitialized = false;
+	
+	// validate that an IK rig has been assigned to the parent op
+	const FIKRetargetRunIKRigOp* ParentOp = reinterpret_cast<const FIKRetargetRunIKRigOp*>(InParentOp);
+	if (ParentOp->Settings.IKRigAsset == nullptr)
+	{
+		InLog.LogWarning( FText::Format(
+		LOCTEXT("MissingIKRig", "{0}, is missing an IK rig. No chains can be retargeted."), FText::FromName(GetName())));
+		return false;
+	}
+
+	FloorConstraints.Reset();
+	const FRetargeterBoneChains& BoneChains = InProcessor.GetBoneChains();
+	for (const FFloorConstraintChainSettings& Chain : Settings.ChainsToAffect)
+	{
+		const FName TargetChainName = Chain.TargetChainName;
+		const FResolvedBoneChain* TargetBoneChain = BoneChains.GetResolvedBoneChainByName(TargetChainName, ERetargetSourceOrTarget::Target,ParentOp->Settings.IKRigAsset);
+		if (!TargetBoneChain)
+		{
+			InLog.LogWarning( FText::Format(
+			LOCTEXT("FloorConstraintChainNotFound", "Floor Constraint: target chain not found in IK Rig, '{0}."), FText::FromName(TargetChainName)));
+			continue;
+		}
+
+		if (TargetBoneChain->IKGoalName == NAME_None)
+		{
+			InLog.LogWarning( FText::Format(
+			LOCTEXT("FloorConstraintIKNotFound", "Floor Constraint: target chain has no IK Goal, '{0}."), FText::FromName(TargetChainName)));
+			continue;
+		}
+
+		// which source chain was this target chain mapped to?
+		const FName SourceChainName = ParentOp->ChainMapping.GetChainMappedTo(TargetChainName, ERetargetSourceOrTarget::Target);
+		const FResolvedBoneChain* SourceBoneChain = BoneChains.GetResolvedBoneChainByName(SourceChainName, ERetargetSourceOrTarget::Source);
+		if (!SourceBoneChain)
+		{
+			InLog.LogWarning( FText::Format(
+			LOCTEXT("FloorConstraintChainNotMapped", "Floor Constraint: found IK chain that was not mapped to a source chain, '{0}."), FText::FromName(TargetChainName)));
+			continue;
+		}
+
+		// all prerequisites met, create a new floor constraint
+		FFloorConstraint FloorConstraint;
+		FloorConstraint.IKGoalName = TargetBoneChain->IKGoalName;
+		FloorConstraint.SourceEndBoneIndex = SourceBoneChain->BoneIndices.Last();
+		FloorConstraint.Settings = &Chain;
+
+		// cache ref pose offset
+		const FTransform& SourceEndBoneRefPose = SourceBoneChain->RefPoseGlobalTransforms.Last();
+		const FTransform& TargetEndBoneRefPose = TargetBoneChain->RefPoseGlobalTransforms.Last();
+		FloorConstraint.HeightOffsetInRefPose = TargetEndBoneRefPose.GetTranslation().Z - SourceEndBoneRefPose.GetTranslation().Z;
+
+		// store floor constraint data for runtime
+		FloorConstraints.Add(FloorConstraint);
+	}
+	
+	bIsInitialized = true;
+    return true;
+}
+
+void FIKRetargetFloorConstraintOp::Run(
+    FIKRetargetProcessor& InProcessor,
+    const double InDeltaTime,
+    const TArray<FTransform>& InSourceGlobalPose,
+    TArray<FTransform>& OutTargetGlobalPose)
+{
+	FIKRigGoalContainer& GoalContainer = InProcessor.GetIKRigGoalContainer();
+	for (FFloorConstraint& FloorConstraint : FloorConstraints)
+	{
+		if (!FloorConstraint.Settings->EnableFloorConstraint)
+		{
+			continue; // constraint disabled
+		}
+		
+		if (FMath::IsNearlyZero(FloorConstraint.Settings->Alpha))
+		{
+			continue; // constraint weight near zero
+		}
+
+		// get the goal to adjust
+		FIKRigGoal* Goal = GoalContainer.FindGoalByName(FloorConstraint.IKGoalName);
+		if (!ensure(Goal))
+		{
+			continue; // goal missing (should not happen at runtime)
+		}
+
+		// calculate falloff range
+		const float FalloffStartHeight = Settings.HeightFalloffOffset;
+		float FalloffEndHeight = FalloffStartHeight + FMath::Max(0.0, Settings.HeightFalloffDistance);
+
+		// get the current height of the source bone
+		const FTransform& CurrentSourceEndGlobal = InSourceGlobalPose[FloorConstraint.SourceEndBoneIndex];
+		const float SourceHeight = FMath::Abs(CurrentSourceEndGlobal.GetLocation().Z);
+		if (SourceHeight >= FalloffEndHeight)
+		{
+			continue; // constraint off while source is not near ground
+		}
+		
+		// calculate the target height of the goal by blending towards the source height when close to floor
+		const double FalloffRange = FMath::Max(1.0,FalloffEndHeight - FalloffStartHeight);
+		const double Falloff = SourceHeight < FalloffStartHeight ? 1.0f : 1.0f - (SourceHeight - FalloffStartHeight) / FalloffRange;
+		const double Weight = Falloff * FloorConstraint.Settings->Alpha;
+		const double TargetHeight = SourceHeight  + FloorConstraint.HeightOffsetInRefPose;
+		
+		// lerp goal height from it's current height to the target height by the weight
+		Goal->Position.Z = FMath::Lerp(Goal->Position.Z, TargetHeight, Weight);
+	}
+}
+
+void FIKRetargetFloorConstraintOp::OnAddedToStack(const UIKRetargeter* InRetargetAsset, const FIKRetargetOpBase* InParentOp)
+{
+	// regenerate chain settings
+	constexpr bool bSkipUnmappedChains = false;
+	constexpr bool bSkipNonIKChains = true;
+	IKRetargetOpUtils::SynchronizeChainSettingsWithIKRig(Settings.ChainsToAffect, InParentOp, bSkipUnmappedChains, bSkipNonIKChains);
+}
+
+void FIKRetargetFloorConstraintOp::OnAssignIKRig(
+    const ERetargetSourceOrTarget SourceOrTarget,
+    const UIKRigDefinition* InIKRig,
+    const FIKRetargetOpBase* InParentOp)
+{
+	// regenerate chain settings
+	constexpr bool bSkipUnmappedChains = false;
+	constexpr bool bSkipNonIKChains = true;
+	IKRetargetOpUtils::SynchronizeChainSettingsWithIKRig(Settings.ChainsToAffect, InParentOp, bSkipUnmappedChains, bSkipNonIKChains);
+}
+
+FIKRetargetOpSettingsBase* FIKRetargetFloorConstraintOp::GetSettings()
+{
+    return &Settings;
+}
+
+void FIKRetargetFloorConstraintOp::SetSettings(const FIKRetargetOpSettingsBase* InSettings)
+{
+    if (InSettings)
+    {
+        Settings = *static_cast<const FIKRetargetFloorConstraintOpSettings*>(InSettings);
+    }
+}
+
+const UScriptStruct* FIKRetargetFloorConstraintOp::GetSettingsType() const
+{
+    return FIKRetargetFloorConstraintOpSettings::StaticStruct();
+}
+
+const UScriptStruct* FIKRetargetFloorConstraintOp::GetType() const
+{
+    return StaticStruct();
+}
+
+const UScriptStruct* FIKRetargetFloorConstraintOp::GetParentOpType() const
+{
+    return FIKRetargetRunIKRigOp::StaticStruct();
+}
+
+void FIKRetargetFloorConstraintOp::OnTargetChainRenamed(const FName InOldChainName, const FName InNewChainName)
+{
+	IKRetargetOpUtils::OnRetargetChainRenamed(Settings.ChainsToAffect, InOldChainName, InNewChainName);
+}
+
+void FIKRetargetFloorConstraintOp::OnParentReinitPropertyEdited(
+    const FIKRetargetOpBase& InParentOp,
+    const FPropertyChangedEvent* InPropertyChangedEvent)
+{
+	// regenerate chain settings
+	constexpr bool bSkipUnmappedChains = false;
+	constexpr bool bSkipNonIKChains = true;
+	IKRetargetOpUtils::SynchronizeChainSettingsWithIKRig(Settings.ChainsToAffect, &InParentOp, bSkipUnmappedChains, bSkipNonIKChains);
+}
+
+// ------------------------------------------------------------
+// Editor-only
+// ------------------------------------------------------------
+#if WITH_EDITOR
+
+void FIKRetargetFloorConstraintOp::ResetChainSettingsToDefault(const FName& InChainName)
+{
+	for (FFloorConstraintChainSettings& Chain : Settings.ChainsToAffect)
+	{
+		if (Chain.TargetChainName == InChainName)
+		{
+			Chain = FFloorConstraintChainSettings(InChainName);
+			return;
+		}
+	}
+}
+
+bool FIKRetargetFloorConstraintOp::AreChainSettingsAtDefault(const FName& InChainName)
+{
+	for (FFloorConstraintChainSettings& ChainToRetarget : Settings.ChainsToAffect)
+	{
+		if (ChainToRetarget.TargetChainName == InChainName)
+		{
+			FFloorConstraintChainSettings DefaultSettings = FFloorConstraintChainSettings();
+			return ChainToRetarget == DefaultSettings;
+		}
+	}
+
+	return true;
+}
+
+#endif // WITH_EDITOR
+
+FIKRetargetIKChainsOpSettings UIKRetargetFloorGoalsController::GetSettings()
+{
+	return *reinterpret_cast<FIKRetargetIKChainsOpSettings*>(OpSettingsToControl);
+}
+
+void UIKRetargetFloorGoalsController::SetSettings(FIKRetargetIKChainsOpSettings InSettings)
+{
+	OpSettingsToControl->CopySettingsAtRuntime(&InSettings);
+}
+
+#undef LOCTEXT_NAMESPACE
diff --git a/Engine/Plugins/Animation/IKRig/Source/IKRig/Private/Retargeter/RetargetOps/RunIKRigOp.cpp b/Engine/Plugins/Animation/IKRig/Source/IKRig/Private/Retargeter/RetargetOps/RunIKRigOp.cpp
--- a/Engine/Plugins/Animation/IKRig/Source/IKRig/Private/Retargeter/RetargetOps/RunIKRigOp.cpp
+++ b/Engine/Plugins/Animation/IKRig/Source/IKRig/Private/Retargeter/RetargetOps/RunIKRigOp.cpp
@@ -52,6 +52,17 @@ bool FIKRetargetRunIKRigOp::Initialize(
 			LOCTEXT("CouldNotInitializeIKRig", "Run IK Rig Op: unable to initialize the IK Rig, {0} for the Skeletal Mesh {1}. See previous warnings."),
 			FText::FromString(Settings.IKRigAsset->GetName()), FText::FromString(InTargetSkeleton.SkeletalMesh->GetName())));
 	}
+
+	// cache goal bone indices
+	GoalBoneIndicesMap.Reset();
+	for (UIKRigEffectorGoal* Goal : Settings.IKRigAsset->GetGoalArray())
+	{
+		const int32 BoneIndex = InTargetSkeleton.FindBoneIndexByName(Goal->BoneName);
+		if (BoneIndex != INDEX_NONE)
+		{
+			GoalBoneIndicesMap.Add(Goal->BoneName, BoneIndex);
+		}
+	}
 	
 	bIsInitialized = IKRigProcessor.IsInitialized();
 	return bIsInitialized;
@@ -106,6 +117,50 @@ void FIKRetargetRunIKRigOp::Run(
 #endif
 }
 
+void FIKRetargetRunIKRigOp::RunBeforeChildren(
+	FIKRetargetProcessor& InProcessor,
+	const double InDeltaTime,
+	const TArray<FTransform>& InSourceGlobalPose,
+	TArray<FTransform>& OutTargetGlobalPose)
+{
+	auto ResetGoalContainer = [this, &OutTargetGlobalPose](FIKRigGoalContainer& InGoalContainer)
+	{
+		// if no IK Rig assigned, leave container empty
+		if (!Settings.IKRigAsset)
+		{
+			InGoalContainer.Empty();
+			return;
+		}
+
+		// add all the goals from the IK Rig (empties and refills)
+		InGoalContainer.FillWithGoalArray(Settings.IKRigAsset->GetGoalArray());
+
+		// set all goals to the input pose in component space
+		for (FIKRigGoal& Goal : InGoalContainer.GetGoalArray())
+		{
+			const int32* BoneIndexPtr = GoalBoneIndicesMap.Find(Goal.BoneName);
+			if (BoneIndexPtr)
+			{
+				Goal.PositionSpace = EIKRigGoalSpace::Component;
+				Goal.RotationSpace = EIKRigGoalSpace::Component;
+				Goal.Position = OutTargetGlobalPose[*BoneIndexPtr].GetLocation();
+				Goal.Rotation = OutTargetGlobalPose[*BoneIndexPtr].GetRotation().Rotator();
+			}
+			else
+			{
+				Goal.PositionSpace = EIKRigGoalSpace::Additive;
+				Goal.RotationSpace = EIKRigGoalSpace::Additive;
+				Goal.Position = FVector::ZeroVector;
+				Goal.Rotation = FRotator::ZeroRotator;
+			}
+		}
+	};
+
+	// reset the goal container for this IK Rig
+	FIKRigGoalContainer& GoalContainer = InProcessor.GetIKRigGoalContainer();
+	ResetGoalContainer(GoalContainer);
+}
+
 void FIKRetargetRunIKRigOp::OnAddedToStack(const UIKRetargeter* InRetargetAsset, const FIKRetargetOpBase* InParentOp)
 {
 	// load the target IK Rig asset to execute
diff --git a/Engine/Plugins/Animation/IKRig/Source/IKRig/Public/Retargeter/IKRetargetOpUtils.h b/Engine/Plugins/Animation/IKRig/Source/IKRig/Public/Retargeter/IKRetargetOpUtils.h
new file mode 100644
--- /dev/null
+++ b/Engine/Plugins/Animation/IKRig/Source/IKRig/Public/Retargeter/IKRetargetOpUtils.h
@@ -0,0 +1,126 @@
+﻿// Copyright Epic Games, Inc. All Rights Reserved.
+
+#pragma once
+#include "IKRetargetOps.h"
+#include "IKRetargetSettings.h"
+#include "Chaos/Character/CharacterGroundConstraintContainer.h"
+#include "Rig/IKRigDefinition.h"
+#include "Retargeter/IKRetargetChainMapping.h"
+
+namespace IKRetargetOpUtils
+{
+	/**
+	 * This is a convenience function for retarget ops that need to keep an
+	 * array of custom chain settings in-sync with the chains in and IK Rig asset.
+	 *
+	 * NOTE: to use this, pass in an array of custom types T and ensure they
+	 *
+	 * 1. Implement these functions:
+	 * FName GetName()
+	 * void SetName(FName InChainName)
+	 *
+	 * 2. Provide a constructor that takes the chain name
+	 */
+	template<typename T>
+	void SynchronizeChainSettingsWithIKRig(
+		TArray<T>& InOutChainSettings,
+		const FIKRetargetOpBase* InRetargetOp,
+		const bool bSkipUnmappedChains,
+		const bool bSkipNonIKChains)
+	{
+		// NOTE: this function early returns if there's no valid chains found: we do NOT clear the settings
+		// this allows users to clear and reassign a different rig and potentially retain/restore compatible settings
+		
+		if (!InRetargetOp)
+		{
+			return;
+		}
+
+		const UIKRigDefinition* TargetIKRig = InRetargetOp->GetCustomTargetIKRig();
+		if (!TargetIKRig)
+		{
+			return;
+		}
+
+		const TArray<FBoneChain>& TargetChains = TargetIKRig->GetRetargetChains();
+		if (TargetChains.IsEmpty())
+		{
+			return;
+		}
+
+		const FRetargetChainMapping* ChainMapping = InRetargetOp->GetChainMapping();
+		if (bSkipUnmappedChains && !ChainMapping)
+		{
+			return;
+		}
+
+		// find the target chains that children ops should deal with
+		TArray<FName> RequiredTargetChains;
+		for (const FBoneChain& TargetChain : TargetChains)
+		{
+			if (bSkipNonIKChains && TargetChain.IKGoalName == NAME_None)
+			{
+				continue; // skip non-IK chains
+			}
+
+			if (!bSkipUnmappedChains)
+			{
+				const FName SourceChain = ChainMapping->GetChainMappedTo(TargetChain.ChainName, ERetargetSourceOrTarget::Target);
+				if (bSkipUnmappedChains && SourceChain == NAME_None)
+				{
+					continue; // skip unmapped chains
+				}
+			}
+			
+			RequiredTargetChains.Add(TargetChain.ChainName);
+		}
+
+		if (RequiredTargetChains.IsEmpty())
+		{
+			return;
+		}
+
+		// remove chains that are not required
+		InOutChainSettings.RemoveAll([&RequiredTargetChains](const T& InChainSettings)
+		{
+			return !RequiredTargetChains.Contains(InChainSettings.GetName());
+		});
+    	
+		// add any required chains not already present
+		for (FName RequiredTargetChain : RequiredTargetChains)
+		{
+			bool bFound = false;
+			for (const T& ChainSettings : InOutChainSettings)
+			{
+				if (ChainSettings.GetName() == RequiredTargetChain)
+				{
+					bFound = true;
+					break;
+				}
+			}
+			if (!bFound)
+			{
+				InOutChainSettings.Emplace(RequiredTargetChain);
+			}
+		}
+	};
+
+	/**
+	 * This is a convenience function for retarget ops that need to respond to an IK Rig chain being renamed.
+	 *
+	 * NOTE: to use this, pass in an array of custom types T and ensure they implement these functions:
+	 * FName GetName()
+	 * void SetName(FName InChainName) */
+	template<typename T>
+	void OnRetargetChainRenamed(TArray<T>& InOutChainSettings, const FName InOldChainName, const FName InNewChainName)
+	{
+		for (T& ChainSettings : InOutChainSettings)
+		{
+			if (ChainSettings.GetName() == InOldChainName)
+			{
+				ChainSettings.SetName(InNewChainName);
+			}
+		}
+	};
+	
+}
diff --git a/Engine/Plugins/Animation/IKRig/Source/IKRig/Public/Retargeter/IKRetargetOps.h b/Engine/Plugins/Animation/IKRig/Source/IKRig/Public/Retargeter/IKRetargetOps.h
--- a/Engine/Plugins/Animation/IKRig/Source/IKRig/Public/Retargeter/IKRetargetOps.h
+++ b/Engine/Plugins/Animation/IKRig/Source/IKRig/Public/Retargeter/IKRetargetOps.h
@@ -226,6 +226,15 @@ public:
 	/** (optional) ops can optionally behave as a 'child' where they can be parented to ops of the type returned by this function. */
 	virtual const UScriptStruct* GetParentOpType() const { return nullptr; };
 
+	/** (optional) override to give this op a chance to do work before it's children are run
+	 * This is executed right after the previous op in the stack and right before the first child op Run() is called
+	 * NOTE: unless this is a parent op, this will have no effect */
+	virtual void RunBeforeChildren(
+		FIKRetargetProcessor& InProcessor,
+		const double InDeltaTime,
+		const TArray<FTransform>& InSourceGlobalPose,
+		TArray<FTransform>& OutTargetGlobalPose){};
+
 	/** (optional) return true to disallow multiple copies of this op in the stack */
 	virtual bool IsSingleton() const { return false; };
 
@@ -233,9 +242,13 @@ public:
 	 * NOTE: The retarget processor will resolve the bone chains in the rig so that they can be queried by the op. */
 	virtual const UIKRigDefinition* GetCustomTargetIKRig() const { return nullptr; };
 	
-	/** (optional) ops can optionally store their own chain mapping, this allows outside systems to query/edit it. */
+	/** (optional) ops can optionally store their own chain mapping, this allows outside systems to query/edit it.
+	 * NOTE: this is conceptually const and should not alter the op state, but only return the chain mapping! */
 	virtual FRetargetChainMapping* GetChainMapping() { return nullptr; };
 
+	/** wraps the non-const GetChainMapping() to provide const access */
+	const FRetargetChainMapping* GetChainMapping() const { return const_cast<FIKRetargetOpBase*>(this)->GetChainMapping(); };
+
 	/** (optional) implement this if the op stores chain settings by name to allow the settings to be maintained after a chain is renamed
 	 * NOTE: this is only called when the target IK Rig this op references has a chain that is renamed.
 	 * NOTE: ops do not have to update chain mappings in this function, that is managed for them if GetChainMapping() is implemented */
diff --git a/Engine/Plugins/Animation/IKRig/Source/IKRig/Public/Retargeter/IKRetargetProcessor.h b/Engine/Plugins/Animation/IKRig/Source/IKRig/Public/Retargeter/IKRetargetProcessor.h
--- a/Engine/Plugins/Animation/IKRig/Source/IKRig/Public/Retargeter/IKRetargetProcessor.h
+++ b/Engine/Plugins/Animation/IKRig/Source/IKRig/Public/Retargeter/IKRetargetProcessor.h
@@ -429,7 +429,7 @@ public:
 	UE_API TArray<const FIKRetargetOpBase*> GetRetargetOpsByType(const UScriptStruct* OpType) const;
 
 	/** Get an op with the given name */
-	UE_API const FIKRetargetOpBase* GetRetargetOpByName(const FName InOpName) const;
+	UE_API FIKRetargetOpBase* GetRetargetOpByName(const FName InOpName);
 
 	/** Get the first op in the stack of the given type */
 	template <typename T>
@@ -572,9 +572,6 @@ private:
 	/** Internal retarget phase that does simple bone-to-bone copying from source*/
 	void GenerateBasePoses(TArray<FTransform>& InSourceGlobalPose);
 
-	/** Setup the goal container for IK Rigs to use */
-	void ResetIKRigGoalContainer();
-
 	/** Initial setup of the retarget op stack (copies op stack from asset into processor) */
 	void InitialOpStackSetup(const TArray<FInstancedStruct>& OpsFromAsset, const FRetargetProfile& InRetargetProfile);
 
diff --git a/Engine/Plugins/Animation/IKRig/Source/IKRig/Public/Retargeter/RetargetOps/FloorConstraintOp.h b/Engine/Plugins/Animation/IKRig/Source/IKRig/Public/Retargeter/RetargetOps/FloorConstraintOp.h
new file mode 100644
--- /dev/null
+++ b/Engine/Plugins/Animation/IKRig/Source/IKRig/Public/Retargeter/RetargetOps/FloorConstraintOp.h
@@ -0,0 +1,145 @@
+﻿// Copyright Epic Games, Inc. All Rights Reserved.
+#pragma once
+
+#include "Retargeter/IKRetargetChainMapping.h"
+#include "Retargeter/IKRetargetOps.h"
+#include "FloorConstraintOp.generated.h"
+
+#define UE_API IKRIG_API
+
+#define LOCTEXT_NAMESPACE "FloorConstraintOp"
+
+
+USTRUCT(BlueprintType)
+struct FFloorConstraintChainSettings
+{
+	GENERATED_BODY()
+
+	FFloorConstraintChainSettings() = default;
+	FFloorConstraintChainSettings(const FName InTargetChainName) : TargetChainName(InTargetChainName) {}
+
+	/** The name of the TARGET chain to transfer animation onto. */
+	UPROPERTY(EditAnywhere, BlueprintReadWrite, Category = "Setup", meta=(ReinitializeOnEdit))
+	FName TargetChainName;
+	
+	/** Whether to apply the floor constraint to the location of the IK goal on this chain. Default is false.
+	 * When ON, the floor constraint will adjust the vertical position of the IK Goal according to the following rules.
+	 * 1. When the source goal bone is LOWER than FloorHeightFalloffStart, the height of the goal is smoothly blended to the height of the source bone.
+	 * 2. When the source goal bone is HIGHER than FloorHeightFalloffStart, the height of the goal is left at its normal retargeted location.
+	 * NOTE: the floor is assumed to be the XY plane where Z = 0.
+	 * NOTE: This only has an effect if the chain has an IK Goal assigned to it in the Target IK Rig asset.
+	 */
+	UPROPERTY(EditAnywhere, BlueprintReadWrite, Category = "Floor Constraint Chain Settings", meta=(ReinitializeOnEdit))
+	bool EnableFloorConstraint = false;
+
+	/** Range 0 to 1. Default is 0. Blend the effect of the constraint on this goal on/off. */
+	UPROPERTY(EditAnywhere, BlueprintReadWrite, Category = "Floor Constraint Chain Settings", meta = (ClampMin = "0.0", ClampMax = "1.0", UIMin = "0.0", UIMax = "1.0"))
+	double Alpha = 1.0f;
+
+	UE_API bool operator==(const FFloorConstraintChainSettings& Other) const;
+
+	FName GetName() const { return TargetChainName; };
+	void SetName(const FName InName) { TargetChainName = InName; };
+};
+
+USTRUCT(BlueprintType, meta = (DisplayName = "Floor Goal Op Settings"))
+struct FIKRetargetFloorConstraintOpSettings : public FIKRetargetOpSettingsBase
+{
+	GENERATED_BODY()
+
+	/** The per-chain settings (exposed indirectly to the UI through a detail customization) */
+	UPROPERTY(EditAnywhere, BlueprintReadWrite, Category = Chains, meta=(ReinitializeOnEdit))
+	TArray<FFloorConstraintChainSettings> ChainsToAffect;
+
+	/** Range 0 to inf. Default is 8. The height in cm from the floor below which the goal is snapped directly to the source bone height.
+     * NOTE: if the source bone height is greater than this value, but lower than FloorHeightFalloffEnd, then the height will smoothly blend from the source
+     * bone height, to the height of the goal in its normal retargeted position. */
+    UPROPERTY(EditAnywhere, BlueprintReadWrite, Category = "Floor Constraint Op Settings", meta = (ClampMin = "0.1", UIMin = "0.1", UIMax = "100.0"))
+    double HeightFalloffOffset = 8.0f;
+    
+    /** Range 0 to inf. Default is 20. The height in cm from the floor below which the goal is gradually blended towards the source bone height.
+     * NOTE: if the source bone is higher than this value, the height of the goal is left at its normal retargeted height. */
+    UPROPERTY(EditAnywhere, BlueprintReadWrite, Category = "Floor Constraint Op Settings", meta = (ClampMin = "0.1", UIMin = "0.1", UIMax = "100.0"))
+    double HeightFalloffDistance = 20.0f;
+
+	UE_API virtual const UClass* GetControllerType() const override;
+
+	UE_API virtual void CopySettingsAtRuntime(const FIKRetargetOpSettingsBase* InSettingsToCopyFrom) override;
+};
+
+struct FFloorConstraint
+{
+	FName IKGoalName = NAME_None;
+	int32 SourceEndBoneIndex = INDEX_NONE;
+	double HeightOffsetInRefPose = 0.0;
+	const FFloorConstraintChainSettings* Settings = nullptr;
+};
+
+USTRUCT(BlueprintType, meta = (DisplayName = "Floor Constraint"))
+struct FIKRetargetFloorConstraintOp : public FIKRetargetOpBase
+{
+	GENERATED_BODY()
+	
+	UE_API virtual bool Initialize(
+		const FIKRetargetProcessor& InProcessor,
+		const FRetargetSkeleton& InSourceSkeleton,
+		const FTargetSkeleton& InTargetSkeleton,
+		const FIKRetargetOpBase* InParentOp,
+		FIKRigLogger& InLog) override;
+	
+	UE_API virtual void Run(
+		FIKRetargetProcessor& InProcessor,
+		const double InDeltaTime,
+		const TArray<FTransform>& InSourceGlobalPose,
+		TArray<FTransform>& OutTargetGlobalPose) override;
+
+	UE_API virtual void OnAddedToStack(const UIKRetargeter* InRetargetAsset, const FIKRetargetOpBase* InParentOp) override;
+
+	UE_API virtual void OnAssignIKRig(const ERetargetSourceOrTarget SourceOrTarget, const UIKRigDefinition* InIKRig, const FIKRetargetOpBase* InParentOp) override;
+
+	UE_API virtual FIKRetargetOpSettingsBase* GetSettings() override;
+	
+	UE_API virtual void SetSettings(const FIKRetargetOpSettingsBase* InSettings) override;
+	
+	UE_API virtual const UScriptStruct* GetSettingsType() const override;
+
+	UE_API virtual const UScriptStruct* GetType() const override;
+
+	UE_API virtual const UScriptStruct* GetParentOpType() const override;
+
+	UE_API virtual void OnTargetChainRenamed(const FName InOldChainName, const FName InNewChainName) override;
+
+	UE_API virtual void OnParentReinitPropertyEdited(const FIKRetargetOpBase& InParentOp, const FPropertyChangedEvent* InPropertyChangedEvent) override;
+
+	UPROPERTY()
+	FIKRetargetFloorConstraintOpSettings Settings;
+
+#if WITH_EDITOR
+	UE_API virtual void ResetChainSettingsToDefault(const FName& InChainName) override;
+	UE_API virtual bool AreChainSettingsAtDefault(const FName& InChainName) override;
+#endif
+
+	TArray<FFloorConstraint> FloorConstraints;
+};
+
+/* The blueprint/python API for editing a Floor Goal Op */
+UCLASS(MinimalAPI, BlueprintType)
+class UIKRetargetFloorGoalsController : public UIKRetargetOpControllerBase
+{
+	GENERATED_BODY()
+	
+public:
+	/* Get the current op settings as a struct.
+	 * @return FIKRetargetIKChainsOpSettings struct with the current settings used by the op. */
+	UFUNCTION(BlueprintCallable, Category = Settings)
+	UE_API FIKRetargetIKChainsOpSettings GetSettings();
+
+	/* Set the op settings. Input is a custom struct type for this op.
+	 * @param InSettings a FIKRetargetIKChainsOpSettings struct containing all the settings to apply to this op */
+	UFUNCTION(BlueprintCallable, Category = Settings)
+	UE_API void SetSettings(FIKRetargetIKChainsOpSettings InSettings);
+};
+
+#undef LOCTEXT_NAMESPACE
+
+#undef UE_API
diff --git a/Engine/Plugins/Animation/IKRig/Source/IKRig/Public/Retargeter/RetargetOps/RunIKRigOp.h b/Engine/Plugins/Animation/IKRig/Source/IKRig/Public/Retargeter/RetargetOps/RunIKRigOp.h
--- a/Engine/Plugins/Animation/IKRig/Source/IKRig/Public/Retargeter/RetargetOps/RunIKRigOp.h
+++ b/Engine/Plugins/Animation/IKRig/Source/IKRig/Public/Retargeter/RetargetOps/RunIKRigOp.h
@@ -84,6 +84,12 @@ struct FIKRetargetRunIKRigOp : public FIKRetargetOpBase
 	
 	virtual bool CanHaveChildOps() const override { return true; };
 
+	UE_API virtual void RunBeforeChildren(
+		FIKRetargetProcessor& InProcessor,
+		const double InDeltaTime,
+		const TArray<FTransform>& InSourceGlobalPose,
+		TArray<FTransform>& OutTargetGlobalPose) override;
+
 	UE_API virtual const UIKRigDefinition* GetCustomTargetIKRig() const override;
 
 	UE_API virtual FRetargetChainMapping* GetChainMapping() override;
@@ -122,6 +128,8 @@ private:
 
 	/** The IK Rig processor for running IK on the target */
 	FIKRigProcessor IKRigProcessor;
+
+	TMap<FName, int32> GoalBoneIndicesMap;
 };
 
 /* The blueprint/python API for editing a Run IK Rig Op */
diff --git a/Engine/Plugins/Animation/IKRig/Source/IKRigEditor/Private/IKRigEditor.cpp b/Engine/Plugins/Animation/IKRig/Source/IKRigEditor/Private/IKRigEditor.cpp
--- a/Engine/Plugins/Animation/IKRig/Source/IKRigEditor/Private/IKRigEditor.cpp
+++ b/Engine/Plugins/Animation/IKRig/Source/IKRigEditor/Private/IKRigEditor.cpp
@@ -107,6 +107,9 @@ void FIKRigEditor::StartupModule()
 	PropertyEditorModule.RegisterCustomPropertyTypeLayout(
 			"IKRetargetStretchChainOpSettings",
 			FOnGetPropertyTypeCustomizationInstance::CreateStatic(&FStretchChainOpCustomization::MakeInstance));
+	PropertyEditorModule.RegisterCustomPropertyTypeLayout(
+			"IKRetargetFloorConstraintOpSettings",
+			FOnGetPropertyTypeCustomizationInstance::CreateStatic(&FFloorConstraintOpCustomization::MakeInstance));
 
 	// register a thumbnail renderer for the assets
 	UThumbnailManager::Get().RegisterCustomRenderer(UIKRigDefinition::StaticClass(), UIKRigThumbnailRenderer::StaticClass());
diff --git a/Engine/Plugins/Animation/IKRig/Source/IKRigEditor/Private/RetargetEditor/IKRetargetOpDetails.cpp b/Engine/Plugins/Animation/IKRig/Source/IKRigEditor/Private/RetargetEditor/IKRetargetOpDetails.cpp
--- a/Engine/Plugins/Animation/IKRig/Source/IKRigEditor/Private/RetargetEditor/IKRetargetOpDetails.cpp
+++ b/Engine/Plugins/Animation/IKRig/Source/IKRigEditor/Private/RetargetEditor/IKRetargetOpDetails.cpp
@@ -743,4 +743,73 @@ void FStretchChainOpCustomization::CustomizeChildren(
 	AddBaseOpProperties(StructPropertyHandle, ChildBuilder);
 }
 
+void FFloorConstraintOpCustomization::CustomizeChildren(
+	TSharedRef<IPropertyHandle> StructPropertyHandle,
+	IDetailChildrenBuilder& ChildBuilder,
+	IPropertyTypeCustomizationUtils& StructCustomizationUtils)
+{
+	if (!LoadAndValidateStructToCustomize(StructPropertyHandle, StructCustomizationUtils))
+	{
+		return;
+	}
+	
+	// show only IK chains
+	FChainMapFilterOptions ShowOnlyIKFilter;
+	ShowOnlyIKFilter.bNeverShowChainsWithoutIK = true;
+
+	// configure chain map list for the IK op
+	FChainMapListConfig ChainListConfig;
+	ChainListConfig.OpWithChainSettings = OpName;
+	ChainListConfig.OpWithChainMapping = AssetController->GetParentOpByName(OpName); // use chain mapping from parent
+	ChainListConfig.Controller = AssetController;
+	ChainListConfig.bEnableGoalColumn = true;
+	ChainListConfig.bEnableChainMapping = false;
+	ChainListConfig.Filter = ShowOnlyIKFilter;
+	ChainListConfig.ChainSettingsGetterFunc = [this](const FName InTargetChainName) -> UObject*
+	{
+		// this lambda fetches the chain settings struct for the given target chain in the given op (by name)
+		auto MemoryProvider = [this, InTargetChainName]() -> uint8*
+		{
+			if (FIKRetargetOpBase* Op = AssetController->GetRetargetOpByName(OpName))
+			{
+				FIKRetargetFloorConstraintOpSettings* OpSettings = reinterpret_cast<FIKRetargetFloorConstraintOpSettings*>(Op->GetSettings());
+				for (FFloorConstraintChainSettings& ChainSettings : OpSettings->ChainsToAffect)
+				{
+					if (ChainSettings.TargetChainName == InTargetChainName)
+					{
+						return reinterpret_cast<uint8*>(&ChainSettings);
+					}
+				}
+			}
+			return nullptr;
+		};
+			
+		FIKRigStructToView StructToView;
+		StructToView.Owner = AssetController->GetAsset();
+		StructToView.Type = FRetargetSpeedPlantingSettings::StaticStruct();
+		StructToView.MemoryProvider = MemoryProvider;
+		StructToView.UniqueName = OpName;
+		
+		UFloorConstraintSettingsWrapper* Wrapper = NewObject<UFloorConstraintSettingsWrapper>(GetTransientPackage(), NAME_None, RF_Standalone);
+		const FName SettingsPropertyName = GET_MEMBER_NAME_CHECKED(UFloorConstraintSettingsWrapper, Settings);
+		Wrapper->InitializeWithRetargeter(StructToView, SettingsPropertyName, AssetController);
+		Wrapper->SetPropertyHidden(GET_MEMBER_NAME_CHECKED(FFloorConstraintChainSettings, TargetChainName), true);
+		return Wrapper;
+	};
+
+	// add the chain mapping list
+	FText ChainMapRowName = FText::FromString("Chain Map");
+	ChildBuilder.AddCustomRow(ChainMapRowName)
+	.WholeRowContent()
+	[
+		SNew(SIKRetargetChainMapList).InChainMapListConfig(ChainListConfig)
+	];
+
+	// add op properties
+	AddChildPropertiesInCategory(StructPropertyHandle, ChildBuilder, FName("Floor Constraint Op Settings"));
+	
+	// add the properties from the base op
+	AddBaseOpProperties(StructPropertyHandle, ChildBuilder);
+}
+
 #undef LOCTEXT_NAMESPACE
diff --git a/Engine/Plugins/Animation/IKRig/Source/IKRigEditor/Public/RetargetEditor/IKRetargetOpDetails.h b/Engine/Plugins/Animation/IKRig/Source/IKRigEditor/Public/RetargetEditor/IKRetargetOpDetails.h
--- a/Engine/Plugins/Animation/IKRig/Source/IKRigEditor/Public/RetargetEditor/IKRetargetOpDetails.h
+++ b/Engine/Plugins/Animation/IKRig/Source/IKRigEditor/Public/RetargetEditor/IKRetargetOpDetails.h
@@ -4,6 +4,7 @@
 #include "IPropertyTypeCustomization.h"
 #include "Retargeter/RetargetOps/AlignPoleVectorOp.h"
 #include "Retargeter/RetargetOps/FKChainsOp.h"
+#include "Retargeter/RetargetOps/FloorConstraintOp.h"
 #include "Retargeter/RetargetOps/IKChainsOp.h"
 #include "Retargeter/RetargetOps/RetargetPoseOp.h"
 #include "Retargeter/RetargetOps/SpeedPlantingOp.h"
@@ -291,3 +292,32 @@ public:
 	UPROPERTY(EditAnywhere, Category = "Stretch Chain Settings")
 	FRetargetStretchChainSettings Settings;
 };
+
+//
+// Floor Constraint Op Customization
+//
+
+class FFloorConstraintOpCustomization : public FIKRetargetOpBaseSettingsCustomization
+{
+public:
+	static TSharedRef<IPropertyTypeCustomization> MakeInstance()
+	{
+		return MakeShared<FFloorConstraintOpCustomization>();
+	}
+
+	// IPropertyTypeCustomization
+	virtual void CustomizeChildren(
+		TSharedRef<IPropertyHandle> StructPropertyHandle,
+		IDetailChildrenBuilder& ChildBuilder,
+		IPropertyTypeCustomizationUtils& StructCustomizationUtils) override;
+	// End IPropertyTypeCustomization
+};
+
+UCLASS(BlueprintType)
+class UFloorConstraintSettingsWrapper : public UIKRigStructWrapperBase
+{
+	GENERATED_BODY()
+public:
+	UPROPERTY(EditAnywhere, Category = "Floor Constraint Chain Settings")
+	FFloorConstraintChainSettings Settings;
+};
