[IK Rig] Adding squash direction to stretch limb solver.

diff --git a/Engine/Plugins/Animation/IKRig/Source/IKRig/Private/Rig/Solvers/IKRigStretchLimb.cpp b/Engine/Plugins/Animation/IKRig/Source/IKRig/Private/Rig/Solvers/IKRigStretchLimb.cpp
--- a/Engine/Plugins/Animation/IKRig/Source/IKRig/Private/Rig/Solvers/IKRigStretchLimb.cpp
+++ b/Engine/Plugins/Animation/IKRig/Source/IKRig/Private/Rig/Solvers/IKRigStretchLimb.cpp
@@ -75,6 +75,13 @@ void FIKRigStretchLimbSolver::Initialize(const FIKRigSkeleton& InIKRigSkeleton)
 		PoleVectorParams.Add(T);
 	}
 
+	// cache bone settings bone indices
+	for (FIKRigStretchLimbBoneSettings& BoneSettings : AllBoneSettings)
+	{
+		const int32 BoneIndex = InIKRigSkeleton.GetBoneIndexFromName(BoneSettings.Bone);
+		BoneSettings.CachedChainIndex = BoneIndices.Find(BoneIndex);
+	}
+
 	bIsInitialized = true;
 }
 
@@ -210,6 +217,11 @@ void FIKRigStretchLimbSolver::SquashChainPerpendicular(FIKRigSkeleton& InIKRigSk
 	{
 		return;
 	}
+
+	if (Settings.SquashMode == EStretchLimbSquashMode::None || FMath::IsNearlyZero(Settings.SquashStrength))
+	{
+		return;
+	}
 	
 	// calc pole vector
 	const FVector& ChainStart = InIKRigSkeleton.CurrentPoseGlobal[BoneIndices[0]].GetLocation();
@@ -222,7 +234,7 @@ void FIKRigStretchLimbSolver::SquashChainPerpendicular(FIKRigSkeleton& InIKRigSk
 	}
 	const FVector PoleNorm = PoleVector.GetSafeNormal();
 
-	// calc bending directions for each bone (skipping start/end)
+	// calc squash direction for each bone (skipping start/end)
 	TArray<FVector> SquashDirections;
 	SquashDirections.Init(FVector::UpVector, BoneIndices.Num());
 	for (int32 ChainIndex = 1; ChainIndex < BoneIndices.Num() - 1 ; ++ChainIndex)
@@ -240,19 +252,43 @@ void FIKRigStretchLimbSolver::SquashChainPerpendicular(FIKRigSkeleton& InIKRigSk
 		SquashDirections[ChainIndex] = BendDirection;
 	}
 
+	// override squash directions from bone settings
+	for (FIKRigStretchLimbBoneSettings BoneSetting : AllBoneSettings)
+	{
+		if (BoneSetting.CachedChainIndex == INDEX_NONE)
+		{
+			continue;
+		}
+
+		const FTransform& BoneTransform = InIKRigSkeleton.CurrentPoseGlobal[BoneIndices[BoneSetting.CachedChainIndex]];
+		const FVector SquashDirection = BoneTransform.GetRotation().RotateVector(BoneSetting.SquashDirection);
+		SquashDirections[BoneSetting.CachedChainIndex] = SquashDirection;
+	}
+
 	// push bones along squash directions
 	const double SquashRatio = 1.0 - NewPoleVectorLenSq / InitialPoleVectorLenSq;
 	for (int32 ChainIndex = 1; ChainIndex < BoneIndices.Num() - 1 ; ++ChainIndex)
 	{
-		auto SmoothTriangleFalloff = [](double InT) -> double
+		auto GetSquashFalloff = [](double InT, EStretchLimbSquashMode InSquashMode) -> double
 		{
-			InT = FMath::Clamp(InT, 0.0f, 1.0f);
-			return 2.0f * FMath::Sqrt(InT * (1.0f - InT)); // 0 at 0, 0 at 1, 1 at 0.5
+			if (InSquashMode == EStretchLimbSquashMode::Uniform)
+			{
+				return 1;
+			}
+				
+			if (InSquashMode == EStretchLimbSquashMode::Bulge)
+			{
+				InT = FMath::Clamp(InT, 0.0f, 1.0f);
+				return 2.0f * FMath::Sqrt(InT * (1.0f - InT)); // 0 at 0, 0 at 1, 1 at 0.5
+			}
+
+			checkNoEntry();
+			return 1.0;
 		};
 
 		const int32 BoneIndex = BoneIndices[ChainIndex];
 		const FVector& BendDirection = SquashDirections[ChainIndex];
-		const double Falloff = SmoothTriangleFalloff(PoleVectorParams[ChainIndex]);
+		const double Falloff = GetSquashFalloff(PoleVectorParams[ChainIndex], Settings.SquashMode);
 		InIKRigSkeleton.CurrentPoseGlobal[BoneIndex].AddToTranslation(Settings.SquashStrength * BendDirection * Falloff * SquashRatio);
 	}
 }
@@ -269,10 +305,10 @@ void FIKRigStretchLimbSolver::StretchChainFinal(FIKRigSkeleton& InIKRigSkeleton,
 	const FVector& ChainEnd = InIKRigSkeleton.CurrentPoseGlobal[BoneIndices.Last()].GetLocation();
 	const FVector CurrentPoleVector = ChainEnd - ChainStart;
 	FVector NewPoleVector = GoalLocation - ChainStart;
-	
-	// clamp to max stretch distance
-	if (Settings.MaximumStretchDistance > 0.0)
+
+	if (Settings.MaximumStretchDistance >= 0.f)
 	{
+		// hard clamp
 		NewPoleVector = NewPoleVector.GetClampedToMaxSize(TotalRestLength + Settings.MaximumStretchDistance);
 	}
 	
@@ -522,6 +558,71 @@ void FIKRigStretchLimbSolver::SetStartBone(const FName& InBoneName)
 	Settings.StartBone = InBoneName;
 }
 
+// BONE SETTINGS
+
+bool FIKRigStretchLimbSolver::UsesCustomBoneSettings() const
+{
+	return true;
+}
+
+void FIKRigStretchLimbSolver::AddSettingsToBone(const FName& InBoneName)
+{
+	if (AllBoneSettings.ContainsByPredicate([&](const FIKRigStretchLimbBoneSettings& Element){return Element.Bone == InBoneName;}))
+	{
+		// bone already has settings
+		return;
+	}
+	
+	AllBoneSettings.Emplace(InBoneName);
+}
+
+void FIKRigStretchLimbSolver::RemoveSettingsOnBone(const FName& InBoneName)
+{
+	AllBoneSettings.RemoveAll([&](const FIKRigStretchLimbBoneSettings& Element)
+	{
+		return Element.Bone == InBoneName;
+	});
+}
+
+FIKRigBoneSettingsBase* FIKRigStretchLimbSolver::GetBoneSettings(const FName& InBoneName)
+{
+	for (FIKRigStretchLimbBoneSettings& BoneSetting : AllBoneSettings)
+	{
+		if (BoneSetting.Bone == InBoneName)
+		{
+			return &BoneSetting;
+		}
+	}
+	
+	return nullptr;
+}
+
+const UScriptStruct* FIKRigStretchLimbSolver::GetBoneSettingsType() const
+{
+	return FIKRigStretchLimbBoneSettings::StaticStruct();
+}
+
+bool FIKRigStretchLimbSolver::HasSettingsOnBone(const FName& InBoneName) const
+{
+	for (const FIKRigStretchLimbBoneSettings& BoneSetting : AllBoneSettings)
+	{
+		if (BoneSetting.Bone == InBoneName)
+		{
+			return true;
+		}
+	}
+
+	return false;
+}
+
+void FIKRigStretchLimbSolver::GetBonesWithSettings(TSet<FName>& OutBonesWithSettings) const
+{
+	for (const FIKRigStretchLimbBoneSettings& BoneSetting : AllBoneSettings)
+	{
+		OutBonesWithSettings.Add(BoneSetting.Bone);
+	}
+}
+
 #if WITH_EDITOR
 
 UIKRigSolverControllerBase* FIKRigStretchLimbSolver::GetSolverController(UObject* Outer)
diff --git a/Engine/Plugins/Animation/IKRig/Source/IKRig/Public/Rig/Solvers/IKRigStretchLimb.h b/Engine/Plugins/Animation/IKRig/Source/IKRig/Public/Rig/Solvers/IKRigStretchLimb.h
--- a/Engine/Plugins/Animation/IKRig/Source/IKRig/Public/Rig/Solvers/IKRigStretchLimb.h
+++ b/Engine/Plugins/Animation/IKRig/Source/IKRig/Public/Rig/Solvers/IKRigStretchLimb.h
@@ -11,6 +11,33 @@ struct FIKRigGoal;
 
 #define UE_API IKRIG_API
 
+USTRUCT(BlueprintType)
+struct FIKRigStretchLimbBoneSettings : public FIKRigBoneSettingsBase
+{
+	GENERATED_BODY()
+
+	FIKRigStretchLimbBoneSettings() = default;
+	FIKRigStretchLimbBoneSettings(FName InBoneName) { Bone = InBoneName; }
+	
+	/** The direction to push this bone when the limb is squashed. This is relative to the local bone axes.
+	 * NOTE: This direction is used by the "Squash Strength" feature of the stretch limb solver.
+	 * If no custom Squash direction is specified, the solver will push the bone in the direction away from its projection onto the pole vector.
+	 * If the bone is lying directly on the pole vector, the default squash direction is undefined and the bone may mot squash in a deterministic manner.
+	 */
+	UPROPERTY(BlueprintReadWrite, EditAnywhere, Category = "Preferred Squash Direction")
+	FVector SquashDirection = FVector::ZAxisVector;
+
+	int32 CachedChainIndex;
+};
+
+UENUM(BlueprintType)
+enum class EStretchLimbSquashMode : uint8
+{
+	None,
+	Uniform,
+	Bulge
+};
+
 UENUM(BlueprintType)
 enum class EStretchLimbRotationMode : uint8
 {
@@ -57,12 +84,17 @@ struct FIKRigStretchLimbSettings : public FIKRigSolverSettingsBase
 	double RotateEndBoneWithGoal = 1.0f;
 
 	/** Number of FABRIK iterations to correct bones lengths. Adjust until the end bone converges on the goal. Default is 8. */
-	UPROPERTY(BlueprintReadWrite, EditAnywhere, Category = "Constraint Solving", meta = (UIMin = "0", UIMax = "50.0"))
+	UPROPERTY(BlueprintReadWrite, EditAnywhere, Category = "Constraint Solving", meta = (UIMin = "0", UIMax = "20.0"))
 	int32 Iterations = 8;
 
+	/** Control the falloff shape of the squash effect applied to the bones when the goal compresses the limb. Default is Uniform.
+	 * NOTE: this has no effect if the Squash Strength is 0. */
+	UPROPERTY(BlueprintReadWrite, EditAnywhere, Category = "Constraint Solving")
+	EStretchLimbSquashMode SquashMode = EStretchLimbSquashMode::Uniform;
+	
 	/** The distance to squash the bones perpendicular to the pole vector when the limb is fully compressed.
 	 * NOTE: tune this to help dense chains achieve a reasonable pose. Must be used with constraint iterations to fix bones lengths. */
-	UPROPERTY(BlueprintReadWrite, EditAnywhere, Category = "Constraint Solving", meta = (UIMin = "0", UIMax = "5000.0"))
+	UPROPERTY(BlueprintReadWrite, EditAnywhere, Category = "Constraint Solving", meta = (UIMin = "0", UIMax = "200.0"))
 	double SquashStrength = 50;
 };
 
@@ -89,6 +121,16 @@ struct FIKRigStretchLimbSolver : public FIKRigSolverBase
 
 	// goal settings
 	virtual bool UsesCustomGoalSettings() const override { return false; };
+
+	// bone settings
+	UE_API virtual bool UsesCustomBoneSettings() const override;
+	UE_API virtual void AddSettingsToBone(const FName& InBoneName) override;
+	UE_API virtual void RemoveSettingsOnBone(const FName& InBoneName) override;
+	UE_API virtual FIKRigBoneSettingsBase* GetBoneSettings(const FName& InBoneName) override;
+	UE_API virtual const UScriptStruct* GetBoneSettingsType() const override;
+	UE_API virtual bool HasSettingsOnBone(const FName& InBoneName) const override;
+	UE_API virtual void GetBonesWithSettings(TSet<FName>& OutBonesWithSettings) const override;
+	// END bone settings
 	
 	// start bone can be set on this solver
 	virtual bool UsesStartBone() const override { return true; };
@@ -107,6 +149,9 @@ struct FIKRigStretchLimbSolver : public FIKRigSolverBase
 	UPROPERTY()
 	FIKRigStretchLimbSettings Settings;
 
+	UPROPERTY()
+	TArray<FIKRigStretchLimbBoneSettings> AllBoneSettings;
+
 private:
 
 	void AimChainAtGoal(FIKRigSkeleton& InIKRigSkeleton, const FVector& GoalLocation) const;
