[IK Retarget] Removing temp floor constraint from IK Goals Op.

diff --git a/Engine/Plugins/Animation/IKRig/Source/IKRig/Private/Retargeter/RetargetOps/IKChainsOp.cpp b/Engine/Plugins/Animation/IKRig/Source/IKRig/Private/Retargeter/RetargetOps/IKChainsOp.cpp
--- a/Engine/Plugins/Animation/IKRig/Source/IKRig/Private/Retargeter/RetargetOps/IKChainsOp.cpp
+++ b/Engine/Plugins/Animation/IKRig/Source/IKRig/Private/Retargeter/RetargetOps/IKChainsOp.cpp
@@ -139,20 +139,6 @@ FVector FIKChainRetargeter::CalculateGoalPosition(
 	const FVector CurrentTargetStartLocation = InTargetGlobalPose[TargetStartBoneIndex].GetTranslation() - InvRootModification;
 	FVector GoalPosition = CurrentTargetStartLocation + (CurrentSourceEndDirectionNormalized * InitialTargetChainLength);
 
-	// apply floor constraint
-	if (!FMath::IsNearlyZero(Settings->FloorConstraintWeight))
-	{
-		const float SourceHeight = FMath::Abs(CurrentSourceEndGlobal.GetLocation().Z);
-		if (SourceHeight <= Settings->FloorHeightFalloffEnd)
-		{
-			const float Start = Settings->FloorHeightFalloffStart;
-			float End = Settings->FloorHeightFalloffEnd;
-			End = End <= Start ? Start + 1.0f : End;
-			double Falloff = SourceHeight < Start ? 1.0f : 1.0f - (SourceHeight - Start) / (End - Start);
-			GoalPosition.Z = FMath::Lerp(GoalPosition.Z, SourceHeight, Falloff * Settings->FloorConstraintWeight);
-		}
-	}
-
 	// blend to source location
 	const double BlendToSourceTranslation = Settings->BlendToSource * Settings->BlendToSourceTranslation;
 	if (BlendToSourceTranslation > UE_KINDA_SMALL_NUMBER)
@@ -189,9 +175,6 @@ FVector FIKChainRetargeter::CalculateGoalPosition(
 bool FRetargetIKChainSettings::operator==(const FRetargetIKChainSettings& Other) const
 {
 	return EnableIK == Other.EnableIK
-		&& FMath::IsNearlyEqualByULP(FloorConstraintWeight, Other.FloorConstraintWeight)
-		&& FMath::IsNearlyEqualByULP(FloorHeightFalloffStart, Other.FloorHeightFalloffStart)
-		&& FMath::IsNearlyEqualByULP(FloorHeightFalloffEnd, Other.FloorHeightFalloffEnd)
 		&& FMath::IsNearlyEqualByULP(BlendToSource, Other.BlendToSource)
 		&& BlendToSourceWeights.Equals(Other.BlendToSourceWeights)
 		&& ApplyPelvisOffsetToSourceGoals == Other.ApplyPelvisOffsetToSourceGoals
diff --git a/Engine/Plugins/Animation/IKRig/Source/IKRig/Public/Retargeter/RetargetOps/IKChainsOp.h b/Engine/Plugins/Animation/IKRig/Source/IKRig/Public/Retargeter/RetargetOps/IKChainsOp.h
--- a/Engine/Plugins/Animation/IKRig/Source/IKRig/Public/Retargeter/RetargetOps/IKChainsOp.h
+++ b/Engine/Plugins/Animation/IKRig/Source/IKRig/Public/Retargeter/RetargetOps/IKChainsOp.h
@@ -79,26 +79,6 @@ struct FRetargetIKChainSettings
 	UPROPERTY(EditAnywhere, BlueprintReadWrite, Category = "Setup", meta=(ReinitializeOnEdit))
 	bool EnableIK = true;
 
-	/** Range 0 to 1. Default is 0. Set this to 1 to turn the floor constraint ON.
-	 * When ON, the floor constraint will adjust the vertical position of the IK Goal according to the following rules.
-	 * 1. When the source goal bone is LOWER than FloorHeightFalloffStart, the height of the goal is smoothly blended to the height of the source bone.
-	 * 2. When the source goal bone is HIGHER than FloorHeightFalloffStart, the height of the goal is left at its normal retargeted location.
-	 * NOTE: the floor is assumed to be the XY plane where Z = 0.
-	 */
-	UPROPERTY(EditAnywhere, BlueprintReadWrite, Category = FloorHeightConstraint, meta = (ClampMin = "0.0", ClampMax = "1.0", UIMin = "0.0", UIMax = "1.0"))
-	double FloorConstraintWeight = 0.0f;
-
-	/** Range 0 to inf. Default is 8. The height in cm from the floor below which the goal is snapped directly to the source bone height.
-	 * NOTE: if the source bone height is greater than this value, but lower than FloorHeightFalloffEnd, then the height will smoothly blend from the source
-	 * bone height, to the height of the goal in its normal retargeted position. */
-	UPROPERTY(EditAnywhere, BlueprintReadWrite, Category = FloorHeightConstraint, meta = (ClampMin = "0.1", UIMin = "0.1", UIMax = "100.0"))
-	double FloorHeightFalloffStart = 8.0f;
-	
-	/** Range 0 to inf. Default is 20. The height in cm from the floor below which the goal is gradually blended towards the source bone height.
-	 * NOTE: if the source bone is higher than this value, the height of the goal is left at its normal retargeted height. */
-	UPROPERTY(EditAnywhere, BlueprintReadWrite, Category = FloorHeightConstraint, meta = (ClampMin = "0.1", UIMin = "0.1", UIMax = "100.0"))
-	double FloorHeightFalloffEnd = 20.0f;
-
 	/** Range 0 to 1. Default 0. Blends IK goal transform from retargeted transform (0) to source bone transform (1).
 	*  At 0 the goal is placed at the retargeted location and rotation.
 	*  At 1 the goal is placed at the location and rotation of the source chain's end bone. */
