[IK Retarget] Added global/local translation/rotation offsets to pelvis motion op.

diff --git a/Engine/Plugins/Animation/IKRig/Source/IKRig/Private/Retargeter/IKRetargetProfile.cpp b/Engine/Plugins/Animation/IKRig/Source/IKRig/Private/Retargeter/IKRetargetProfile.cpp
--- a/Engine/Plugins/Animation/IKRig/Source/IKRig/Private/Retargeter/IKRetargetProfile.cpp
+++ b/Engine/Plugins/Animation/IKRig/Source/IKRig/Private/Retargeter/IKRetargetProfile.cpp
@@ -160,8 +160,8 @@ void FRetargetProfile::MergeWithOtherProfile(const FRetargetProfile& OtherProfil
 				SettingsInProfile->BlendToSourceTranslationWeights = OtherProfile.RootSettings.BlendToSourceWeights;
 				SettingsInProfile->ScaleHorizontal = OtherProfile.RootSettings.ScaleHorizontal;
 				SettingsInProfile->ScaleVertical = OtherProfile.RootSettings.ScaleVertical;
-				SettingsInProfile->TranslationOffset = OtherProfile.RootSettings.TranslationOffset;
-				SettingsInProfile->RotationOffset = OtherProfile.RootSettings.RotationOffset;
+				SettingsInProfile->TranslationOffsetGlobal = OtherProfile.RootSettings.TranslationOffset;
+				SettingsInProfile->RotationOffsetLocal = OtherProfile.RootSettings.RotationOffset;
 				SettingsInProfile->AffectIKHorizontal = OtherProfile.RootSettings.AffectIKHorizontal;
 				SettingsInProfile->AffectIKVertical = OtherProfile.RootSettings.AffectIKVertical;
 			}
diff --git a/Engine/Plugins/Animation/IKRig/Source/IKRig/Private/Retargeter/IKRetargeter.cpp b/Engine/Plugins/Animation/IKRig/Source/IKRig/Private/Retargeter/IKRetargeter.cpp
--- a/Engine/Plugins/Animation/IKRig/Source/IKRig/Private/Retargeter/IKRetargeter.cpp
+++ b/Engine/Plugins/Animation/IKRig/Source/IKRig/Private/Retargeter/IKRetargeter.cpp
@@ -294,8 +294,8 @@ void UIKRetargeter::PostLoadConvertEverythingToOps()
 		Settings->BlendToSourceTranslationWeights = RootSettings_DEPRECATED->Settings.BlendToSourceWeights;
 		Settings->ScaleHorizontal = RootSettings_DEPRECATED->Settings.ScaleHorizontal;
 		Settings->ScaleVertical = RootSettings_DEPRECATED->Settings.ScaleVertical;
-		Settings->TranslationOffset = RootSettings_DEPRECATED->Settings.TranslationOffset;
-		Settings->RotationOffset = RootSettings_DEPRECATED->Settings.RotationOffset;
+		Settings->TranslationOffsetGlobal = RootSettings_DEPRECATED->Settings.TranslationOffset;
+		Settings->RotationOffsetLocal = RootSettings_DEPRECATED->Settings.RotationOffset;
 		Settings->AffectIKHorizontal = RootSettings_DEPRECATED->Settings.AffectIKHorizontal;
 		Settings->AffectIKVertical = RootSettings_DEPRECATED->Settings.AffectIKVertical;
 	}
diff --git a/Engine/Plugins/Animation/IKRig/Source/IKRig/Private/Retargeter/RetargetOps/PelvisMotionOp.cpp b/Engine/Plugins/Animation/IKRig/Source/IKRig/Private/Retargeter/RetargetOps/PelvisMotionOp.cpp
--- a/Engine/Plugins/Animation/IKRig/Source/IKRig/Private/Retargeter/RetargetOps/PelvisMotionOp.cpp
+++ b/Engine/Plugins/Animation/IKRig/Source/IKRig/Private/Retargeter/RetargetOps/PelvisMotionOp.cpp
@@ -258,6 +258,12 @@ void FIKRetargetPelvisMotionOp::PostLoad(const FIKRigObjectVersion::Type InVersi
 	{
 		Settings.bDebugDraw = Settings.bEnableDebugDraw_DEPRECATED;
 	}
+
+	if (InVersion < FIKRigObjectVersion::AddedLocalGlobalOffsetsToPelvisOp)
+	{
+		Settings.TranslationOffsetGlobal = Settings.TranslationOffset_DEPRECATED;
+		Settings.RotationOffsetLocal = Settings.RotationOffset_DEPRECATED;
+	}
 }
 
 int32 FIKRetargetPelvisMotionOp::GetPelvisBoneIndex(ERetargetSourceOrTarget SourceOrTarget) const
@@ -281,6 +287,28 @@ void FIKRetargetPelvisMotionOp::EncodePose(const TArray<FTransform>& SourceGloba
 
 void FIKRetargetPelvisMotionOp::DecodePose(FTransform& OutPelvisGlobalPose)
 {
+	// retarget rotation
+	// NOTE: rotation must be done before translation to get proper orientation for local position offsets
+	FQuat Rotation;
+	{
+		// calc offset between initial source/target root rotations
+		const FQuat RotationDelta = Source.CurrentRotation * Source.InitialRotation.Inverse();
+		// add retarget pose delta to the current source rotation
+		const FQuat RetargetedRotation = RotationDelta * Target.InitialRotation;
+
+		// apply static LOCAL rotation offset
+		Rotation = RetargetedRotation * Settings.RotationOffsetLocal.Quaternion();
+		// apply static GLOBAL rotation offset
+		Rotation = Settings.RotationOffsetGlobal.Quaternion() * Rotation;
+
+		// blend with alpha
+		Rotation = FQuat::FastLerp(Target.InitialRotation, Rotation, Settings.RotationAlpha);
+		Rotation.Normalize();
+
+		// record the delta created by all the modifications made to the root rotation
+		Target.PelvisRotationDelta = RetargetedRotation * Target.InitialRotation.Inverse();
+	}
+	
 	// retarget position
 	FVector Position;
 	{
@@ -320,8 +348,11 @@ void FIKRetargetPelvisMotionOp::DecodePose(FTransform& OutPelvisGlobalPose)
 		const FVector HorizontalOffset = (ScaledRetargetedPosition - Target.InitialPosition) * FVector(Settings.ScaleHorizontal, Settings.ScaleHorizontal, 1.0f);
 		Position = Target.InitialPosition + HorizontalOffset;
 		
-		// apply a static offset
-		Position += Settings.TranslationOffset;
+		// apply a static GLOBAL offset
+		Position += Settings.TranslationOffsetGlobal;
+
+		// apply a static LOCAL offset
+		Position += Rotation.RotateVector(Settings.TranslationOffsetLocal);
 
 		// blend with alpha
 		Position = FMath::Lerp(Target.InitialPosition, Position, Settings.TranslationAlpha);
@@ -330,25 +361,6 @@ void FIKRetargetPelvisMotionOp::DecodePose(FTransform& OutPelvisGlobalPose)
 		Target.PelvisTranslationDelta = Position - RetargetedPosition;
 	}
 
-	// retarget rotation
-	FQuat Rotation;
-	{
-		// calc offset between initial source/target root rotations
-		const FQuat RotationDelta = Source.CurrentRotation * Source.InitialRotation.Inverse();
-		// add retarget pose delta to the current source rotation
-		const FQuat RetargetedRotation = RotationDelta * Target.InitialRotation;
-
-		// add static rotation offset
-		Rotation = RetargetedRotation * Settings.RotationOffset.Quaternion();
-
-		// blend with alpha
-		Rotation = FQuat::FastLerp(Target.InitialRotation, Rotation, Settings.RotationAlpha);
-		Rotation.Normalize();
-
-		// record the delta created by all the modifications made to the root rotation
-		Target.PelvisRotationDelta = RetargetedRotation * Target.InitialRotation.Inverse();
-	}
-
 	// apply to target
 	OutPelvisGlobalPose.SetTranslation(Position);
 	OutPelvisGlobalPose.SetRotation(Rotation);
diff --git a/Engine/Plugins/Animation/IKRig/Source/IKRig/Public/IKRigObjectVersion.h b/Engine/Plugins/Animation/IKRig/Source/IKRig/Public/IKRigObjectVersion.h
--- a/Engine/Plugins/Animation/IKRig/Source/IKRig/Public/IKRigObjectVersion.h
+++ b/Engine/Plugins/Animation/IKRig/Source/IKRig/Public/IKRigObjectVersion.h
@@ -47,6 +47,9 @@ struct FIKRigObjectVersion
 		// added "bEnableDebugDraw" to base op settings
 		AddedDebugDrawTogglePerOp,
 
+		// added local/global offsets to pelvis motion op
+		AddedLocalGlobalOffsetsToPelvisOp,
+
 		// -----<new versions can be added above this line>-------------------------------------------------
 		VersionPlusOne,
 		LatestVersion = VersionPlusOne - 1
diff --git a/Engine/Plugins/Animation/IKRig/Source/IKRig/Public/Retargeter/RetargetOps/PelvisMotionOp.h b/Engine/Plugins/Animation/IKRig/Source/IKRig/Public/Retargeter/RetargetOps/PelvisMotionOp.h
--- a/Engine/Plugins/Animation/IKRig/Source/IKRig/Public/Retargeter/RetargetOps/PelvisMotionOp.h
+++ b/Engine/Plugins/Animation/IKRig/Source/IKRig/Public/Retargeter/RetargetOps/PelvisMotionOp.h
@@ -49,9 +49,13 @@ struct FIKRetargetPelvisMotionOpSettings : public FIKRetargetOpSettingsBase
 	UPROPERTY(EditAnywhere, BlueprintReadWrite, Category = "Pelvis Rotation", meta = (ClampMin = "0.0", ClampMax = "1.0", UIMin = "0.0", UIMax = "1.0"))
 	double RotationAlpha = 1.0f;
 
-	/** Applies a static local-space rotation offset to the retarget pelvis.*/
+	/** Applies a static local-space rotation offset to the pelvis.*/
 	UPROPERTY(EditAnywhere, BlueprintReadWrite, Category = "Pelvis Rotation", meta = (ClampMin = "-180.0", ClampMax = "180.0", UIMin = "-180.0", UIMax = "180.0"))
-	FRotator RotationOffset = FRotator::ZeroRotator;
+	FRotator RotationOffsetLocal = FRotator::ZeroRotator;
+
+	/** Applies a static global-space rotation offset to the pelvis.*/
+	UPROPERTY(EditAnywhere, BlueprintReadWrite, Category = "Pelvis Rotation", meta = (ClampMin = "-180.0", ClampMax = "180.0", UIMin = "-180.0", UIMax = "180.0"))
+	FRotator RotationOffsetGlobal = FRotator::ZeroRotator;
 	
 	/** Range 0 to 1. Default 1. Blends the amount of retargeted pelvis translation to apply.
 	*  At 0 the pelvis is left at the position from the retarget pose.
@@ -59,9 +63,13 @@ struct FIKRetargetPelvisMotionOpSettings : public FIKRetargetOpSettingsBase
 	UPROPERTY(EditAnywhere, BlueprintReadWrite, Category = "Pelvis Translation", meta = (ClampMin = "0.0", ClampMax = "1.0", UIMin = "0.0", UIMax = "1.0"))
 	double TranslationAlpha = 1.0f;
 
-	/** Applies a static component-space translation offset to the pelvis.*/
+	/** Applies a static local-space translation offset to the pelvis.*/
+	UPROPERTY(EditAnywhere, BlueprintReadWrite, Category = "Pelvis Translation")
+	FVector TranslationOffsetLocal = FVector::ZeroVector;
+	
+	/** Applies a static global-space translation offset to the pelvis.*/
 	UPROPERTY(EditAnywhere, BlueprintReadWrite, Category = "Pelvis Translation")
-	FVector TranslationOffset = FVector::ZeroVector;
+	FVector TranslationOffsetGlobal = FVector::ZeroVector;
 	
 	/** Range 0 to 1. Default 0. Blends the retarget pelvis' translation to the exact source location.
 	*  At 0 the pelvis is placed at the retargeted location.
@@ -110,6 +118,12 @@ struct FIKRetargetPelvisMotionOpSettings : public FIKRetargetOpSettingsBase
 	IKRIG_API virtual USkeleton* GetSkeleton(const FName InPropertyName) override;
 #endif
 
+	UPROPERTY(meta = (DeprecatedProperty))
+	FVector TranslationOffset_DEPRECATED = FVector::ZeroVector;
+
+	UPROPERTY(meta = (DeprecatedProperty))
+	FRotator RotationOffset_DEPRECATED = FRotator::ZeroRotator;
+	
 	UPROPERTY(meta = (DeprecatedProperty))
 	bool bEnableDebugDraw_DEPRECATED = true;
 };
