[IK Rig] Fixes crash when deleting skeletal mesh bone while mesh open in IK Rig.

diff --git a/Engine/Plugins/Animation/IKRig/Source/IKRigEditor/Private/RigEditor/IKRigEditMode.cpp b/Engine/Plugins/Animation/IKRig/Source/IKRigEditor/Private/RigEditor/IKRigEditMode.cpp
--- a/Engine/Plugins/Animation/IKRig/Source/IKRigEditor/Private/RigEditor/IKRigEditMode.cpp
+++ b/Engine/Plugins/Animation/IKRig/Source/IKRigEditor/Private/RigEditor/IKRigEditMode.cpp
@@ -30,22 +30,23 @@ bool FIKRigEditMode::GetCameraTarget(FSphere& OutTarget) const
 		return false; 
 	}
 	
-	// target union of selected goals and bones
-	TArray<FName> OutGoalNames, OutBoneNames;
+	
+	// collect selected points
+	TArray<FVector> TargetPoints;
+	
+	// get selected goal locations
+	TArray<FName> OutGoalNames;
 	Controller->GetSelectedGoalNames(OutGoalNames);
+	for (const FName& GoalName : OutGoalNames)
+	{			
+		TargetPoints.Add(Controller->AssetController->GetGoal(GoalName)->CurrentTransform.GetLocation());
+	}
+
+	// get selected bone locations
+	TArray<FName> OutBoneNames;
 	Controller->GetSelectedBoneNames(OutBoneNames);
-	
-	if (!(OutGoalNames.IsEmpty() && OutBoneNames.IsEmpty()))
+	if (!OutBoneNames.IsEmpty())
 	{
-		TArray<FVector> TargetPoints;
-
-		// get goal locations
-		for (const FName& GoalName : OutGoalNames)
-		{			
-			TargetPoints.Add(Controller->AssetController->GetGoal(GoalName)->CurrentTransform.GetLocation());
-		}
-
-		// get bone locations
 		const FIKRigSkeleton* CurrentIKRigSkeleton = Controller->GetCurrentIKRigSkeleton();
 		const FIKRigSkeleton& Skeleton = CurrentIKRigSkeleton ? *CurrentIKRigSkeleton : Controller->AssetController->GetIKRigSkeleton();
 		for (const FName& BoneName : OutBoneNames)
@@ -61,9 +62,13 @@ bool FIKRigEditMode::GetCameraTarget(FSphere& OutTarget) const
 				}
 				TargetPoints.Add(Skeleton.CurrentPoseGlobal[BoneIndex].GetLocation());
 			}
-		}
+		}		
+	}
 
-		// create a sphere that contains all the goal points
+	// target union of selected goals and bones
+	if (!TargetPoints.IsEmpty())
+	{
+		// create a sphere that contains all the target points
 		OutTarget = FSphere(&TargetPoints[0], TargetPoints.Num());
 		return true;
 	}
