[IK Retarget] Fixes ScaleSourcePivot not initialized properly.

diff --git a/Engine/Plugins/Animation/IKRig/Source/IKRig/Public/Retargeter/RetargetOps/ScaleSourceOp.h b/Engine/Plugins/Animation/IKRig/Source/IKRig/Public/Retargeter/RetargetOps/ScaleSourceOp.h
--- a/Engine/Plugins/Animation/IKRig/Source/IKRig/Public/Retargeter/RetargetOps/ScaleSourceOp.h
+++ b/Engine/Plugins/Animation/IKRig/Source/IKRig/Public/Retargeter/RetargetOps/ScaleSourceOp.h
@@ -27,7 +27,7 @@ struct FIKRetargetScaleSourceOpSettings : public FIKRetargetOpSettingsBase
 	double SourceScaleFactor = 1.0f;
 
 	UPROPERTY(EditAnywhere, BlueprintReadWrite, Category = ScalePivot)
-	EScaleSourcePivot ScalePivot;
+	EScaleSourcePivot ScalePivot = EScaleSourcePivot::ComponentOrigin;
 
 	UPROPERTY(EditAnywhere, Category = ScalePivot)
 	FBoneReference ScalePivotBone;
