[IK Retarget] Added PostSerialize to Pelvis Motion Op settings to load deprecated offsets.

diff --git a/Engine/Plugins/Animation/IKRig/Source/IKRig/Private/Retargeter/RetargetOps/PelvisMotionOp.cpp b/Engine/Plugins/Animation/IKRig/Source/IKRig/Private/Retargeter/RetargetOps/PelvisMotionOp.cpp
--- a/Engine/Plugins/Animation/IKRig/Source/IKRig/Private/Retargeter/RetargetOps/PelvisMotionOp.cpp
+++ b/Engine/Plugins/Animation/IKRig/Source/IKRig/Private/Retargeter/RetargetOps/PelvisMotionOp.cpp
@@ -49,6 +49,18 @@ USkeleton* FIKRetargetPelvisMotionOpSettings::GetSkeleton(const FName InProperty
 }
 #endif
 
+void FIKRetargetPelvisMotionOpSettings::PostSerialize(const FArchive& Ar)
+{
+	if (Ar.IsLoading())
+	{
+		if (Ar.CustomVer(FIKRigObjectVersion::GUID) < FIKRigObjectVersion::AddedLocalGlobalOffsetsToPelvisOp)
+		{
+			TranslationOffsetGlobal = TranslationOffset_DEPRECATED;
+			RotationOffsetLocal = RotationOffset_DEPRECATED;
+		}
+	}
+}
+
 bool FIKRetargetPelvisMotionOp::Initialize(
 	const FIKRetargetProcessor& Processor,
 	const FRetargetSkeleton& SourceSkeleton,
diff --git a/Engine/Plugins/Animation/IKRig/Source/IKRig/Public/Retargeter/RetargetOps/PelvisMotionOp.h b/Engine/Plugins/Animation/IKRig/Source/IKRig/Public/Retargeter/RetargetOps/PelvisMotionOp.h
--- a/Engine/Plugins/Animation/IKRig/Source/IKRig/Public/Retargeter/RetargetOps/PelvisMotionOp.h
+++ b/Engine/Plugins/Animation/IKRig/Source/IKRig/Public/Retargeter/RetargetOps/PelvisMotionOp.h
@@ -126,6 +126,14 @@ struct FIKRetargetPelvisMotionOpSettings : public FIKRetargetOpSettingsBase
 	
 	UPROPERTY(meta = (DeprecatedProperty))
 	bool bEnableDebugDraw_DEPRECATED = true;
+
+	void PostSerialize(const FArchive& Ar);
+};
+
+template<>
+struct TStructOpsTypeTraits<FIKRetargetPelvisMotionOpSettings> : public TStructOpsTypeTraitsBase2<FIKRetargetPelvisMotionOpSettings>
+{
+	enum { WithPostSerialize = true };
 };
 
 struct FPelvisSource
