[IK Retarget] Strip scale from DebugDraw source transform to prevent double transforms.

diff --git a/Engine/Plugins/Animation/IKRig/Source/IKRig/Private/Retargeter/IKRetargetProcessor.cpp b/Engine/Plugins/Animation/IKRig/Source/IKRig/Private/Retargeter/IKRetargetProcessor.cpp
--- a/Engine/Plugins/Animation/IKRig/Source/IKRig/Private/Retargeter/IKRetargetProcessor.cpp
+++ b/Engine/Plugins/Animation/IKRig/Source/IKRig/Private/Retargeter/IKRetargetProcessor.cpp
@@ -1685,13 +1685,27 @@ FName FIKRetargetProcessor::GetFirstChainMappedToChain(
 
 void FIKRetargetProcessor::DebugDrawAllOps(
 	FPrimitiveDrawInterface* InPDI,
-	const FTransform& InSourceTransform,
-	const FTransform& InComponentTransform,
-	const FIKRetargetDebugDrawState& EditorState,
-	const int32 InSelectedOpIndex) const
+	const FIKRetargetDebugDrawState& EditorState) const
 {
-	const float ComponentScale = static_cast<float>(InComponentTransform.GetScale3D().GetMax());
+	const UIKRetargeter* RetargetAsset = GetRetargetAsset();
+	if (!ensure(RetargetAsset))
+	{
+		return;
+	}
+
+	// construct the target preview transform
+	const double TargetScale = RetargetAsset->TargetMeshScale;
+	const FTransform TargetTransform = FTransform(FQuat::Identity, RetargetAsset->TargetMeshOffset, FVector(TargetScale));
 	
+	// NOTE: we cannot directly use the source mesh component transform because it has translation and scaling applied to preview the
+	// custom pivot source scaling results. The source pose bone transforms that the ops themselves have access to use scaled translations,
+	// but have scale at 1,1,1. We avoid actual bone scaling in the retargeter entirely, but if we visualized the source mesh in this manner,
+	// it would appear like the skin is stretched out and thin (when scaled up) and squished up and fat (when scaled down); hence the use of scaling.
+	//
+	// So to avoid double transforms in debug drawing, the ops should use the bone translations passed into them, and this source transform with NO scaling.
+	const FTransform SourceTransform = FTransform(FQuat::Identity, RetargetAsset->SourceMeshOffset, FVector::OneVector);
+
+	// debug draw each op in the stack
 	for (int32 OpIndex=0; OpIndex<OpStack.Num(); ++OpIndex)
 	{
 		// filter disabled or uninitialized ops
@@ -1701,7 +1715,7 @@ void FIKRetargetProcessor::DebugDrawAllOps(
 			continue;
 		}
 
-		Op->DebugDraw(InPDI, InSourceTransform, InComponentTransform, ComponentScale, EditorState);
+		Op->DebugDraw(InPDI, SourceTransform, TargetTransform, TargetScale, EditorState);
 	}
 }
 
diff --git a/Engine/Plugins/Animation/IKRig/Source/IKRig/Private/Retargeter/RetargetOps/PelvisMotionOp.cpp b/Engine/Plugins/Animation/IKRig/Source/IKRig/Private/Retargeter/RetargetOps/PelvisMotionOp.cpp
--- a/Engine/Plugins/Animation/IKRig/Source/IKRig/Private/Retargeter/RetargetOps/PelvisMotionOp.cpp
+++ b/Engine/Plugins/Animation/IKRig/Source/IKRig/Private/Retargeter/RetargetOps/PelvisMotionOp.cpp
@@ -400,11 +400,8 @@ void FIKRetargetPelvisMotionOp::DebugDraw(
 	const FIKRetargetDebugDrawState& InEditorState) const
 {
 	FScopeLock ScopeLock(&DebugDataMutex);
-
-	// strip scale out of the source transform (already applied to pose)
-	const FTransform SourceComponentTransform(InSourceTransform.GetRotation(), InSourceTransform.GetLocation(), FVector::OneVector);
 	
-	const FTransform SourcePelvis = FTransform(Source.CurrentRotation, Source.CurrentPosition) * SourceComponentTransform;
+	const FTransform SourcePelvis = FTransform(Source.CurrentRotation, Source.CurrentPosition) * InSourceTransform;
 	const FTransform TargetPelvis = CurrentPelvisTransform * InComponentTransform;
 	const FVector TargetFloor = TargetPelvis.GetLocation() * FVector(1,1,0);
 	const FVector TargetCrotch = TargetPelvis.GetLocation() + FVector(0,0, -Settings.TargetCrotchOffset);
diff --git a/Engine/Plugins/Animation/IKRig/Source/IKRig/Public/Retargeter/IKRetargetProcessor.h b/Engine/Plugins/Animation/IKRig/Source/IKRig/Public/Retargeter/IKRetargetProcessor.h
--- a/Engine/Plugins/Animation/IKRig/Source/IKRig/Public/Retargeter/IKRetargetProcessor.h
+++ b/Engine/Plugins/Animation/IKRig/Source/IKRig/Public/Retargeter/IKRetargetProcessor.h
@@ -526,12 +526,7 @@ public:
 	/** Attach a delegate to be notified whenever this processor is initialized. */
 	FOnRetargeterInitialized& OnRetargeterInitialized(){ return RetargeterInitialized; };
 	/** Run all debug drawing on all ops in the stack (or just selected op if one is provided) */
-	UE_API void DebugDrawAllOps(
-		FPrimitiveDrawInterface* InPDI,
-		const FTransform& InSourceTransform,
-		const FTransform& InComponentTransform,
-		const FIKRetargetDebugDrawState& EditorState,
-		const int32 InSelectedOpIndex) const;
+	UE_API void DebugDrawAllOps(FPrimitiveDrawInterface* InPDI, const FIKRetargetDebugDrawState& EditorState) const;
 	
 #endif
 
diff --git a/Engine/Plugins/Animation/IKRig/Source/IKRigEditor/Private/RetargetEditor/IKRetargetDefaultMode.cpp b/Engine/Plugins/Animation/IKRig/Source/IKRigEditor/Private/RetargetEditor/IKRetargetDefaultMode.cpp
--- a/Engine/Plugins/Animation/IKRig/Source/IKRigEditor/Private/RetargetEditor/IKRetargetDefaultMode.cpp
+++ b/Engine/Plugins/Animation/IKRig/Source/IKRigEditor/Private/RetargetEditor/IKRetargetDefaultMode.cpp
@@ -78,12 +78,8 @@ void FIKRetargetDefaultMode::RenderDebugProxies(FPrimitiveDrawInterface* PDI, co
 		return;
 	}
 
-	// debug draw the currently selected op
-	int32 SelectedOpIndex = Controller->GetSelectedOpIndex();
-	UDebugSkelMeshComponent* SourceSkelMesh = Controller->GetSkeletalMeshComponent(ERetargetSourceOrTarget::Source);
-	UDebugSkelMeshComponent* TargetSkelMesh = Controller->GetSkeletalMeshComponent(ERetargetSourceOrTarget::Target);
-	const FIKRetargetDebugDrawState& EditorState = Controller->GetSelectionState();
-	RetargetProcessor->DebugDrawAllOps(PDI, SourceSkelMesh->GetComponentTransform(), TargetSkelMesh->GetComponentTransform(), EditorState, SelectedOpIndex);
+	// give ops a chance to debug draw
+	RetargetProcessor->DebugDrawAllOps(PDI, Controller->GetSelectionState());
 }
 
 bool FIKRetargetDefaultMode::HandleClick(FEditorViewportClient* InViewportClient, HHitProxy* HitProxy, const FViewportClick& Click)
