[IK Retarget] Outdated comment fix in IKRetargetProcessor caught by [REDACTED_PERSON]

diff --git a/Engine/Plugins/Animation/IKRig/Source/IKRig/Public/Retargeter/IKRetargetProcessor.h b/Engine/Plugins/Animation/IKRig/Source/IKRig/Public/Retargeter/IKRetargetProcessor.h
--- a/Engine/Plugins/Animation/IKRig/Source/IKRig/Public/Retargeter/IKRetargetProcessor.h
+++ b/Engine/Plugins/Animation/IKRig/Source/IKRig/Public/Retargeter/IKRetargetProcessor.h
@@ -525,7 +525,7 @@ public:
 	UE_API FName GetFirstChainMappedToChain(const FName InChainName, const ERetargetSourceOrTarget InSourceOrTarget) const;
 	/** Attach a delegate to be notified whenever this processor is initialized. */
 	FOnRetargeterInitialized& OnRetargeterInitialized(){ return RetargeterInitialized; };
-	/** Run all debug drawing on all ops in the stack (or just selected op if one is provided) */
+	/** Run debug drawing on all ops in the stack */
 	UE_API void DebugDrawAllOps(FPrimitiveDrawInterface* InPDI, const FIKRetargetDebugDrawState& EditorState) const;
 	
 #endif
