[IK Retarget] Added new rotation filtering op

diff --git a/Engine/Plugins/Animation/IKRig/Source/IKRig/Private/Retargeter/IKRetargetProcessor.cpp b/Engine/Plugins/Animation/IKRig/Source/IKRig/Private/Retargeter/IKRetargetProcessor.cpp
--- a/Engine/Plugins/Animation/IKRig/Source/IKRig/Private/Retargeter/IKRetargetProcessor.cpp
+++ b/Engine/Plugins/Animation/IKRig/Source/IKRig/Private/Retargeter/IKRetargetProcessor.cpp
@@ -240,10 +240,12 @@ void FRetargetSkeleton::SetGlobalTransformAndUpdateChildren(
 
 	TArray<int32> ChildBoneIndices;
 	GetChildrenIndicesRecursive(InBoneToSetIndex, ChildBoneIndices);
+
 	for (const int32 ChildBoneIndex : ChildBoneIndices)
 	{
 		const FTransform RelativeToPrev = InOutGlobalPose[ChildBoneIndex].GetRelativeTransform(PrevTransform);
 		InOutGlobalPose[ChildBoneIndex] = RelativeToPrev * InNewTransform;
+		InOutGlobalPose[ChildBoneIndex].NormalizeRotation();
 	}
 }
 
diff --git a/Engine/Plugins/Animation/IKRig/Source/IKRig/Private/Retargeter/RetargetOps/FilterBoneOp.cpp b/Engine/Plugins/Animation/IKRig/Source/IKRig/Private/Retargeter/RetargetOps/FilterBoneOp.cpp
new file mode 100644
--- /dev/null
+++ b/Engine/Plugins/Animation/IKRig/Source/IKRig/Private/Retargeter/RetargetOps/FilterBoneOp.cpp
@@ -0,0 +1,231 @@
+// Copyright Epic Games, Inc. All Rights Reserved.
+
+#include "Retargeter/RetargetOps/FilterBoneOp.h"
+
+#define LOCTEXT_NAMESPACE "FilterBoneOp"
+
+
+FQuat FOneEuroRotationFilter::Update(
+	const FQuat& InTargetRotation,
+	double InDeltaTime,
+	const FIKRetargetFilterBoneOpSettings& InSettings)
+{
+	if (bFirstRun)
+	{
+		PrevFilteredRot = InTargetRotation;
+		PrevHalfAngVel = FVector::ZeroVector;
+		bFirstRun = false;
+		return InTargetRotation;
+	}
+
+	// sanity check delta time
+	constexpr double MAX_DELTATIME = 1.0/15.0;
+	constexpr double MIN_DELTATIME = 1.0/240.0;
+	InDeltaTime = FMath::Clamp(InDeltaTime, MIN_DELTATIME, MAX_DELTATIME);
+
+	auto AlphaFromHz = [&InDeltaTime](double FreqCutoff) -> double
+	{
+		const double Tau = (FreqCutoff > 0.0) ? (1.0 / (2.0 * PI * FreqCutoff)) : UE_SMALL_NUMBER;
+		return 1.0 / (1.0 + Tau / FMath::Max(InDeltaTime, UE_KINDA_SMALL_NUMBER));
+	};
+
+	// This an adaptive first-order low-pass filter based on the "One-Euro" 2012 paper: https://gery.casiez.net/1euro/
+	//
+	// NOTE: because we are filtering rotations, and desire to use linear math,
+	// all calculations are done with the exponential map of the delta quaternion
+
+	// calc error from current filtered to target
+	FQuat ErrorQ = InTargetRotation * PrevFilteredRot.Inverse();
+	ErrorQ = ErrorQ.W < 0 ? -ErrorQ : ErrorQ; // negate if W negative
+
+	// convert into tangent space
+	// NOTE: this produces a 3d vector amenable to filtering, with |Phi| = angle/2 
+	const FQuat LogQ = ErrorQ.Log();
+	const FVector Phi(LogQ.X, LogQ.Y, LogQ.Z);
+	
+	// compute angular velocity that would reduce error to zero if fully applied
+	const FVector HalfAngVel = (InDeltaTime > 0.0) ? (Phi / InDeltaTime) : FVector::ZeroVector; // rad/s
+	// low-pass filter the velocity with fixed cutoff frequency
+	const double VelCutoffFreq = (InSettings.VelocityCutoffHz > 0.0) ? InSettings.VelocityCutoffHz : 20.0; // Hz
+	PrevHalfAngVel = FMath::Lerp(PrevHalfAngVel, HalfAngVel, AlphaFromHz(VelCutoffFreq));
+
+	// adaptive cutoff for the signal based on filtered speed
+	const double Speed = 2.0 * PrevHalfAngVel.Size(); // rad/s
+	const double FreqCutoff = InSettings.MinFrequency + InSettings.Responsiveness * Speed; // Hz
+	// take only a fraction of the filtered velocity this frame
+	const FVector PhiStep = PrevHalfAngVel * (AlphaFromHz(FreqCutoff) * InDeltaTime);
+	
+	// exponentiate back into SO(3)
+	const FQuat DeltaQ = FQuat(PhiStep.X, PhiStep.Y, PhiStep.Z, 0).Exp();
+	// integrate the angular velocity to get new rotation
+	const FQuat NewRotation = (DeltaQ * PrevFilteredRot).GetNormalized();
+	// store prev filtered rotation for next frame
+	PrevFilteredRot = NewRotation;
+	
+	return NewRotation;
+}
+
+const UClass* FIKRetargetFilterBoneOpSettings::GetControllerType() const
+{
+	return UIKRetargetFilterBoneController::StaticClass();
+}
+
+void FIKRetargetFilterBoneOpSettings::CopySettingsAtRuntime(const FIKRetargetOpSettingsBase* InSettingsToCopyFrom)
+{
+	// copies everything except the bones we are operating on (those require reinit)
+	const TArray<FName> PropertiesToIgnore = {"BonesToFilter"};
+	FIKRetargetOpBase::CopyStructProperties(
+		FIKRetargetFilterBoneOpSettings::StaticStruct(),
+		InSettingsToCopyFrom,
+		this,
+		PropertiesToIgnore);
+}
+
+bool FIKRetargetFilterBoneOp::Initialize(
+	const FIKRetargetProcessor& Processor,
+	const FRetargetSkeleton& SourceSkeleton,
+	const FTargetSkeleton& TargetSkeleton,
+	const FIKRetargetOpBase* InParentOp, FIKRigLogger& Log)
+{
+	bIsInitialized = false;
+	
+	for (FFilterBoneData& FilterBoneData : Settings.BonesToFilter)
+	{
+		// find the target bone
+		FilterBoneData.TargetBone.BoneIndex = TargetSkeleton.FindBoneIndexByName(FilterBoneData.TargetBone.BoneName);
+		FilterBoneData.bIsReady = FilterBoneData.TargetBone.BoneIndex != INDEX_NONE;
+		if (!FilterBoneData.bIsReady)
+		{
+			Log.LogWarning(FText::Format(
+				LOCTEXT("MissingSourceBone", "Filter Bone op refers to non-existant target bone, {0}."),
+				FText::FromName(FilterBoneData.TargetBone.BoneName)));
+		}
+	}
+	
+	// always treat this op as "initialized", individual filters will only execute if their prerequisites are met
+	bIsInitialized = true;
+	// force to initialize bone rotations on next update
+	bResetPlayback = true;
+	return true;
+}
+
+void FIKRetargetFilterBoneOp::Run(
+	FIKRetargetProcessor& Processor,
+	const double InDeltaTime,
+	const TArray<FTransform>& InSourceGlobalPose,
+	TArray<FTransform>& OutTargetGlobalPose)
+{
+	// true on the first run or if playback was reset or reversed
+	if (bResetPlayback)
+	{
+		// do not filter on first run, just reset prev rotations
+		ResetFilters(OutTargetGlobalPose);
+		bResetPlayback = false;
+	}
+
+	const FRetargetSkeleton& TargetSkeleton = Processor.GetSkeleton(ERetargetSourceOrTarget::Target);
+	for (FFilterBoneData& FilterBoneData : Settings.BonesToFilter)
+	{
+		if (!FilterBoneData.bIsReady)
+		{
+			continue;
+		}
+
+		FTransform TargetGlobalTransform = OutTargetGlobalPose[FilterBoneData.TargetBone.BoneIndex];
+		FQuat FilteredRotation = FilterBoneData.Filter.Update(TargetGlobalTransform.GetRotation(), InDeltaTime, Settings);
+		TargetGlobalTransform.SetRotation(FilteredRotation);
+
+		// assign result and update children
+		TargetSkeleton.SetGlobalTransformAndUpdateChildren(FilterBoneData.TargetBone.BoneIndex, TargetGlobalTransform,OutTargetGlobalPose);
+	}
+}
+
+const UScriptStruct* FIKRetargetFilterBoneOp::GetSettingsType() const
+{
+	return FIKRetargetFilterBoneOpSettings::StaticStruct();
+}
+
+const UScriptStruct* FIKRetargetFilterBoneOp::GetType() const
+{
+	return FIKRetargetFilterBoneOp::StaticStruct();
+}
+
+void FIKRetargetFilterBoneOp::OnPlaybackReset()
+{
+	if (Settings.bResetPlayback)
+	{
+		bResetPlayback = true;	
+	}
+}
+
+void FIKRetargetFilterBoneOp::ResetFilters(TArray<FTransform>& OutTargetGlobalPose)
+{
+	for (FFilterBoneData& FilterBoneData : Settings.BonesToFilter)
+	{
+		if (!FilterBoneData.bIsReady)
+		{
+			continue;
+		}
+
+		FilterBoneData.Filter.Reset();
+	}
+}
+
+FIKRetargetOpSettingsBase* FIKRetargetFilterBoneOp::GetSettings()
+{
+	return &Settings;
+}
+
+// BEGIN CONTROLLER
+
+FIKRetargetFilterBoneOpSettings UIKRetargetFilterBoneController::GetSettings()
+{
+	return *reinterpret_cast<FIKRetargetFilterBoneOpSettings*>(OpSettingsToControl);
+}
+
+void UIKRetargetFilterBoneController::SetSettings(FIKRetargetFilterBoneOpSettings InSettings)
+{
+	OpSettingsToControl->CopySettingsAtRuntime(&InSettings);
+}
+
+void UIKRetargetFilterBoneController::ClearBonesToFilter()
+{
+	FIKRetargetFilterBoneOpSettings* Settings = reinterpret_cast<FIKRetargetFilterBoneOpSettings*>(OpSettingsToControl);
+	Settings->BonesToFilter.Reset();
+}
+
+void UIKRetargetFilterBoneController::AddBoneToFilter(const FName InTargetBone)
+{
+	FIKRetargetFilterBoneOpSettings* Settings = reinterpret_cast<FIKRetargetFilterBoneOpSettings*>(OpSettingsToControl);
+
+	// skip if bone already present
+	for (FFilterBoneData& BonePair : Settings->BonesToFilter)
+	{
+		if (BonePair.TargetBone == InTargetBone)
+		{
+			return;
+		}
+	}
+
+	// add new bone
+	FFilterBoneData NewBonePair;
+	NewBonePair.TargetBone.BoneName = InTargetBone;
+	Settings->BonesToFilter.Add(NewBonePair);
+}
+
+TArray<FName> UIKRetargetFilterBoneController::GetAllBonesToFilter()
+{
+	TArray<FName> AllBonesToFilter;
+	
+	FIKRetargetFilterBoneOpSettings* Settings = reinterpret_cast<FIKRetargetFilterBoneOpSettings*>(OpSettingsToControl);
+	for (FFilterBoneData& BoneToFilter : Settings->BonesToFilter)
+	{
+		AllBonesToFilter.Add(BoneToFilter.TargetBone.BoneName);
+	}
+	
+	return MoveTemp(AllBonesToFilter);
+}
+
+
+#undef LOCTEXT_NAMESPACE
+
diff --git a/Engine/Plugins/Animation/IKRig/Source/IKRig/Public/Retargeter/RetargetOps/FilterBoneOp.h b/Engine/Plugins/Animation/IKRig/Source/IKRig/Public/Retargeter/RetargetOps/FilterBoneOp.h
new file mode 100644
--- /dev/null
+++ b/Engine/Plugins/Animation/IKRig/Source/IKRig/Public/Retargeter/RetargetOps/FilterBoneOp.h
@@ -0,0 +1,157 @@
+﻿// Copyright Epic Games, Inc. All Rights Reserved.
+
+#pragma once
+
+#include "Animation/BoneReference.h"
+#include "Retargeter/IKRetargetOps.h"
+#include "Retargeter/IKRetargetProcessor.h"
+
+#include "FilterBoneOp.generated.h"
+
+#define UE_API IKRIG_API
+
+#define LOCTEXT_NAMESPACE "FilterBoneOp"
+
+struct FIKRetargetOpSettingsBase;
+
+struct FOneEuroRotationFilter
+{
+	FQuat Update(
+		const FQuat& InTargetRotation,
+		double InDeltaTime,
+		const FIKRetargetFilterBoneOpSettings& InSettings);
+
+	void Reset() { bFirstRun = true; };
+
+private:
+	
+	FQuat PrevFilteredRot = FQuat::Identity;
+	FVector PrevHalfAngVel = FVector::ZeroVector;
+	FVector PrevPhiFilt = FVector::ZeroVector;
+	bool bFirstRun = true;
+};
+
+USTRUCT(BlueprintType)
+struct FFilterBoneData
+{
+	GENERATED_BODY()
+
+	// The target bone to filter.
+	UPROPERTY(EditAnywhere, Category=Settings, meta=(ReinitializeOnEdit))
+	FBoneReference TargetBone;
+
+	bool bIsReady = false;
+	FOneEuroRotationFilter Filter;
+};
+
+USTRUCT(BlueprintType, meta = (DisplayName = "Filter Bone Settings"))
+struct FIKRetargetFilterBoneOpSettings : public FIKRetargetOpSettingsBase
+{
+	GENERATED_BODY()
+	
+	/** A list of bone-pairs to copy transforms between */
+	UPROPERTY(EditAnywhere, BlueprintReadWrite, Category=Setup, meta=(ReinitializeOnEdit))
+	TArray<FFilterBoneData> BonesToFilter;
+
+	/*/* When true, applies a One-Euro filter to smooth the rotation of the target bone #1#
+	UPROPERTY(EditAnywhere, BlueprintReadWrite, Category="One Euro Filter")
+	bool bFilterRotation = true;*/
+
+	/* Hz. Sets the low-pass cutoff when motion is near zero.
+	 * Higher = less smoothing at rest (more responsive but more jitter).
+	 * Lower = more smoothing at rest (less jitter but more “stickiness”). */
+	UPROPERTY(EditAnywhere, BlueprintReadWrite, Category="One Euro Filter")
+	double MinFrequency = 1.5;
+
+	/* Raises the cutoff as angular speed grows.
+	 * Larger values are snappier on fast turns.
+	 * Smaller values are smoother but laggier during quick motion.
+	 * Typical sweet spot: 0.3 – 0.8 */
+	UPROPERTY(EditAnywhere, BlueprintReadWrite, Category="One Euro Filter")
+	double Responsiveness = 0.5;
+
+	/* Hz. Low-passes the raw angular velocity before we use it to adapt the derivative cutoff.
+	* If you see breathing/pumping of the smoothing during motion onsets or reversals, lower this value (e.g., 30 → 20 Hz).
+	* If responsiveness to fast changes is sluggish, raise this a bit.
+	* Good starting range: 15–30 Hz. Keep under Nyquist frequency (frame_rate/2). */
+	UPROPERTY(EditAnywhere, BlueprintReadWrite, Category="One Euro Filter")
+	double VelocityCutoffHz = 20.0;
+
+	/* If true, filter is reset when playback loops. */
+	UPROPERTY(EditAnywhere, BlueprintReadWrite, Category="One Euro Filter")
+	bool bResetPlayback = true;
+
+	UE_API virtual const UClass* GetControllerType() const override;
+	
+	UE_API virtual void CopySettingsAtRuntime(const FIKRetargetOpSettingsBase* InSettingsToCopyFrom) override;
+};
+
+USTRUCT(BlueprintType, meta = (DisplayName = "Filter Bones"))
+struct FIKRetargetFilterBoneOp : public FIKRetargetOpBase
+{
+	GENERATED_BODY()
+
+	UE_API virtual bool Initialize(
+		const FIKRetargetProcessor& Processor,
+		const FRetargetSkeleton& SourceSkeleton,
+		const FTargetSkeleton& TargetSkeleton,
+		const FIKRetargetOpBase* InParentOp,
+		FIKRigLogger& Log) override;
+	
+	UE_API virtual void Run(
+		FIKRetargetProcessor& Processor,
+		const double InDeltaTime,
+		const TArray<FTransform>& InSourceGlobalPose,
+		TArray<FTransform>& OutTargetGlobalPose) override;
+
+	UE_API virtual FIKRetargetOpSettingsBase* GetSettings() override;
+	
+	UE_API virtual const UScriptStruct* GetSettingsType() const override;
+	
+	UE_API virtual const UScriptStruct* GetType() const override;
+
+	UE_API virtual void OnPlaybackReset();
+	
+	UPROPERTY()
+	FIKRetargetFilterBoneOpSettings Settings;
+
+private:
+	void ResetFilters(TArray<FTransform>& OutTargetGlobalPose);
+
+	bool bResetPlayback = true;
+};
+
+/* The blueprint/python API for editing a Filter Bone Op */
+UCLASS(MinimalAPI, BlueprintType)
+class UIKRetargetFilterBoneController : public UIKRetargetOpControllerBase
+{
+	GENERATED_BODY()
+	
+public:
+	/* Get the current op settings as a struct.
+	 * @return FIKRetargetPinBoneOpSettings struct with the current settings used by the op. */
+	UFUNCTION(BlueprintCallable, Category = Settings)
+	UE_API FIKRetargetFilterBoneOpSettings GetSettings();
+
+	/* Set the op settings. Input is a custom struct type for this op.
+	 * @param InSettings a FIKRetargetPinBoneOpSettings struct containing all the settings to apply to this op */
+	UFUNCTION(BlueprintCallable, Category = Settings)
+	UE_API void SetSettings(FIKRetargetFilterBoneOpSettings InSettings);
+
+	/* Clear all the bones */
+	UFUNCTION(BlueprintCallable, Category = BonePairs)
+	UE_API void ClearBonesToFilter();
+
+	/* Add a target bone to filter. */
+	UFUNCTION(BlueprintCallable, Category = BonePairs)
+	UE_API void AddBoneToFilter(const FName InTargetBone);
+
+	/* Get all the bones currently stored in the op.
+	 * @return an array of target bone names*/
+	UFUNCTION(BlueprintCallable, Category = BonePairs)
+	UE_API TArray<FName> GetAllBonesToFilter();
+};
+
+#undef LOCTEXT_NAMESPACE
+
+#undef UE_API
