[IK Retarget] Fixes bug where Source Scale with moving pivot forces reinitialization every frame.

diff --git a/Engine/Plugins/Animation/IKRig/Source/IKRig/Public/Retargeter/IKRetargetProcessor.h b/Engine/Plugins/Animation/IKRig/Source/IKRig/Public/Retargeter/IKRetargetProcessor.h
--- a/Engine/Plugins/Animation/IKRig/Source/IKRig/Public/Retargeter/IKRetargetProcessor.h
+++ b/Engine/Plugins/Animation/IKRig/Source/IKRig/Public/Retargeter/IKRetargetProcessor.h
@@ -38,7 +38,7 @@ struct FRetargetPoseScaleWithPivot
 
 	bool operator==(const FRetargetPoseScaleWithPivot& Other) const
 	{
-		return FMath::IsNearlyEqual(Factor, Other.Factor) && Pivot.Equals(Other.Pivot);
+		return FMath::IsNearlyEqual(Factor, Other.Factor);
 	}
 
 	bool ScalePose(TArray<FTransform>& InOutGlobalPose) const;
