[IK Retarget] Fixed regression causing IK goal ops not to initialize correctly.

diff --git a/Engine/Plugins/Animation/IKRig/Source/IKRig/Private/Retargeter/IKRetargetProcessor.cpp b/Engine/Plugins/Animation/IKRig/Source/IKRig/Private/Retargeter/IKRetargetProcessor.cpp
--- a/Engine/Plugins/Animation/IKRig/Source/IKRig/Private/Retargeter/IKRetargetProcessor.cpp
+++ b/Engine/Plugins/Animation/IKRig/Source/IKRig/Private/Retargeter/IKRetargetProcessor.cpp
@@ -1219,10 +1219,19 @@ void FIKRetargetProcessor::InitialOpStackSetup(
 void FIKRetargetProcessor::InitializeRetargetOps()
 {
 	// initialize the ops
+	TArray<FIKRetargetOpBase*> OpsInitBeforeChildrenThisTick;
 	for (FInstancedStruct& OpStruct : OpStack)
 	{
 		FIKRetargetOpBase& Op = OpStruct.GetMutable<FIKRetargetOpBase>();
 		FIKRetargetOpBase* ParentOp = GetRetargetOpByName(Op.GetParentOpName());
+
+		// give parent ops a chance to initialize right before their children
+		if (ParentOp && !OpsInitBeforeChildrenThisTick.Contains(ParentOp))
+		{
+			ParentOp->InitializeBeforeChildren(*this, SourceSkeleton, TargetSkeleton, Log);
+			OpsInitBeforeChildrenThisTick.Add(ParentOp);
+		}
+		
 		Op.Initialize(*this, SourceSkeleton, TargetSkeleton, ParentOp, Log);
 	}
 	
diff --git a/Engine/Plugins/Animation/IKRig/Source/IKRig/Private/Retargeter/RetargetOps/RunIKRigOp.cpp b/Engine/Plugins/Animation/IKRig/Source/IKRig/Private/Retargeter/RetargetOps/RunIKRigOp.cpp
--- a/Engine/Plugins/Animation/IKRig/Source/IKRig/Private/Retargeter/RetargetOps/RunIKRigOp.cpp
+++ b/Engine/Plugins/Animation/IKRig/Source/IKRig/Private/Retargeter/RetargetOps/RunIKRigOp.cpp
@@ -117,48 +117,26 @@ void FIKRetargetRunIKRigOp::Run(
 #endif
 }
 
+void FIKRetargetRunIKRigOp::InitializeBeforeChildren(
+	FIKRetargetProcessor& InProcessor,
+	const FRetargetSkeleton& InSourceSkeleton,
+	const FTargetSkeleton& InTargetSkeleton,
+	FIKRigLogger& Log)
+{
+	// reset the goal container for this IK Rig
+	FIKRigGoalContainer& GoalContainer = InProcessor.GetIKRigGoalContainer();
+	ResetGoalContainer(InTargetSkeleton.RetargetPoses.GetGlobalRetargetPose(), GoalContainer);
+}
+
 void FIKRetargetRunIKRigOp::RunBeforeChildren(
 	FIKRetargetProcessor& InProcessor,
 	const double InDeltaTime,
 	const TArray<FTransform>& InSourceGlobalPose,
 	TArray<FTransform>& OutTargetGlobalPose)
 {
-	auto ResetGoalContainer = [this, &OutTargetGlobalPose](FIKRigGoalContainer& InGoalContainer)
-	{
-		// if no IK Rig assigned, leave container empty
-		if (!Settings.IKRigAsset)
-		{
-			InGoalContainer.Empty();
-			return;
-		}
-
-		// add all the goals from the IK Rig (empties and refills)
-		InGoalContainer.FillWithGoalArray(Settings.IKRigAsset->GetGoalArray());
-
-		// set all goals to the input pose in component space
-		for (FIKRigGoal& Goal : InGoalContainer.GetGoalArray())
-		{
-			const int32* BoneIndexPtr = GoalBoneIndicesMap.Find(Goal.BoneName);
-			if (BoneIndexPtr)
-			{
-				Goal.PositionSpace = EIKRigGoalSpace::Component;
-				Goal.RotationSpace = EIKRigGoalSpace::Component;
-				Goal.Position = OutTargetGlobalPose[*BoneIndexPtr].GetLocation();
-				Goal.Rotation = OutTargetGlobalPose[*BoneIndexPtr].GetRotation().Rotator();
-			}
-			else
-			{
-				Goal.PositionSpace = EIKRigGoalSpace::Additive;
-				Goal.RotationSpace = EIKRigGoalSpace::Additive;
-				Goal.Position = FVector::ZeroVector;
-				Goal.Rotation = FRotator::ZeroRotator;
-			}
-		}
-	};
-
 	// reset the goal container for this IK Rig
 	FIKRigGoalContainer& GoalContainer = InProcessor.GetIKRigGoalContainer();
-	ResetGoalContainer(GoalContainer);
+	ResetGoalContainer(OutTargetGlobalPose, GoalContainer);
 }
 
 void FIKRetargetRunIKRigOp::OnAddedToStack(const UIKRetargeter* InRetargetAsset, const FIKRetargetOpBase* InParentOp)
@@ -353,6 +331,41 @@ TArray<FName> FIKRetargetRunIKRigOp::GetRequiredTargetChains() const
 	return MoveTemp(RequiredTargetChains);
 }
 
+void FIKRetargetRunIKRigOp::ResetGoalContainer(
+	const TArray<FTransform>& InTargetGlobalPose,
+	FIKRigGoalContainer& InOutGoalContainer)
+{
+	// if no IK Rig assigned, leave container empty
+	if (!Settings.IKRigAsset)
+	{
+		InOutGoalContainer.Empty();
+		return;
+	}
+
+	// add all the goals from the IK Rig (empties and refills)
+	InOutGoalContainer.FillWithGoalArray(Settings.IKRigAsset->GetGoalArray());
+
+	// set all goals to the input pose in component space
+	for (FIKRigGoal& Goal : InOutGoalContainer.GetGoalArray())
+	{
+		const int32* BoneIndexPtr = GoalBoneIndicesMap.Find(Goal.BoneName);
+		if (BoneIndexPtr)
+		{
+			Goal.PositionSpace = EIKRigGoalSpace::Component;
+			Goal.RotationSpace = EIKRigGoalSpace::Component;
+			Goal.Position = InTargetGlobalPose[*BoneIndexPtr].GetLocation();
+			Goal.Rotation = InTargetGlobalPose[*BoneIndexPtr].GetRotation().Rotator();
+		}
+		else
+		{
+			Goal.PositionSpace = EIKRigGoalSpace::Additive;
+			Goal.RotationSpace = EIKRigGoalSpace::Additive;
+			Goal.Position = FVector::ZeroVector;
+			Goal.Rotation = FRotator::ZeroRotator;
+		}
+	}
+}
+
 FIKRetargetRunIKRigOpSettings UIKRetargetRunIKRigController::GetSettings()
 {
 	return *reinterpret_cast<FIKRetargetRunIKRigOpSettings*>(OpSettingsToControl);
diff --git a/Engine/Plugins/Animation/IKRig/Source/IKRig/Public/Retargeter/IKRetargetOps.h b/Engine/Plugins/Animation/IKRig/Source/IKRig/Public/Retargeter/IKRetargetOps.h
--- a/Engine/Plugins/Animation/IKRig/Source/IKRig/Public/Retargeter/IKRetargetOps.h
+++ b/Engine/Plugins/Animation/IKRig/Source/IKRig/Public/Retargeter/IKRetargetOps.h
@@ -226,6 +226,15 @@ public:
 	/** (optional) ops can optionally behave as a 'child' where they can be parented to ops of the type returned by this function. */
 	virtual const UScriptStruct* GetParentOpType() const { return nullptr; };
 
+	/** (optional) override to give this op a chance to initialize before it's children are initialized
+	 * This is executed right after the previous op in the stack and right before the first child op Initialize() is called
+	 * NOTE: unless this is a parent op, this will have no effect */
+	virtual void InitializeBeforeChildren(
+		FIKRetargetProcessor& InProcessor,
+		const FRetargetSkeleton& InSourceSkeleton,
+		const FTargetSkeleton& InTargetSkeleton,
+		FIKRigLogger& Log) {};
+	
 	/** (optional) override to give this op a chance to do work before it's children are run
 	 * This is executed right after the previous op in the stack and right before the first child op Run() is called
 	 * NOTE: unless this is a parent op, this will have no effect */
diff --git a/Engine/Plugins/Animation/IKRig/Source/IKRig/Public/Retargeter/RetargetOps/RunIKRigOp.h b/Engine/Plugins/Animation/IKRig/Source/IKRig/Public/Retargeter/RetargetOps/RunIKRigOp.h
--- a/Engine/Plugins/Animation/IKRig/Source/IKRig/Public/Retargeter/RetargetOps/RunIKRigOp.h
+++ b/Engine/Plugins/Animation/IKRig/Source/IKRig/Public/Retargeter/RetargetOps/RunIKRigOp.h
@@ -84,6 +84,12 @@ struct FIKRetargetRunIKRigOp : public FIKRetargetOpBase
 	
 	virtual bool CanHaveChildOps() const override { return true; };
 
+	UE_API virtual void InitializeBeforeChildren(
+		FIKRetargetProcessor& InProcessor,
+		const FRetargetSkeleton& InSourceSkeleton,
+		const FTargetSkeleton& InTargetSkeleton,
+		FIKRigLogger& Log) override;
+	
 	UE_API virtual void RunBeforeChildren(
 		FIKRetargetProcessor& InProcessor,
 		const double InDeltaTime,
@@ -126,6 +132,8 @@ struct FIKRetargetRunIKRigOp : public FIKRetargetOpBase
 	
 private:
 
+	void ResetGoalContainer(const TArray<FTransform> &InTargetGlobalPose, FIKRigGoalContainer& InOutGoalContainer);
+
 	/** The IK Rig processor for running IK on the target */
 	FIKRigProcessor IKRigProcessor;
 
