[RigMapper] move RigMapper Op out into experimental plugin and fixed automation test

diff --git a/Engine/Plugins/Animation/IKRig/IKRig.uplugin b/Engine/Plugins/Animation/IKRig/IKRig.uplugin
--- a/Engine/Plugins/Animation/IKRig/IKRig.uplugin
+++ b/Engine/Plugins/Animation/IKRig/IKRig.uplugin
@@ -23,10 +23,6 @@
 		{
 			"Name": "FullBodyIK",
 			"Enabled": true
-		},
-		{
-			"Name": "RigMapper",
-			"Enabled": true
 		}
 	],
 	"Modules": [
diff --git a/Engine/Plugins/Animation/IKRig/Source/IKRig/IKRig.Build.cs b/Engine/Plugins/Animation/IKRig/Source/IKRig/IKRig.Build.cs
--- a/Engine/Plugins/Animation/IKRig/Source/IKRig/IKRig.Build.cs
+++ b/Engine/Plugins/Animation/IKRig/Source/IKRig/IKRig.Build.cs
@@ -18,7 +18,6 @@ public class IKRig : ModuleRules
 				"CoreUObject",
 				"Engine",
 				"AnimationCore",
-				"RigMapper",
 			});
 
 		PublicDependencyModuleNames.AddRange(
diff --git a/Engine/Plugins/Animation/IKRig/Source/IKRig/Private/Retargeter/RetargetOps/RigMapperOp.cpp b/Engine/Plugins/Animation/IKRig/Source/IKRig/Private/Retargeter/RetargetOps/RigMapperOp.cpp
deleted file mode 100644
diff --git a/Engine/Plugins/Animation/IKRig/Source/IKRig/Public/Retargeter/RetargetOps/RigMapperOp.h b/Engine/Plugins/Animation/IKRig/Source/IKRig/Public/Retargeter/RetargetOps/RigMapperOp.h
deleted file mode 100644
diff --git a/Engine/Plugins/Animation/IKRig/Source/IKRigEditor/IKRigEditor.Build.cs b/Engine/Plugins/Animation/IKRig/Source/IKRigEditor/IKRigEditor.Build.cs
--- a/Engine/Plugins/Animation/IKRig/Source/IKRigEditor/IKRigEditor.Build.cs
+++ b/Engine/Plugins/Animation/IKRig/Source/IKRigEditor/IKRigEditor.Build.cs
@@ -51,8 +51,7 @@ namespace UnrealBuildTool.Rules
                     "AnimationWidgets",
                     "ApplicationCore",
 					"MeshDescription",
-					"SkeletalMeshDescription",
-					"RigMapperEditor"
+					"SkeletalMeshDescription"
                 }
             );
         }
diff --git a/Engine/Plugins/Animation/IKRig/Source/IKRigEditor/Private/RetargetEditor/CurveUtils.cpp b/Engine/Plugins/Animation/IKRig/Source/IKRigEditor/Private/RetargetEditor/CurveUtils.cpp
new file mode 100644
--- /dev/null
+++ b/Engine/Plugins/Animation/IKRig/Source/IKRigEditor/Private/RetargetEditor/CurveUtils.cpp
@@ -0,0 +1,209 @@
+// Copyright Epic Games, Inc. All Rights Reserved.
+
+#include "CurveUtils.h"
+#include "Animation/AnimSequence.h"
+#include "Animation/Skeleton.h"
+#include "Interfaces//ITargetPlatformManagerModule.h"
+
+bool FCurveUtils::LoadCurveValuesFromAnimSequence(const UAnimSequence* InSource, const USkeleton* InSourceSkeleton, TArray<FName>& OutCurveNames, TArray<FFrameTime>& OutFrameTimes, FCurveUtils::FFrameValues& OutCurveValuesPerFrame,
+	TArray<int32>& OutCurveFlags, TArray<FLinearColor>& OutCurveColors)
+{
+	TArray<FString> CurveNames;
+	if (LoadCurveValuesFromAnimSequence(InSource, CurveNames, OutFrameTimes, OutCurveValuesPerFrame))
+	{
+		OutCurveNames.Reserve(CurveNames.Num());
+		OutCurveFlags.Reserve(CurveNames.Num());
+		OutCurveColors.Reserve(CurveNames.Num());
+
+		const IAnimationDataModel* SourceDataModel = InSource->GetDataModel();
+
+		for (const FString& CurveName : CurveNames)
+		{
+			const FName CurCurveName(*CurveName);
+			const FAnimationCurveIdentifier CurCurveId = UAnimationCurveIdentifierExtensions::FindCurveIdentifier(InSourceSkeleton, CurCurveName, ERawCurveTrackTypes::RCT_Float);
+
+			const FFloatCurve* SourceCurve = SourceDataModel->FindFloatCurve(CurCurveId);
+			if (SourceCurve)
+			{
+				OutCurveFlags.Add(SourceCurve->GetCurveTypeFlags());
+				OutCurveColors.Add(SourceCurve->Color);
+			}
+			else
+			{
+				return false;
+			}
+
+			OutCurveNames.Add(CurCurveName);
+		}
+	}
+	else
+	{
+		return false;
+	}
+
+	return true;
+}
+
+FFrameRate FCurveUtils::GetAnimSequenceRate(const UAnimSequence* InAnimSequence)
+{
+	return InAnimSequence ? InAnimSequence->GetTargetSamplingFrameRate(GetTargetPlatformManager()->GetRunningTargetPlatform()) : FFrameRate();
+}
+
+bool FCurveUtils::LoadCurveValuesFromAnimSequence(const UAnimSequence* InSource, TArray<FString>& OutCurveNames, TArray<FFrameTime>& OutFrameTimes, FFrameValues& OutCurveValuesPerFrame)
+{
+	FFramePoses Poses;
+
+	const TArray<FFloatCurve>& FloatCurves = InSource->GetDataModel()->GetFloatCurves();
+
+	if (!FloatCurves.IsEmpty())
+	{
+		OutFrameTimes.Reserve(FloatCurves[0].FloatCurve.GetNumKeys());
+
+		for (const FFloatCurve& Curve : FloatCurves)
+		{
+			const FString CurveName = Curve.GetName().ToString();
+
+			if (!Curve.FloatCurve.Keys.IsEmpty())
+			{
+				OutCurveNames.AddUnique(CurveName);
+
+				for (const FRichCurveKey& Key : Curve.FloatCurve.Keys)
+				{
+					const FFrameTime FrameNumber = GetAnimSequenceRate(InSource).AsFrameTime(Key.Time);
+					int32 PoseIndex = OutFrameTimes.Find(FrameNumber);
+
+					if (PoseIndex == INDEX_NONE)
+					{
+						PoseIndex = OutFrameTimes.Add(FrameNumber);
+						Poses.AddDefaulted_GetRef().Reserve(FloatCurves.Num());
+					}
+
+					Poses[PoseIndex].Add(CurveName, Key.Value);
+				}
+			}
+		}
+	}
+
+	return BakeSparseKeys(Poses, OutCurveNames, OutFrameTimes, OutCurveValuesPerFrame);
+}
+
+bool FCurveUtils::BakeSparseKeys(const FFramePoses& InPoses, const TArray<FString>& InCurveNames, TArray<FFrameTime>& OutFrameTimes, FFrameValues& OutCurveValuesPerFrame)
+{
+	if (InPoses.IsEmpty())
+	{
+		return false;
+	}
+	if (InCurveNames.IsEmpty())
+	{
+		return false;
+	}
+	if (InPoses.Num() != OutFrameTimes.Num())
+	{
+		return false;
+	}
+
+	// This extra work is needed to make sure frame times and thus resulting curve values are in order
+	TArray<FFrameTime> FrameTimesInOrder = OutFrameTimes;
+	FrameTimesInOrder.Sort();
+
+	OutCurveValuesPerFrame.Reset(FrameTimesInOrder.Num());
+	for (int32 FrameIndex = 0; FrameIndex < OutFrameTimes.Num(); FrameIndex += 1)
+	{
+		OutCurveValuesPerFrame.AddDefaulted_GetRef().AddZeroed(InCurveNames.Num());
+	}
+
+	TArray<FString> BakedCurves;
+	BakedCurves.Reserve(InCurveNames.Num());
+
+	// If all curves are present on frame 0, we won't go further than one loop
+	for (int32 FrameIndex = 0; FrameIndex < OutFrameTimes.Num() && BakedCurves.Num() < InCurveNames.Num(); FrameIndex += 1)
+	{
+		const int32 ActualFrameIndex = OutFrameTimes.Find(FrameTimesInOrder[FrameIndex]); // Lookup for non ordered frames
+
+		for (const TPair<FString, float>& Keyframe : InPoses[ActualFrameIndex])
+		{
+			const FString& CurveName = Keyframe.Key;
+
+			if (BakedCurves.Contains(CurveName))
+			{
+				continue;
+			}
+			BakedCurves.Add(CurveName);
+
+			const float CurveValue = Keyframe.Value;
+			const int32 CurveIndex = InCurveNames.Find(CurveName);
+
+			SparseBakeCurve(CurveName, CurveIndex, CurveValue, OutCurveValuesPerFrame, FrameTimesInOrder, FrameIndex, InPoses, OutFrameTimes, ActualFrameIndex);
+
+			if (BakedCurves.Num() == InCurveNames.Num())
+			{
+				break;
+			}
+		}
+	}
+
+	OutFrameTimes = FrameTimesInOrder;
+
+	if (OutCurveValuesPerFrame.Num() != OutFrameTimes.Num())
+	{
+		return false;
+	}
+	if (OutCurveValuesPerFrame[0].Num() != InCurveNames.Num())
+	{
+		return false;
+	}
+	if (OutCurveValuesPerFrame.Last().Num() != InCurveNames.Num())
+	{
+		return false;
+	}
+	return true;
+}
+
+void FCurveUtils::SparseBakeCurve(const FString& InCurveName, int32 InCurveIndex, float InCurveValue, FFrameValues& OutCurveValuesPerFrame, const TArray<FFrameTime>& InFrameTimesInOrder, int32 InFrameIndex,
+	const FFramePoses& InPoses, const TArray<FFrameTime>& InFrameTimes, const int32 InActualFrameIndex)
+{
+	// We haven't seen this control so far. So keep it const on all previous frames as well as the current one
+	for (int32 BakeFrameIndex = 0; BakeFrameIndex <= InFrameIndex; BakeFrameIndex += 1)
+	{
+		OutCurveValuesPerFrame[BakeFrameIndex][InCurveIndex] = InCurveValue;
+	}
+
+	float LastKeyedValue = InCurveValue;
+	int32 LastKeyedFrameIndex = InFrameIndex;
+	int32 ActualLastKeyedFrameIndex = InActualFrameIndex; // Lookup for non ordered frames
+
+	// For following frames, add a key a lerp all previously missed keys if need be
+	for (int32 BakeFrameIndex = InFrameIndex + 1; BakeFrameIndex < InFrameTimes.Num(); BakeFrameIndex += 1)
+	{
+		const int32 ActualBakeFrameIndex = InFrameTimes.Find(InFrameTimesInOrder[BakeFrameIndex]); // Lookup for non ordered frames
+
+		// If the current frame has a key for the control, set it, and try to lerp all the keys we might have missed since last time we saw the control
+		if (InPoses[ActualBakeFrameIndex].Contains(InCurveName))
+		{
+			// Key this frame
+			OutCurveValuesPerFrame[BakeFrameIndex][InCurveIndex] = InPoses[ActualBakeFrameIndex][InCurveName];
+			LastKeyedValue = OutCurveValuesPerFrame[BakeFrameIndex][InCurveIndex].Get(0);
+
+			// If we have not seen this control on the previous frame, lerp all keys with missed
+			if (BakeFrameIndex > LastKeyedFrameIndex + 1)
+			{
+				// Lerp from the last key we found for this control to the current key, adding a lerp key for each frame
+				for (int32 LerpFrameIndex = LastKeyedFrameIndex + 1; LerpFrameIndex < BakeFrameIndex; LerpFrameIndex += 1)
+				{
+					const double LerpAlpha = (InFrameTimesInOrder[LerpFrameIndex].AsDecimal() - InFrameTimesInOrder[LastKeyedFrameIndex].AsDecimal()) / (InFrameTimesInOrder[BakeFrameIndex].AsDecimal() - InFrameTimesInOrder[LastKeyedFrameIndex].AsDecimal());
+					const double LerpValue = FMath::Lerp(InPoses[ActualLastKeyedFrameIndex][InCurveName], InPoses[ActualBakeFrameIndex][InCurveName], LerpAlpha);
+
+					OutCurveValuesPerFrame[LerpFrameIndex][InCurveIndex] = static_cast<float>(LerpValue);
+				}
+			}
+			LastKeyedFrameIndex = BakeFrameIndex;
+			ActualLastKeyedFrameIndex = ActualBakeFrameIndex; // Lookup for non ordered frames
+		}
+		else
+		{
+			// Per default, add a const key to all next frames (in case there isn't any key next).
+			// If it happens that there is one later, LastKeyedFrameIndex has not changed, and the above will lerp and override all const keys 
+			OutCurveValuesPerFrame[BakeFrameIndex][InCurveIndex] = LastKeyedValue;
+		}
+	}
+}
\ No newline at end of file
diff --git a/Engine/Plugins/Animation/IKRig/Source/IKRigEditor/Private/RetargetEditor/CurveUtils.h b/Engine/Plugins/Animation/IKRig/Source/IKRigEditor/Private/RetargetEditor/CurveUtils.h
new file mode 100644
--- /dev/null
+++ b/Engine/Plugins/Animation/IKRig/Source/IKRigEditor/Private/RetargetEditor/CurveUtils.h
@@ -0,0 +1,44 @@
+// Copyright Epic Games, Inc. All Rights Reserved.
+
+#pragma once
+
+#include "CoreMinimal.h"
+#include "Misc/FrameRate.h"
+
+
+class UAnimSequence;
+class USkeleton;
+
+/**
+ * Simple helper class to contain function for loading curve values from anim sequence.
+ * Note, these functions are copies from ones in the RigMapper plugin Subsystem, but we do not want to make IKRig depend upon this plugin whilst it is experimental, and does not seem worthwhile
+ * making an additional plugin.
+ * When RigMapper plugin moves out of experimental, this code should be re-consolidated.
+ */
+class FCurveUtils
+{
+	
+public:
+
+	using FPoseValues = TArray<TOptional<float>>;
+	using FFrameValues = TArray<FCurveUtils::FPoseValues>;
+
+	/**
+	 * Load curve values from animation sequence
+	 */
+	static bool LoadCurveValuesFromAnimSequence(const UAnimSequence* InSource, const USkeleton* InSourceSkeleton, TArray<FName>& OutCurveNames, TArray<FFrameTime>& OutFrameTimes, FCurveUtils::FFrameValues& OutCurveValuesPerFrame, TArray<int32> & OutCurveFlags, TArray<FLinearColor> & OutCurveColors);
+
+	/**
+	 * Get the frame rate for the animation sequence
+	 */
+	static FFrameRate GetAnimSequenceRate(const UAnimSequence* AnimSequence);
+
+
+private:
+	static bool LoadCurveValuesFromAnimSequence(const UAnimSequence* InSource, TArray<FString>& OutCurveNames, TArray<FFrameTime>& OutFrameTimes, FFrameValues& OutCurveValuesPerFrame);
+	using FPose = TMap<FString, float>;
+	using FFramePoses = TArray<FCurveUtils::FPose>;
+	static bool BakeSparseKeys(const FFramePoses& InPoses, const TArray<FString>& InCurveNames, TArray<FFrameTime>& OutFrameTimes, FFrameValues& OutCurveValuesPerFrame);
+	static void SparseBakeCurve(const FString& InCurveName, int32 InCurveIndex, float InCurveValue, FFrameValues& OutCurveValuesPerFrame, const TArray<FFrameTime>& InFrameTimesInOrder, int32 InFrameIndex, const FFramePoses& InPoses, const TArray<FFrameTime>& InFrameTimes, const int32 InActualFrameIndex);
+};
+
diff --git a/Engine/Plugins/Animation/IKRig/Source/IKRigEditor/Private/RetargetEditor/IKRetargetBatchOperation.cpp b/Engine/Plugins/Animation/IKRig/Source/IKRigEditor/Private/RetargetEditor/IKRetargetBatchOperation.cpp
--- a/Engine/Plugins/Animation/IKRig/Source/IKRigEditor/Private/RetargetEditor/IKRetargetBatchOperation.cpp
+++ b/Engine/Plugins/Animation/IKRig/Source/IKRigEditor/Private/RetargetEditor/IKRetargetBatchOperation.cpp
@@ -24,7 +24,7 @@
 #include "ObjectTools.h"
 #include "AssetRegistry/AssetRegistryModule.h"
 #include "Retargeter/RetargetOps/SpeedPlantingOp.h"
-#include "RigMapperEditorSubsystem.h"
+#include "CurveUtils.h"
 
 #include UE_INLINE_GENERATED_CPP_BY_NAME(IKRetargetBatchOperation)
 
@@ -576,7 +576,7 @@ void UIKRetargetBatchOperation::ApplyCurveOps(const FIKRetargetBatchOperationCon
 		FIKRetargetOpBase::FCurveData InputCurveMetaData;
 
 		TArray<FFrameTime> InputFrameTimes;
-		bool bLoadedData = URigMapperEditorSubsystem::LoadCurveValuesFromAnimSequence(SourceSequence, SourceSkeleton, InputCurveMetaData.Names, InputFrameTimes, InputCurveFrameValues, InputCurveMetaData.Flags, InputCurveMetaData.Colors);
+		bool bLoadedData = FCurveUtils::LoadCurveValuesFromAnimSequence(SourceSequence, SourceSkeleton, InputCurveMetaData.Names, InputFrameTimes, InputCurveFrameValues, InputCurveMetaData.Flags, InputCurveMetaData.Colors);
 		if (bLoadedData)
 		{
 			FIKRetargetOpBase::FFrameValues OutputCurveFrameValues;
@@ -601,13 +601,13 @@ void UIKRetargetBatchOperation::ApplyCurveOps(const FIKRetargetBatchOperationCon
 
 			// set the output values and metadata into the anim sequence
 			// note that outputs are contained in the Input data variables since they are moved there after each operation
-			const FFrameRate InputFrameRate = URigMapperEditorSubsystem::GetAnimSequenceRate(SourceSequence);
+			const FFrameRate InputFrameRate = FCurveUtils::GetAnimSequenceRate(SourceSequence);
 			AddCurveValuesToAnimSequence(SourceSkeleton, InputCurveMetaData, InputCurveFrameValues, InputFrameRate, InputFrameTimes, bShouldTransact, TargetSeqController);
 		}
 			
 		if(!bLoadedData)
 		{
-			UE_LOG(LogTemp, Warning, TEXT("Batch retarget aborted. Failed to get data from Anim Sequence."));
+			UE_LOG(LogTemp, Warning, TEXT("No animation curves found for retargeting in Anim Sequence: %s"), *SourceSequence->GetName());
 		}
 
 		TargetSeqController.CloseBracket(bShouldTransact);
