[IK Retarget] Restoring op settings serialization patches due to new Op::PostLoad() not being called in all instances.

diff --git a/Engine/Plugins/Animation/IKRig/Source/IKRig/Private/Retargeter/RetargetOps/CopyBasePoseOp.cpp b/Engine/Plugins/Animation/IKRig/Source/IKRig/Private/Retargeter/RetargetOps/CopyBasePoseOp.cpp
--- a/Engine/Plugins/Animation/IKRig/Source/IKRig/Private/Retargeter/RetargetOps/CopyBasePoseOp.cpp
+++ b/Engine/Plugins/Animation/IKRig/Source/IKRig/Private/Retargeter/RetargetOps/CopyBasePoseOp.cpp
@@ -26,6 +26,19 @@ void FIKRetargetCopyBasePoseOpSettings::CopySettingsAtRuntime(const FIKRetargetO
 		PropertiesToIgnore);
 }
 
+void FIKRetargetCopyBasePoseOpSettings::PostSerialize(const FArchive& Ar)
+{
+	if (Ar.IsLoading())
+	{
+		if (Ar.CustomVer(FIKRigObjectVersion::GUID) < FIKRigObjectVersion::CopyBasePoseConvertedToBoneReference)
+		{
+			PRAGMA_DISABLE_DEPRECATION_WARNINGS
+			CopyFromStart.BoneName = CopyBasePoseRoot_DEPRECATED;
+			PRAGMA_ENABLE_DEPRECATION_WARNINGS
+		}
+	}
+}
+
 bool FIKRetargetCopyBasePoseOp::Initialize(
 	const FIKRetargetProcessor& InProcessor,
 	const FRetargetSkeleton& InSourceSkeleton,
diff --git a/Engine/Plugins/Animation/IKRig/Source/IKRig/Private/Retargeter/RetargetOps/IKChainsOp.cpp b/Engine/Plugins/Animation/IKRig/Source/IKRig/Private/Retargeter/RetargetOps/IKChainsOp.cpp
--- a/Engine/Plugins/Animation/IKRig/Source/IKRig/Private/Retargeter/RetargetOps/IKChainsOp.cpp
+++ b/Engine/Plugins/Animation/IKRig/Source/IKRig/Private/Retargeter/RetargetOps/IKChainsOp.cpp
@@ -185,6 +185,17 @@ bool FRetargetIKChainSettings::operator==(const FRetargetIKChainSettings& Other)
 		&& FMath::IsNearlyEqualByULP(Extension, Other.Extension);
 }
 
+void FRetargetIKChainSettings::PostSerialize(const FArchive& Ar)
+{
+	if (Ar.IsLoading())
+	{
+		if (Ar.CustomVer(FIKRigObjectVersion::GUID) < FIKRigObjectVersion::RemovedPelvisMotionFromSourceGoals)
+		{
+			ApplyPelvisOffsetToSourceGoals = true;
+		}
+	}
+}
+
 const UClass* FIKRetargetIKChainsOpSettings::GetControllerType() const
 {
 	return UIKRetargetIKChainsController::StaticClass();
diff --git a/Engine/Plugins/Animation/IKRig/Source/IKRig/Private/Retargeter/RetargetOps/PinBoneOp.cpp b/Engine/Plugins/Animation/IKRig/Source/IKRig/Private/Retargeter/RetargetOps/PinBoneOp.cpp
--- a/Engine/Plugins/Animation/IKRig/Source/IKRig/Private/Retargeter/RetargetOps/PinBoneOp.cpp
+++ b/Engine/Plugins/Animation/IKRig/Source/IKRig/Private/Retargeter/RetargetOps/PinBoneOp.cpp
@@ -37,6 +37,20 @@ void FPinBoneData::CachePinData(ERetargetSourceOrTarget InSkeletonToCopyFrom, co
     LocalRefPoseBoneToCopyTo = BoneToCopyToParentRefPoseGlobal.GetRelativeTransform(BoneToCopyToRefPoseGlobal);
 }
 
+void FPinBoneData::PostSerialize(const FArchive& Ar)
+{
+	if (Ar.IsLoading())
+	{
+		if (Ar.CustomVer(FIKRigObjectVersion::GUID) < FIKRigObjectVersion::PinBoneTypeAndOffsetsUpgraded)
+		{
+			PRAGMA_DISABLE_DEPRECATION_WARNINGS
+			BoneToCopyTo.BoneName = BoneToPin_DEPRECATED;
+			BoneToCopyFrom.BoneName = BoneToPinTo_DEPRECATED;
+			PRAGMA_ENABLE_DEPRECATION_WARNINGS
+		}
+	}
+}
+
 const UClass* FIKRetargetPinBoneOpSettings::GetControllerType() const
 {
 	return UIKRetargetPinBoneController::StaticClass();
diff --git a/Engine/Plugins/Animation/IKRig/Source/IKRig/Private/Retargeter/RetargetOps/RootMotionGeneratorOp.cpp b/Engine/Plugins/Animation/IKRig/Source/IKRig/Private/Retargeter/RetargetOps/RootMotionGeneratorOp.cpp
--- a/Engine/Plugins/Animation/IKRig/Source/IKRig/Private/Retargeter/RetargetOps/RootMotionGeneratorOp.cpp
+++ b/Engine/Plugins/Animation/IKRig/Source/IKRig/Private/Retargeter/RetargetOps/RootMotionGeneratorOp.cpp
@@ -56,6 +56,21 @@ USkeleton* FIKRetargetRootMotionOpSettings::GetSkeleton(const FName InPropertyNa
 }
 #endif
 
+void FIKRetargetRootMotionOpSettings::PostSerialize(const FArchive& Ar)
+{
+	if (Ar.IsLoading())
+	{
+		if (Ar.CustomVer(FIKRigObjectVersion::GUID) < FIKRigObjectVersion::ModularRetargeterOps)
+		{
+			PRAGMA_DISABLE_DEPRECATION_WARNINGS
+			SourceRoot.BoneName = SourceRootBone_DEPRECATED;
+			TargetRoot.BoneName = TargetRootBone_DEPRECATED;
+			TargetPelvis.BoneName = TargetPelvisBone_DEPRECATED;
+			PRAGMA_ENABLE_DEPRECATION_WARNINGS
+		}
+	}
+}
+
 bool FIKRetargetRootMotionOp::Initialize(
 	const FIKRetargetProcessor& Processor,
 	const FRetargetSkeleton& SourceSkeleton,
diff --git a/Engine/Plugins/Animation/IKRig/Source/IKRig/Public/Retargeter/RetargetOps/CopyBasePoseOp.h b/Engine/Plugins/Animation/IKRig/Source/IKRig/Public/Retargeter/RetargetOps/CopyBasePoseOp.h
--- a/Engine/Plugins/Animation/IKRig/Source/IKRig/Public/Retargeter/RetargetOps/CopyBasePoseOp.h
+++ b/Engine/Plugins/Animation/IKRig/Source/IKRig/Public/Retargeter/RetargetOps/CopyBasePoseOp.h
@@ -35,6 +35,18 @@ struct FIKRetargetCopyBasePoseOpSettings : public FIKRetargetOpSettingsBase
 	UE_API virtual const UClass* GetControllerType() const override;
 
 	UE_API virtual void CopySettingsAtRuntime(const FIKRetargetOpSettingsBase* InSettingsToCopyFrom) override;
+
+	/** Custom serialization to handle conversion from deprecated CopyBasePoseRoot to CopyFromStart */
+	void PostSerialize(const FArchive& Ar);
+};
+
+// Custom serialization to load deprecated settings
+template<> struct TStructOpsTypeTraits<FIKRetargetCopyBasePoseOpSettings> : public TStructOpsTypeTraitsBase2<FIKRetargetCopyBasePoseOpSettings>
+{
+	enum
+	{
+		WithPostSerialize = true,
+	};
 };
 
 USTRUCT(BlueprintType, meta = (DisplayName = "Copy Base Pose"))
diff --git a/Engine/Plugins/Animation/IKRig/Source/IKRig/Public/Retargeter/RetargetOps/IKChainsOp.h b/Engine/Plugins/Animation/IKRig/Source/IKRig/Public/Retargeter/RetargetOps/IKChainsOp.h
--- a/Engine/Plugins/Animation/IKRig/Source/IKRig/Public/Retargeter/RetargetOps/IKChainsOp.h
+++ b/Engine/Plugins/Animation/IKRig/Source/IKRig/Public/Retargeter/RetargetOps/IKChainsOp.h
@@ -134,6 +134,18 @@ struct FRetargetIKChainSettings
 	double Extension = 1.0f;
 
 	UE_API bool operator==(const FRetargetIKChainSettings& Other) const;
+
+	/** Custom serialization to handle loading old assets */
+	void PostSerialize(const FArchive& Ar);
+};
+
+// Custom serialization to patch old settings
+template<> struct TStructOpsTypeTraits<FRetargetIKChainSettings> : public TStructOpsTypeTraitsBase2<FRetargetIKChainSettings>
+{
+	enum
+	{
+		WithPostSerialize = true,
+	};
 };
 
 USTRUCT(BlueprintType, meta = (DisplayName = "Retarget IK Goals Settings"))
diff --git a/Engine/Plugins/Animation/IKRig/Source/IKRig/Public/Retargeter/RetargetOps/PinBoneOp.h b/Engine/Plugins/Animation/IKRig/Source/IKRig/Public/Retargeter/RetargetOps/PinBoneOp.h
--- a/Engine/Plugins/Animation/IKRig/Source/IKRig/Public/Retargeter/RetargetOps/PinBoneOp.h
+++ b/Engine/Plugins/Animation/IKRig/Source/IKRig/Public/Retargeter/RetargetOps/PinBoneOp.h
@@ -49,6 +49,17 @@ struct FPinBoneData
 	FName BoneToPin_DEPRECATED;
 	UPROPERTY(meta=(DeprecatedProperty))
 	FName BoneToPinTo_DEPRECATED;
+
+	void PostSerialize(const FArchive& Ar);
+};
+
+template<>
+struct TStructOpsTypeTraits<FPinBoneData> : public TStructOpsTypeTraitsBase2<FPinBoneData>
+{
+	enum
+	{
+		WithPostSerialize = true,
+	};
 };
 
 UENUM()
diff --git a/Engine/Plugins/Animation/IKRig/Source/IKRig/Public/Retargeter/RetargetOps/RootMotionGeneratorOp.h b/Engine/Plugins/Animation/IKRig/Source/IKRig/Public/Retargeter/RetargetOps/RootMotionGeneratorOp.h
--- a/Engine/Plugins/Animation/IKRig/Source/IKRig/Public/Retargeter/RetargetOps/RootMotionGeneratorOp.h
+++ b/Engine/Plugins/Animation/IKRig/Source/IKRig/Public/Retargeter/RetargetOps/RootMotionGeneratorOp.h
@@ -84,6 +84,9 @@ struct FIKRetargetRootMotionOpSettings : public FIKRetargetOpSettingsBase
 #if WITH_EDITOR
 	UE_API virtual USkeleton* GetSkeleton(const FName InPropertyName) override;
 #endif
+
+	// patch to load old bone names into new FBoneReference
+	UE_API void PostSerialize(const FArchive& Ar);
 	
 	// Changed to FBoneReference to support bone selector UI
 	UPROPERTY(meta=(DeprecatedProperty))
@@ -94,6 +97,15 @@ struct FIKRetargetRootMotionOpSettings : public FIKRetargetOpSettingsBase
 	FName TargetPelvisBone_DEPRECATED;
 };
 
+template<>
+struct TStructOpsTypeTraits<FIKRetargetRootMotionOpSettings> : public TStructOpsTypeTraitsBase2<FIKRetargetRootMotionOpSettings>
+{
+	enum
+	{
+		WithPostSerialize = true,
+	};
+};
+
 USTRUCT(BlueprintType, meta = (DisplayName = "Root Motion"))
 struct FIKRetargetRootMotionOp : public FIKRetargetOpBase
 {
