[IK Retarget] Fixes height offset in Floor Constraint when used on non-feet goals

diff --git a/Engine/Plugins/Animation/IKRig/Source/IKRig/Private/Retargeter/RetargetOps/FloorConstraintOp.cpp b/Engine/Plugins/Animation/IKRig/Source/IKRig/Private/Retargeter/RetargetOps/FloorConstraintOp.cpp
--- a/Engine/Plugins/Animation/IKRig/Source/IKRig/Private/Retargeter/RetargetOps/FloorConstraintOp.cpp
+++ b/Engine/Plugins/Animation/IKRig/Source/IKRig/Private/Retargeter/RetargetOps/FloorConstraintOp.cpp
@@ -14,7 +14,8 @@
 bool FFloorConstraintChainSettings::operator==(const FFloorConstraintChainSettings& Other) const
 {
 	return EnableFloorConstraint == Other.EnableFloorConstraint
-		&& FMath::IsNearlyEqualByULP(Alpha,Other.Alpha);
+		&& FMath::IsNearlyEqualByULP(Alpha,Other.Alpha)
+		&& FMath::IsNearlyEqualByULP(MaintainHeightOffset, Other.MaintainHeightOffset);
 }
 
 const UClass* FIKRetargetFloorConstraintOpSettings::GetControllerType() const
@@ -155,7 +156,8 @@ void FIKRetargetFloorConstraintOp::Run(
 		const double FalloffRange = FMath::Max(1.0,FalloffEndHeight - FalloffStartHeight);
 		const double Falloff = SourceHeight < FalloffStartHeight ? 1.0f : 1.0f - (SourceHeight - FalloffStartHeight) / FalloffRange;
 		const double Weight = Falloff * FloorConstraint.Settings->Alpha;
-		const double TargetHeight = SourceHeight  + FloorConstraint.HeightOffsetInRefPose;
+		const double HeightOffset = FMath::Lerp(0.0f, FloorConstraint.HeightOffsetInRefPose, FloorConstraint.Settings->MaintainHeightOffset);
+		const double TargetHeight = SourceHeight + HeightOffset;
 		
 		// lerp goal height from it's current height to the target height by the weight
 		Goal->Position.Z = FMath::Lerp(Goal->Position.Z, TargetHeight, Weight);
diff --git a/Engine/Plugins/Animation/IKRig/Source/IKRig/Public/Retargeter/RetargetOps/FloorConstraintOp.h b/Engine/Plugins/Animation/IKRig/Source/IKRig/Public/Retargeter/RetargetOps/FloorConstraintOp.h
--- a/Engine/Plugins/Animation/IKRig/Source/IKRig/Public/Retargeter/RetargetOps/FloorConstraintOp.h
+++ b/Engine/Plugins/Animation/IKRig/Source/IKRig/Public/Retargeter/RetargetOps/FloorConstraintOp.h
@@ -36,6 +36,10 @@ struct FFloorConstraintChainSettings
 	UPROPERTY(EditAnywhere, BlueprintReadWrite, Category = "Floor Constraint Chain Settings", meta = (ClampMin = "0.0", ClampMax = "1.0", UIMin = "0.0", UIMax = "1.0"))
 	double Alpha = 1.0f;
 
+	/** Range 0 to 1. Default is 0. Maintain the height different between the source and target from the retarget pose. */
+	UPROPERTY(EditAnywhere, BlueprintReadWrite, Category = "Floor Constraint Chain Settings", meta = (ClampMin = "0.0", ClampMax = "1.0", UIMin = "0.0", UIMax = "1.0"))
+	double MaintainHeightOffset = 0.0f;
+
 	UE_API bool operator==(const FFloorConstraintChainSettings& Other) const;
 
 	FName GetName() const { return TargetChainName; };
diff --git a/Engine/Plugins/Animation/IKRig/Source/IKRigEditor/Private/RetargetEditor/IKRetargetOpDetails.cpp b/Engine/Plugins/Animation/IKRig/Source/IKRigEditor/Private/RetargetEditor/IKRetargetOpDetails.cpp
--- a/Engine/Plugins/Animation/IKRig/Source/IKRigEditor/Private/RetargetEditor/IKRetargetOpDetails.cpp
+++ b/Engine/Plugins/Animation/IKRig/Source/IKRigEditor/Private/RetargetEditor/IKRetargetOpDetails.cpp
@@ -786,7 +786,7 @@ void FFloorConstraintOpCustomization::CustomizeChildren(
 			
 		FIKRigStructToView StructToView;
 		StructToView.Owner = AssetController->GetAsset();
-		StructToView.Type = FRetargetSpeedPlantingSettings::StaticStruct();
+		StructToView.Type = FFloorConstraintChainSettings::StaticStruct();
 		StructToView.MemoryProvider = MemoryProvider;
 		StructToView.UniqueName = OpName;
 		
