[IK Retarget] Minor check for out-of-bounds in RunIKRigOp (BugHawk)

diff --git a/Engine/Plugins/Animation/IKRig/Source/IKRig/Private/Retargeter/RetargetOps/RunIKRigOp.cpp b/Engine/Plugins/Animation/IKRig/Source/IKRig/Private/Retargeter/RetargetOps/RunIKRigOp.cpp
--- a/Engine/Plugins/Animation/IKRig/Source/IKRig/Private/Retargeter/RetargetOps/RunIKRigOp.cpp
+++ b/Engine/Plugins/Animation/IKRig/Source/IKRig/Private/Retargeter/RetargetOps/RunIKRigOp.cpp
@@ -349,7 +349,7 @@ void FIKRetargetRunIKRigOp::ResetGoalContainer(
 	for (FIKRigGoal& Goal : InOutGoalContainer.GetGoalArray())
 	{
 		const int32* BoneIndexPtr = GoalBoneIndicesMap.Find(Goal.BoneName);
-		if (BoneIndexPtr)
+		if (BoneIndexPtr && InTargetGlobalPose.IsValidIndex(*BoneIndexPtr))
 		{
 			Goal.PositionSpace = EIKRigGoalSpace::Component;
 			Goal.RotationSpace = EIKRigGoalSpace::Component;
