[IK Rig] Fixed warning spam when trying to add bone settings to a solver that doesn't support them.

diff --git a/Engine/Plugins/Animation/IKRig/Source/IKRigEditor/Private/RigEditor/IKRigController.cpp b/Engine/Plugins/Animation/IKRig/Source/IKRigEditor/Private/RigEditor/IKRigController.cpp
--- a/Engine/Plugins/Animation/IKRig/Source/IKRigEditor/Private/RigEditor/IKRigController.cpp
+++ b/Engine/Plugins/Animation/IKRig/Source/IKRigEditor/Private/RigEditor/IKRigController.cpp
@@ -56,9 +56,10 @@ bool UIKRigController::AddBoneSetting(const FName BoneName, int32 SolverIndex) c
 {
 	check(Asset->SolverStack.IsValidIndex(SolverIndex))
 
-	if (!CanAddBoneSetting(BoneName, SolverIndex))
+	FText ErrorMessage;
+	if (!CanAddBoneSetting(BoneName, SolverIndex, &ErrorMessage))
 	{
-		UE_LOG(LogIKRigEditor, Warning, TEXT("Bone settings refused by solver. See output for details."));
+		UE_LOG(LogIKRigEditor, Warning, TEXT("Bone settings refused by solver. Reason: %s"), *ErrorMessage.ToString());
 		return false; // prerequisites not met
 	}
 	
@@ -70,31 +71,54 @@ bool UIKRigController::AddBoneSetting(const FName BoneName, int32 SolverIndex) c
 	return true;
 }
 
-bool UIKRigController::CanAddBoneSetting(const FName BoneName, int32 SolverIndex) const
+bool UIKRigController::CanAddBoneSetting(const FName BoneName, int32 SolverIndex, FText* OutError) const
 {
 	const FIKRigSolverBase* Solver = GetSolverAtIndex(SolverIndex);
 	if (!Solver)
 	{
-		UE_LOG(LogIKRigEditor, Warning, TEXT("Solver does not exist at index: %d."), SolverIndex);
-		return false; // solver doesn't exist
+		if (OutError)
+		{
+			*OutError = FText::Format(
+				NSLOCTEXT("IKRig", "NoSolverAtIndex", "Solver does not exist at index: {0}."),
+				FText::AsNumber(SolverIndex));
+		}
+		return false;
 	}
 
-	if (Asset->Skeleton.GetBoneIndexFromName(BoneName) == INDEX_NONE)
+	if (!Asset || Asset->Skeleton.GetBoneIndexFromName(BoneName) == INDEX_NONE)
 	{
-		UE_LOG(LogIKRigEditor, Warning, TEXT("Bone does not exist, %s."), *BoneName.ToString());
-		return false; // bone doesn't exist
+		if (OutError)
+		{
+			*OutError = FText::Format(
+				NSLOCTEXT("IKRig", "BoneMissing", "Bone does not exist: {0}."),
+				FText::FromName(BoneName));
+		}
+		return false;
 	}
 
 	if (!Solver->UsesCustomBoneSettings())
 	{
-		UE_LOG(LogIKRigEditor, Warning, TEXT("Solver does not support bone settings."));
-		return false; // solver doesn't support per-bone settings
+		if (OutError)
+		{
+			*OutError = NSLOCTEXT("IKRig", "NoCustomBoneSettings", "Solver does not support bone settings.");
+		}
+		return false;
 	}
 
 	if (Solver->HasSettingsOnBone(BoneName))
 	{
-		UE_LOG(LogIKRigEditor, Warning, TEXT("Solver already has settings on bone, %s."), *BoneName.ToString());
-		return false; // solve already has settings on this bone
+		if (OutError)
+		{
+			*OutError = FText::Format(
+				NSLOCTEXT("IKRig", "AlreadyHasSettings", "Solver already has settings on bone: {0}."),
+				FText::FromName(BoneName));
+		}
+		return false;
+	}
+
+	if (OutError)
+	{
+		*OutError = FText::GetEmpty(); // no error
 	}
 	
 	return true;
diff --git a/Engine/Plugins/Animation/IKRig/Source/IKRigEditor/Public/RigEditor/IKRigController.h b/Engine/Plugins/Animation/IKRig/Source/IKRigEditor/Public/RigEditor/IKRigController.h
--- a/Engine/Plugins/Animation/IKRig/Source/IKRigEditor/Public/RigEditor/IKRigController.h
+++ b/Engine/Plugins/Animation/IKRig/Source/IKRigEditor/Public/RigEditor/IKRigController.h
@@ -291,7 +291,7 @@ public:
 	UE_API FIKRigBoneSettingsBase* GetSettingsForBone(const FName BoneName, int32 SolverIndex) const;
 
 	// Returns true if this Bone can have settings in the given Solver. 
-	UE_API bool CanAddBoneSetting(const FName BoneName, int32 SolverIndex) const;
+	UE_API bool CanAddBoneSetting(const FName BoneName, int32 SolverIndex, FText* OutError = nullptr) const;
 
 	// Returns true if settings for this Bone can be removed from the given Solver. 
 	UE_API bool CanRemoveBoneSetting(const FName BoneName, int32 SolverIndex) const;
