[IK Retarget] Fixed bug, Source Scale not resetting when deleting op.

diff --git a/Engine/Plugins/Animation/IKRig/Source/IKRig/Private/Retargeter/IKRetargetProcessor.cpp b/Engine/Plugins/Animation/IKRig/Source/IKRig/Private/Retargeter/IKRetargetProcessor.cpp
--- a/Engine/Plugins/Animation/IKRig/Source/IKRig/Private/Retargeter/IKRetargetProcessor.cpp
+++ b/Engine/Plugins/Animation/IKRig/Source/IKRig/Private/Retargeter/IKRetargetProcessor.cpp
@@ -1869,6 +1869,8 @@ void FIKRetargetProcessor::ApplySourceScaleToPose(TArray<FTransform>& InOutSourc
 	const FIKRetargetScaleSourceOp* ScaleSourceOp = GetFirstRetargetOpOfType<FIKRetargetScaleSourceOp>();
 	if (!ScaleSourceOp)
 	{
+		CurrentSourceScale.Factor = 1.0;
+		CurrentSourceScale.Pivot = FVector::ZeroVector;
 		return;
 	}
 	
