[IK Retarget] Fixes issue, retarget op settings structs not getting PostLoad() patch

diff --git a/Engine/Plugins/Animation/IKRig/Source/IKRig/Private/Retargeter/IKRetargetOps.cpp b/Engine/Plugins/Animation/IKRig/Source/IKRig/Private/Retargeter/IKRetargetOps.cpp
--- a/Engine/Plugins/Animation/IKRig/Source/IKRig/Private/Retargeter/IKRetargetOps.cpp
+++ b/Engine/Plugins/Animation/IKRig/Source/IKRig/Private/Retargeter/IKRetargetOps.cpp
@@ -30,6 +30,11 @@ UIKRetargetOpControllerBase* FIKRetargetOpSettingsBase::CreateControllerIfNeeded
 	return Controller.Get();
 }
 
+void FIKRetargetOpBase::PostLoad(const FIKRigObjectVersion::Type InVersion)
+{
+	GetSettings()->PostLoad(InVersion);
+}
+
 void FIKRetargetOpBase::SetSettings(const FIKRetargetOpSettingsBase* InSettings)
 {
 	GetSettings()->CopySettingsAtRuntime(InSettings);
diff --git a/Engine/Plugins/Animation/IKRig/Source/IKRig/Private/Retargeter/RetargetOps/CopyBasePoseOp.cpp b/Engine/Plugins/Animation/IKRig/Source/IKRig/Private/Retargeter/RetargetOps/CopyBasePoseOp.cpp
--- a/Engine/Plugins/Animation/IKRig/Source/IKRig/Private/Retargeter/RetargetOps/CopyBasePoseOp.cpp
+++ b/Engine/Plugins/Animation/IKRig/Source/IKRig/Private/Retargeter/RetargetOps/CopyBasePoseOp.cpp
@@ -26,19 +26,24 @@ void FIKRetargetCopyBasePoseOpSettings::CopySettingsAtRuntime(const FIKRetargetO
 		PropertiesToIgnore);
 }
 
-void FIKRetargetCopyBasePoseOpSettings::PostSerialize(const FArchive& Ar)
+void FIKRetargetCopyBasePoseOpSettings::PostLoad(const FIKRigObjectVersion::Type InVersion)
 {
-	if (Ar.IsLoading())
+	FIKRetargetOpSettingsBase::PostLoad(InVersion);
+
+	// custom serialization to handle conversion from deprecated CopyBasePoseRoot to CopyFromStart
+	if (InVersion < FIKRigObjectVersion::CopyBasePoseConvertedToBoneReference)
 	{
-		if (Ar.CustomVer(FIKRigObjectVersion::GUID) < FIKRigObjectVersion::CopyBasePoseConvertedToBoneReference)
-		{
-			PRAGMA_DISABLE_DEPRECATION_WARNINGS
-			CopyFromStart.BoneName = CopyBasePoseRoot_DEPRECATED;
-			PRAGMA_ENABLE_DEPRECATION_WARNINGS
-		}
+		PRAGMA_DISABLE_DEPRECATION_WARNINGS
+		CopyFromStart.BoneName = CopyBasePoseRoot_DEPRECATED;
+		PRAGMA_ENABLE_DEPRECATION_WARNINGS
 	}
 }
 
+bool FIKRetargetCopyBasePoseOpSettings::Serialize(FArchive& Ar)
+{
+	return SerializeOpWithVersion(Ar, *this);
+}
+
 bool FIKRetargetCopyBasePoseOp::Initialize(
 	const FIKRetargetProcessor& InProcessor,
 	const FRetargetSkeleton& InSourceSkeleton,
@@ -196,17 +201,6 @@ const UScriptStruct* FIKRetargetCopyBasePoseOp::GetType() const
 	return FIKRetargetCopyBasePoseOp::StaticStruct();
 }
 
-void FIKRetargetCopyBasePoseOp::PostLoad(const FIKRigObjectVersion::Type InVersion)
-{
-	// custom serialization to handle conversion from deprecated CopyBasePoseRoot to CopyFromStart
-	if (InVersion < FIKRigObjectVersion::CopyBasePoseConvertedToBoneReference)
-	{
-		PRAGMA_DISABLE_DEPRECATION_WARNINGS
-		Settings.CopyFromStart.BoneName = Settings.CopyBasePoseRoot_DEPRECATED;
-		PRAGMA_ENABLE_DEPRECATION_WARNINGS
-	}
-}
-
 FIKRetargetCopyBasePoseOpSettings UIKRetargetCopyBasePoseController::GetSettings()
 {
 	return *reinterpret_cast<FIKRetargetCopyBasePoseOpSettings*>(OpSettingsToControl);
diff --git a/Engine/Plugins/Animation/IKRig/Source/IKRig/Private/Retargeter/RetargetOps/FKChainsOp.cpp b/Engine/Plugins/Animation/IKRig/Source/IKRig/Private/Retargeter/RetargetOps/FKChainsOp.cpp
--- a/Engine/Plugins/Animation/IKRig/Source/IKRig/Private/Retargeter/RetargetOps/FKChainsOp.cpp
+++ b/Engine/Plugins/Animation/IKRig/Source/IKRig/Private/Retargeter/RetargetOps/FKChainsOp.cpp
@@ -1078,6 +1078,8 @@ const UScriptStruct* FIKRetargetFKChainsOp::GetType() const
 
 void FIKRetargetFKChainsOp::PostLoad(const FIKRigObjectVersion::Type InVersion)
 {
+	Super::PostLoad(InVersion);
+	
 	if (InVersion >= FIKRigObjectVersion::OpsOwnChainMapping
 		&& InVersion < FIKRigObjectVersion::MovedChainMappingToFKChainOpSettings)
 	{
diff --git a/Engine/Plugins/Animation/IKRig/Source/IKRig/Private/Retargeter/RetargetOps/IKChainsOp.cpp b/Engine/Plugins/Animation/IKRig/Source/IKRig/Private/Retargeter/RetargetOps/IKChainsOp.cpp
--- a/Engine/Plugins/Animation/IKRig/Source/IKRig/Private/Retargeter/RetargetOps/IKChainsOp.cpp
+++ b/Engine/Plugins/Animation/IKRig/Source/IKRig/Private/Retargeter/RetargetOps/IKChainsOp.cpp
@@ -185,8 +185,11 @@ bool FRetargetIKChainSettings::operator==(const FRetargetIKChainSettings& Other)
 		&& FMath::IsNearlyEqualByULP(Extension, Other.Extension);
 }
 
-void FRetargetIKChainSettings::PostSerialize(const FArchive& Ar)
+bool FRetargetIKChainSettings::Serialize(FArchive& Ar)
 {
+	Ar.UsingCustomVersion(FIKRigObjectVersion::GUID);
+	StaticStruct()->SerializeTaggedProperties(Ar, reinterpret_cast<uint8*>(this), StaticStruct(), nullptr, nullptr);
+
 	if (Ar.IsLoading())
 	{
 		if (Ar.CustomVer(FIKRigObjectVersion::GUID) < FIKRigObjectVersion::RemovedPelvisMotionFromSourceGoals)
@@ -194,6 +197,8 @@ void FRetargetIKChainSettings::PostSerialize(const FArchive& Ar)
 			ApplyPelvisOffsetToSourceGoals = true;
 		}
 	}
+		
+	return true;
 }
 
 const UClass* FIKRetargetIKChainsOpSettings::GetControllerType() const
@@ -225,6 +230,24 @@ void FIKRetargetIKChainsOpSettings::CopySettingsAtRuntime(const FIKRetargetOpSet
 	}
 }
 
+void FIKRetargetIKChainsOpSettings::PostLoad(const FIKRigObjectVersion::Type InVersion)
+{
+	FIKRetargetOpSettingsBase::PostLoad(InVersion);
+	
+	if (InVersion < FIKRigObjectVersion::RemovedPelvisMotionFromSourceGoals)
+	{
+		for (FRetargetIKChainSettings& ChainSettings : ChainsToRetarget)
+		{
+			ChainSettings.ApplyPelvisOffsetToSourceGoals = true;
+		}
+	}
+}
+
+bool FIKRetargetIKChainsOpSettings::Serialize(FArchive& Ar)
+{
+	return SerializeOpWithVersion(Ar, *this);
+}
+
 bool FIKRetargetIKChainsOp::Initialize(
 	const FIKRetargetProcessor& InProcessor,
 	const FRetargetSkeleton& InSourceSkeleton,
@@ -448,17 +471,6 @@ void FIKRetargetIKChainsOp::OnParentReinitPropertyEdited(
 	RegenerateChainSettings(&InParentOp);
 }
 
-void FIKRetargetIKChainsOp::PostLoad(const FIKRigObjectVersion::Type InVersion)
-{
-	if (InVersion < FIKRigObjectVersion::RemovedPelvisMotionFromSourceGoals)
-	{
-		for (FRetargetIKChainSettings& ChainSettings : Settings.ChainsToRetarget)
-		{
-			ChainSettings.ApplyPelvisOffsetToSourceGoals = true;
-		}
-	}
-}
-
 #if WITH_EDITOR
 
 FCriticalSection FIKRetargetIKChainsOp::DebugDataMutex;
diff --git a/Engine/Plugins/Animation/IKRig/Source/IKRig/Private/Retargeter/RetargetOps/PelvisMotionOp.cpp b/Engine/Plugins/Animation/IKRig/Source/IKRig/Private/Retargeter/RetargetOps/PelvisMotionOp.cpp
--- a/Engine/Plugins/Animation/IKRig/Source/IKRig/Private/Retargeter/RetargetOps/PelvisMotionOp.cpp
+++ b/Engine/Plugins/Animation/IKRig/Source/IKRig/Private/Retargeter/RetargetOps/PelvisMotionOp.cpp
@@ -10,6 +10,8 @@
 #include "PrimitiveDrawingUtils.h"
 #endif
 
+#include "Retargeter/RetargetOps/RetargetPoseOp.h"
+
 #include UE_INLINE_GENERATED_CPP_BY_NAME(PelvisMotionOp)
 
 #define LOCTEXT_NAMESPACE "PelvisMotionOp"
@@ -49,16 +51,25 @@ USkeleton* FIKRetargetPelvisMotionOpSettings::GetSkeleton(const FName InProperty
 }
 #endif
 
-void FIKRetargetPelvisMotionOpSettings::PostSerialize(const FArchive& Ar)
+void FIKRetargetPelvisMotionOpSettings::PostLoad(const FIKRigObjectVersion::Type InVersion)
 {
-	if (Ar.IsLoading())
+	FIKRetargetOpSettingsBase::PostLoad(InVersion);
+
+	if (InVersion < FIKRigObjectVersion::AddedDebugDrawTogglePerOp)
 	{
-		if (Ar.CustomVer(FIKRigObjectVersion::GUID) < FIKRigObjectVersion::AddedLocalGlobalOffsetsToPelvisOp)
-		{
-			TranslationOffsetGlobal = TranslationOffset_DEPRECATED;
-			RotationOffsetLocal = RotationOffset_DEPRECATED;
-		}
+		bDebugDraw = bEnableDebugDraw_DEPRECATED;
 	}
+
+	if (InVersion < FIKRigObjectVersion::AddedLocalGlobalOffsetsToPelvisOp)
+	{
+		TranslationOffsetGlobal = TranslationOffset_DEPRECATED;
+		RotationOffsetLocal = RotationOffset_DEPRECATED;
+	}
+}
+
+bool FIKRetargetPelvisMotionOpSettings::Serialize(FArchive& Ar)
+{
+	return SerializeOpWithVersion(Ar, *this);
 }
 
 bool FIKRetargetPelvisMotionOp::Initialize(
@@ -264,20 +275,6 @@ FVector FIKRetargetPelvisMotionOp::GetPelvisTranslationOffset() const
 	return Target.PelvisTranslationDelta;
 }
 
-void FIKRetargetPelvisMotionOp::PostLoad(const FIKRigObjectVersion::Type InVersion)
-{
-	if (InVersion < FIKRigObjectVersion::AddedDebugDrawTogglePerOp)
-	{
-		Settings.bDebugDraw = Settings.bEnableDebugDraw_DEPRECATED;
-	}
-
-	if (InVersion < FIKRigObjectVersion::AddedLocalGlobalOffsetsToPelvisOp)
-	{
-		Settings.TranslationOffsetGlobal = Settings.TranslationOffset_DEPRECATED;
-		Settings.RotationOffsetLocal = Settings.RotationOffset_DEPRECATED;
-	}
-}
-
 int32 FIKRetargetPelvisMotionOp::GetPelvisBoneIndex(ERetargetSourceOrTarget SourceOrTarget) const
 {
 	return SourceOrTarget == ERetargetSourceOrTarget::Source ? Source.BoneIndex : Target.BoneIndex;
diff --git a/Engine/Plugins/Animation/IKRig/Source/IKRig/Private/Retargeter/RetargetOps/PinBoneOp.cpp b/Engine/Plugins/Animation/IKRig/Source/IKRig/Private/Retargeter/RetargetOps/PinBoneOp.cpp
--- a/Engine/Plugins/Animation/IKRig/Source/IKRig/Private/Retargeter/RetargetOps/PinBoneOp.cpp
+++ b/Engine/Plugins/Animation/IKRig/Source/IKRig/Private/Retargeter/RetargetOps/PinBoneOp.cpp
@@ -88,6 +88,27 @@ USkeleton* FIKRetargetPinBoneOpSettings::GetSkeleton(const FName InPropertyName)
 }
 #endif
 
+void FIKRetargetPinBoneOpSettings::PostLoad(const FIKRigObjectVersion::Type InVersion)
+{
+	FIKRetargetOpSettingsBase::PostLoad(InVersion);
+
+	if (InVersion < FIKRigObjectVersion::PinBoneTypeAndOffsetsUpgraded)
+	{
+		for (FPinBoneData& BoneToPin : BonesToPin)
+		{
+			PRAGMA_DISABLE_DEPRECATION_WARNINGS
+			BoneToPin.BoneToCopyTo.BoneName = BoneToPin.BoneToPin_DEPRECATED;
+			BoneToPin.BoneToCopyFrom.BoneName = BoneToPin.BoneToPinTo_DEPRECATED;
+			PRAGMA_ENABLE_DEPRECATION_WARNINGS
+		}
+	}
+}
+
+bool FIKRetargetPinBoneOpSettings::Serialize(FArchive& Ar)
+{
+	return SerializeOpWithVersion(Ar, *this);
+}
+
 bool FIKRetargetPinBoneOp::Initialize(
 	const FIKRetargetProcessor& Processor,
 	const FRetargetSkeleton& SourceSkeleton,
@@ -210,21 +231,6 @@ const UScriptStruct* FIKRetargetPinBoneOp::GetType() const
 	return FIKRetargetPinBoneOp::StaticStruct();
 }
 
-void FIKRetargetPinBoneOp::PostLoad(const FIKRigObjectVersion::Type InVersion)
-{
-	// patch old versions to use FBoneReference instead of FName
-	if (InVersion < FIKRigObjectVersion::PinBoneTypeAndOffsetsUpgraded)
-	{
-		PRAGMA_DISABLE_DEPRECATION_WARNINGS
-		for (FPinBoneData& PinBoneData : Settings.BonesToPin)
-		{
-			PinBoneData.BoneToCopyTo.BoneName = PinBoneData.BoneToPin_DEPRECATED;
-			PinBoneData.BoneToCopyFrom.BoneName = PinBoneData.BoneToPinTo_DEPRECATED;
-		}
-		PRAGMA_ENABLE_DEPRECATION_WARNINGS
-	}
-}
-
 #if WITH_EDITOR
 FText FIKRetargetPinBoneOp::GetWarningMessage() const
 {
diff --git a/Engine/Plugins/Animation/IKRig/Source/IKRig/Private/Retargeter/RetargetOps/RootMotionGeneratorOp.cpp b/Engine/Plugins/Animation/IKRig/Source/IKRig/Private/Retargeter/RetargetOps/RootMotionGeneratorOp.cpp
--- a/Engine/Plugins/Animation/IKRig/Source/IKRig/Private/Retargeter/RetargetOps/RootMotionGeneratorOp.cpp
+++ b/Engine/Plugins/Animation/IKRig/Source/IKRig/Private/Retargeter/RetargetOps/RootMotionGeneratorOp.cpp
@@ -56,21 +56,23 @@ USkeleton* FIKRetargetRootMotionOpSettings::GetSkeleton(const FName InPropertyNa
 }
 #endif
 
-void FIKRetargetRootMotionOpSettings::PostSerialize(const FArchive& Ar)
+void FIKRetargetRootMotionOpSettings::PostLoad(const FIKRigObjectVersion::Type InVersion)
 {
-	if (Ar.IsLoading())
+	if (InVersion < FIKRigObjectVersion::ModularRetargeterOps)
 	{
-		if (Ar.CustomVer(FIKRigObjectVersion::GUID) < FIKRigObjectVersion::ModularRetargeterOps)
-		{
-			PRAGMA_DISABLE_DEPRECATION_WARNINGS
-			SourceRoot.BoneName = SourceRootBone_DEPRECATED;
-			TargetRoot.BoneName = TargetRootBone_DEPRECATED;
-			TargetPelvis.BoneName = TargetPelvisBone_DEPRECATED;
-			PRAGMA_ENABLE_DEPRECATION_WARNINGS
-		}
+		PRAGMA_DISABLE_DEPRECATION_WARNINGS
+		SourceRoot.BoneName = SourceRootBone_DEPRECATED;
+		TargetRoot.BoneName = TargetRootBone_DEPRECATED;
+		TargetPelvis.BoneName = TargetPelvisBone_DEPRECATED;
+		PRAGMA_ENABLE_DEPRECATION_WARNINGS
 	}
 }
 
+bool FIKRetargetRootMotionOpSettings::Serialize(FArchive& Ar)
+{
+	return SerializeOpWithVersion(Ar, *this);
+}
+
 bool FIKRetargetRootMotionOp::Initialize(
 	const FIKRetargetProcessor& Processor,
 	const FRetargetSkeleton& SourceSkeleton,
@@ -227,19 +229,6 @@ const UScriptStruct* FIKRetargetRootMotionOp::GetType() const
 	return FIKRetargetRootMotionOp::StaticStruct();
 }
 
-void FIKRetargetRootMotionOp::PostLoad(const FIKRigObjectVersion::Type InVersion)
-{
-	// load old FNames into new FBoneReference
-	if (InVersion < FIKRigObjectVersion::ModularRetargeterOps)
-	{
-		PRAGMA_DISABLE_DEPRECATION_WARNINGS
-		Settings.SourceRoot.BoneName = Settings.SourceRootBone_DEPRECATED;
-		Settings.TargetRoot.BoneName = Settings.TargetRootBone_DEPRECATED;
-		Settings.TargetPelvis.BoneName = Settings.TargetPelvisBone_DEPRECATED;
-		PRAGMA_ENABLE_DEPRECATION_WARNINGS
-	}
-}
-
 void FIKRetargetRootMotionOp::Reset()
 {
 	SourceRootIndex = INDEX_NONE;
diff --git a/Engine/Plugins/Animation/IKRig/Source/IKRig/Private/Retargeter/RetargetOps/StrideWarpingOp.cpp b/Engine/Plugins/Animation/IKRig/Source/IKRig/Private/Retargeter/RetargetOps/StrideWarpingOp.cpp
--- a/Engine/Plugins/Animation/IKRig/Source/IKRig/Private/Retargeter/RetargetOps/StrideWarpingOp.cpp
+++ b/Engine/Plugins/Animation/IKRig/Source/IKRig/Private/Retargeter/RetargetOps/StrideWarpingOp.cpp
@@ -70,6 +70,21 @@ bool FIKRetargetStrideWarpingOpSettings::operator==(const FIKRetargetStrideWarpi
 		&& FMath::IsNearlyEqualByULP(WarpSplay, Other.WarpSplay);
 }
 
+void FIKRetargetStrideWarpingOpSettings::PostLoad(const FIKRigObjectVersion::Type InVersion)
+{
+	FIKRetargetOpSettingsBase::PostLoad(InVersion);
+	
+	if (InVersion < FIKRigObjectVersion::AddedDebugDrawTogglePerOp)
+	{
+		bDebugDraw = bEnableDebugDraw_DEPRECATED;
+	}
+}
+
+bool FIKRetargetStrideWarpingOpSettings::Serialize(FArchive& Ar)
+{
+	return SerializeOpWithVersion(Ar, *this);
+}
+
 bool FIKRetargetStrideWarpingOp::Initialize(
 	const FIKRetargetProcessor& InProcessor,
 	const FRetargetSkeleton& InSourceSkeleton,
@@ -393,14 +408,6 @@ void FIKRetargetStrideWarpingOp::OnParentReinitPropertyEdited(
 	RegenerateChainSettings(&InParentOp);
 }
 
-void FIKRetargetStrideWarpingOp::PostLoad(const FIKRigObjectVersion::Type InVersion)
-{
-	if (InVersion < FIKRigObjectVersion::AddedDebugDrawTogglePerOp)
-	{
-		Settings.bDebugDraw = Settings.bEnableDebugDraw_DEPRECATED;
-	}
-}
-
 #if WITH_EDITOR
 
 FCriticalSection FIKRetargetStrideWarpingOp::DebugDataMutex;
diff --git a/Engine/Plugins/Animation/IKRig/Source/IKRig/Public/Retargeter/IKRetargetOps.h b/Engine/Plugins/Animation/IKRig/Source/IKRig/Public/Retargeter/IKRetargetOps.h
--- a/Engine/Plugins/Animation/IKRig/Source/IKRig/Public/Retargeter/IKRetargetOps.h
+++ b/Engine/Plugins/Animation/IKRig/Source/IKRig/Public/Retargeter/IKRetargetOps.h
@@ -57,6 +57,43 @@ struct FIKRetargetOpSettingsBase
 
 	/** (required) override to specify how settings should be applied in a way that will not require reinitialization (ie runtime compatible)*/
 	virtual void CopySettingsAtRuntime(const FIKRetargetOpSettingsBase* InSettingsToCopyFrom) { checkNoEntry(); };
+
+	// NOTE ON RETARGET OP POSTLOAD() / SERIALIZE() BEHAVIOR:
+	// If custom op settings require load-patching for backwards compatibility, it is NOT enough to rely on the archive
+	// being tagged with custom versioning from the UIKRetargeter itself. This is because these settings structs can be stored
+	// OUTSIDE of the retargeter (ie, in BP classes), which means they must take full control of their own back compatibility.
+	//
+	// BUT, because the settings are also stored as FInstancedStruct inside the retarget asset, and because these instanced structs
+	// do not run custom serialization, the patch must ALSO be applied explicitly from within FIKRetargetOpBase::PostLoad().
+	//
+	// This leads us to this recommended pattern for op settings that need serialization patches:
+	// 1. Override virtual FIKRetargetOpSettingsBase::PostLoad() in your custom settings struct.
+	//		This will be called by the retargeter to apply patching for any op settings stored in the asset.
+	// 2. Add the type trait WithSerializer and implement bool MyOp::Serialize().
+	// 3. Call SerializeOpWithVersion() within the custom MyOp::Serialize()
+	//		This will ensure the settings are saved with a custom version and it will call OpSettings::PostLoad() after loading
+	//
+	// Following this pattern ensures that op settings are properly patched both when they are stored in a retarget asset
+	// as an instanced struct AND if they are stored alone as a BP variable.
+	//
+	template <typename T>
+	static bool SerializeOpWithVersion(FArchive& Ar, T& Self)
+	{
+		Ar.UsingCustomVersion(FIKRigObjectVersion::GUID);
+		T::StaticStruct()->SerializeTaggedProperties(Ar, reinterpret_cast<uint8*>(&Self), T::StaticStruct(), nullptr, nullptr);
+
+		if (Ar.IsLoading())
+		{
+			const FIKRigObjectVersion::Type CustomVersion = static_cast<FIKRigObjectVersion::Type>(Ar.CustomVer(FIKRigObjectVersion::GUID));
+			Self.PostLoad(CustomVersion);
+		}
+		
+		return true;
+	}
+
+	/** (optional) called during the PostLoad of the owning asset. Can be used to apply patches to the data for backwards compatibility.
+	 * NOTE: see description above on correct setup for retarget op settings serialization patches */
+	virtual void PostLoad(const FIKRigObjectVersion::Type InVersion) {};
 	
 	/** (optional, but recommended) provide a custom controller type (deriving from UIKRetargetOpControllerBase) as an API for editing the op
 	 * NOTE: this type will automatically be lazy instantiated when/if the user calls GetController() on the op. */
@@ -129,7 +166,7 @@ public:
 		FIKRigLogger& InLog) {};
 
 	/** (optional) called during the PostLoad of the owning asset. Can be used to apply patches to the data for backwards compatibility. */
-	virtual void PostLoad(const FIKRigObjectVersion::Type InVersion) {};
+	virtual void PostLoad(const FIKRigObjectVersion::Type InVersion);
 
 	using FFrameValues = TArray<TArray<TOptional<float>>>;
 
diff --git a/Engine/Plugins/Animation/IKRig/Source/IKRig/Public/Retargeter/RetargetOps/CopyBasePoseOp.h b/Engine/Plugins/Animation/IKRig/Source/IKRig/Public/Retargeter/RetargetOps/CopyBasePoseOp.h
--- a/Engine/Plugins/Animation/IKRig/Source/IKRig/Public/Retargeter/RetargetOps/CopyBasePoseOp.h
+++ b/Engine/Plugins/Animation/IKRig/Source/IKRig/Public/Retargeter/RetargetOps/CopyBasePoseOp.h
@@ -37,7 +37,8 @@ struct FIKRetargetCopyBasePoseOpSettings : public FIKRetargetOpSettingsBase
 	UE_API virtual void CopySettingsAtRuntime(const FIKRetargetOpSettingsBase* InSettingsToCopyFrom) override;
 
 	/** Custom serialization to handle conversion from deprecated CopyBasePoseRoot to CopyFromStart */
-	void PostSerialize(const FArchive& Ar);
+	virtual void PostLoad(const FIKRigObjectVersion::Type InVersion) override;
+	bool Serialize(FArchive& Ar);
 };
 
 // Custom serialization to load deprecated settings
@@ -45,7 +46,7 @@ template<> struct TStructOpsTypeTraits<FIKRetargetCopyBasePoseOpSettings> : publ
 {
 	enum
 	{
-		WithPostSerialize = true,
+		WithSerializer = true,
 	};
 };
 
@@ -77,9 +78,7 @@ struct FIKRetargetCopyBasePoseOp : public FIKRetargetOpBase
 	UE_API virtual const UScriptStruct* GetType() const override;
 
 	virtual bool IsSingleton() const override { return true; };
-
-	UE_API virtual void PostLoad(const FIKRigObjectVersion::Type InVersion) override;
-
+	
 	UPROPERTY()
 	FIKRetargetCopyBasePoseOpSettings Settings;
 
diff --git a/Engine/Plugins/Animation/IKRig/Source/IKRig/Public/Retargeter/RetargetOps/IKChainsOp.h b/Engine/Plugins/Animation/IKRig/Source/IKRig/Public/Retargeter/RetargetOps/IKChainsOp.h
--- a/Engine/Plugins/Animation/IKRig/Source/IKRig/Public/Retargeter/RetargetOps/IKChainsOp.h
+++ b/Engine/Plugins/Animation/IKRig/Source/IKRig/Public/Retargeter/RetargetOps/IKChainsOp.h
@@ -136,16 +136,13 @@ struct FRetargetIKChainSettings
 	UE_API bool operator==(const FRetargetIKChainSettings& Other) const;
 
 	/** Custom serialization to handle loading old assets */
-	void PostSerialize(const FArchive& Ar);
+	bool Serialize(FArchive& Ar);
 };
 
-// Custom serialization to patch old settings
-template<> struct TStructOpsTypeTraits<FRetargetIKChainSettings> : public TStructOpsTypeTraitsBase2<FRetargetIKChainSettings>
+template<>
+struct TStructOpsTypeTraits<FRetargetIKChainSettings> : public TStructOpsTypeTraitsBase2<FRetargetIKChainSettings>
 {
-	enum
-	{
-		WithPostSerialize = true,
-	};
+	enum { WithSerializer = true };
 };
 
 USTRUCT(BlueprintType, meta = (DisplayName = "Retarget IK Goals Settings"))
@@ -175,6 +172,15 @@ struct FIKRetargetIKChainsOpSettings : public FIKRetargetOpSettingsBase
 	UE_API virtual const UClass* GetControllerType() const override;
 
 	UE_API virtual void CopySettingsAtRuntime(const FIKRetargetOpSettingsBase* InSettingsToCopyFrom) override;
+
+	virtual void PostLoad(const FIKRigObjectVersion::Type InVersion) override;
+	bool Serialize(FArchive& Ar);
+};
+
+template<>
+struct TStructOpsTypeTraits<FIKRetargetIKChainsOpSettings> : public TStructOpsTypeTraitsBase2<FIKRetargetIKChainsOpSettings>
+{
+	enum { WithSerializer = true };
 };
 
 #if WITH_EDITOR
@@ -223,9 +229,7 @@ struct FIKRetargetIKChainsOp : public FIKRetargetOpBase
 	UE_API virtual void OnTargetChainRenamed(const FName InOldChainName, const FName InNewChainName) override;
 
 	UE_API virtual void OnParentReinitPropertyEdited(const FIKRetargetOpBase& InParentOp, const FPropertyChangedEvent* InPropertyChangedEvent) override;
-
-	UE_API virtual void PostLoad(const FIKRigObjectVersion::Type InVersion) override;
-
+	
 	UPROPERTY()
 	FIKRetargetIKChainsOpSettings Settings;
 
diff --git a/Engine/Plugins/Animation/IKRig/Source/IKRig/Public/Retargeter/RetargetOps/PelvisMotionOp.h b/Engine/Plugins/Animation/IKRig/Source/IKRig/Public/Retargeter/RetargetOps/PelvisMotionOp.h
--- a/Engine/Plugins/Animation/IKRig/Source/IKRig/Public/Retargeter/RetargetOps/PelvisMotionOp.h
+++ b/Engine/Plugins/Animation/IKRig/Source/IKRig/Public/Retargeter/RetargetOps/PelvisMotionOp.h
@@ -127,13 +127,14 @@ struct FIKRetargetPelvisMotionOpSettings : public FIKRetargetOpSettingsBase
 	UPROPERTY(meta = (DeprecatedProperty))
 	bool bEnableDebugDraw_DEPRECATED = true;
 
-	void PostSerialize(const FArchive& Ar);
+	virtual void PostLoad(const FIKRigObjectVersion::Type InVersion) override;
+	bool Serialize(FArchive& Ar);
 };
 
 template<>
 struct TStructOpsTypeTraits<FIKRetargetPelvisMotionOpSettings> : public TStructOpsTypeTraitsBase2<FIKRetargetPelvisMotionOpSettings>
 {
-	enum { WithPostSerialize = true };
+	enum { WithSerializer = true };
 };
 
 struct FPelvisSource
@@ -198,8 +199,6 @@ struct FIKRetargetPelvisMotionOp : public FIKRetargetOpBase
 
 	IKRIG_API FVector GetPelvisTranslationOffset() const;
 
-	IKRIG_API virtual void PostLoad(const FIKRigObjectVersion::Type InVersion) override;
-
 	int32 GetPelvisBoneIndex(ERetargetSourceOrTarget SourceOrTarget) const;
 
 	UPROPERTY()
diff --git a/Engine/Plugins/Animation/IKRig/Source/IKRig/Public/Retargeter/RetargetOps/PinBoneOp.h b/Engine/Plugins/Animation/IKRig/Source/IKRig/Public/Retargeter/RetargetOps/PinBoneOp.h
--- a/Engine/Plugins/Animation/IKRig/Source/IKRig/Public/Retargeter/RetargetOps/PinBoneOp.h
+++ b/Engine/Plugins/Animation/IKRig/Source/IKRig/Public/Retargeter/RetargetOps/PinBoneOp.h
@@ -138,6 +138,15 @@ struct FIKRetargetPinBoneOpSettings : public FIKRetargetOpSettingsBase
 #if WITH_EDITOR
 	UE_API virtual USkeleton* GetSkeleton(const FName InPropertyName) override;
 #endif
+
+	virtual void PostLoad(const FIKRigObjectVersion::Type InVersion) override;
+	bool Serialize(FArchive& Ar);
+};
+
+template<>
+struct TStructOpsTypeTraits<FIKRetargetPinBoneOpSettings> : public TStructOpsTypeTraitsBase2<FIKRetargetPinBoneOpSettings>
+{
+	enum { WithSerializer = true };
 };
 
 USTRUCT(BlueprintType, meta = (DisplayName = "Pin Bones"))
@@ -163,9 +172,7 @@ struct FIKRetargetPinBoneOp : public FIKRetargetOpBase
 	UE_API virtual const UScriptStruct* GetSettingsType() const override;
 	
 	UE_API virtual const UScriptStruct* GetType() const override;
-
-	UE_API virtual void PostLoad(const FIKRigObjectVersion::Type InVersion) override;
-
+	
 	UPROPERTY()
 	FIKRetargetPinBoneOpSettings Settings;
 	
diff --git a/Engine/Plugins/Animation/IKRig/Source/IKRig/Public/Retargeter/RetargetOps/RootMotionGeneratorOp.h b/Engine/Plugins/Animation/IKRig/Source/IKRig/Public/Retargeter/RetargetOps/RootMotionGeneratorOp.h
--- a/Engine/Plugins/Animation/IKRig/Source/IKRig/Public/Retargeter/RetargetOps/RootMotionGeneratorOp.h
+++ b/Engine/Plugins/Animation/IKRig/Source/IKRig/Public/Retargeter/RetargetOps/RootMotionGeneratorOp.h
@@ -84,9 +84,6 @@ struct FIKRetargetRootMotionOpSettings : public FIKRetargetOpSettingsBase
 #if WITH_EDITOR
 	UE_API virtual USkeleton* GetSkeleton(const FName InPropertyName) override;
 #endif
-
-	// patch to load old bone names into new FBoneReference
-	UE_API void PostSerialize(const FArchive& Ar);
 	
 	// Changed to FBoneReference to support bone selector UI
 	UPROPERTY(meta=(DeprecatedProperty))
@@ -95,15 +92,15 @@ struct FIKRetargetRootMotionOpSettings : public FIKRetargetOpSettingsBase
 	FName TargetRootBone_DEPRECATED;
 	UPROPERTY(meta=(DeprecatedProperty))
 	FName TargetPelvisBone_DEPRECATED;
+
+	virtual void PostLoad(const FIKRigObjectVersion::Type InVersion) override;
+	bool Serialize(FArchive& Ar);
 };
 
 template<>
 struct TStructOpsTypeTraits<FIKRetargetRootMotionOpSettings> : public TStructOpsTypeTraitsBase2<FIKRetargetRootMotionOpSettings>
 {
-	enum
-	{
-		WithPostSerialize = true,
-	};
+	enum { WithSerializer = true };
 };
 
 USTRUCT(BlueprintType, meta = (DisplayName = "Root Motion"))
@@ -137,9 +134,7 @@ struct FIKRetargetRootMotionOp : public FIKRetargetOpBase
 	UE_API virtual const UScriptStruct* GetSettingsType() const override;
 	
 	UE_API virtual const UScriptStruct* GetType() const override;
-
-	UE_API virtual void PostLoad(const FIKRigObjectVersion::Type InVersion) override;
-
+	
 	UPROPERTY()
 	FIKRetargetRootMotionOpSettings Settings;
 
diff --git a/Engine/Plugins/Animation/IKRig/Source/IKRig/Public/Retargeter/RetargetOps/StrideWarpingOp.h b/Engine/Plugins/Animation/IKRig/Source/IKRig/Public/Retargeter/RetargetOps/StrideWarpingOp.h
--- a/Engine/Plugins/Animation/IKRig/Source/IKRig/Public/Retargeter/RetargetOps/StrideWarpingOp.h
+++ b/Engine/Plugins/Animation/IKRig/Source/IKRig/Public/Retargeter/RetargetOps/StrideWarpingOp.h
@@ -101,6 +101,15 @@ struct FIKRetargetStrideWarpingOpSettings : public FIKRetargetOpSettingsBase
 
 	UPROPERTY(meta = (DeprecatedProperty))
 	bool bEnableDebugDraw_DEPRECATED = true;
+
+	virtual void PostLoad(const FIKRigObjectVersion::Type InVersion) override;
+	bool Serialize(FArchive& Ar);
+};
+
+template<>
+struct TStructOpsTypeTraits<FIKRetargetStrideWarpingOpSettings> : public TStructOpsTypeTraitsBase2<FIKRetargetStrideWarpingOpSettings>
+{
+	enum { WithSerializer = true };
 };
 
 struct FStrideWarpGoalData
@@ -146,9 +155,7 @@ struct FIKRetargetStrideWarpingOp : public FIKRetargetOpBase
 	UE_API virtual void OnTargetChainRenamed(const FName InOldChainName, const FName InNewChainName) override;
 
 	UE_API virtual void OnParentReinitPropertyEdited(const FIKRetargetOpBase& InParentOp, const FPropertyChangedEvent* InPropertyChangedEvent) override;
-
-	IKRIG_API virtual void PostLoad(const FIKRigObjectVersion::Type InVersion) override;
-
+	
 	UPROPERTY()
 	FIKRetargetStrideWarpingOpSettings Settings;
 
