[IK Retarget] Removed unnecessary 3-bone limit on IK Chains.

diff --git a/Engine/Plugins/Animation/IKRig/Source/IKRig/Private/Retargeter/RetargetOps/IKChainsOp.cpp b/Engine/Plugins/Animation/IKRig/Source/IKRig/Private/Retargeter/RetargetOps/IKChainsOp.cpp
--- a/Engine/Plugins/Animation/IKRig/Source/IKRig/Private/Retargeter/RetargetOps/IKChainsOp.cpp
+++ b/Engine/Plugins/Animation/IKRig/Source/IKRig/Private/Retargeter/RetargetOps/IKChainsOp.cpp
@@ -31,12 +31,6 @@ bool FIKChainRetargeter::Initialize(
 		return false;
 	}
 
-	if (InTargetBoneChain.BoneIndices.Num() < 3)
-	{
-		InLog.LogWarning(LOCTEXT("TargetChainLessThanThree", "IK Chains Op: trying to retarget target bone chain with IK but it has less than 3 joints."));
-		return false;
-	}
-
 	// store pointers to required data
 	SourceBoneChain = &InSourceBoneChain;
 	TargetBoneChain = &InTargetBoneChain;
