[IK Retarget] Removed addition of Default Ops when making a retargeter from an IK Rig.

diff --git a/Engine/Plugins/Animation/IKRig/Source/IKRigEditor/Private/RetargetEditor/AssetTypeActions_IKRetargeter.cpp b/Engine/Plugins/Animation/IKRig/Source/IKRigEditor/Private/RetargetEditor/AssetTypeActions_IKRetargeter.cpp
--- a/Engine/Plugins/Animation/IKRig/Source/IKRigEditor/Private/RetargetEditor/AssetTypeActions_IKRetargeter.cpp
+++ b/Engine/Plugins/Animation/IKRig/Source/IKRigEditor/Private/RetargetEditor/AssetTypeActions_IKRetargeter.cpp
@@ -132,9 +132,6 @@ void FAssetTypeActions_IKRetargeter::CreateNewIKRetargeterFromIKRig(UObject* InS
 	const UIKRetargeterController* Controller = UIKRetargeterController::GetController(RetargetAsset);
 	Controller->SetIKRig(ERetargetSourceOrTarget::Source, IKRigAsset);
 	Controller->SetIKRig(ERetargetSourceOrTarget::Target, IKRigAsset);
-
-	// add default ops
-	Controller->AddDefaultOps();
 }
 
 void FAssetTypeActions_IKRetargeter::ExtendAnimAssetMenusForBatchRetargeting()
