[IK Rig] Added support for multi-selection editing

#FailedValidations Preflight

diff --git a/Engine/Plugins/Animation/IKRig/Source/IKRig/Public/Rig/Solvers/IKRigBodyMoverSolver.h b/Engine/Plugins/Animation/IKRig/Source/IKRig/Public/Rig/Solvers/IKRigBodyMoverSolver.h
--- a/Engine/Plugins/Animation/IKRig/Source/IKRig/Public/Rig/Solvers/IKRigBodyMoverSolver.h
+++ b/Engine/Plugins/Animation/IKRig/Source/IKRig/Public/Rig/Solvers/IKRigBodyMoverSolver.h
@@ -9,7 +9,7 @@
 
 #define UE_API IKRIG_API
 
-USTRUCT(BlueprintType)
+USTRUCT(BlueprintType, meta=(UIWrapper="/Script/IKRigEditor.BodyMoverGoalSettingsWrapper"))
 struct FIKRigBodyMoverGoalSettings : public FIKRigGoalSettingsBase
 {
 	GENERATED_BODY()
@@ -23,7 +23,7 @@ struct FIKRigBodyMoverGoalSettings : public FIKRigGoalSettingsBase
 	float InfluenceMultiplier = 1.0f;
 };
 
-USTRUCT(BlueprintType)
+USTRUCT(BlueprintType, meta=(UIWrapper="/Script/IKRigEditor.BodyMoverSettingsWrapper"))
 struct FIKRigBodyMoverSettings : public FIKRigSolverSettingsBase
 {
 	GENERATED_BODY()
diff --git a/Engine/Plugins/Animation/IKRig/Source/IKRig/Public/Rig/Solvers/IKRigFullBodyIK.h b/Engine/Plugins/Animation/IKRig/Source/IKRig/Public/Rig/Solvers/IKRigFullBodyIK.h
--- a/Engine/Plugins/Animation/IKRig/Source/IKRig/Public/Rig/Solvers/IKRigFullBodyIK.h
+++ b/Engine/Plugins/Animation/IKRig/Source/IKRig/Public/Rig/Solvers/IKRigFullBodyIK.h
@@ -11,7 +11,7 @@
 
 #define UE_API IKRIG_API
 
-USTRUCT(BlueprintType)
+USTRUCT(BlueprintType, meta=(UIWrapper="/Script/IKRigEditor.FBIKGoalSettingsWrapper"))
 struct FIKRigFBIKGoalSettings : public FIKRigGoalSettingsBase
 {
 	GENERATED_BODY()
@@ -55,7 +55,7 @@ struct FIKRigFBIKGoalSettings : public FIKRigGoalSettingsBase
 	int32 IndexInSolver = -1;
 };
 
-USTRUCT(BlueprintType)
+USTRUCT(BlueprintType, meta=(UIWrapper="/Script/IKRigEditor.FBIKBoneSettingsWrapper"))
 struct FIKRigFBIKBoneSettings : public FIKRigBoneSettingsBase
 {
 	GENERATED_BODY()
@@ -141,7 +141,7 @@ struct FIKRigFBIKBoneSettings : public FIKRigBoneSettingsBase
 	}
 };
 
-USTRUCT(BlueprintType)
+USTRUCT(BlueprintType, meta=(UIWrapper="/Script/IKRigEditor.FBIKSettingsWrapper"))
 struct FIKRigFBIKSettings : public FIKRigSolverSettingsBase
 {
 	GENERATED_BODY()
diff --git a/Engine/Plugins/Animation/IKRig/Source/IKRig/Public/Rig/Solvers/IKRigLimbSolver.h b/Engine/Plugins/Animation/IKRig/Source/IKRig/Public/Rig/Solvers/IKRigLimbSolver.h
--- a/Engine/Plugins/Animation/IKRig/Source/IKRig/Public/Rig/Solvers/IKRigLimbSolver.h
+++ b/Engine/Plugins/Animation/IKRig/Source/IKRig/Public/Rig/Solvers/IKRigLimbSolver.h
@@ -12,7 +12,7 @@
 #define UE_API IKRIG_API
 
 
-USTRUCT(BlueprintType)
+USTRUCT(BlueprintType, meta=(UIWrapper="/Script/IKRigEditor.LimbSolverSettingsWrapper"))
 struct FIKRigLimbSolverSettings : public FLimbSolverSettings
 {
 	GENERATED_BODY()
diff --git a/Engine/Plugins/Animation/IKRig/Source/IKRig/Public/Rig/Solvers/IKRigPoleSolver.h b/Engine/Plugins/Animation/IKRig/Source/IKRig/Public/Rig/Solvers/IKRigPoleSolver.h
--- a/Engine/Plugins/Animation/IKRig/Source/IKRig/Public/Rig/Solvers/IKRigPoleSolver.h
+++ b/Engine/Plugins/Animation/IKRig/Source/IKRig/Public/Rig/Solvers/IKRigPoleSolver.h
@@ -9,7 +9,7 @@
 
 #define UE_API IKRIG_API
 
-USTRUCT(BlueprintType)
+USTRUCT(BlueprintType, meta=(UIWrapper="/Script/IKRigEditor.PoleSolverSettingsWrapper"))
 struct FIKRigPoleSolverSettings : public FIKRigSolverSettingsBase
 {
 	GENERATED_BODY()
diff --git a/Engine/Plugins/Animation/IKRig/Source/IKRig/Public/Rig/Solvers/IKRigSetTransform.h b/Engine/Plugins/Animation/IKRig/Source/IKRig/Public/Rig/Solvers/IKRigSetTransform.h
--- a/Engine/Plugins/Animation/IKRig/Source/IKRig/Public/Rig/Solvers/IKRigSetTransform.h
+++ b/Engine/Plugins/Animation/IKRig/Source/IKRig/Public/Rig/Solvers/IKRigSetTransform.h
@@ -9,7 +9,7 @@
 
 #define UE_API IKRIG_API
 
-USTRUCT(BlueprintType)
+USTRUCT(BlueprintType, meta=(UIWrapper="/Script/IKRigEditor.SetTransformSettingsWrapper"))
 struct FIKRigSetTransformSettings : public FIKRigSolverSettingsBase
 {
 	GENERATED_BODY()
diff --git a/Engine/Plugins/Animation/IKRig/Source/IKRig/Public/Rig/Solvers/IKRigStretchLimb.h b/Engine/Plugins/Animation/IKRig/Source/IKRig/Public/Rig/Solvers/IKRigStretchLimb.h
--- a/Engine/Plugins/Animation/IKRig/Source/IKRig/Public/Rig/Solvers/IKRigStretchLimb.h
+++ b/Engine/Plugins/Animation/IKRig/Source/IKRig/Public/Rig/Solvers/IKRigStretchLimb.h
@@ -11,7 +11,7 @@ struct FIKRigGoal;
 
 #define UE_API IKRIG_API
 
-USTRUCT(BlueprintType)
+USTRUCT(BlueprintType, meta=(UIWrapper="/Script/IKRigEditor.StretchLimbBoneSettingsWrapper"))
 struct FIKRigStretchLimbBoneSettings : public FIKRigBoneSettingsBase
 {
 	GENERATED_BODY()
@@ -45,7 +45,7 @@ enum class EStretchLimbRotationMode : uint8
 	OrientToGoal,
 };
 
-USTRUCT(BlueprintType)
+USTRUCT(BlueprintType, meta=(UIWrapper="/Script/IKRigEditor.StretchLimbSettingsWrapper"))
 struct FIKRigStretchLimbSettings : public FIKRigSolverSettingsBase
 {
 	GENERATED_BODY()
@@ -187,4 +187,4 @@ private:
 	double InitialPoleVectorLenSq;
 };
 
-
+#undef UE_API
diff --git a/Engine/Plugins/Animation/IKRig/Source/IKRigEditor/Private/RigEditor/IKRigController.cpp b/Engine/Plugins/Animation/IKRig/Source/IKRigEditor/Private/RigEditor/IKRigController.cpp
--- a/Engine/Plugins/Animation/IKRig/Source/IKRigEditor/Private/RigEditor/IKRigController.cpp
+++ b/Engine/Plugins/Animation/IKRig/Source/IKRigEditor/Private/RigEditor/IKRigController.cpp
@@ -21,8 +21,7 @@ UIKRigController::UIKRigController()
 {
 	AutoCharacterizer = MakeUnique<FAutoCharacterizer>();
 	AutoFBIKCreator = MakeUnique<FAutoFBIKCreator>();
-
-	StructViewer = CreateDefaultSubobject<UIKRigStructViewer>(TEXT("SettingsViewer"));
+	StructViewerPool.Init(this);
 }
 
 UIKRigController* UIKRigController::GetController(const UIKRigDefinition* InIKRigDefinition)
@@ -625,10 +624,9 @@ bool UIKRigController::ValidateChain(
 	}
 }
 
-UIKRigStructViewer* UIKRigController::GetStructViewer()
+TArray<UIKRigStructWrapperBase*> UIKRigController::GetStructWrappers(const int32 InNum, TSubclassOf<UIKRigStructWrapperBase> InType)
 {
-	StructViewer->Reset();
-	return StructViewer;
+	return StructViewerPool.GetStructWrappers(InNum, InType);
 }
 
 FName UIKRigController::GetRetargetChainFromBone(const FName BoneName, const FIKRigSkeleton* OptionalSkeleton) const
diff --git a/Engine/Plugins/Animation/IKRig/Source/IKRigEditor/Private/RigEditor/IKRigEditorController.cpp b/Engine/Plugins/Animation/IKRig/Source/IKRigEditor/Private/RigEditor/IKRigEditorController.cpp
--- a/Engine/Plugins/Animation/IKRig/Source/IKRigEditor/Private/RigEditor/IKRigEditorController.cpp
+++ b/Engine/Plugins/Animation/IKRig/Source/IKRigEditor/Private/RigEditor/IKRigEditorController.cpp
@@ -461,13 +461,13 @@ void FIKRigEditorController::HandleDeleteSelectedElements()
 	{
 		switch(SelectedItem->ElementType)
 		{
-		case IKRigTreeElementType::GOAL:
+		case EIKRigTreeElementType::GOAL:
 			AssetController->RemoveGoal(SelectedItem->GoalName);
 			break;
-		case IKRigTreeElementType::SOLVERGOAL:
+		case EIKRigTreeElementType::GOAL_SETTINGS:
 			AssetController->DisconnectGoalFromSolver(SelectedItem->EffectorGoalName, SelectedItem->EffectorIndex);
 			break;
-		case IKRigTreeElementType::BONE_SETTINGS:
+		case EIKRigTreeElementType::BONE_SETTINGS:
 			AssetController->RemoveBoneSetting(SelectedItem->BoneSettingBoneName, SelectedItem->BoneSettingsSolverIndex);
 			break;
 		default:
@@ -672,7 +672,7 @@ void FIKRigEditorController::AddNewGoals(const TArray<FName>& GoalNames, const T
 	// show last created goal in details view
 	if (LastCreatedGoalName != NAME_None)
 	{
-		ShowDetailsForGoal(LastCreatedGoalName);
+		ShowDetailsForGoals({LastCreatedGoalName});
 	}
 }
 
@@ -690,19 +690,19 @@ void FIKRigEditorController::HandleGoalSelectedInViewport(const FName& GoalName,
 {
 	if (SkeletonView.IsValid())
 	{
-		SkeletonView->AddSelectedItemFromViewport(GoalName, IKRigTreeElementType::GOAL, bReplace);
+		SkeletonView->AddSelectedItemFromViewport(GoalName, EIKRigTreeElementType::GOAL, bReplace);
 		ShowDetailsForElements(SkeletonView->GetSelectedItems());
 		return;
 	}
 
-	ShowDetailsForGoal(GoalName);
+	ShowDetailsForGoals({GoalName});
 }
 
 void FIKRigEditorController::HandleBoneSelectedInViewport(const FName& BoneName, bool bReplace) const
 {
 	if (SkeletonView.IsValid())
 	{
-		SkeletonView->AddSelectedItemFromViewport(BoneName, IKRigTreeElementType::BONE, bReplace);
+		SkeletonView->AddSelectedItemFromViewport(BoneName, EIKRigTreeElementType::BONE, bReplace);
 		ShowDetailsForElements(SkeletonView->GetSelectedItems());
 		return;
 	}
@@ -956,98 +956,246 @@ void FIKRigEditorController::ShowDetailsForBone(const FName BoneName) const
 	DetailsView->SetObject(BoneDetails);
 }
 
-void FIKRigEditorController::ShowDetailsForBoneSettings(const FName& BoneName, int32 SolverIndex) const
+void FIKRigEditorController::ShowDetailsForBoneSettings(const TArray<TSharedPtr<FIKRigTreeElement>>& InItems) const
 {
+	if (InItems.IsEmpty())
+	{
+		return;
+	}
+	
 	const UIKRigDefinition* Asset = AssetController->GetAsset();
 	
-	auto MemoryProvider = [Asset, SolverIndex, BoneName]() -> uint8*
+	auto GetTypeOfItem = [Asset](const TSharedPtr<FIKRigTreeElement>& InItem)->TSubclassOf<UIKRigStructWrapperBase>
 	{
+		const int32 SolverIndex = InItem->BoneSettingsSolverIndex;
 		const UIKRigController* Controller = UIKRigController::GetController(Asset);
-		if (FIKRigSolverBase* Solver = Controller->GetSolverAtIndex(SolverIndex))
+		FIKRigSolverBase* Solver = Controller->GetSolverAtIndex(SolverIndex);
+		if (!Solver)
 		{
-			return reinterpret_cast<uint8*>(Solver->GetBoneSettings(BoneName));
+			return nullptr;
 		}
-		
-		return nullptr;
+		const UScriptStruct* BoneSettingsType = Solver->GetBoneSettingsType();
+		if (!BoneSettingsType)
+		{
+			return nullptr;
+		}
+
+		return UIKRigStructWrapperBase::GetStructWrapperFromMetaData(BoneSettingsType);
 	};
 
-	FInstancedStruct* InstancedStruct = AssetController->GetSolverStructAtIndex(SolverIndex);
-	const FIKRigSolverBase* Solver = InstancedStruct->GetPtr<FIKRigSolverBase>();
+	TSubclassOf<UIKRigStructWrapperBase> TypeToDisplay = GetTypeOfItem(InItems[0]);
+	if (!ensureMsgf(TypeToDisplay.Get(), TEXT("Selected settings do not have a UI wrapper class. Cannot be displayed.")))
+	{
+		return;
+	}
+	
+	// can only display homogenous array
+	TArray<TSharedPtr<FIKRigTreeElement>> ItemsToDisplay = FilterByFirstItemType(InItems, GetTypeOfItem);
+	TArray<UIKRigStructWrapperBase*> StructWrappers = AssetController->GetStructWrappers(ItemsToDisplay.Num(), TypeToDisplay);
+	for (int32 Index = 0; Index < ItemsToDisplay.Num(); ++Index)
+	{
+		TSharedPtr<FIKRigTreeElement> Item = ItemsToDisplay[Index];
+		const int32 SolverIndex = Item->BoneSettingsSolverIndex;
+		const FName BoneName = Item->BoneSettingBoneName;
+		
+		auto MemoryProvider = [Asset, SolverIndex, BoneName]() -> uint8*
+			{
+				const UIKRigController* Controller = UIKRigController::GetController(Asset);
+				if (FIKRigSolverBase* Solver = Controller->GetSolverAtIndex(SolverIndex))
+				{
+					return reinterpret_cast<uint8*>(Solver->GetBoneSettings(BoneName));
+				}
+		
+				return nullptr;
+			};
 
-	FIKRigStructToView StructToView;
-	StructToView.Type = Solver->GetBoneSettingsType();
-	StructToView.MemoryProvider = MemoryProvider;
-	StructToView.Owner = AssetController->GetAsset();
-	StructToView.UniqueName = BoneName;
+		FInstancedStruct* InstancedStruct = AssetController->GetSolverStructAtIndex(SolverIndex);
+		const FIKRigSolverBase* Solver = InstancedStruct->GetPtr<FIKRigSolverBase>();
+
+		FIKRigStructToView StructToView;
+		StructToView.Type = Solver->GetBoneSettingsType();
+		StructToView.MemoryProvider = MemoryProvider;
+		StructToView.Owner = AssetController->GetAsset();
+		StructToView.UniqueName = BoneName;
 	
-	UIKRigStructViewer* StructViewer = AssetController->GetStructViewer();
-	StructViewer->SetStructToView(StructToView);
+		StructWrappers[Index]->Initialize(StructToView, FName("Settings"));
+	}
 	
-	ShowStructViewerInDetailsView(StructViewer);
+	ShowStructWrappersInDetailsView(StructWrappers);
 }
 
-void FIKRigEditorController::ShowDetailsForGoal(const FName& GoalName) const
+void FIKRigEditorController::ShowDetailsForGoals(const TArray<FName>& InGoalNames) const
 {
-	DetailsView->SetObject(AssetController->GetGoal(GoalName));
+	TArray<UObject*> Goals;
+	for (const FName& GoalName : InGoalNames)
+	{
+		if (UIKRigEffectorGoal* Goal = AssetController->GetGoal(GoalName))
+		{
+			Goals.Add(Goal);
+		}
+	}
+	
+	DetailsView->SetObjects(Goals);
 }
 
-void FIKRigEditorController::ShowDetailsForGoalSettings(const FName GoalName, const int32 SolverIndex) const
+void FIKRigEditorController::ShowDetailsForGoalSettings(const TArray<TSharedPtr<FIKRigTreeElement>>& InItems) const
 {
+	if (InItems.IsEmpty())
+	{
+		return;
+	}
+	
 	const UIKRigDefinition* Asset = AssetController->GetAsset();
 	
-	auto MemoryProvider = [Asset, SolverIndex, GoalName]() -> uint8*
+	auto GetTypeOfItem = [Asset](const TSharedPtr<FIKRigTreeElement>& InItem)->TSubclassOf<UIKRigStructWrapperBase>
 	{
+		const int32 SolverIndex = InItem->EffectorIndex;
 		const UIKRigController* Controller = UIKRigController::GetController(Asset);
-		if (FIKRigSolverBase* Solver = Controller->GetSolverAtIndex(SolverIndex))
+		FIKRigSolverBase* Solver = Controller->GetSolverAtIndex(SolverIndex);
+		if (!Solver)
 		{
-			return reinterpret_cast<uint8*>(Solver->GetGoalSettings(GoalName));
+			return nullptr;
 		}
-		
-		return nullptr;
+		const UScriptStruct* GoalSettingsType = Solver->GetGoalSettingsType();
+		if (!GoalSettingsType)
+		{
+			return nullptr;
+		}
+
+		return UIKRigStructWrapperBase::GetStructWrapperFromMetaData(GoalSettingsType);
 	};
+
+	TSubclassOf<UIKRigStructWrapperBase> TypeToDisplay = GetTypeOfItem(InItems[0]);
+	if (!ensureMsgf(TypeToDisplay.Get(), TEXT("Selected settings do not have a UI wrapper class. Cannot be displayed.")))
+	{
+		return;
+	}
 	
-	FInstancedStruct* InstancedStruct = AssetController->GetSolverStructAtIndex(SolverIndex);
+	// can only display homogenous array
+	TArray<TSharedPtr<FIKRigTreeElement>> ItemsToDisplay = FilterByFirstItemType(InItems, GetTypeOfItem);
+	TArray<UIKRigStructWrapperBase*> StructWrappers = AssetController->GetStructWrappers(ItemsToDisplay.Num(), TypeToDisplay);
+	for (int32 Index = 0; Index < ItemsToDisplay.Num(); ++Index)
+	{
+		TSharedPtr<FIKRigTreeElement> Item = InItems[Index];
+		const int32 SolverIndex = Item->EffectorIndex;
+		const FName GoalName = Item->EffectorGoalName;
+		
+		auto MemoryProvider = [Asset, SolverIndex, GoalName]() -> uint8*
+			{
+				const UIKRigController* Controller = UIKRigController::GetController(Asset);
+				if (FIKRigSolverBase* Solver = Controller->GetSolverAtIndex(SolverIndex))
+				{
+					return reinterpret_cast<uint8*>(Solver->GetGoalSettings(GoalName));
+				}
+		
+				return nullptr;
+			};
 
-	FIKRigStructToView StructToView;
-	StructToView.Type = InstancedStruct->GetPtr<FIKRigSolverBase>()->GetGoalSettingsType();
-	StructToView.MemoryProvider = MemoryProvider;
-	StructToView.Owner = AssetController->GetAsset();
-	StructToView.UniqueName = GoalName;
+		FInstancedStruct* InstancedStruct = AssetController->GetSolverStructAtIndex(SolverIndex);
+		const FIKRigSolverBase* Solver = InstancedStruct->GetPtr<FIKRigSolverBase>();
+
+		FIKRigStructToView StructToView;
+		StructToView.Type = Solver->GetGoalSettingsType();
+		StructToView.MemoryProvider = MemoryProvider;
+		StructToView.Owner = AssetController->GetAsset();
+		StructToView.UniqueName = GoalName;
 	
-	UIKRigStructViewer* StructViewer = AssetController->GetStructViewer();
-	StructViewer->SetStructToView(StructToView);
+		StructWrappers[Index]->Initialize(StructToView, FName("Settings"));
+	}
 	
-	ShowStructViewerInDetailsView(StructViewer);
+	ShowStructWrappersInDetailsView(StructWrappers);
 }
 
-void FIKRigEditorController::ShowDetailsForSolver(const int32 SolverIndex) const
+void FIKRigEditorController::ShowDetailsForSolvers(const TArray<int32>& InSolverIndices) const
 {
+	if (InSolverIndices.IsEmpty())
+	{
+		return;
+	}
+	
 	const UIKRigDefinition* Asset = AssetController->GetAsset();
 	
-	auto MemoryProvider = [Asset, SolverIndex]() -> uint8*
+	auto GetTypeOfItem = [Asset](const int32 InSolverIndex)->TSubclassOf<UIKRigStructWrapperBase>
 	{
 		const UIKRigController* Controller = UIKRigController::GetController(Asset);
-		if (FIKRigSolverBase* Solver = Controller->GetSolverAtIndex(SolverIndex))
+		FIKRigSolverBase* Solver = Controller->GetSolverAtIndex(InSolverIndex);
+		if (!Solver)
 		{
-			return reinterpret_cast<uint8*>(Solver->GetSolverSettings());
+			return nullptr;
 		}
-		
-		return nullptr;
+		const UScriptStruct* SettingsType = Solver->GetSolverSettingsType();
+		if (!SettingsType)
+		{
+			return nullptr;
+		}
+
+		return UIKRigStructWrapperBase::GetStructWrapperFromMetaData(SettingsType);
 	};
+
+	TSubclassOf<UIKRigStructWrapperBase> TypeToDisplay = GetTypeOfItem(InSolverIndices[0]);
+	if (!ensureMsgf(TypeToDisplay.Get(), TEXT("Selected settings do not have a UI wrapper class. Cannot be displayed.")))
+	{
+		return;
+	}
+
+	auto FilterByFirstType = [](
+		const TArray<int32>& InIndices,
+		TFunctionRef<TSubclassOf<UObject>(int32)> GetTypeOfItem) -> TArray<int32>
+		{
+			TArray<int32> Out;
+			if (InIndices.IsEmpty())
+			{
+				return Out;
+			}
+
+			// filter by first item type
+			TSubclassOf<UObject> RefType = GetTypeOfItem(InIndices[0]);
+
+			// collect only items with the same type
+			Out.Reserve(InIndices.Num());
+			for (const int32 Index : InIndices)
+			{
+				const TSubclassOf<UObject> ItemType = GetTypeOfItem(Index);
+				if (ItemType == RefType) 
+				{
+					Out.Add(Index);
+				}
+			}
 	
-	FInstancedStruct* InstancedStruct = AssetController->GetSolverStructAtIndex(SolverIndex);
-	const FIKRigSolverBase* Solver = InstancedStruct->GetPtr<FIKRigSolverBase>();
+			return Out;	
+		};
 	
-	FIKRigStructToView StructToView;
-	StructToView.Type = Solver->GetSolverSettingsType();
-	StructToView.MemoryProvider = MemoryProvider;
-	StructToView.Owner = AssetController->GetAsset();
-	StructToView.UniqueName = FName(Solver->GetNiceName().ToString());
+	// can only display homogenous array
+	TArray<int32> SolversToDisplay = FilterByFirstType(InSolverIndices, GetTypeOfItem);
+	TArray<UIKRigStructWrapperBase*> StructWrappers = AssetController->GetStructWrappers(SolversToDisplay.Num(), TypeToDisplay);
+	for (int32 Index = 0; Index < SolversToDisplay.Num(); ++Index)
+	{
+		const int32 SolverIndex = SolversToDisplay[Index];
+		
+		auto MemoryProvider = [Asset, SolverIndex]() -> uint8*
+			{
+				const UIKRigController* Controller = UIKRigController::GetController(Asset);
+				if (FIKRigSolverBase* Solver = Controller->GetSolverAtIndex(SolverIndex))
+				{
+					return reinterpret_cast<uint8*>(Solver->GetSolverSettings());
+				}
+		
+				return nullptr;
+			};
+
+		FInstancedStruct* InstancedStruct = AssetController->GetSolverStructAtIndex(SolverIndex);
+		const FIKRigSolverBase* Solver = InstancedStruct->GetPtr<FIKRigSolverBase>();
+
+		FIKRigStructToView StructToView;
+		StructToView.Type = Solver->GetSolverSettingsType();
+		StructToView.MemoryProvider = MemoryProvider;
+		StructToView.Owner = AssetController->GetAsset();
+		StructToView.UniqueName = FName(Solver->GetNiceName().ToString());
 	
-	UIKRigStructViewer* StructViewer = AssetController->GetStructViewer();
-	StructViewer->SetStructToView(StructToView);
+		StructWrappers[Index]->Initialize(StructToView, FName("Settings"));
+	}
 	
-	ShowStructViewerInDetailsView(StructViewer);
+	ShowStructWrappersInDetailsView(StructWrappers);
 }
 
 void FIKRigEditorController::ShowAssetDetails() const
@@ -1057,30 +1205,60 @@ void FIKRigEditorController::ShowAssetDetails() const
 
 void FIKRigEditorController::ShowDetailsForElements(const TArray<TSharedPtr<FIKRigTreeElement>>& InItems) const
 {
-	if (!InItems.Num())
+	if (InItems.IsEmpty())
 	{
 		ShowAssetDetails();
 		return;
 	}
 
-	// edit the last selected element
-	const TSharedPtr<FIKRigTreeElement>& LastItem = InItems.Last();
-	switch (LastItem->ElementType)
+	auto FilterSameTypeAsFirst = [](const TArray<TSharedPtr<FIKRigTreeElement>>& InItems) -> TArray<TSharedPtr<FIKRigTreeElement>>
+		{
+			TArray<TSharedPtr<FIKRigTreeElement>> Result;
+
+			if (InItems.Num() == 0)
+			{
+				return Result;
+			}
+
+			const EIKRigTreeElementType FirstType = InItems[0]->ElementType;
+
+			for (const TSharedPtr<FIKRigTreeElement>& Item : InItems)
+			{
+				if (Item->ElementType == FirstType)
+				{
+					Result.Add(Item);
+				}
+			}
+
+			return Result;
+		};
+
+	// edit the last selected element (first in list) and all selected elements of that same type
+	TArray<TSharedPtr<FIKRigTreeElement>> ToEdit = FilterSameTypeAsFirst(InItems);
+	const TSharedPtr<FIKRigTreeElement>& LastSelected = InItems[0];
+	switch (LastSelected->ElementType)
 	{
-	case IKRigTreeElementType::BONE:
-		ShowDetailsForBone(LastItem->BoneName);
+	case EIKRigTreeElementType::BONE:
+		ShowDetailsForBone(LastSelected->BoneName); // multi-select not supported for bones
 		break;
 		
-	case IKRigTreeElementType::GOAL:
-		ShowDetailsForGoal(LastItem->GoalName);
+	case EIKRigTreeElementType::GOAL:
+		{
+			TArray<FName> GoalsToEdit;
+			for (const TSharedPtr<FIKRigTreeElement>& GoalToEdit : ToEdit)
+			{
+				GoalsToEdit.Add(GoalToEdit->GoalName);
+			}
+			ShowDetailsForGoals(GoalsToEdit);	
+		}
 		break;
 		
-	case IKRigTreeElementType::SOLVERGOAL:
-		ShowDetailsForGoalSettings(LastItem->EffectorGoalName, LastItem->EffectorIndex);
+	case EIKRigTreeElementType::GOAL_SETTINGS:
+		ShowDetailsForGoalSettings(ToEdit);
 		break;
 		
-	case IKRigTreeElementType::BONE_SETTINGS:
-		ShowDetailsForBoneSettings(LastItem->BoneSettingBoneName, LastItem->BoneSettingsSolverIndex);
+	case EIKRigTreeElementType::BONE_SETTINGS:
+		ShowDetailsForBoneSettings(ToEdit);
 		break;
 		
 	default:
@@ -1089,10 +1267,17 @@ void FIKRigEditorController::ShowDetailsForElements(const TArray<TSharedPtr<FIKR
 	}
 }
 
-void FIKRigEditorController::ShowStructViewerInDetailsView(UIKRigStructViewer* StructViewer) const
+void FIKRigEditorController::ShowStructWrappersInDetailsView(TArray<UIKRigStructWrapperBase*> InStructWrappers) const
 {
-	constexpr bool bForceRefresh = true; // must forcibly refresh because struct viewer is a recycled UObject
-	DetailsView->SetObject(StructViewer, bForceRefresh);
+	// convert to TArray<UObject*>
+	TArray<UObject*> Viewers;
+	for (UIKRigStructWrapperBase* StructWrapper : InStructWrappers)
+	{
+		Viewers.Add(StructWrapper);
+	}
+	
+	constexpr bool bForceRefresh = true; // must forcibly refresh because struct wrapper is a recycled UObject
+	DetailsView->SetObjects(Viewers, bForceRefresh);
 }
 
 void FIKRigEditorController::OnFinishedChangingDetails(const FPropertyChangedEvent& PropertyChangedEvent) const
@@ -1370,7 +1555,7 @@ void FIKRigEditorController::SetLastSelectedType(EIKRigSelectionType SelectionTy
 TObjectPtr<UIKRigBoneDetails> FIKRigEditorController::CreateBoneDetails(const TSharedPtr<FIKRigTreeElement const>& InBoneItem) const
 {
 	// ensure that the element is related to a bone
-	if (InBoneItem->ElementType != IKRigTreeElementType::BONE)
+	if (InBoneItem->ElementType != EIKRigTreeElementType::BONE)
 	{
 		return nullptr;
 	}
@@ -1384,5 +1569,32 @@ TObjectPtr<UIKRigBoneDetails> FIKRigEditorController::CreateBoneDetails(const TS
 	return NewBoneDetails;
 }
 
+TArray<TSharedPtr<FIKRigTreeElement>> FIKRigEditorController::FilterByFirstItemType(
+	const TArray<TSharedPtr<FIKRigTreeElement>>& InItems,
+	TFunctionRef<TSubclassOf<UObject>(const TSharedPtr<FIKRigTreeElement>&)> GetTypeOfItem) const
+{
+	TArray<TSharedPtr<FIKRigTreeElement>> Out;
+	if (InItems.IsEmpty())
+	{
+		return Out;
+	}
+
+	// filter by first item type
+	TSubclassOf<UObject> RefType = GetTypeOfItem(InItems[0]);
+
+	// collect only items with the same type
+	Out.Reserve(InItems.Num());
+	for (const TSharedPtr<FIKRigTreeElement>& Item : InItems)
+	{
+		const TSubclassOf<UObject> ItemType = GetTypeOfItem(Item);
+		if (ItemType == RefType) 
+		{
+			Out.Add(Item);
+		}
+	}
+	
+	return Out;
+}
+
 #undef LOCTEXT_NAMESPACE
 
diff --git a/Engine/Plugins/Animation/IKRig/Source/IKRigEditor/Private/RigEditor/SIKRigHierarchy.cpp b/Engine/Plugins/Animation/IKRig/Source/IKRigEditor/Private/RigEditor/SIKRigHierarchy.cpp
--- a/Engine/Plugins/Animation/IKRig/Source/IKRigEditor/Private/RigEditor/SIKRigHierarchy.cpp
+++ b/Engine/Plugins/Animation/IKRig/Source/IKRigEditor/Private/RigEditor/SIKRigHierarchy.cpp
@@ -28,7 +28,7 @@
 static const FName IKRigElementColumnName(TEXT("ElementName"));
 static const FName IKRigChainColumnName(TEXT("RetargetChainName"));
 
-FIKRigTreeElement::FIKRigTreeElement(const FText& InKey, IKRigTreeElementType InType,
+FIKRigTreeElement::FIKRigTreeElement(const FText& InKey, EIKRigTreeElementType InType,
 									 const TSharedRef<FIKRigEditorController>& InEditorController)
 	: Key(InKey)
 	, ElementType(InType)
@@ -84,13 +84,13 @@ FName FIKRigTreeElement::GetChainName() const
 	
 	switch(ElementType)
 	{
-	case IKRigTreeElementType::BONE:
+	case EIKRigTreeElementType::BONE:
 		return AssetController->GetRetargetChainFromBone(BoneName, Skeleton);
-	case IKRigTreeElementType::BONE_SETTINGS:
+	case EIKRigTreeElementType::BONE_SETTINGS:
 		return AssetController->GetRetargetChainFromBone(BoneSettingBoneName, Skeleton);
-	case IKRigTreeElementType::GOAL:
+	case EIKRigTreeElementType::GOAL:
 		return AssetController->GetRetargetChainFromGoal(GoalName);
-	case IKRigTreeElementType::SOLVERGOAL:
+	case EIKRigTreeElementType::GOAL_SETTINGS:
 		return AssetController->GetRetargetChainFromGoal(EffectorGoalName);
 	default:
 		return NAME_None;
@@ -138,7 +138,7 @@ void SIKRigHierarchyItem::Construct(const FArguments& InArgs)
 	const FSlateBrush* Brush = FAppStyle::Get().GetBrush("SkeletonTree.Bone");
 	switch(TreeElement->ElementType)
 	{
-		case IKRigTreeElementType::BONE:
+		case EIKRigTreeElementType::BONE:
 			if (!HierarchyView.Pin()->IsElementExcludedBone(WeakTreeElement))
 			{
 				Brush = FAppStyle::Get().GetBrush("SkeletonTree.Bone");
@@ -148,13 +148,13 @@ void SIKRigHierarchyItem::Construct(const FArguments& InArgs)
 				Brush = FAppStyle::Get().GetBrush("SkeletonTree.BoneNonWeighted");
 			}
 			break;
-		case IKRigTreeElementType::BONE_SETTINGS:
+		case EIKRigTreeElementType::BONE_SETTINGS:
 			Brush = FIKRigEditorStyle::Get().GetBrush("IKRig.Tree.BoneWithSettings");
 			break;
-		case IKRigTreeElementType::GOAL:
+		case EIKRigTreeElementType::GOAL:
 			Brush = FIKRigEditorStyle::Get().GetBrush("IKRig.Tree.Goal");
 			break;
-		case IKRigTreeElementType::SOLVERGOAL:
+		case EIKRigTreeElementType::GOAL_SETTINGS:
 			Brush = FIKRigEditorStyle::Get().GetBrush("IKRig.Tree.Effector");
 			break;
 		default:
@@ -177,7 +177,7 @@ void SIKRigHierarchyItem::Construct(const FArguments& InArgs)
 		]
 	];
 	
-	if (TreeElement->ElementType == IKRigTreeElementType::BONE)
+	if (TreeElement->ElementType == EIKRigTreeElementType::BONE)
 	{
 		RowBox->AddSlot()
         .AutoWidth()
@@ -320,7 +320,7 @@ FReply SIKRigSkeletonRow::HandleDragDetected(const FGeometry& MyGeometry, const
 	}
 
 	const TSharedPtr<FIKRigTreeElement> DraggedElement = SelectedItems[0];
-	if (DraggedElement->ElementType != IKRigTreeElementType::GOAL)
+	if (DraggedElement->ElementType != EIKRigTreeElementType::GOAL)
 	{
 		return FReply::Unhandled();
 	}
@@ -344,7 +344,7 @@ TOptional<EItemDropZone> SIKRigSkeletonRow::HandleCanAcceptDrop(
 	const TSharedPtr<FIKRigSkeletonDragDropOp> DragDropOp = DragDropEvent.GetOperationAs<FIKRigSkeletonDragDropOp>();
 	if (DragDropOp.IsValid())
 	{
-		if (TargetItem.Get()->ElementType == IKRigTreeElementType::BONE)
+		if (TargetItem.Get()->ElementType == EIKRigTreeElementType::BONE)
         {
         	ReturnedDropZone = EItemDropZone::BelowItem;	
         }
@@ -371,8 +371,8 @@ FReply SIKRigSkeletonRow::HandleAcceptDrop(
 	}
 
 	const FIKRigTreeElement& DraggedElement = *DragDropOp.Get()->Element.Pin().Get();
-	ensure(DraggedElement.ElementType == IKRigTreeElementType::GOAL);		// drag only supported for goals
-	ensure(TargetItem.Get()->ElementType == IKRigTreeElementType::BONE);	// drop only supported for bones
+	ensure(DraggedElement.ElementType == EIKRigTreeElementType::GOAL);		// drag only supported for goals
+	ensure(TargetItem.Get()->ElementType == EIKRigTreeElementType::BONE);	// drop only supported for bones
 
 	// re-parent the goal to a different bone
 	const UIKRigController* AssetController = Controller->AssetController;
@@ -510,7 +510,7 @@ void SIKRigHierarchy::Construct(const FArguments& InArgs, TSharedRef<FIKRigEdito
 
 void SIKRigHierarchy::AddSelectedItemFromViewport(
 	const FName& ItemName,
-	IKRigTreeElementType ItemType,
+	EIKRigTreeElementType ItemType,
 	const bool bReplace)
 {	
 	// nothing to add
@@ -527,13 +527,13 @@ void SIKRigHierarchy::AddSelectedItemFromViewport(
 		bool bIsBeingAdded = false;
 		switch (ItemType)
 		{
-		case IKRigTreeElementType::GOAL:
+		case EIKRigTreeElementType::GOAL:
 			if (ItemName == Item->GoalName)
 			{
 				bIsBeingAdded = true;
 			}
 			break;
-		case IKRigTreeElementType::BONE:
+		case EIKRigTreeElementType::BONE:
 			if (ItemName == Item->BoneName)
 			{
 				bIsBeingAdded = true;
@@ -669,25 +669,25 @@ bool SIKRigHierarchy::IsElementConnectedToSolver(TWeakPtr<FIKRigTreeElement> InT
 		return false; // not a valid solver index
 	}
 	
-	if (TreeElement->ElementType == IKRigTreeElementType::BONE)
+	if (TreeElement->ElementType == EIKRigTreeElementType::BONE)
 	{
 		// is this bone affected by this solver?
 		return Solver->IsBoneAffectedBySolver(TreeElement->BoneName, AssetController->GetIKRigSkeleton());
 	}
 
-	if (TreeElement->ElementType == IKRigTreeElementType::BONE_SETTINGS)
+	if (TreeElement->ElementType == EIKRigTreeElementType::BONE_SETTINGS)
 	{
 		// is this bone setting belonging to the solver?
 		return Solver->HasSettingsOnBone(TreeElement->BoneSettingBoneName);
 	}
 
-	if (TreeElement->ElementType == IKRigTreeElementType::GOAL)
+	if (TreeElement->ElementType == EIKRigTreeElementType::GOAL)
 	{
 		// is goal connected to the solver?
 		return AssetController->IsGoalConnectedToSolver(TreeElement->GoalName, SolverIndex);
 	}
 
-	if (TreeElement->ElementType == IKRigTreeElementType::SOLVERGOAL)
+	if (TreeElement->ElementType == EIKRigTreeElementType::GOAL_SETTINGS)
 	{
 		// is this an effector for this solver?
 		return TreeElement->EffectorIndex == SolverIndex;
@@ -716,7 +716,7 @@ bool SIKRigHierarchy::IsElementConnectedToAnySolver(TWeakPtr<FIKRigTreeElement>
 
 bool SIKRigHierarchy::IsElementExcludedBone(TWeakPtr<FIKRigTreeElement> TreeElement)
 {
-	if (TreeElement.Pin()->ElementType != IKRigTreeElementType::BONE)
+	if (TreeElement.Pin()->ElementType != EIKRigTreeElementType::BONE)
 	{
 		return false;
 	}
@@ -892,7 +892,7 @@ void SIKRigHierarchy::HandleNewGoal() const
 	const TArray<TSharedPtr<FIKRigTreeElement>> SelectedItems = TreeView.Get()->GetSelectedItems();
 	for (const TSharedPtr<FIKRigTreeElement>& Item : SelectedItems)
 	{
-		if (Item->ElementType != IKRigTreeElementType::BONE)
+		if (Item->ElementType != EIKRigTreeElementType::BONE)
 		{
 			continue; // can only add goals to bones
 		}
@@ -921,7 +921,7 @@ bool SIKRigHierarchy::CanAddNewGoal() const
 	// can only add goals to selected bones
 	for (const TSharedPtr<FIKRigTreeElement>& Item : SelectedItems)
 	{
-		if (Item->ElementType != IKRigTreeElementType::BONE)
+		if (Item->ElementType != EIKRigTreeElementType::BONE)
 		{
 			return false;
 		}
@@ -952,7 +952,7 @@ bool SIKRigHierarchy::CanDeleteElement() const
 	// are all selected items goals, effectors or bone settings?
 	for (const TSharedPtr<FIKRigTreeElement>& Item : SelectedItems)
 	{
-		const bool bIsBone = Item->ElementType == IKRigTreeElementType::BONE;
+		const bool bIsBone = Item->ElementType == EIKRigTreeElementType::BONE;
 		if (bIsBone)
 		{
 			return false;
@@ -1069,15 +1069,17 @@ void SIKRigHierarchy::HandleSetRootBoneOnSolvers()
 	UIKRigController* AssetController = Controller->AssetController;
 	TArray<TSharedPtr<FSolverStackElement>> SelectedSolvers;
 	Controller->GetSelectedSolvers(SelectedSolvers);
-	int32 SolverToShow = 0;
+	TArray<int32> SolversToShow;
 	for (const TSharedPtr<FSolverStackElement>& Solver : SelectedSolvers)
 	{
-		AssetController->SetStartBone(RootBoneName, Solver->IndexInStack);
-		SolverToShow = Solver->IndexInStack;
+		if (AssetController->SetStartBone(RootBoneName, Solver->IndexInStack))
+		{
+			SolversToShow.Add(Solver->IndexInStack);
+		}
 	}
 
 	// show solver that had it's root bone updated
-	Controller->ShowDetailsForSolver(SolverToShow);
+	Controller->ShowDetailsForSolvers(SolversToShow);
 	
 	// show new icon when bone has settings applied
 	RefreshTreeView();
@@ -1131,15 +1133,17 @@ void SIKRigHierarchy::HandleSetEndBoneOnSolvers()
 	UIKRigController* AssetController = Controller->AssetController;
 	TArray<TSharedPtr<FSolverStackElement>> SelectedSolvers;
 	Controller->GetSelectedSolvers(SelectedSolvers);
-	int32 SolverToShow = 0;
+	TArray<int32> SolversToShow;
 	for (const TSharedPtr<FSolverStackElement>& Solver : SelectedSolvers)
 	{
-		AssetController->SetEndBone(RootBoneName, Solver->IndexInStack);
-		SolverToShow = Solver->IndexInStack;
+		if (AssetController->SetEndBone(RootBoneName, Solver->IndexInStack))
+		{
+			SolversToShow.Add(Solver->IndexInStack);
+		}
 	}
 
 	// show solver that had it's root bone updated
-	Controller->ShowDetailsForSolver(SolverToShow);
+	Controller->ShowDetailsForSolvers(SolversToShow);
 	
 	// show new icon when bone has settings applied
 	RefreshTreeView();
@@ -1197,19 +1201,15 @@ void SIKRigHierarchy::HandleAddBoneSettings()
 	UIKRigController* AssetController = Controller->AssetController;
 	TArray<TSharedPtr<FSolverStackElement>> SelectedSolvers;
 	Controller->GetSelectedSolvers(SelectedSolvers);
-	FName BoneNameForSettings;
-	int32 SolverIndex = INDEX_NONE;
 	for (const TSharedPtr<FIKRigTreeElement>& BoneItem : SelectedBones)
 	{
 		for (const TSharedPtr<FSolverStackElement>& Solver : SelectedSolvers)
 		{
 			AssetController->AddBoneSetting(BoneItem->BoneName, Solver->IndexInStack);
-			BoneNameForSettings = BoneItem->BoneName;
-			SolverIndex = Solver->IndexInStack;
         }
 	}
 
-	Controller->ShowDetailsForBoneSettings(BoneNameForSettings, SolverIndex);
+	Controller->ShowDetailsForBone(SelectedBones[0]->BoneName);
 
 	// show new icon when bone has settings applied
 	RefreshTreeView();
@@ -1509,7 +1509,7 @@ void SIKRigHierarchy::GetSelectedBones(TArray<TSharedPtr<FIKRigTreeElement>>& Ou
 	const TArray<TSharedPtr<FIKRigTreeElement>> SelectedItems = TreeView->GetSelectedItems();
 	for (const TSharedPtr<FIKRigTreeElement>& Item : SelectedItems)
 	{
-		if (Item->ElementType == IKRigTreeElementType::BONE)
+		if (Item->ElementType == EIKRigTreeElementType::BONE)
 		{
 			OutBoneItems.Add(Item);
         }
@@ -1533,7 +1533,7 @@ void SIKRigHierarchy::GetSelectedGoals(TArray<TSharedPtr<FIKRigTreeElement>>& Ou
 	TArray<TSharedPtr<FIKRigTreeElement>> SelectedItems = TreeView->GetSelectedItems();
 	for (const TSharedPtr<FIKRigTreeElement>& Item : SelectedItems)
 	{
-		if (Item->ElementType == IKRigTreeElementType::GOAL)
+		if (Item->ElementType == EIKRigTreeElementType::GOAL)
 		{
 			OutSelectedGoals.Add(Item);
 		}
@@ -1772,7 +1772,7 @@ void SIKRigHierarchy::RefreshTreeView(bool IsInitialSetup /*=false*/)
 		// create "Bone" tree element for this bone
 		const FName BoneName = Skeleton.BoneNames[BoneIndex];
 		const FText BoneDisplayName = FText::FromName(BoneName);
-		TSharedPtr<FIKRigTreeElement> BoneElement = MakeShared<FIKRigTreeElement>(BoneDisplayName, IKRigTreeElementType::BONE, ControllerRef);
+		TSharedPtr<FIKRigTreeElement> BoneElement = MakeShared<FIKRigTreeElement>(BoneDisplayName, EIKRigTreeElementType::BONE, ControllerRef);
 		BoneElement.Get()->BoneName = BoneName;
 		const int32 BoneElementIndex = AllElements.Add(BoneElement);
 		BoneTreeElementIndices.Add(BoneName, BoneElementIndex);
@@ -1815,7 +1815,7 @@ void SIKRigHierarchy::RefreshTreeView(bool IsInitialSetup /*=false*/)
 			
 			const FText SolverDisplayName = FText::FromString(AssetController->GetSolverUniqueName(SolverIndex));
 			const FText BoneSettingDisplayName = FText::Format(LOCTEXT("BoneSettings", "{0} bone settings."), SolverDisplayName);
-			TSharedPtr<FIKRigTreeElement> SettingsItem = MakeShared<FIKRigTreeElement>(BoneSettingDisplayName, IKRigTreeElementType::BONE_SETTINGS, ControllerRef);
+			TSharedPtr<FIKRigTreeElement> SettingsItem = MakeShared<FIKRigTreeElement>(BoneSettingDisplayName, EIKRigTreeElementType::BONE_SETTINGS, ControllerRef);
 			SettingsItem->BoneSettingBoneName = BoneName;
 			SettingsItem->BoneSettingsSolverIndex = SolverIndex;
 			AllElements.Add(SettingsItem);
@@ -1846,7 +1846,7 @@ void SIKRigHierarchy::RefreshTreeView(bool IsInitialSetup /*=false*/)
 			
 			// make new element for goal
 			const FText GoalDisplayName = FText::FromName(Goal->GoalName);
-			TSharedPtr<FIKRigTreeElement> GoalItem = MakeShared<FIKRigTreeElement>(GoalDisplayName, IKRigTreeElementType::GOAL, ControllerRef);
+			TSharedPtr<FIKRigTreeElement> GoalItem = MakeShared<FIKRigTreeElement>(GoalDisplayName, EIKRigTreeElementType::GOAL, ControllerRef);
 			GoalItem->GoalName = Goal->GoalName;
 			AllElements.Add(GoalItem);
 
@@ -1878,7 +1878,7 @@ void SIKRigHierarchy::RefreshTreeView(bool IsInitialSetup /*=false*/)
 				// make new element for effector
 				const FText SolverDisplayName = FText::FromString(AssetController->GetSolverUniqueName(SolverIndex));
 				const FText EffectorDisplayName = FText::Format(LOCTEXT("GoalSettingsForSolver", "{0} goal settings."), SolverDisplayName);
-				TSharedPtr<FIKRigTreeElement> EffectorItem = MakeShared<FIKRigTreeElement>(EffectorDisplayName, IKRigTreeElementType::SOLVERGOAL, ControllerRef);
+				TSharedPtr<FIKRigTreeElement> EffectorItem = MakeShared<FIKRigTreeElement>(EffectorDisplayName, EIKRigTreeElementType::GOAL_SETTINGS, ControllerRef);
 				EffectorItem->EffectorIndex = SolverIndex;
 				EffectorItem->EffectorGoalName = Goal->GoalName;
 				AllElements.Add(EffectorItem);
diff --git a/Engine/Plugins/Animation/IKRig/Source/IKRigEditor/Private/RigEditor/SIKRigSolverStack.cpp b/Engine/Plugins/Animation/IKRig/Source/IKRigEditor/Private/RigEditor/SIKRigSolverStack.cpp
--- a/Engine/Plugins/Animation/IKRig/Source/IKRigEditor/Private/RigEditor/SIKRigSolverStack.cpp
+++ b/Engine/Plugins/Animation/IKRig/Source/IKRigEditor/Private/RigEditor/SIKRigSolverStack.cpp
@@ -284,7 +284,7 @@ void SIKRigSolverStack::Construct(const FArguments& InArgs, TSharedRef<FIKRigEdi
         .Padding(0.0f)
         [
 			SAssignNew( ListView, SSolverStackListViewType )
-			.SelectionMode(ESelectionMode::Single)
+			.SelectionMode(ESelectionMode::Multi)
 			.IsEnabled(this, &SIKRigSolverStack::IsAddSolverEnabled)
 			.ListItemsSource( &ListViewItems )
 			.OnGenerateRow( this, &SIKRigSolverStack::MakeListRowWidget )
@@ -392,7 +392,7 @@ void SIKRigSolverStack::AddNewSolver(UScriptStruct* SolverType)
 	// select it
 	ListView->SetSelection(ListViewItems[NewSolverIndex]);
 	// show details for it
-	Controller->ShowDetailsForSolver(NewSolverIndex);
+	Controller->ShowDetailsForSolvers({NewSolverIndex});
 }
 
 void SIKRigSolverStack::DeleteSolver(TSharedPtr<FSolverStackElement> SolverToDelete)
@@ -474,17 +474,17 @@ void SIKRigSolverStack::OnSelectionChanged(TSharedPtr<FSolverStackElement> InIte
 {
 	if (SelectInfo != ESelectInfo::Direct)
 	{
-		ShowDetailsForItem(InItem);	
+		ShowDetailsForItems(ListView->GetSelectedItems());	
 	}
 }
 
 void SIKRigSolverStack::OnItemClicked(TSharedPtr<FSolverStackElement> InItem)
 {
-	ShowDetailsForItem(InItem);
+	ShowDetailsForItems({InItem});
 	EditorController.Pin()->SetLastSelectedType(EIKRigSelectionType::SolverStack);
 }
 
-void SIKRigSolverStack::ShowDetailsForItem(TSharedPtr<FSolverStackElement> InItem)
+void SIKRigSolverStack::ShowDetailsForItems(const TArray<TSharedPtr<FSolverStackElement>>& InItems)
 {
 	const TSharedPtr<FIKRigEditorController> Controller = EditorController.Pin();
 	if (!Controller.IsValid())
@@ -496,7 +496,7 @@ void SIKRigSolverStack::ShowDetailsForItem(TSharedPtr<FSolverStackElement> InIte
 	Controller->RefreshTreeView();
 
 	// the solver selection must be done after we rebuilt the skeleton tree has it keeps the selection
-	if (!InItem.IsValid())
+	if (InItems.IsEmpty())
 	{
 		// clear the details panel only if there's nothing selected in the skeleton view
 		if (Controller->DoesSkeletonHaveSelectedItems())
@@ -508,7 +508,12 @@ void SIKRigSolverStack::ShowDetailsForItem(TSharedPtr<FSolverStackElement> InIte
 	}
 	else
 	{
-		Controller->ShowDetailsForSolver(InItem.Get()->IndexInStack);
+		TArray<int32> SolverIndices;
+		for (const TSharedPtr<FSolverStackElement>& Item : InItems)
+		{
+			SolverIndices.Add(Item->IndexInStack);
+		}
+		Controller->ShowDetailsForSolvers(SolverIndices);
 	}
 }
 
@@ -582,7 +587,7 @@ FReply SIKRigSolverStack::OnAcceptDrop(
 	{
 		RefreshStackView();
 		Controller->RefreshTreeView(); // update solver indices in effector items
-		Controller->ShowDetailsForSolver(TargetIndex);
+		Controller->ShowDetailsForSolvers({TargetIndex});
 		ListView->SetSelection(ListViewItems[TargetIndex]);
 	}
 	
diff --git a/Engine/Plugins/Animation/IKRig/Source/IKRigEditor/Public/RigEditor/IKRigController.h b/Engine/Plugins/Animation/IKRig/Source/IKRigEditor/Public/RigEditor/IKRigController.h
--- a/Engine/Plugins/Animation/IKRig/Source/IKRigEditor/Public/RigEditor/IKRigController.h
+++ b/Engine/Plugins/Animation/IKRig/Source/IKRigEditor/Public/RigEditor/IKRigController.h
@@ -3,6 +3,7 @@
 #pragma once
 
 #include "CoreMinimal.h"
+#include "IKRigStructViewer.h"
 #include "StructUtils/InstancedStruct.h"
 #include "UObject/Object.h"
 
@@ -400,8 +401,9 @@ public:
 	// Optionally provide a runtime skeleton from an IKRigProcessor to get indices for a running instance (otherwise uses stored hierarchy in asset)
 	UE_API bool ValidateChain( const FName ChainName, const FIKRigSkeleton* OptionalSkeleton, TSet<int32>& OutChainIndices) const;
 
-	// a singleton wrapper to display structs in the details panel
-	UE_API UIKRigStructViewer* GetStructViewer();
+	// returns an array of wrappers to display structs in the details panel
+	// NOTE: these are recycled by the editor using this controller (assumes one editor per asset)
+	UE_API TArray<UIKRigStructWrapperBase*> GetStructWrappers(const int32 InNum, TSubclassOf<UIKRigStructWrapperBase> InType);
 
 private:
 
@@ -444,9 +446,10 @@ private:
 	UPROPERTY(transient)
 	TObjectPtr<UIKRigDefinition> Asset = nullptr;
 
-	// a dummy UObject to wrap UStructs for the details panel
+	// a pool of objects inheriting from UIKRigStructWrapperBase for editing various IK Rig struct types in a details view
+	// NOTE: pool contains a bucket per settings type
 	UPROPERTY(transient)
-	TObjectPtr<UIKRigStructViewer> StructViewer = nullptr;
+	FIKRigStructWrapperPool StructViewerPool;
 
 	// Static global map of assets to their associated controller
 	static UE_API TMap<UIKRigDefinition*, UIKRigController*> Controllers;
diff --git a/Engine/Plugins/Animation/IKRig/Source/IKRigEditor/Public/RigEditor/IKRigEditorController.h b/Engine/Plugins/Animation/IKRig/Source/IKRigEditor/Public/RigEditor/IKRigEditorController.h
--- a/Engine/Plugins/Animation/IKRig/Source/IKRigEditor/Public/RigEditor/IKRigEditorController.h
+++ b/Engine/Plugins/Animation/IKRig/Source/IKRigEditor/Public/RigEditor/IKRigEditorController.h
@@ -24,6 +24,7 @@ class FIKRigTreeElement;
 class UAnimInstance;
 class UDebugSkelMeshComponent;
 class UIKRigStructViewer;
+class UIKRigStructWrapperBase;
 struct FIKRigProcessor;
 
 enum EIKRigSelectionType : int8
@@ -174,19 +175,19 @@ public:
 	// show single transform of bone in details view 
 	void ShowDetailsForBone(const FName BoneName) const;
 	// show single BONE settings in details view 
-	void ShowDetailsForBoneSettings(const FName& BoneName, int32 SolverIndex) const;
+	void ShowDetailsForBoneSettings(const TArray<TSharedPtr<FIKRigTreeElement>>& InItems) const;
 	// show single GOAL settings in details view 
-	void ShowDetailsForGoal(const FName& GoalName) const;
+	void ShowDetailsForGoals(const TArray<FName>& GoalName) const;
 	// show single EFFECTOR settings in details view 
-	void ShowDetailsForGoalSettings(const FName GoalName, const int32 SolverIndex) const;
+	void ShowDetailsForGoalSettings(const TArray<TSharedPtr<FIKRigTreeElement>>& InItems) const;
 	// show single SOLVER settings in details view 
-	void ShowDetailsForSolver(const int32 SolverIndex) const;
+	void ShowDetailsForSolvers(const TArray<int32>& InSolverIndices) const;
 	// show nothing in details view 
 	void ShowAssetDetails() const;
 	// show selected items in details view 
 	void ShowDetailsForElements(const TArray<TSharedPtr<FIKRigTreeElement>>& InItems) const;
 	// show the currently configured StructViewer object in the details tab
-	void ShowStructViewerInDetailsView(UIKRigStructViewer* StructViewer) const;
+	void ShowStructWrappersInDetailsView(TArray<UIKRigStructWrapperBase*> InStructViewers) const;
 	// callback when detail is edited 
 	void OnFinishedChangingDetails(const FPropertyChangedEvent& PropertyChangedEvent) const;
 	// returns true if the supplied UObject is being shown in the details panel
@@ -251,6 +252,10 @@ public:
 	
 private:
 
+	TArray<TSharedPtr<FIKRigTreeElement>> FilterByFirstItemType(
+		const TArray<TSharedPtr<FIKRigTreeElement>>& InItems,
+		TFunctionRef<TSubclassOf<UObject>(const TSharedPtr<FIKRigTreeElement>&)> GetTypeOfItem) const;
+
 	// asset properties tab 
 	TSharedPtr<IDetailsView> DetailsView;
 
diff --git a/Engine/Plugins/Animation/IKRig/Source/IKRigEditor/Public/RigEditor/IKRigStructViewer.h b/Engine/Plugins/Animation/IKRig/Source/IKRigEditor/Public/RigEditor/IKRigStructViewer.h
--- a/Engine/Plugins/Animation/IKRig/Source/IKRigEditor/Public/RigEditor/IKRigStructViewer.h
+++ b/Engine/Plugins/Animation/IKRig/Source/IKRigEditor/Public/RigEditor/IKRigStructViewer.h
@@ -7,6 +7,7 @@
 #include "ScopedTransaction.h"
 #include "Interfaces/Interface_BoneReferenceSkeletonProvider.h"
 #include "RigEditor/IKRigEditorController.h"
+#include "UObject/SavePackage.h"
 
 #include "IKRigStructViewer.generated.h"
 
@@ -169,6 +170,29 @@ public:
 	void UpdateWrapperStructWithLatestValues();
 	
 	virtual void PostEditChangeProperty( struct FPropertyChangedEvent& PropertyChangedEvent ) override;
+
+	// NOTE: How to add multi-edit support to IK Rig settings structs.
+	// IK Rig settings structs can supply a "UIWrapper" meta data string with the path to a custom wrapper class
+	// to be displayed in the details panel when editing that struct.
+	//
+	// In order to ensure that edits to the wrapper are pushed back to the actual struct-to-be-edited,
+	// the wrapper must inherit from UIKRigStructWrapperBase. This stores a pointer to the original memory and
+	// triggers callbacks when properties are edited to push those values back to the struct being edited.
+	//
+	// By default, only a single struct under a UProperty named "Settings" will be displayed in the details panel.
+	// See example subclasses of UIKRigStructWrapperBase.
+	static TSubclassOf<UIKRigStructWrapperBase> GetStructWrapperFromMetaData(const UScriptStruct* S)
+	{
+		const FString Path = S->GetMetaData(TEXT("UIWrapper"));
+		if (Path.IsEmpty())
+		{
+			return nullptr;
+		}
+	
+		const FSoftClassPath ClassPath(Path);
+		TSoftClassPtr<UIKRigStructWrapperBase> SoftClass(ClassPath);
+		return SoftClass.LoadSynchronous();
+	}
 	
 private:
 	
@@ -185,4 +209,71 @@ public:
 	// IDetailCustomization interface
 	virtual void CustomizeDetails(IDetailLayoutBuilder& DetailBuilder) override;
 	// END IDetailCustomization interface
+};
+
+// NOTE: this is a dummy wrapper around TArray<TObjectPtr<>> that is necessary because we cannot store TArray<rawpointer> as a TMap value in the pool
+USTRUCT()
+struct FIKRigStructWrapperBucket
+{
+	GENERATED_BODY()
+
+	UPROPERTY(Transient)
+	TArray<TObjectPtr<UIKRigStructWrapperBase>> Items;
+};
+
+USTRUCT()
+struct FIKRigStructWrapperPool
+{
+	GENERATED_BODY()
+
+public:
+	// call once to give the pool an owning UObject to prevent GC, otherwise uses GetTransientPackage()
+	void Init(UObject* InOuter) { CreationOuter = InOuter; }
+
+	// get (or lazily instantiate) Num wrappers of the requested subclass.
+	// NOTE: buckets will grow, but never shrink
+	TArray<UIKRigStructWrapperBase*> GetStructWrappers(int32 Num, TSubclassOf<UIKRigStructWrapperBase> InType)
+	{
+		TArray<UIKRigStructWrapperBase*> Out;
+		if (Num <= 0 || !InType)
+		{
+			return Out;
+		}
+
+		FIKRigStructWrapperBucket& Bucket = Pool.FindOrAdd(InType);
+
+		// grow, but never shrink (preserves indices; no compaction)
+		if (Bucket.Items.Num() < Num)
+		{
+			Bucket.Items.SetNum(Num); // new entries default to nullptr
+		}
+
+		Out.Reserve(Num);
+
+		UObject* Outer = CreationOuter.IsValid() ? CreationOuter.Get() : CastChecked<UObject>(GetTransientPackage());
+
+		for (int32 InstanceIndex = 0; InstanceIndex < Num; ++InstanceIndex)
+		{
+			// Create if missing or of the wrong type (defensive in case bucket was reused)
+			if (!IsValid(Bucket.Items[InstanceIndex]) || Bucket.Items[InstanceIndex]->GetClass() != InType)
+			{
+				Bucket.Items[InstanceIndex] = NewObject<UIKRigStructWrapperBase>(Outer, InType, NAME_None, RF_Transient);
+			}
+
+			UIKRigStructWrapperBase* Instance = Bucket.Items[InstanceIndex];
+			Instance->Reset();
+			Out.Add(Instance);
+		}
+
+		return Out;
+	}
+
+private:
+	// used as the Outer for all pooled objects
+	UPROPERTY(Transient)
+	TWeakObjectPtr<UObject> CreationOuter;
+
+	// one bucket per subclass; values are GC-tracked via TObjectPtr
+	UPROPERTY(Transient)
+	TMap<TSubclassOf<UIKRigStructWrapperBase>, FIKRigStructWrapperBucket> Pool;
 };
\ No newline at end of file
diff --git a/Engine/Plugins/Animation/IKRig/Source/IKRigEditor/Public/RigEditor/IKRigStructWrappers.h b/Engine/Plugins/Animation/IKRig/Source/IKRigEditor/Public/RigEditor/IKRigStructWrappers.h
new file mode 100644
--- /dev/null
+++ b/Engine/Plugins/Animation/IKRig/Source/IKRigEditor/Public/RigEditor/IKRigStructWrappers.h
@@ -0,0 +1,129 @@
+﻿// Copyright Epic Games, Inc. All Rights Reserved.
+
+#pragma once
+
+#include "IKRigStructViewer.h"
+
+#include "Rig/Solvers/IKRigBodyMoverSolver.h"
+#include "Rig/Solvers/IKRigFullBodyIK.h"
+#include "Rig/Solvers/IKRigLimbSolver.h"
+#include "Rig/Solvers/IKRigPoleSolver.h"
+#include "Rig/Solvers/IKRigSetTransform.h"
+#include "Rig/Solvers/IKRigStretchLimb.h"
+
+#include "IKRigStructWrappers.generated.h"
+
+
+//
+// BODY MOVER Struct Wrappers
+//
+
+UCLASS()
+class UBodyMoverSettingsWrapper : public UIKRigStructWrapperBase
+{
+	GENERATED_BODY()
+public:
+	UPROPERTY(EditAnywhere, Category = "Body Mover Settings")
+	FIKRigBodyMoverSettings Settings;
+};
+
+UCLASS()
+class UBodyMoverGoalSettingsWrapper : public UIKRigStructWrapperBase
+{
+	GENERATED_BODY()
+public:
+	UPROPERTY(EditAnywhere, Category = "Body Mover Goal Settings")
+	FIKRigBodyMoverGoalSettings Settings;
+};
+
+//
+// FBIK Struct Wrappers
+//
+
+UCLASS()
+class UFBIKSettingsWrapper : public UIKRigStructWrapperBase
+{
+	GENERATED_BODY()
+public:
+	UPROPERTY(EditAnywhere, Category = "Full-Body IK Settings")
+	FIKRigFBIKSettings Settings;
+};
+
+UCLASS()
+class UFBIKGoalSettingsWrapper : public UIKRigStructWrapperBase
+{
+	GENERATED_BODY()
+public:
+	UPROPERTY(EditAnywhere, Category = "FBIK Goal Settings")
+	FIKRigFBIKGoalSettings Settings;
+};
+
+UCLASS()
+class UFBIKBoneSettingsWrapper : public UIKRigStructWrapperBase
+{
+	GENERATED_BODY()
+public:
+	UPROPERTY(EditAnywhere, Category = "FBIK Bone Settings")
+	FIKRigFBIKBoneSettings Settings;
+};
+
+//
+// Limb Solver Struct Wrappers
+//
+
+UCLASS()
+class ULimbSolverSettingsWrapper : public UIKRigStructWrapperBase
+{
+	GENERATED_BODY()
+public:
+	UPROPERTY(EditAnywhere, Category = "Limb Solver Settings")
+	FIKRigLimbSolverSettings Settings;
+};
+
+//
+// Pole Solver Struct Wrappers
+//
+
+UCLASS()
+class UPoleSolverSettingsWrapper : public UIKRigStructWrapperBase
+{
+	GENERATED_BODY()
+public:
+	UPROPERTY(EditAnywhere, Category = "Pole Solver Settings")
+	FIKRigPoleSolverSettings Settings;
+};
+
+//
+// SetTransform Solver Struct Wrappers
+//
+
+UCLASS()
+class USetTransformSettingsWrapper : public UIKRigStructWrapperBase
+{
+	GENERATED_BODY()
+public:
+	UPROPERTY(EditAnywhere, Category = "Set Transform Settings")
+	FIKRigSetTransformSettings Settings;
+};
+
+//
+// Stretch Limb Solver Struct Wrappers
+//
+
+UCLASS()
+class UStretchLimbSettingsWrapper : public UIKRigStructWrapperBase
+{
+	GENERATED_BODY()
+public:
+	UPROPERTY(EditAnywhere, Category = "Stretch Limb Settings")
+	FIKRigStretchLimbSettings Settings;
+};
+
+UCLASS()
+class UStretchLimbBoneSettingsWrapper : public UIKRigStructWrapperBase
+{
+	GENERATED_BODY()
+public:
+	UPROPERTY(EditAnywhere, Category = "Stretch Limb Bone Settings")
+	FIKRigStretchLimbBoneSettings Settings;
+};
diff --git a/Engine/Plugins/Animation/IKRig/Source/IKRigEditor/Public/RigEditor/SIKRigHierarchy.h b/Engine/Plugins/Animation/IKRig/Source/IKRigEditor/Public/RigEditor/SIKRigHierarchy.h
--- a/Engine/Plugins/Animation/IKRig/Source/IKRigEditor/Public/RigEditor/SIKRigHierarchy.h
+++ b/Engine/Plugins/Animation/IKRig/Source/IKRigEditor/Public/RigEditor/SIKRigHierarchy.h
@@ -19,7 +19,7 @@ class FIKRigEditorToolkit;
 class USkeletalMesh;
 class UIKRigBoneDetails;
 
-enum class IKRigTreeElementType { BONE, GOAL, SOLVERGOAL, BONE_SETTINGS };
+enum class EIKRigTreeElementType { BONE, GOAL, GOAL_SETTINGS, BONE_SETTINGS };
 
 class FIKRigTreeElement : public TSharedFromThis<FIKRigTreeElement>
 {
@@ -27,11 +27,11 @@ public:
 	
 	FIKRigTreeElement(
 		const FText& InKey,
-		IKRigTreeElementType InType,
+		EIKRigTreeElementType InType,
 		const TSharedRef<FIKRigEditorController>& InEditorController);
 
 	FText Key;
-	IKRigTreeElementType ElementType;
+	EIKRigTreeElementType ElementType;
 
 	/** determines if this item is filtered out of the view */
 	bool bIsHidden = false;
@@ -196,7 +196,7 @@ public:
 	bool IsGoalSelected(const FName& GoalName);
 	void AddSelectedItemFromViewport(
 		const FName& ItemName,
-		IKRigTreeElementType ItemType,
+		EIKRigTreeElementType ItemType,
 		const bool bReplace);
 	void GetSelectedBoneChains(TArray<FBoneChain>& OutChains) const;
 	TArray<TSharedPtr<FIKRigTreeElement>> GetSelectedItems() const;
diff --git a/Engine/Plugins/Animation/IKRig/Source/IKRigEditor/Public/RigEditor/SIKRigSolverStack.h b/Engine/Plugins/Animation/IKRig/Source/IKRigEditor/Public/RigEditor/SIKRigSolverStack.h
--- a/Engine/Plugins/Animation/IKRig/Source/IKRigEditor/Public/RigEditor/SIKRigSolverStack.h
+++ b/Engine/Plugins/Animation/IKRig/Source/IKRigEditor/Public/RigEditor/SIKRigSolverStack.h
@@ -120,7 +120,7 @@ private:
 	void OnSelectionChanged(TSharedPtr<FSolverStackElement> InItem, ESelectInfo::Type SelectInfo);
 	/** when a solver is clicked on in the stack view */
 	void OnItemClicked(TSharedPtr<FSolverStackElement> InItem);
-	void ShowDetailsForItem(TSharedPtr<FSolverStackElement> InItem); 
+	void ShowDetailsForItems(const TArray<TSharedPtr<FSolverStackElement>>& InItems); 
 
 	/** list view generate row callback */
 	TSharedRef<ITableRow> MakeListRowWidget(TSharedPtr<FSolverStackElement> InElement, const TSharedRef<STableViewBase>& OwnerTable);
