[IK Retarget] Trivial Bughawk fixes after 46121270

#FailedValidations Preflight

diff --git a/Engine/Plugins/Animation/IKRig/Source/IKRigEditor/Private/RigEditor/IKRigEditorController.cpp b/Engine/Plugins/Animation/IKRig/Source/IKRigEditor/Private/RigEditor/IKRigEditorController.cpp
--- a/Engine/Plugins/Animation/IKRig/Source/IKRigEditor/Private/RigEditor/IKRigEditorController.cpp
+++ b/Engine/Plugins/Animation/IKRig/Source/IKRigEditor/Private/RigEditor/IKRigEditorController.cpp
@@ -1076,7 +1076,7 @@ void FIKRigEditorController::ShowDetailsForGoalSettings(const TArray<TSharedPtr<
 	TArray<UIKRigStructWrapperBase*> StructWrappers = AssetController->GetStructWrappers(ItemsToDisplay.Num(), TypeToDisplay);
 	for (int32 Index = 0; Index < ItemsToDisplay.Num(); ++Index)
 	{
-		TSharedPtr<FIKRigTreeElement> Item = InItems[Index];
+		TSharedPtr<FIKRigTreeElement> Item = ItemsToDisplay[Index];
 		const int32 SolverIndex = Item->EffectorIndex;
 		const FName GoalName = Item->EffectorGoalName;
 		
diff --git a/Engine/Plugins/Animation/IKRig/Source/IKRigEditor/Private/RigEditor/SIKRigHierarchy.cpp b/Engine/Plugins/Animation/IKRig/Source/IKRigEditor/Private/RigEditor/SIKRigHierarchy.cpp
--- a/Engine/Plugins/Animation/IKRig/Source/IKRigEditor/Private/RigEditor/SIKRigHierarchy.cpp
+++ b/Engine/Plugins/Animation/IKRig/Source/IKRigEditor/Private/RigEditor/SIKRigHierarchy.cpp
@@ -6,6 +6,7 @@
 #include "Rig/Solvers/IKRigSolverBase.h"
 #include "Engine/SkeletalMesh.h"
 #include "IPersonaToolkit.h"
+#include "RenderGraphResources.h"
 #include "SKismetInspector.h"
 #include "Dialogs/Dialogs.h"
 #include "SPositiveActionButton.h"
@@ -370,6 +371,11 @@ FReply SIKRigSkeletonRow::HandleAcceptDrop(
 		return FReply::Handled();
 	}
 
+	if (DragDropOp->Element.IsValid())
+	{
+		return FReply::Handled();
+	}
+
 	const FIKRigTreeElement& DraggedElement = *DragDropOp.Get()->Element.Pin().Get();
 	ensure(DraggedElement.ElementType == EIKRigTreeElementType::GOAL);		// drag only supported for goals
 	ensure(TargetItem.Get()->ElementType == EIKRigTreeElementType::BONE);	// drop only supported for bones
diff --git a/Engine/Plugins/Animation/IKRig/Source/IKRigEditor/Public/RigEditor/IKRigEditorController.h b/Engine/Plugins/Animation/IKRig/Source/IKRigEditor/Public/RigEditor/IKRigEditorController.h
--- a/Engine/Plugins/Animation/IKRig/Source/IKRigEditor/Public/RigEditor/IKRigEditorController.h
+++ b/Engine/Plugins/Animation/IKRig/Source/IKRigEditor/Public/RigEditor/IKRigEditorController.h
@@ -174,20 +174,20 @@ public:
 	
 	// show single transform of bone in details view 
 	void ShowDetailsForBone(const FName BoneName) const;
-	// show single BONE settings in details view 
+	// show BONE settings in details view 
 	void ShowDetailsForBoneSettings(const TArray<TSharedPtr<FIKRigTreeElement>>& InItems) const;
-	// show single GOAL settings in details view 
+	// show GOAL settings in details view 
 	void ShowDetailsForGoals(const TArray<FName>& GoalName) const;
-	// show single EFFECTOR settings in details view 
+	// show EFFECTOR settings in details view 
 	void ShowDetailsForGoalSettings(const TArray<TSharedPtr<FIKRigTreeElement>>& InItems) const;
-	// show single SOLVER settings in details view 
+	// show SOLVER settings in details view 
 	void ShowDetailsForSolvers(const TArray<int32>& InSolverIndices) const;
 	// show nothing in details view 
 	void ShowAssetDetails() const;
-	// show selected items in details view 
+	// show details for items 
 	void ShowDetailsForElements(const TArray<TSharedPtr<FIKRigTreeElement>>& InItems) const;
-	// show the currently configured StructViewer object in the details tab
-	void ShowStructWrappersInDetailsView(TArray<UIKRigStructWrapperBase*> InStructViewers) const;
+	// show the currently configured IKRigStructWrapper objects in the details tab
+	void ShowStructWrappersInDetailsView(TArray<UIKRigStructWrapperBase*> InStructWrappers) const;
 	// callback when detail is edited 
 	void OnFinishedChangingDetails(const FPropertyChangedEvent& PropertyChangedEvent) const;
 	// returns true if the supplied UObject is being shown in the details panel
