[IK Retarget] Fixed bug, source scale not reset when Source Scale Op disabled.

#FailedValidations Preflight

diff --git a/Engine/Plugins/Animation/IKRig/Source/IKRig/Private/Retargeter/IKRetargetProcessor.cpp b/Engine/Plugins/Animation/IKRig/Source/IKRig/Private/Retargeter/IKRetargetProcessor.cpp
--- a/Engine/Plugins/Animation/IKRig/Source/IKRig/Private/Retargeter/IKRetargetProcessor.cpp
+++ b/Engine/Plugins/Animation/IKRig/Source/IKRig/Private/Retargeter/IKRetargetProcessor.cpp
@@ -1867,7 +1867,7 @@ FName FIKRetargetProcessor::GetPelvisBone(ERetargetSourceOrTarget SourceOrTarget
 void FIKRetargetProcessor::ApplySourceScaleToPose(TArray<FTransform>& InOutSourceGlobalPose) const
 {
 	const FIKRetargetScaleSourceOp* ScaleSourceOp = GetFirstRetargetOpOfType<FIKRetargetScaleSourceOp>();
-	if (!ScaleSourceOp)
+	if (!ScaleSourceOp || !ScaleSourceOp->IsEnabled())
 	{
 		CurrentSourceScale.Factor = 1.0;
 		CurrentSourceScale.Pivot = FVector::ZeroVector;
