[IKRig] Fixed crash in IKRig anim node copy pose logic

diff --git a/Engine/Plugins/Animation/IKRig/Source/IKRig/Private/AnimNodes/AnimNode_IKRig.cpp b/Engine/Plugins/Animation/IKRig/Source/IKRig/Private/AnimNodes/AnimNode_IKRig.cpp
--- a/Engine/Plugins/Animation/IKRig/Source/IKRig/Private/AnimNodes/AnimNode_IKRig.cpp
+++ b/Engine/Plugins/Animation/IKRig/Source/IKRig/Private/AnimNodes/AnimNode_IKRig.cpp
@@ -88,7 +88,7 @@ void FAnimNode_IKRig::CopyInputPoseToSolver(FCompactPose& InputPose)
 			// bones that were recorded with rig indices == -1 were not in the
 			// Reference Skeleton that the IK Rig was initialized with and therefore
 			// are not considered as part of the solve.
-			if (*Index != -1)
+			if (IKRigSkeleton.CurrentPoseLocal.IsValidIndex(*Index))
 			{
 				IKRigSkeleton.CurrentPoseLocal[*Index] = InputPose[CPIndex];	
 			}
