[IK Retarget] Removed broken old dropdown menu in retarget editor toolbar.

diff --git a/Engine/Plugins/Animation/IKRig/Source/IKRigEditor/Private/RetargetEditor/IKRetargetEditor.cpp b/Engine/Plugins/Animation/IKRig/Source/IKRigEditor/Private/RetargetEditor/IKRetargetEditor.cpp
--- a/Engine/Plugins/Animation/IKRig/Source/IKRigEditor/Private/RetargetEditor/IKRetargetEditor.cpp
+++ b/Engine/Plugins/Animation/IKRig/Source/IKRigEditor/Private/RetargetEditor/IKRetargetEditor.cpp
@@ -270,14 +270,6 @@ void FIKRetargetEditor::FillToolbar(FToolBarBuilder& ToolbarBuilder)
 			TAttribute<FText>(), 
 			TAttribute<FSlateIcon>::Create(TAttribute<FSlateIcon>::FGetter::CreateSP(EditorController, &FIKRetargetEditorController::GetCurrentRetargetModeIcon))
 		);
-		
-	ToolbarBuilder.AddComboButton(
-		FUIAction(),
-		FOnGetContent::CreateSP(this, &FIKRetargetEditor::GenerateRetargetModesMenu),
-		LOCTEXT("RetargetMode_Label", "UI Modes"),
-		LOCTEXT("RetargetMode_ToolTip", "Choose which mode to display in the viewport."),
-		FSlateIcon(FAppStyle::GetAppStyleSetName(), "LevelEditor.Recompile"),
-		true);
 
 	ToolbarBuilder.EndSection();
 	
@@ -295,18 +287,6 @@ void FIKRetargetEditor::FillToolbar(FToolBarBuilder& ToolbarBuilder)
 	ToolbarBuilder.EndSection();
 }
 
-TSharedRef<SWidget> FIKRetargetEditor::GenerateRetargetModesMenu()
-{
-	FMenuBuilder MenuBuilder(true, GetToolkitCommands());
-	
-	MenuBuilder.BeginSection(TEXT("Retarget Modes"));
-	MenuBuilder.AddMenuEntry(FIKRetargetCommands::Get().RunRetargeter, TEXT("Run Retargeter"), TAttribute<FText>(), TAttribute<FText>(),  EditorController->GetRetargeterModeIcon(ERetargeterOutputMode::RunRetarget));
-	MenuBuilder.AddMenuEntry(FIKRetargetCommands::Get().EditRetargetPose, TEXT("Edit Retarget Pose"), TAttribute<FText>(), TAttribute<FText>(), EditorController->GetRetargeterModeIcon(ERetargeterOutputMode::EditRetargetPose));
-	MenuBuilder.EndSection();
-	
-	return MenuBuilder.MakeWidget();
-}
-
 FName FIKRetargetEditor::GetToolkitFName() const
 {
 	return FName("IKRetargetEditor");
diff --git a/Engine/Plugins/Animation/IKRig/Source/IKRigEditor/Public/RetargetEditor/IKRetargetEditor.h b/Engine/Plugins/Animation/IKRig/Source/IKRigEditor/Public/RetargetEditor/IKRetargetEditor.h
--- a/Engine/Plugins/Animation/IKRig/Source/IKRigEditor/Public/RetargetEditor/IKRetargetEditor.h
+++ b/Engine/Plugins/Animation/IKRig/Source/IKRigEditor/Public/RetargetEditor/IKRetargetEditor.h
@@ -72,7 +72,6 @@ private:
 	void BindCommands();
 	void ExtendToolbar();
 	void FillToolbar(FToolBarBuilder& ToolbarBuilder);
-	TSharedRef<SWidget> GenerateRetargetModesMenu();
 	/** END toolbar */
 	
 	/** preview scene setup */
