[IK Retarget] Allow displaying retarget chains column while pose editing.

diff --git a/Engine/Plugins/Animation/IKRig/Source/IKRig/Private/Retargeter/IKRetargetProcessor.cpp b/Engine/Plugins/Animation/IKRig/Source/IKRig/Private/Retargeter/IKRetargetProcessor.cpp
--- a/Engine/Plugins/Animation/IKRig/Source/IKRig/Private/Retargeter/IKRetargetProcessor.cpp
+++ b/Engine/Plugins/Animation/IKRig/Source/IKRig/Private/Retargeter/IKRetargetProcessor.cpp
@@ -1854,14 +1854,20 @@ FName FIKRetargetProcessor::GetPelvisBone(ERetargetSourceOrTarget SourceOrTarget
 		}
 	}
 
+	const UIKRetargeter* Asset = GetRetargetAsset();
+	if (!Asset)
+	{
+		return NAME_None;
+	}
+
 	// fallback to going directly to the IK Rig
-	const UIKRigDefinition* IKRigAsset = GetRetargetAsset()->GetIKRig(SourceOrTarget);
-	if (IKRigAsset)
+	const UIKRigDefinition* IKRigAsset = Asset->GetIKRig(SourceOrTarget);
+	if (!IKRigAsset)
 	{
-		return IKRigAsset->GetPelvis();
+		return NAME_None;	
 	}
 
-	return NAME_None;
+	return IKRigAsset->GetPelvis();
 }
 
 void FIKRetargetProcessor::ApplySourceScaleToPose(TArray<FTransform>& InOutSourceGlobalPose) const
diff --git a/Engine/Plugins/Animation/IKRig/Source/IKRigEditor/Private/RetargetEditor/IKRetargetEditorController.cpp b/Engine/Plugins/Animation/IKRig/Source/IKRigEditor/Private/RetargetEditor/IKRetargetEditorController.cpp
--- a/Engine/Plugins/Animation/IKRig/Source/IKRigEditor/Private/RetargetEditor/IKRetargetEditorController.cpp
+++ b/Engine/Plugins/Animation/IKRig/Source/IKRigEditor/Private/RetargetEditor/IKRetargetEditorController.cpp
@@ -555,7 +555,7 @@ FName FIKRetargetEditorController::GetChainNameFromBone(const FName& BoneName, E
 {
 	// get an initialized processor
 	const FIKRetargetProcessor* Processor = GetRetargetProcessor();
-	if (!(Processor && Processor->IsInitialized()))
+	if (!Processor)
 	{
 		return NAME_None;
 	}
