[IK Retarget] Exposed "bOverwriteExistingFiles" to batch retarget operation.

diff --git a/Engine/Plugins/Animation/IKRig/Source/IKRigEditor/Private/RetargetEditor/IKRetargetBatchOperation.cpp b/Engine/Plugins/Animation/IKRig/Source/IKRigEditor/Private/RetargetEditor/IKRetargetBatchOperation.cpp
--- a/Engine/Plugins/Animation/IKRig/Source/IKRigEditor/Private/RetargetEditor/IKRetargetBatchOperation.cpp
+++ b/Engine/Plugins/Animation/IKRig/Source/IKRigEditor/Private/RetargetEditor/IKRetargetBatchOperation.cpp
@@ -832,7 +832,8 @@ TArray<FAssetData> UIKRetargetBatchOperation::DuplicateAndRetarget(
 	const FString& Replace,
 	const FString& Prefix,
 	const FString& Suffix,
-	const bool bIncludeReferencedAssets)
+	const bool bIncludeReferencedAssets,
+	const bool bOverwriteExistingFiles)
 {
 	// fill the context with all the data needed to run a batch retarget
 	FIKRetargetBatchOperationContext Context;
@@ -851,6 +852,7 @@ TArray<FAssetData> UIKRetargetBatchOperation::DuplicateAndRetarget(
 	Context.NameRule.ReplaceFrom = Search;
 	Context.NameRule.ReplaceTo = Replace;
 	Context.bIncludeReferencedAssets = bIncludeReferencedAssets;
+	Context.bOverwriteExistingFiles = bOverwriteExistingFiles;
 
 	// actually run the batch operation
 	UIKRetargetBatchOperation* BatchOperation = NewObject<UIKRetargetBatchOperation>();
diff --git a/Engine/Plugins/Animation/IKRig/Source/IKRigEditor/Public/RetargetEditor/IKRetargetBatchOperation.h b/Engine/Plugins/Animation/IKRig/Source/IKRigEditor/Public/RetargetEditor/IKRetargetBatchOperation.h
--- a/Engine/Plugins/Animation/IKRig/Source/IKRigEditor/Public/RetargetEditor/IKRetargetBatchOperation.h
+++ b/Engine/Plugins/Animation/IKRig/Source/IKRigEditor/Public/RetargetEditor/IKRetargetBatchOperation.h
@@ -127,7 +127,8 @@ public:
 	 * @param Replace A string to replace with in the file name
 	 * @param Suffix A string to add at the end of the new file name
 	 * @param Prefix A string to add to the start of the new file name
-	 * @param bRemapReferencedAssets Whether to remap existing references in the animation assets
+	 * @param bIncludeReferencedAssets Whether to retarget animation assets referenced by the AssetsToRetarget
+	 * @param bOverwriteExistingFiles Whether to overwrite any existing output files with the same name, or create new files with a unique name
 	 * 
 	 * Return: list of new animation files created.*/
 	UFUNCTION(BlueprintCallable, Category=IKBatchRetarget)
@@ -140,7 +141,8 @@ public:
 		const FString& Replace = "",
 		const FString& Prefix = "",
 		const FString& Suffix = "",
-		const bool bIncludeReferencedAssets=true);
+		const bool bIncludeReferencedAssets=true,
+		const bool bOverwriteExistingFiles=false);
 	
 	// Actually run the process to duplicate and retarget the assets for the given context
 	UE_API void RunRetarget(FIKRetargetBatchOperationContext& Context);
