[IK Rig] Exposed IK Rig Control Rig node inputs to details panel.

diff --git a/Engine/Plugins/Animation/IKRig/Source/IKRig/Public/AnimNodes/RigUnit_IKRig.h b/Engine/Plugins/Animation/IKRig/Source/IKRig/Public/AnimNodes/RigUnit_IKRig.h
--- a/Engine/Plugins/Animation/IKRig/Source/IKRig/Public/AnimNodes/RigUnit_IKRig.h
+++ b/Engine/Plugins/Animation/IKRig/Source/IKRig/Public/AnimNodes/RigUnit_IKRig.h
@@ -37,19 +37,19 @@ struct FIKRigGoalInput
 	}
 
 	/** The name of the goal to affect. Must match the name in the IK Rig. */
-	UPROPERTY()
+	UPROPERTY(EditAnywhere, Category=GoalSettings)
 	FName GoalName;
 
 	/** The position and rotation target for the IK Goal. */
-	UPROPERTY()
+	UPROPERTY(EditAnywhere, Category=GoalSettings)
 	FTransform Transform;
 
 	/** Range 0-1, default 1. Blends the Goal position from the input pose position to the Goal position. */
-	UPROPERTY()
+	UPROPERTY(EditAnywhere, Category=GoalSettings)
 	double PositionAlpha;
 
 	/** Range 0-1, default 1. Blends the Goal rotation from the input pose rotation to the Goal rotation. */
-	UPROPERTY()
+	UPROPERTY(EditAnywhere, Category=GoalSettings)
 	double RotationAlpha;
 };
 
