[IK Retarget] Better error checking for resolved chains to resolve crash potential.

diff --git a/Engine/Plugins/Animation/IKRig/Source/IKRig/Private/Retargeter/IKRetargetProcessor.cpp b/Engine/Plugins/Animation/IKRig/Source/IKRig/Private/Retargeter/IKRetargetProcessor.cpp
--- a/Engine/Plugins/Animation/IKRig/Source/IKRig/Private/Retargeter/IKRetargetProcessor.cpp
+++ b/Engine/Plugins/Animation/IKRig/Source/IKRig/Private/Retargeter/IKRetargetProcessor.cpp
@@ -574,6 +574,7 @@ FResolvedBoneChain::FResolvedBoneChain(
 	if (BoneIndices.Last() != StartIndex)
 	{
 		BoneIndices.Reset();
+		bEndIsStartOrChildOfStart = false;
 		return;
 	}
 	
@@ -592,7 +593,10 @@ FResolvedBoneChain::FResolvedBoneChain(
 
 bool FResolvedBoneChain::IsValid() const
 {
-	return bFoundStartBone && bFoundEndBone && bEndIsStartOrChildOfStart;
+	return bFoundStartBone &&
+		bFoundEndBone &&
+		bEndIsStartOrChildOfStart &&
+		BoneIndices.Num() == RefPoseGlobalTransforms.Num() == RefPoseLocalTransforms.Num();
 }
 
 void FResolvedBoneChain::UpdatePoseFromSkeleton(const FRetargetSkeleton& InSkeleton)
