[IK Retarget] Fixes issue causing broken retarget chains (regression from 46194288)

diff --git a/Engine/Plugins/Animation/IKRig/Source/IKRig/Private/Retargeter/IKRetargetProcessor.cpp b/Engine/Plugins/Animation/IKRig/Source/IKRig/Private/Retargeter/IKRetargetProcessor.cpp
--- a/Engine/Plugins/Animation/IKRig/Source/IKRig/Private/Retargeter/IKRetargetProcessor.cpp
+++ b/Engine/Plugins/Animation/IKRig/Source/IKRig/Private/Retargeter/IKRetargetProcessor.cpp
@@ -596,7 +596,8 @@ bool FResolvedBoneChain::IsValid() const
 	return bFoundStartBone &&
 		bFoundEndBone &&
 		bEndIsStartOrChildOfStart &&
-		BoneIndices.Num() == RefPoseGlobalTransforms.Num() == RefPoseLocalTransforms.Num();
+		BoneIndices.Num() == RefPoseGlobalTransforms.Num() &&
+		BoneIndices.Num() == RefPoseLocalTransforms.Num();
 }
 
 void FResolvedBoneChain::UpdatePoseFromSkeleton(const FRetargetSkeleton& InSkeleton)
