[IKRetarget] Fixed potential non-normalized rotation in IK Chains Op.

diff --git a/Engine/Plugins/Animation/IKRig/Source/IKRig/Private/Retargeter/RetargetOps/IKChainsOp.cpp b/Engine/Plugins/Animation/IKRig/Source/IKRig/Private/Retargeter/RetargetOps/IKChainsOp.cpp
--- a/Engine/Plugins/Animation/IKRig/Source/IKRig/Private/Retargeter/RetargetOps/IKChainsOp.cpp
+++ b/Engine/Plugins/Animation/IKRig/Source/IKRig/Private/Retargeter/RetargetOps/IKChainsOp.cpp
@@ -87,7 +87,7 @@ FQuat FIKChainRetargeter::CalculateGoalRotation(const TArray<FTransform>& InSour
 	
 	// apply delta rotation from source input
 	const FQuat CurrentSourceEndRotation = CurrentSourceEndGlobal.GetRotation();
-	const FQuat SourceDeltaRotation = CurrentSourceEndRotation * InitialSourceEndGlobal.GetRotation().Inverse();
+	const FQuat SourceDeltaRotation = (CurrentSourceEndRotation * InitialSourceEndGlobal.GetRotation().Inverse()).GetNormalized();
 	FQuat GoalRotation = SourceDeltaRotation * InitialTargetEndGlobal.GetRotation();
 
 	// blend to source rotation
