Remove bone tracks from duplicate anim sequences during retarget process.  This keeps bone track indices for the new sequence the same as the target skeleton.  Previously, bone indices could get out of order which would break root motion as it assumes root motion track is index 0

diff --git a/Engine/Plugins/Animation/IKRig/Source/IKRigEditor/Private/RetargetEditor/IKRetargetBatchOperation.cpp b/Engine/Plugins/Animation/IKRig/Source/IKRigEditor/Private/RetargetEditor/IKRetargetBatchOperation.cpp
--- a/Engine/Plugins/Animation/IKRig/Source/IKRigEditor/Private/RetargetEditor/IKRetargetBatchOperation.cpp
+++ b/Engine/Plugins/Animation/IKRig/Source/IKRigEditor/Private/RetargetEditor/IKRetargetBatchOperation.cpp
@@ -241,7 +241,8 @@ void UIKRetargetBatchOperation::RetargetAssets(
 			AnimSequenceToRetarget->RetargetSource = NAME_None;
 			AnimSequenceToRetarget->SetRetargetSourceAsset(Context.TargetMesh);
 
-			// removes all bone tracks not belonging to this skeleton, adds those that do
+			// removes all bone tracks and add those that exist in the target skeleton
+			Controller.RemoveAllBoneTracks(bShouldTransact);
 			Controller.UpdateWithSkeleton(NewSkeleton, bShouldTransact);
 
 			// done editing sequence data, close bracket
@@ -396,6 +397,7 @@ void UIKRetargetBatchOperation::ConvertAnimation(
 		TargetSeqController.NotifyPopulated(); // must set bIsPopulated=true otherwise "UpdateWithSkeleton" will early out
 		// this removes all bone tracks and reinitializes the FKControlRig to use the new skeleton.
 		// NOTE: we do NOT ResetModel() because we want to keep curves and attributes from the source
+		TargetSeqController.RemoveAllBoneTracks(bShouldTransact);
 		TargetSeqController.UpdateWithSkeleton(const_cast<USkeleton*>(TargetSkeleton.SkeletalMesh->GetSkeleton()), bShouldTransact);
 
 		// number of frames in this animation
