IKRetargetProcessor: Add RunAfterParent functionality for retarget ops.

diff --git a/Engine/Plugins/Animation/IKRig/Source/IKRig/Private/Retargeter/IKRetargetProcessor.cpp b/Engine/Plugins/Animation/IKRig/Source/IKRig/Private/Retargeter/IKRetargetProcessor.cpp
--- a/Engine/Plugins/Animation/IKRig/Source/IKRig/Private/Retargeter/IKRetargetProcessor.cpp
+++ b/Engine/Plugins/Animation/IKRig/Source/IKRig/Private/Retargeter/IKRetargetProcessor.cpp
@@ -1345,6 +1345,10 @@ void FIKRetargetProcessor::RunRetargetOps(
 	};
 
 	TArray<FIKRetargetOpBase*> OpsRunBeforeChildrenThisTick;
+
+	// TODO: This won't work if child ops can have children
+	FIKRetargetOpBase* CurrentParentOp = nullptr;
+	TArray<FIKRetargetOpBase*> OpsToRunAfterParent;
 	
 	for (FInstancedStruct& OpStruct : OpStack)
 	{
@@ -1366,8 +1370,26 @@ void FIKRetargetProcessor::RunRetargetOps(
 		
 		if (ShouldExecuteOp(Op))
 		{
+			// Give child ops entry point to run after parent op finishes
+			if (ParentOp)
+			{
+				ensure(!CurrentParentOp || CurrentParentOp == ParentOp);
+				CurrentParentOp = ParentOp;
+				OpsToRunAfterParent.Add(&Op);
+			}
+			
 			// run the op
-			Op.Run(*this, InDeltaTime, InSourceGlobalPose, OutTargetGlobalPose);	
+			Op.Run(*this, InDeltaTime, InSourceGlobalPose, OutTargetGlobalPose);
+
+			// Run children list after parent
+			if (CurrentParentOp == &Op)
+			{
+				CurrentParentOp = nullptr;
+				for (FIKRetargetOpBase* ChildOp : OpsToRunAfterParent)
+				{
+					ChildOp->RunAfterParent(*this, InDeltaTime, InSourceGlobalPose, OutTargetGlobalPose);
+				}
+			}
 		}
 
 		// END op profiler scope
diff --git a/Engine/Plugins/Animation/IKRig/Source/IKRig/Public/Retargeter/IKRetargetOps.h b/Engine/Plugins/Animation/IKRig/Source/IKRig/Public/Retargeter/IKRetargetOps.h
--- a/Engine/Plugins/Animation/IKRig/Source/IKRig/Public/Retargeter/IKRetargetOps.h
+++ b/Engine/Plugins/Animation/IKRig/Source/IKRig/Public/Retargeter/IKRetargetOps.h
@@ -281,6 +281,15 @@ public:
 		const TArray<FTransform>& InSourceGlobalPose,
 		TArray<FTransform>& OutTargetGlobalPose){};
 
+	/** (optional) override to give this op a chance to do work after it's parent is run
+	* This is executed right after the parent op in the stack and right before the next op after parent
+	* NOTE: unless this is a child op, this will have no effect */
+	virtual void RunAfterParent(
+		FIKRetargetProcessor& InProcessor,
+		const double InDeltaTime,
+		const TArray<FTransform>& InSourceGlobalPose,
+		TArray<FTransform>& OutTargetGlobalPose){};
+
 	/** (optional) return true to disallow multiple copies of this op in the stack */
 	virtual bool IsSingleton() const { return false; };
 
