IK Retargeter: Add the ability to specify a target path when duplicating and retargeting assets through Blueprint.

#rnx

diff --git a/Engine/Plugins/Animation/IKRig/Source/IKRigEditor/Private/RetargetEditor/IKRetargetBatchOperation.cpp b/Engine/Plugins/Animation/IKRig/Source/IKRigEditor/Private/RetargetEditor/IKRetargetBatchOperation.cpp
--- a/Engine/Plugins/Animation/IKRig/Source/IKRigEditor/Private/RetargetEditor/IKRetargetBatchOperation.cpp
+++ b/Engine/Plugins/Animation/IKRig/Source/IKRigEditor/Private/RetargetEditor/IKRetargetBatchOperation.cpp
@@ -832,6 +832,8 @@ TArray<FAssetData> UIKRetargetBatchOperation::DuplicateAndRetarget(
 	const FString& Replace,
 	const FString& Prefix,
 	const FString& Suffix,
+	const FString& TargetPath,
+	const bool bUseSourcePath,
 	const bool bIncludeReferencedAssets,
 	const bool bOverwriteExistingFiles)
 {
@@ -851,6 +853,11 @@ TArray<FAssetData> UIKRetargetBatchOperation::DuplicateAndRetarget(
 	Context.NameRule.Suffix = Suffix;
 	Context.NameRule.ReplaceFrom = Search;
 	Context.NameRule.ReplaceTo = Replace;
+	if (!TargetPath.IsEmpty())
+	{
+		Context.NameRule.FolderPath = TargetPath;
+	}
+	Context.bUseSourcePath = bUseSourcePath;
 	Context.bIncludeReferencedAssets = bIncludeReferencedAssets;
 	Context.bOverwriteExistingFiles = bOverwriteExistingFiles;
 
diff --git a/Engine/Plugins/Animation/IKRig/Source/IKRigEditor/Public/RetargetEditor/IKRetargetBatchOperation.h b/Engine/Plugins/Animation/IKRig/Source/IKRigEditor/Public/RetargetEditor/IKRetargetBatchOperation.h
--- a/Engine/Plugins/Animation/IKRig/Source/IKRigEditor/Public/RetargetEditor/IKRetargetBatchOperation.h
+++ b/Engine/Plugins/Animation/IKRig/Source/IKRigEditor/Public/RetargetEditor/IKRetargetBatchOperation.h
@@ -127,10 +127,13 @@ public:
 	 * @param Replace A string to replace with in the file name
 	 * @param Suffix A string to add at the end of the new file name
 	 * @param Prefix A string to add to the start of the new file name
+	 * @param TargetPath A string containing the destination path for the retargeted assets. Ignored if bUseSourcePath is set to true.
+	 * @param bUseSourcePath Place the assets in the same folder as their source location instead of the given target path.
 	 * @param bIncludeReferencedAssets Whether to retarget animation assets referenced by the AssetsToRetarget
 	 * @param bOverwriteExistingFiles Whether to overwrite any existing output files with the same name, or create new files with a unique name
 	 * 
-	 * Return: list of new animation files created.*/
+	 * Return: list of new animation files created.
+	 */
 	UFUNCTION(BlueprintCallable, Category=IKBatchRetarget)
 	static UE_API TArray<FAssetData> DuplicateAndRetarget(
 		const TArray<FAssetData>& AssetsToRetarget,
@@ -141,6 +144,8 @@ public:
 		const FString& Replace = "",
 		const FString& Prefix = "",
 		const FString& Suffix = "",
+		const FString& TargetPath = "",
+		const bool bUseSourcePath = false,
 		const bool bIncludeReferencedAssets=true,
 		const bool bOverwriteExistingFiles=false);
 	
