[IK Retarget] Fixed crash when exporting retarget assets ontop of redirectors with same name.

diff --git a/Engine/Plugins/Animation/IKRig/Source/IKRigEditor/Private/RetargetEditor/SRetargetAnimAssetsWindow.cpp b/Engine/Plugins/Animation/IKRig/Source/IKRigEditor/Private/RetargetEditor/SRetargetAnimAssetsWindow.cpp
--- a/Engine/Plugins/Animation/IKRig/Source/IKRigEditor/Private/RetargetEditor/SRetargetAnimAssetsWindow.cpp
+++ b/Engine/Plugins/Animation/IKRig/Source/IKRigEditor/Private/RetargetEditor/SRetargetAnimAssetsWindow.cpp
@@ -1109,20 +1109,21 @@ FReply SRetargetAnimAssetsWindow::OnExportRetargetAssets()
 	auto SaveAssetToDisk = [](UObject* Asset, FString& AssetName, const FString& FolderPath)
 	{
 		// create a unique package name and path
-		const FString BasePackageName = FolderPath + "/";
-		FString UniquePackageName;
+		const FString BasePackageName = FolderPath / AssetName;
+		const FString Suffix = "";
+		FString UniquePackageName, UniqueAssetName;
 		IAssetTools& AssetTools = FModuleManager::LoadModuleChecked<FAssetToolsModule>("AssetTools").Get();
-		AssetTools.CreateUniqueAssetName(BasePackageName, AssetName, UniquePackageName, AssetName);
+		AssetTools.CreateUniqueAssetName(BasePackageName, Suffix, UniquePackageName, UniqueAssetName);
 		
 		// duplicate the asset
 		FObjectDuplicationParameters ObjParameters(Asset, GetTransientPackage());
-		ObjParameters.DestName = FName(AssetName);
+		ObjParameters.DestName = MakeUniqueObjectName(GetTransientPackage(), Asset->GetClass(), *UniqueAssetName);
 		ObjParameters.ApplyFlags = RF_Transactional;
 		UObject* AssetToSave = StaticDuplicateObjectEx(ObjParameters);
 			
 		// create a new asset package
 		UPackage* Package = CreatePackage(*UniquePackageName);
-		AssetToSave->Rename(*AssetName, Package);
+		AssetToSave->Rename(*UniqueAssetName, Package);
 		FAssetRegistryModule::AssetCreated(AssetToSave);
 		Package->MarkPackageDirty();
 
