[IK Retarget] Fixes crash when editing retarget pose without an IK Rig.

diff --git a/Engine/Plugins/Animation/IKRig/Source/IKRigEditor/Private/RetargetEditor/IKRetargeterController.cpp b/Engine/Plugins/Animation/IKRig/Source/IKRigEditor/Private/RetargetEditor/IKRetargeterController.cpp
--- a/Engine/Plugins/Animation/IKRig/Source/IKRigEditor/Private/RetargetEditor/IKRetargeterController.cpp
+++ b/Engine/Plugins/Animation/IKRig/Source/IKRigEditor/Private/RetargetEditor/IKRetargeterController.cpp
@@ -812,14 +812,19 @@ FName UIKRetargeterController::GetPelvisBone(const ERetargetSourceOrTarget Sourc
 
 void UIKRetargeterController::CleanPoseList(const ERetargetSourceOrTarget SourceOrTarget) const
 {
-	const UIKRigDefinition* IKRig = GetIKRig(SourceOrTarget);
-	if (!IKRig)
+	USkeletalMesh* SkeletalMesh = GetPreviewMesh(SourceOrTarget);
+	if (!SkeletalMesh)
+	{
+		return;
+	}
+
+	const TWeakPtr<const FIKRigSkeleton> IKRigSkeleton = GetCachedIKRigSkeleton(SkeletalMesh);
+	if (!IKRigSkeleton.IsValid())
 	{
 		return;
 	}
 	
 	// remove all bone offsets that are no longer part of the skeleton
-	const TArray<FName> AllowedBoneNames = IKRig->GetSkeleton().BoneNames;
 	TMap<FName, FIKRetargetPose>& RetargetPoses = GetRetargetPoses(SourceOrTarget);
 	for (TTuple<FName, FIKRetargetPose>& Pose : RetargetPoses)
 	{
@@ -827,7 +832,7 @@ void UIKRetargeterController::CleanPoseList(const ERetargetSourceOrTarget Source
 		TArray<FName> BonesToRemove;
 		for (TTuple<FName, FQuat>& BoneOffset : Pose.Value.BoneRotationOffsets)
 		{
-			if (!AllowedBoneNames.Contains(BoneOffset.Key))
+			if (!IKRigSkeleton.Pin()->BoneNames.Contains(BoneOffset.Key))
 			{
 				BonesToRemove.Add(BoneOffset.Key);
 			}
@@ -840,7 +845,7 @@ void UIKRetargeterController::CleanPoseList(const ERetargetSourceOrTarget Source
 		}
 
 		// sort the pose offset from leaf to root
-		Pose.Value.SortHierarchically(IKRig->GetSkeleton());
+		Pose.Value.SortHierarchically(*IKRigSkeleton.Pin().Get());
 	}
 }
 
@@ -1307,10 +1312,20 @@ void UIKRetargeterController::SetRotationOffsetForRetargetPoseBone(
 	const ERetargetSourceOrTarget SourceOrTarget) const
 {
 	FScopeLock Lock(&ControllerLock);
+
+	// add bone offset to pose
 	FIKRetargetPose& Pose = GetCurrentRetargetPose(SourceOrTarget);
 	Pose.SetDeltaRotationForBone(BoneName, RotationOffset);
-	const UIKRigDefinition* IKRig = GetAsset()->GetIKRig(SourceOrTarget);
-	Pose.SortHierarchically(IKRig->GetSkeleton());
+
+	// sort the pose hierarchically based on skeleton
+	if (USkeletalMesh* SkeletalMesh = GetAsset()->GetPreviewMesh(SourceOrTarget))
+	{
+		TWeakPtr<const FIKRigSkeleton> IKRigSkeleton = GetCachedIKRigSkeleton(SkeletalMesh);
+		if (IKRigSkeleton.IsValid())
+		{
+			Pose.SortHierarchically(*IKRigSkeleton.Pin().Get());
+		}
+	}
 }
 
 FQuat UIKRetargeterController::GetRotationOffsetForRetargetPoseBone(
@@ -1414,6 +1429,24 @@ FName UIKRetargeterController::MakePoseNameUnique(const FString& PoseName, const
 	return FName(UniqueName);
 }
 
+TWeakPtr<const FIKRigSkeleton> UIKRetargeterController::GetCachedIKRigSkeleton(const USkeletalMesh* InSkeletalMesh) const
+{
+	if (!InSkeletalMesh)
+	{
+		return nullptr;
+	}
+
+	if (TSharedPtr<FIKRigSkeleton>* IKRigSkeleton = CachedSkeletons.Find(InSkeletalMesh))
+	{
+		return IKRigSkeleton->ToWeakPtr();
+	}
+
+	TSharedPtr<FIKRigSkeleton>& NewIKRigSkeleton = CachedSkeletons.Add(InSkeletalMesh);
+	NewIKRigSkeleton = MakeShared<FIKRigSkeleton>();
+	NewIKRigSkeleton->SetInputSkeleton(InSkeletalMesh, TArray<FName>() /*ExcludedBones*/);
+	return NewIKRigSkeleton.ToWeakPtr();
+}
+
 //
 // BEGIN DEPRECATED API
 //
diff --git a/Engine/Plugins/Animation/IKRig/Source/IKRigEditor/Public/RetargetEditor/IKRetargeterController.h b/Engine/Plugins/Animation/IKRig/Source/IKRigEditor/Public/RetargetEditor/IKRetargeterController.h
--- a/Engine/Plugins/Animation/IKRig/Source/IKRigEditor/Public/RetargetEditor/IKRetargeterController.h
+++ b/Engine/Plugins/Animation/IKRig/Source/IKRigEditor/Public/RetargetEditor/IKRetargeterController.h
@@ -446,6 +446,8 @@ public:
 	UE_API FName MakePoseNameUnique(
 		const FString& PoseName,
 		const ERetargetSourceOrTarget SourceOrTarget) const;
+
+	TWeakPtr<const FIKRigSkeleton> GetCachedIKRigSkeleton(const USkeletalMesh* InSkeletalMesh) const;
 	
 private:
 	
@@ -536,6 +538,9 @@ private:
 	// a dummy UObject to wrap UStructs for the details panel
 	UPROPERTY(transient)
 	TObjectPtr<UIKRigStructViewer> StructViewer = nullptr;
+
+	// transient cache used for duration of editing
+	mutable TMap<const USkeletalMesh*, TSharedPtr<FIKRigSkeleton>> CachedSkeletons;
 	
 public:
 
