[IK Retarget] Fixes crash when aligning bones with overlapping retarget chains.

diff --git a/Engine/Plugins/Animation/IKRig/Source/IKRig/Private/Retargeter/IKRetargetProcessor.cpp b/Engine/Plugins/Animation/IKRig/Source/IKRig/Private/Retargeter/IKRetargetProcessor.cpp
--- a/Engine/Plugins/Animation/IKRig/Source/IKRig/Private/Retargeter/IKRetargetProcessor.cpp
+++ b/Engine/Plugins/Animation/IKRig/Source/IKRig/Private/Retargeter/IKRetargetProcessor.cpp
@@ -1515,7 +1515,7 @@ bool FIKRetargetProcessor::WasInitializedWithTheseAssets(
 
 #if WITH_EDITOR
 
-bool FIKRetargetProcessor::IsBoneInAMappedChain(
+bool FIKRetargetProcessor::IsBoneMapped(
 	const FName BoneName,
 	const ERetargetSourceOrTarget SourceOrTarget) const
 {
@@ -1524,39 +1524,29 @@ bool FIKRetargetProcessor::IsBoneInAMappedChain(
 	// NOTE: cannot use Skeleton.IsBoneRetargeted because that only exists
 	// on the target skeleton and this function needs to work for either skeleton
 
-	// bone must exist
-	const FRetargetSkeleton& Skeleton = GetSkeleton(SourceOrTarget);
-	const int32 BoneIndex = Skeleton.FindBoneIndexByName(BoneName);
-	if (BoneIndex == INDEX_NONE)
+	if (BoneName == NAME_None)
 	{
 		return false;
 	}
 
-	// pelvis is always retargeted
+	// special case: if the bone is a pelvis, check if it's mapped to a pelvis on the other side
 	const FName PelvisBoneName = GetPelvisBone(SourceOrTarget, ERetargetOpsToSearch::AssetOps);
 	if (BoneName == PelvisBoneName)
 	{
-		return true;
-	}
-	
-	// bones must be in at least one resolved chain
-	const TSet<FName> ChainsThatContainBone = AllBoneChains.GetChainsThatContainBone(BoneIndex, SourceOrTarget);
-	if (ChainsThatContainBone.IsEmpty())
-	{
-		return false;
-	}
-
-	// bone's chain must be mapped to alternate chain (source-to-target or vice versa)
-	for (const FName ChainContainingBone : ChainsThatContainBone)
-	{
-		const FName MappedChain = GetFirstChainMappedToChain(ChainContainingBone, SourceOrTarget);
-		if (MappedChain != NAME_None)
+		const ERetargetSourceOrTarget OtherSide = SourceOrTarget == ERetargetSourceOrTarget::Source ? ERetargetSourceOrTarget::Target : ERetargetSourceOrTarget::Source;
+		const FName OtherPelvisBoneName = GetPelvisBone(OtherSide, ERetargetOpsToSearch::AssetOps);
+		if (OtherPelvisBoneName != NAME_None)
 		{
 			return true;
 		}
 	}
 	
-	return false;
+	// find chains on both sides
+	FName ChainWithBone;
+	FName MappedChain;
+	const bool bFoundMapping = GetFirstChainMappingForBone(BoneName, SourceOrTarget, ChainWithBone, MappedChain);
+
+	return bFoundMapping;
 }
 
 int32 FIKRetargetProcessor::GetBoneIndexFromName(
@@ -1572,47 +1562,56 @@ const FRetargeterBoneChains& FIKRetargetProcessor::GetBoneChains()
 	return AllBoneChains;
 }
 
-FName FIKRetargetProcessor::GetChainNameForBone(const FName BoneName, const ERetargetSourceOrTarget SourceOrTarget) const
+FName FIKRetargetProcessor::GetFirstChainForBone(const FName BoneName, const ERetargetSourceOrTarget SourceOrTarget) const
+{
+	const int32 BoneIndex = GetBoneIndexFromName(BoneName, SourceOrTarget);
+	const TSet<FName> ChainsThatContainBone = AllBoneChains.GetChainsThatContainBone(BoneIndex, SourceOrTarget);
+	return !ChainsThatContainBone.IsEmpty() ? *ChainsThatContainBone.begin() : NAME_None;
+}
+
+TSet<FName> FIKRetargetProcessor::GetAllChainNamesForBone(
+	const FName BoneName,
+	const ERetargetSourceOrTarget SourceOrTarget) const
 {
-	const bool bUseSource = SourceOrTarget == ERetargetSourceOrTarget::Source;
-	const FRetargetSkeleton& Skeleton = bUseSource ? SourceSkeleton : TargetSkeleton;
 	const int32 BoneIndex = GetBoneIndexFromName(BoneName, SourceOrTarget);
-	const FName RootBoneName = GetPelvisBone(SourceOrTarget, ERetargetOpsToSearch::AssetOps);
-	if (BoneName == RootBoneName)
+	return AllBoneChains.GetChainsThatContainBone(BoneIndex, SourceOrTarget);
+}
+
+bool FIKRetargetProcessor::GetFirstChainMappingForBone(
+	const FName BoneName,
+	const ERetargetSourceOrTarget SourceOrTarget,
+	FName& OutChainContainingBone,
+	FName& OutMappedChain) const
+{
+	if (!ensure(bIsInitialized))
 	{
-		return FName("Retarget Pelvis");
+		return false;
 	}
+	
+	OutChainContainingBone = NAME_None;
+	OutMappedChain = NAME_None;
 
-	auto GetCombinedName = [](const TSet<FName>& NameSet) -> FName
+	// bone must exist
+	const int32 BoneIndex = GetBoneIndexFromName(BoneName, SourceOrTarget);
+	if (BoneIndex == INDEX_NONE)
 	{
-		if (NameSet.Num() == 0)
-		{
-			return NAME_None;
-		}
-    
-		if (NameSet.Num() == 1)
-		{
-			return NameSet.Array()[0];
-		}
-    
-		FString Result;
-		bool bFirst = true;
-    
-		for (const FName& Name : NameSet)
+		return false;
+	}
+	
+	// bone must be in a chain that is mapped to alternate chain (source-to-target or vice versa)
+	const TSet<FName> ChainsThatContainBone = AllBoneChains.GetChainsThatContainBone(BoneIndex, SourceOrTarget);
+	for (const FName ChainContainingBone : ChainsThatContainBone)
+	{
+		const FName MappedChain = GetFirstChainMappedToChain(ChainContainingBone, SourceOrTarget);
+		if (MappedChain != NAME_None)
 		{
-			if (!bFirst)
-			{
-				Result += TEXT(", ");
-			}
-			Result += Name.ToString();
-			bFirst = false;
+			OutChainContainingBone = ChainContainingBone;
+			OutMappedChain = MappedChain;
+			return true;
 		}
-    
-		return FName(Result);
-	};
+	}
 	
-	const TSet<FName> ChainsThatContainBone = AllBoneChains.GetChainsThatContainBone(BoneIndex, SourceOrTarget);
-	return GetCombinedName(ChainsThatContainBone);
+	return false;
 }
 
 FTransform FIKRetargetProcessor::GetGlobalRetargetPoseAtParam(
@@ -1648,9 +1647,15 @@ FTransform FIKRetargetProcessor::GetRetargetPoseBoneTransform(
 
 float FIKRetargetProcessor::GetParamOfBoneInChain(
 	const FName InBoneName,
-	const ERetargetSourceOrTarget SourceOrTarget) const
+	const ERetargetSourceOrTarget SourceOrTarget,
+	const FName InChainName) const
 {
-	const FName ChainName = GetChainNameForBone(InBoneName, SourceOrTarget);
+	FName ChainName = InChainName;
+	if (ChainName == NAME_None)
+	{
+		ChainName = GetFirstChainForBone(InBoneName, SourceOrTarget);
+	}
+	
 	const FResolvedBoneChain* Chain = GetBoneChains().GetResolvedBoneChainByName(ChainName, SourceOrTarget);
 	if (!Chain)
 	{
@@ -1659,7 +1664,13 @@ float FIKRetargetProcessor::GetParamOfBoneInChain(
 
 	const int32 BoneIndexInSkeleton = GetBoneIndexFromName(InBoneName, SourceOrTarget);
 	const int32 BoneIndexInChain =  Chain->BoneIndices.Find(BoneIndexInSkeleton);
-	return Chain->Params[BoneIndexInChain];
+	if (Chain->Params.IsValidIndex(BoneIndexInChain))
+	{
+		return Chain->Params[BoneIndexInChain];	
+	}
+
+	ensureMsgf(false, TEXT("Should never fail to find a chain index if bone exists in chain."));
+	return INDEX_NONE;
 }
 
 FName FIKRetargetProcessor::GetClosestBoneToParam(
diff --git a/Engine/Plugins/Animation/IKRig/Source/IKRig/Public/Retargeter/IKRetargetProcessor.h b/Engine/Plugins/Animation/IKRig/Source/IKRig/Public/Retargeter/IKRetargetProcessor.h
--- a/Engine/Plugins/Animation/IKRig/Source/IKRig/Public/Retargeter/IKRetargetProcessor.h
+++ b/Engine/Plugins/Animation/IKRig/Source/IKRig/Public/Retargeter/IKRetargetProcessor.h
@@ -505,20 +505,24 @@ private:
 	FOnRetargeterInitialized RetargeterInitialized;
 public:
 	/** Returns true if the bone is part of a retarget chain or root bone, false otherwise. */
-	UE_API bool IsBoneInAMappedChain(const FName BoneName, const ERetargetSourceOrTarget SourceOrTarget) const;
+	UE_API bool IsBoneMapped(const FName BoneName, const ERetargetSourceOrTarget SourceOrTarget) const;
 	/** Returns index of the bone with the given name in either Source or Target skeleton. */
 	UE_API int32 GetBoneIndexFromName(const FName BoneName, const ERetargetSourceOrTarget SourceOrTarget) const;
 	/** Get access to all the fully resolved bone chain data */
 	UE_API const FRetargeterBoneChains& GetBoneChains();
 	/** Returns name of the chain associated with this bone. Returns NAME_None if bone is not in a chain. */
-	UE_API FName GetChainNameForBone(const FName BoneName, const ERetargetSourceOrTarget SourceOrTarget) const;
+	UE_API FName GetFirstChainForBone(const FName BoneName, const ERetargetSourceOrTarget SourceOrTarget) const;
+	/** Returns names of all the chains associated with this bone. Returns empty set if bone is not in a chain. */
+	UE_API TSet<FName> GetAllChainNamesForBone(const FName BoneName, const ERetargetSourceOrTarget SourceOrTarget) const;
+	// get the first mapping that contains the given bone, returns false if no mapping exists
+	UE_API bool GetFirstChainMappingForBone(const FName BoneName, const ERetargetSourceOrTarget SourceOrTarget, FName& OutChainContainingBone, FName& OutMappedChain) const;
 	/** Get a transform at a given param in a chain */
 	UE_API FTransform GetGlobalRetargetPoseAtParam(const FName InChainName, const float Param, const ERetargetSourceOrTarget SourceOrTarget) const;
 	/** Get a transform of a bone in the retarget pose */
 	UE_API FTransform GetRetargetPoseBoneTransform(const FName InBoneName, const ERetargetSourceOrTarget SourceOrTarget, ERetargetBoneSpace BoneSpace) const;
 	
 	/** Get the param of the bone in it's retarget chain. Ranges from 0 to NumBonesInChain. */
-	UE_API float GetParamOfBoneInChain(const FName InBoneName, const ERetargetSourceOrTarget SourceOrTarget) const;
+	UE_API float GetParamOfBoneInChain(const FName InBoneName, const ERetargetSourceOrTarget SourceOrTarget, const FName InChain = NAME_None) const;
 	/** Get the bone in the chain at the given param. */
 	UE_API FName GetClosestBoneToParam(const FName InChainName, const float InParam, const ERetargetSourceOrTarget SourceOrTarget) const;
 	/** Get the chain mapped to this one. */
@@ -664,7 +668,7 @@ public:
 #if WITH_EDITOR
 	bool IsBoneRetargeted(const FName BoneName, const ERetargetSourceOrTarget SourceOrTarget) const
 	{
-		return Processor.IsBoneInAMappedChain(BoneName, SourceOrTarget);
+		return Processor.IsBoneMapped(BoneName, SourceOrTarget);
 	}
 	int32 GetBoneIndexFromName(const FName BoneName, const ERetargetSourceOrTarget SourceOrTarget) const
 	{
@@ -672,7 +676,7 @@ public:
 	}
 	FName GetChainNameForBone(const FName BoneName, const ERetargetSourceOrTarget SourceOrTarget) const
 	{
-		return Processor.GetChainNameForBone(BoneName, SourceOrTarget);
+		return Processor.GetFirstChainForBone(BoneName, SourceOrTarget);
 	}
 	FTransform GetGlobalRetargetPoseAtParam(const FName InChainName, const float Param, const ERetargetSourceOrTarget SourceOrTarget) const
 	{
diff --git a/Engine/Plugins/Animation/IKRig/Source/IKRigEditor/Private/RetargetEditor/IKRetargetEditorController.cpp b/Engine/Plugins/Animation/IKRig/Source/IKRigEditor/Private/RetargetEditor/IKRetargetEditorController.cpp
--- a/Engine/Plugins/Animation/IKRig/Source/IKRigEditor/Private/RetargetEditor/IKRetargetEditorController.cpp
+++ b/Engine/Plugins/Animation/IKRig/Source/IKRigEditor/Private/RetargetEditor/IKRetargetEditorController.cpp
@@ -547,11 +547,11 @@ bool FIKRetargetEditorController::IsBoneRetargeted(const FName& BoneName, ERetar
 		return false;
 	}
 
-	// return if it's a retargeted bone
-	return Processor->IsBoneInAMappedChain(BoneName, SourceOrTarget);
+	// return true if the bone is mapped to another bone on the source/target
+	return Processor->IsBoneMapped(BoneName, SourceOrTarget);
 }
 
-FName FIKRetargetEditorController::GetChainNameFromBone(const FName& BoneName, ERetargetSourceOrTarget SourceOrTarget) const
+FName FIKRetargetEditorController::GetChainDisplayNameForBone(const FName& BoneName, ERetargetSourceOrTarget SourceOrTarget) const
 {
 	// get an initialized processor
 	const FIKRetargetProcessor* Processor = GetRetargetProcessor();
@@ -559,8 +559,44 @@ FName FIKRetargetEditorController::GetChainNameFromBone(const FName& BoneName, E
 	{
 		return NAME_None;
 	}
+
+	// special case, retarget pelvis should always display as such
+	const FName RootBoneName = Processor->GetPelvisBone(SourceOrTarget, ERetargetOpsToSearch::AssetOps);
+	if (BoneName == RootBoneName)
+	{
+		return FName("Retarget Pelvis");
+	}
+
+	auto GetCombinedName = [](const TSet<FName>& Names) -> FName
+		{
+			if (Names.IsEmpty())
+			{
+				return NAME_None;
+			}
+    
+			if (Names.Num() == 1)
+			{
+				return *Names.begin();
+			}
+    
+			FString Result;
+			bool bFirst = true;
+    
+			for (const FName& Name : Names)
+			{
+				if (!bFirst)
+				{
+					Result += TEXT(", ");
+				}
+				Result += Name.ToString();
+				bFirst = false;
+			}
+    
+			return FName(Result);
+		};
 	
-	return Processor->GetChainNameForBone(BoneName, SourceOrTarget);
+	const TSet<FName> ChainsThatContainBone = Processor->GetAllChainNamesForBone(BoneName, SourceOrTarget);
+	return GetCombinedName(ChainsThatContainBone);
 }
 
 TObjectPtr<UIKRetargetBoneDetails> FIKRetargetEditorController::GetOrCreateBoneDetailsObject(const FName& BoneName)
diff --git a/Engine/Plugins/Animation/IKRig/Source/IKRigEditor/Private/RetargetEditor/IKRetargeterPoseGenerator.cpp b/Engine/Plugins/Animation/IKRig/Source/IKRigEditor/Private/RetargetEditor/IKRetargeterPoseGenerator.cpp
--- a/Engine/Plugins/Animation/IKRig/Source/IKRigEditor/Private/RetargetEditor/IKRetargeterPoseGenerator.cpp
+++ b/Engine/Plugins/Animation/IKRig/Source/IKRigEditor/Private/RetargetEditor/IKRetargeterPoseGenerator.cpp
@@ -72,7 +72,7 @@ void FRetargetAutoPoseGenerator::AlignBones(
 		// warn user if they tried to align a non-retargeted bone (doesn't work)
 		for (const FName Bone : BonesToAlign)
 		{
-			if (!Processor->IsBoneInAMappedChain(Bone, SourceOrTarget))
+			if (!Processor->IsBoneMapped(Bone, SourceOrTarget))
 			{
 				Processor->Log.LogWarning(FText::Format(LOCTEXT("SkipAutoAlignNonRetargeted", "Skipped aligning {0}. The bone is not retargeted and therefore has no 'matching' equivalent on other skeleton."), FText::FromName(Bone)));
 			}
@@ -92,7 +92,7 @@ void FRetargetAutoPoseGenerator::AlignBone(
 	}
 	
 	// can't auto-align if the bone is not retargeted  (warning is elsewhere)
-	if (!Processor->IsBoneInAMappedChain(BoneToAlign, SourceOrTarget))
+	if (!Processor->IsBoneMapped(BoneToAlign, SourceOrTarget))
 	{
 		return;
 	}
@@ -113,12 +113,16 @@ void FRetargetAutoPoseGenerator::AlignBone(
 	const ERetargetSourceOrTarget SideToMatch = SideToAlign == ERetargetSourceOrTarget::Source ? ERetargetSourceOrTarget::Target : ERetargetSourceOrTarget::Source;
 	
 	// find chains on both sides
-	const FName ChainToAlign = Processor->GetChainNameForBone(BoneToAlign, SideToAlign);
-	const FName ChainToMatch = Processor->GetFirstChainMappedToChain(ChainToAlign, SideToAlign);
-	checkf(ChainToMatch != NAME_None != 0, TEXT("Bone should never be retargeted and not in a mapped chain."));
+	FName ChainToAlign;
+	FName ChainToMatch;
+	if (!Processor->GetFirstChainMappingForBone(BoneToAlign, SideToAlign, ChainToAlign, ChainToMatch))
+	{
+		ensureMsgf(false, TEXT("Bone should never pass IsBoneMapped but not have both chains."));
+		return;
+	}
 	
 	// get the normalized chain parameter to match
-	const float Param = Processor->GetParamOfBoneInChain(BoneToAlign, SideToAlign);
+	const float Param = Processor->GetParamOfBoneInChain(BoneToAlign, SideToAlign, ChainToAlign);
 
 	// find tangents to align between source/target
 	FVector DirectionToAlign;
@@ -302,7 +306,7 @@ void FRetargetAutoPoseGenerator::SnapToGround(const FName BoneToPutOnGround, con
 		// find the lowest bone in the current retarget pose
 		for (int32 BoneIndex=0; BoneIndex<Skeleton.BoneNames.Num(); ++BoneIndex)
 		{
-			if (!Processor->IsBoneInAMappedChain(Skeleton.BoneNames[BoneIndex], SourceOrTarget))
+			if (!Processor->IsBoneMapped(Skeleton.BoneNames[BoneIndex], SourceOrTarget))
 			{
 				continue;
 			}
@@ -600,7 +604,7 @@ bool FRetargetAutoPoseGenerator::IsARetargetLeaf(
 	GetAllChildrenOfBone(BoneName, SourceOrTarget, AllChildren);
 	for (const FName Child : AllChildren)
 	{
-		if (Processor->IsBoneInAMappedChain(Child, SourceOrTarget))
+		if (Processor->IsBoneMapped(Child, SourceOrTarget))
 		{
 			return false;
 		}
diff --git a/Engine/Plugins/Animation/IKRig/Source/IKRigEditor/Private/RetargetEditor/SIKRetargetHierarchy.cpp b/Engine/Plugins/Animation/IKRig/Source/IKRigEditor/Private/RetargetEditor/SIKRetargetHierarchy.cpp
--- a/Engine/Plugins/Animation/IKRig/Source/IKRigEditor/Private/RetargetEditor/SIKRetargetHierarchy.cpp
+++ b/Engine/Plugins/Animation/IKRig/Source/IKRigEditor/Private/RetargetEditor/SIKRetargetHierarchy.cpp
@@ -35,7 +35,7 @@ FName FIKRetargetHierarchyElement::GetChainName() const
 		return NAME_None;
 	}
 
-	return Controller->GetChainNameFromBone(Name, Controller->GetSourceOrTarget());
+	return Controller->GetChainDisplayNameForBone(Name, Controller->GetSourceOrTarget());
 }
 
 TSharedRef<SWidget> SIKRetargetHierarchyRow::GenerateWidgetForColumn(const FName& ColumnName)
@@ -446,7 +446,7 @@ void SIKRetargetHierarchy::RefreshTreeView(bool IsInitialSetup)
 		BoneTreeElementIndices.Add(BoneName, BoneElementIndex);
 
 		// get the chain this bone belongs to
-		FName ChainName = Controller->GetChainNameFromBone(BoneName, Controller->GetSourceOrTarget());
+		FName ChainName = Controller->GetChainDisplayNameForBone(BoneName, Controller->GetSourceOrTarget());
 		
 		// apply text filter to items
 		if (!(TextFilter->GetFilterText().IsEmpty() ||
diff --git a/Engine/Plugins/Animation/IKRig/Source/IKRigEditor/Public/RetargetEditor/IKRetargetEditorController.h b/Engine/Plugins/Animation/IKRig/Source/IKRigEditor/Public/RetargetEditor/IKRetargetEditorController.h
--- a/Engine/Plugins/Animation/IKRig/Source/IKRigEditor/Public/RetargetEditor/IKRetargetEditorController.h
+++ b/Engine/Plugins/Animation/IKRig/Source/IKRigEditor/Public/RetargetEditor/IKRetargetEditorController.h
@@ -259,7 +259,7 @@ public:
 	// determine if bone in the specified skeleton is part of the retarget (in a mapped chain)
 	bool IsBoneRetargeted(const FName& BoneName, ERetargetSourceOrTarget SourceOrTarget) const;
 	// get the name of the chain that contains this bone
-	FName GetChainNameFromBone(const FName& BoneName, ERetargetSourceOrTarget SourceOrTarget) const;
+	FName GetChainDisplayNameForBone(const FName& BoneName, ERetargetSourceOrTarget SourceOrTarget) const;
 
 	// factory to get/create bone details UObject
 	TObjectPtr<UIKRetargetBoneDetails> GetOrCreateBoneDetailsObject(const FName& BoneName);
