IK Rig: Allow enabling/disabling anim graph node w/ cvar

diff --git a/Engine/Plugins/Animation/IKRig/Source/IKRig/Private/AnimNodes/AnimNode_IKRig.cpp b/Engine/Plugins/Animation/IKRig/Source/IKRig/Private/AnimNodes/AnimNode_IKRig.cpp
--- a/Engine/Plugins/Animation/IKRig/Source/IKRig/Private/AnimNodes/AnimNode_IKRig.cpp
+++ b/Engine/Plugins/Animation/IKRig/Source/IKRig/Private/AnimNodes/AnimNode_IKRig.cpp
@@ -20,6 +20,12 @@
 
 #include UE_INLINE_GENERATED_CPP_BY_NAME(AnimNode_IKRig)
 
+namespace UE::Anim::IKRig
+{
+	static bool bEnableIKRig = true;
+	FAutoConsoleVariableRef CVarEnableIKRigNode(TEXT("a.animnode.IKRig.Enable"), bEnableIKRig, TEXT("Enables/disables IK Rig Node in Anim Graph."));
+}
+
 FAnimNode_IKRig::FAnimNode_IKRig()
 	: AlphaInputType(EAnimAlphaInputType::Float)
 	, bAlphaBoolEnabled(true)
@@ -53,7 +59,7 @@ void FAnimNode_IKRig::Evaluate_AnyThread(FPoseContext& Output)
 	}
 
 	// bail out if uninitialized
-	if (!IKRigProcessor.IsInitialized())
+	if (!IKRigProcessor.IsInitialized() || !UE::Anim::IKRig::bEnableIKRig)
 	{
 		return;
 	}
