[IK Retarget] Fixed bug reinitializing children ops.

diff --git a/Engine/Plugins/Animation/IKRig/Source/IKRigEditor/Private/RetargetEditor/IKRetargeterController.cpp b/Engine/Plugins/Animation/IKRig/Source/IKRigEditor/Private/RetargetEditor/IKRetargeterController.cpp
--- a/Engine/Plugins/Animation/IKRig/Source/IKRigEditor/Private/RetargetEditor/IKRetargeterController.cpp
+++ b/Engine/Plugins/Animation/IKRig/Source/IKRigEditor/Private/RetargetEditor/IKRetargeterController.cpp
@@ -515,7 +515,7 @@ void UIKRetargeterController::OnOpPropertyChanged(const FName& InOpName, const F
 
 	// notify children of parent property edit
 	TArray<int32> ChildOps = GetChildOpIndices(OpIndex);
-	for (int32 ChildOpIndex = 0; ChildOpIndex < ChildOps.Num(); ++ChildOpIndex)
+	for (int32 ChildOpIndex : ChildOps)
 	{
 		if (FIKRetargetOpBase* ChildOp = GetRetargetOpByIndex(ChildOpIndex))
 		{
