[IK Retarget] Fixes "use after free" issue caught by address sanitizer.

diff --git a/Engine/Plugins/Animation/IKRig/Source/IKRigEditor/Private/RetargetEditor/SIKRetargetHierarchy.cpp b/Engine/Plugins/Animation/IKRig/Source/IKRigEditor/Private/RetargetEditor/SIKRetargetHierarchy.cpp
--- a/Engine/Plugins/Animation/IKRig/Source/IKRigEditor/Private/RetargetEditor/SIKRetargetHierarchy.cpp
+++ b/Engine/Plugins/Animation/IKRig/Source/IKRigEditor/Private/RetargetEditor/SIKRetargetHierarchy.cpp
@@ -280,11 +280,11 @@ void SIKRetargetHierarchy::ShowItemAfterSelection(FName ItemName)
 	
 	if(GetDefault<UPersonaOptions>()->bExpandTreeOnSelection)
 	{
-		TSharedPtr<FIKRetargetHierarchyElement> ItemToExpand = ItemToShow->Parent;
+		TWeakPtr<FIKRetargetHierarchyElement> ItemToExpand = ItemToShow->Parent;
 		while(ItemToExpand.IsValid())
 		{
-			TreeView->SetItemExpansion(ItemToExpand, true);
-			ItemToExpand = ItemToExpand->Parent;
+			TreeView->SetItemExpansion(ItemToExpand.Pin(), true);
+			ItemToExpand = ItemToExpand.Pin()->Parent;
 		}
 	}
     
@@ -484,15 +484,15 @@ void SIKRetargetHierarchy::RefreshTreeView(bool IsInitialSetup)
 	for (int32 BoneIndex=0; BoneIndex<BoneNames.Num(); ++BoneIndex)
 	{
 		const FName BoneName = BoneNames[BoneIndex];
-		const TSharedPtr<FIKRetargetHierarchyElement> BoneTreeElement = AllElements[BoneTreeElementIndices[BoneName]];
+		TWeakPtr<FIKRetargetHierarchyElement> BoneTreeElement = AllElements[BoneTreeElementIndices[BoneName]];
 
-		if (BoneTreeElement->bIsHidden)
+		if (BoneTreeElement.IsValid() && BoneTreeElement.Pin()->bIsHidden)
 		{
 			continue;
 		}
 		
 		// find first parent that is not filtered
-		const TSharedPtr<FIKRetargetHierarchyElement>* FirstNonHiddenParentElement = nullptr;
+		TWeakPtr<FIKRetargetHierarchyElement> FirstNonHiddenParentElement;
 		int32 ParentBoneIndex = ParentIndices[BoneIndex];
 		while (true)
 		{
@@ -502,27 +502,27 @@ void SIKRetargetHierarchy::RefreshTreeView(bool IsInitialSetup)
 			}
 
 			const FName ParentBoneName = BoneNames[ParentBoneIndex];
-			TSharedPtr<FIKRetargetHierarchyElement> ParentBoneTreeElement = AllElements[BoneTreeElementIndices[ParentBoneName]];
+			TSharedPtr<FIKRetargetHierarchyElement>& ParentBoneTreeElement = AllElements[BoneTreeElementIndices[ParentBoneName]];
 			if (!ParentBoneTreeElement->bIsHidden)
 			{
-				FirstNonHiddenParentElement = &ParentBoneTreeElement;
+				FirstNonHiddenParentElement = ParentBoneTreeElement;
 				break;
 			}
 			
 			ParentBoneIndex = ParentIndices[ParentBoneIndex];
 		}
 
-		if (FirstNonHiddenParentElement)
+		if (FirstNonHiddenParentElement.IsValid())
 		{
 			// store pointer to child on parent
-			FirstNonHiddenParentElement->Get()->Children.Add(BoneTreeElement);
+			FirstNonHiddenParentElement.Pin()->Children.Add(BoneTreeElement);
 			// store pointer to parent on child
-			BoneTreeElement->Parent = *FirstNonHiddenParentElement;
+			BoneTreeElement.Pin()->Parent = FirstNonHiddenParentElement;
 		}
 		else
 		{
 			// has no parent, store a root element
-			RootElements.Add(BoneTreeElement);
+			RootElements.Add(BoneTreeElement.Pin());
 		}
 	}
 
@@ -550,7 +550,10 @@ void SIKRetargetHierarchy::HandleGetChildrenForTree(
 	TSharedPtr<FIKRetargetHierarchyElement> InItem,
 	TArray<TSharedPtr<FIKRetargetHierarchyElement>>& OutChildren)
 {
-	OutChildren = InItem->Children;
+	for (TWeakPtr<FIKRetargetHierarchyElement>& Child : InItem->Children)
+	{
+		OutChildren.Add(Child.Pin());
+	}
 }
 
 void SIKRetargetHierarchy::OnSelectionChanged(
@@ -605,16 +608,16 @@ void SIKRetargetHierarchy::SetExpansionRecursive(
     
 	if (bTowardsParent)
 	{
-		if (InElement->Parent.Get())
+		if (InElement->Parent.IsValid())
 		{
-			SetExpansionRecursive(InElement->Parent, bTowardsParent, bShouldBeExpanded);
+			SetExpansionRecursive(InElement->Parent.Pin(), bTowardsParent, bShouldBeExpanded);
 		}
 	}
 	else
 	{
 		for (int32 ChildIndex = 0; ChildIndex < InElement->Children.Num(); ++ChildIndex)
 		{
-			SetExpansionRecursive(InElement->Children[ChildIndex], bTowardsParent, bShouldBeExpanded);
+			SetExpansionRecursive(InElement->Children[ChildIndex].Pin(), bTowardsParent, bShouldBeExpanded);
 		}
 	}
 }
diff --git a/Engine/Plugins/Animation/IKRig/Source/IKRigEditor/Public/RetargetEditor/SIKRetargetHierarchy.h b/Engine/Plugins/Animation/IKRig/Source/IKRigEditor/Public/RetargetEditor/SIKRetargetHierarchy.h
--- a/Engine/Plugins/Animation/IKRig/Source/IKRigEditor/Public/RetargetEditor/SIKRetargetHierarchy.h
+++ b/Engine/Plugins/Animation/IKRig/Source/IKRigEditor/Public/RetargetEditor/SIKRetargetHierarchy.h
@@ -20,8 +20,8 @@ public:
 	FIKRetargetHierarchyElement(const FName& InName, const TSharedRef<FIKRetargetEditorController>& InEditorController);
 
 	FText Key;
-	TSharedPtr<FIKRetargetHierarchyElement> Parent;
-	TArray<TSharedPtr<FIKRetargetHierarchyElement>> Children;
+	TWeakPtr<FIKRetargetHierarchyElement> Parent;
+	TArray<TWeakPtr<FIKRetargetHierarchyElement>> Children;
 	FName Name;
 	bool bIsHidden = false;
 
