Adjust IKRigToolkit to only refresh when transactions affect the underlying stored objects.

#rnx

diff --git a/Engine/Plugins/Animation/IKRig/Source/IKRigEditor/Private/RigEditor/IKRigToolkit.cpp b/Engine/Plugins/Animation/IKRig/Source/IKRigEditor/Private/RigEditor/IKRigToolkit.cpp
--- a/Engine/Plugins/Animation/IKRig/Source/IKRigEditor/Private/RigEditor/IKRigToolkit.cpp
+++ b/Engine/Plugins/Animation/IKRig/Source/IKRigEditor/Private/RigEditor/IKRigToolkit.cpp
@@ -24,6 +24,8 @@
 #include "RigEditor/IKRigEditorController.h"
 #include "RigEditor/IKRigMode.h"
 
+#include "Misc/ITransaction.h"
+
 #define LOCTEXT_NAMESPACE "IKRigEditorToolkit"
 
 const FName IKRigEditorModes::IKRigEditorMode("IKRigEditorMode");
@@ -238,6 +240,22 @@ void FIKRigEditorToolkit::HandleOnPreviewSceneSettingsCustomized(IDetailLayoutBu
 	DetailBuilder.HideCategory("Animation Blueprint");
 }
 
+bool FIKRigEditorToolkit::MatchesContext(const FTransactionContext& InContext, const TArray<TPair<UObject*, FTransactionObjectEvent>>& TransactionObjectContexts) const
+{
+	if (InContext.PrimaryObject == EditorController->AssetController)
+	{
+		return true;
+	}
+	for (const TPair<UObject*, FTransactionObjectEvent>& ObjectTransactionPair : TransactionObjectContexts)
+	{
+		if (ObjectTransactionPair.Get<0>() == EditorController->AssetController)
+		{
+			return true;
+		}
+	}	
+	return false;
+}
+
 void FIKRigEditorToolkit::PostUndo(bool bSuccess)
 {
 	FScopedReinitializeIKRig Reinitialize(EditorController->AssetController);
diff --git a/Engine/Plugins/Animation/IKRig/Source/IKRigEditor/Public/RigEditor/IKRigToolkit.h b/Engine/Plugins/Animation/IKRig/Source/IKRigEditor/Public/RigEditor/IKRigToolkit.h
--- a/Engine/Plugins/Animation/IKRig/Source/IKRigEditor/Public/RigEditor/IKRigToolkit.h
+++ b/Engine/Plugins/Animation/IKRig/Source/IKRigEditor/Public/RigEditor/IKRigToolkit.h
@@ -56,7 +56,8 @@ public:
 
 	/** FSelfRegisteringEditorUndoClient interface */
 	virtual void PostUndo( bool bSuccess );
-	virtual void PostRedo( bool bSuccess );
+	virtual void PostRedo(bool bSuccess);
+	virtual bool MatchesContext( const FTransactionContext& InContext, const TArray<TPair<UObject*, FTransactionObjectEvent>>& TransactionObjectContexts ) const override;
 	/** END FSelfRegisteringEditorUndoClient interface */
 
 	TSharedRef<FIKRigEditorController> GetController() const {return EditorController;};
