Adjust IKRetargetEditor to only refresh when transactions affect the underlying stored objects.

#rnx

diff --git a/Engine/Plugins/Animation/IKRig/Source/IKRigEditor/Private/RetargetEditor/IKRetargetEditorController.cpp b/Engine/Plugins/Animation/IKRig/Source/IKRigEditor/Private/RetargetEditor/IKRetargetEditorController.cpp
--- a/Engine/Plugins/Animation/IKRig/Source/IKRigEditor/Private/RetargetEditor/IKRetargetEditorController.cpp
+++ b/Engine/Plugins/Animation/IKRig/Source/IKRigEditor/Private/RetargetEditor/IKRetargetEditorController.cpp
@@ -35,6 +35,8 @@
 #include "RigEditor/IKRigStructViewer.h"
 #include "Styling/AppStyle.h"
 
+#include "Misc/ITransaction.h"
+
 #define LOCTEXT_NAMESPACE "IKRetargetEditorController"
 
 
@@ -197,6 +199,37 @@ void FIKRetargetEditorController::Close()
 	}
 }
 
+bool FIKRetargetEditorController::MatchesContext(const FTransactionContext& InContext, const TArray<TPair<UObject*, FTransactionObjectEvent>>& TransactionObjectContexts) const
+{
+	auto TransactionIncludesObject = [&InContext, &TransactionObjectContexts](UObject * Object)
+	{
+		if (InContext.PrimaryObject == Object)
+		{
+			return true;		
+		}
+		for (const TPair<UObject*, FTransactionObjectEvent>& ObjectTransactionPair : TransactionObjectContexts)
+		{
+			if (ObjectTransactionPair.Get<0>() == Object)
+			{
+				return true;
+			}
+		}
+		return false;
+	};
+
+	const UIKRetargeter* Asset = AssetController->GetAsset();
+	if (Asset)
+	{
+		const bool bIncludesAssetControllerObject = TransactionIncludesObject(AssetController) || TransactionIncludesObject(AssetController->GetAsset());
+		const bool bIncludesSourceRig = TransactionIncludesObject(Asset->GetIKRigWriteable(ERetargetSourceOrTarget::Source));
+		const bool bIncludesTargetRig = TransactionIncludesObject(Asset->GetIKRigWriteable(ERetargetSourceOrTarget::Target));
+
+		return bIncludesAssetControllerObject || bIncludesSourceRig || bIncludesTargetRig;
+	}
+		
+	return false;
+}
+
 void FIKRetargetEditorController::PostUndo(bool bSuccess)
 {
 	AssetController->CleanAsset();
diff --git a/Engine/Plugins/Animation/IKRig/Source/IKRigEditor/Public/RetargetEditor/IKRetargetEditorController.h b/Engine/Plugins/Animation/IKRig/Source/IKRigEditor/Public/RetargetEditor/IKRetargetEditorController.h
--- a/Engine/Plugins/Animation/IKRig/Source/IKRigEditor/Public/RetargetEditor/IKRetargetEditorController.h
+++ b/Engine/Plugins/Animation/IKRig/Source/IKRigEditor/Public/RetargetEditor/IKRetargetEditorController.h
@@ -12,6 +12,8 @@
 #include "Widgets/SWindow.h"
 #include "EditorUndoClient.h"
 
+#include "Misc/TransactionObjectEvent.h"
+
 class SRetargetOpStack;
 enum class ERetargetSourceOrTarget : uint8;
 class SIKRetargetHierarchy;
@@ -88,7 +90,8 @@ public:
 
 	/** FSelfRegisteringEditorUndoClient interface */
 	virtual void PostUndo( bool bSuccess );
-	virtual void PostRedo( bool bSuccess );
+	virtual void PostRedo(bool bSuccess);
+	virtual bool MatchesContext( const FTransactionContext& InContext, const TArray<TPair<UObject*, FTransactionObjectEvent>>& TransactionObjectContexts ) const;	
 	/** END FSelfRegisteringEditorUndoClient interface */
 	
 	// Bind callbacks to this IK Rig
