[IK Retarget] Fixed chain map list selection being lost when editing settings.

diff --git a/Engine/Plugins/Animation/IKRig/Source/IKRigEditor/Private/RetargetEditor/SIKRetargetChainMapList.cpp b/Engine/Plugins/Animation/IKRig/Source/IKRigEditor/Private/RetargetEditor/SIKRetargetChainMapList.cpp
--- a/Engine/Plugins/Animation/IKRig/Source/IKRigEditor/Private/RetargetEditor/SIKRetargetChainMapList.cpp
+++ b/Engine/Plugins/Animation/IKRig/Source/IKRigEditor/Private/RetargetEditor/SIKRetargetChainMapList.cpp
@@ -470,6 +470,14 @@ void SIKRetargetChainMapList::RefreshView()
 		return Chain->IKGoalName != NAME_None;
 	};
 
+	// record selected item to restore
+	FName SelectedChainName = NAME_None;
+	TArray<TSharedPtr<FRetargetChainMapElement>> SelectedItems = ListView->GetSelectedItems();
+	if (!SelectedItems.IsEmpty())
+	{
+		SelectedChainName = SelectedItems[0].Get()->TargetChainName;	
+	}
+
 	// refresh items
 	ListViewItems.Reset();
 
@@ -524,8 +532,14 @@ void SIKRetargetChainMapList::RefreshView()
 		// create an item for this chain
 		TSharedPtr<FRetargetChainMapElement> ChainItem = FRetargetChainMapElement::Make(TargetChainName);
 		ListViewItems.Add(ChainItem);
+
+		// restore selection
+		if (ChainItem.Get()->TargetChainName == SelectedChainName)
+		{
+			ListView->SetSelection(ChainItem);
+		}
 	}
-	
+
 	ListView->RequestListRefresh();
 }
 
