[IK Retarget] Add more robust support for FBoneReference to retarget ops.
Fix potential crash bug in op details panel.

diff --git a/Engine/Plugins/Animation/IKRig/Source/IKRig/Private/Retargeter/IKRetargetOps.cpp b/Engine/Plugins/Animation/IKRig/Source/IKRig/Private/Retargeter/IKRetargetOps.cpp
--- a/Engine/Plugins/Animation/IKRig/Source/IKRig/Private/Retargeter/IKRetargetOps.cpp
+++ b/Engine/Plugins/Animation/IKRig/Source/IKRig/Private/Retargeter/IKRetargetOps.cpp
@@ -42,6 +42,12 @@ void FIKRetargetOpBase::SetSettings(const FIKRetargetOpSettingsBase* InSettings)
 
 #if WITH_EDITOR
 
+void FIKRetargetOpSettingsBase::AssignSkeletonAssets(const USkeleton* InSourceSkeleton, const USkeleton* InTargetSkeleton)
+{
+	SourceSkeletonAsset = InSourceSkeleton;
+	TargetSkeletonAsset = InTargetSkeleton;
+}
+
 FName FIKRetargetOpBase::GetDefaultName() const
 {
 	// calls virtual GetType() to get the derived type
diff --git a/Engine/Plugins/Animation/IKRig/Source/IKRig/Private/Retargeter/IKRetargetProcessor.cpp b/Engine/Plugins/Animation/IKRig/Source/IKRig/Private/Retargeter/IKRetargetProcessor.cpp
--- a/Engine/Plugins/Animation/IKRig/Source/IKRig/Private/Retargeter/IKRetargetProcessor.cpp
+++ b/Engine/Plugins/Animation/IKRig/Source/IKRig/Private/Retargeter/IKRetargetProcessor.cpp
@@ -1192,8 +1192,9 @@ void FIKRetargetProcessor::InitialOpStackSetup(
 		AssetOp->EditorInstance = &NewOp;
 		NewOp.ResetExecutionTime();
 		// references the skeletons are used for bone selector widgets in the details panel
-		AssetOpSettings->SourceSkeletonAsset = SourceSkeleton.SkeletalMesh->GetSkeleton();
-		AssetOpSettings->TargetSkeletonAsset = TargetSkeleton.SkeletalMesh->GetSkeleton();
+		AssetOpSettings->AssignSkeletonAssets(
+			SourceSkeleton.SkeletalMesh->GetSkeleton(),
+			TargetSkeleton.SkeletalMesh->GetSkeleton());
 #endif
 
 	}
diff --git a/Engine/Plugins/Animation/IKRig/Source/IKRig/Public/Retargeter/IKRetargetOps.h b/Engine/Plugins/Animation/IKRig/Source/IKRig/Public/Retargeter/IKRetargetOps.h
--- a/Engine/Plugins/Animation/IKRig/Source/IKRig/Public/Retargeter/IKRetargetOps.h
+++ b/Engine/Plugins/Animation/IKRig/Source/IKRig/Public/Retargeter/IKRetargetOps.h
@@ -23,10 +23,36 @@ enum class ERetargetSourceOrTarget : uint8;
 class FPrimitiveDrawInterface;
 class FIKRetargetEditorController;
 struct FIKRetargetDebugDrawState;
+class IPropertyHandle;
 #endif
 
+/*
+ * Because the settings structs in the retarget op stack are not UObjects, we cannot use the IBoneReferenceSkeletonProvider interface
+ * Instead, if Op settings require a FBoneReference with a bone selector widget then the settings must inherit either directly from this
+ * (in the case of sub structs), or indirectly from the FIKRetargetOpSettingBase for the op's main settings struct.
+ **/
+USTRUCT(BlueprintType)
+struct FIKRetargetOpSkeletonProvider
+{
+	GENERATED_BODY()
+	
+	virtual ~FIKRetargetOpSkeletonProvider() {};
+
+#if WITH_EDITORONLY_DATA
+public:
+	/** a reference to the skeletons for bone selector widgets (can return these from GetSkeleton) */
+	mutable const USkeleton* SourceSkeletonAsset;
+	mutable const USkeleton* TargetSkeletonAsset;
+	/** allow settings to provide a skeleton for any given FBoneReference widget. Defaults to the target skeleton. */
+	virtual USkeleton* GetSkeleton(const FName InPropertyName) { return const_cast<USkeleton*>(TargetSkeletonAsset); };
+	/** optional alternative to GetSkeleton that provides the full PropertyHandle. This is used if GetSkeleton(const FName) returns null
+	 * NOTE: the property handle is necessary to help differentiate properties in copies of the same struct type */ 
+	virtual USkeleton* GetSkeleton(const IPropertyHandle* InPropertyHandle) { return const_cast<USkeleton*>(TargetSkeletonAsset); };
+#endif
+};
+
 /** This is the base class for defining editable settings for your custom retargeting operation.
- * All user-configurable properties for your "op" should be stored in a subclass of this.
+ * All user-configurable properties for an "op" should be stored within a single instance of a subclass of this type.
  * These settings will automatically be:
  * 1. Displayed in the details panel when the op is selected
  * 2. Saved/loaded with the op in the retarget asset
@@ -39,7 +65,7 @@ struct FIKRetargetDebugDrawState;
  * the runtime Op will get the latest values automatically.
  */
 USTRUCT(BlueprintType)
-struct FIKRetargetOpSettingsBase
+struct FIKRetargetOpSettingsBase : public FIKRetargetOpSkeletonProvider
 {
 	GENERATED_BODY()
 	
@@ -118,11 +144,9 @@ public:
 	FName OwningOpName;
 	/** the instance of this op currently running in the editor viewport */
 	mutable FIKRetargetOpSettingsBase* EditorInstance;
-	/** a reference to the skeletons for bone selector widgets (can return these from GetSkeleton) */
-	mutable const USkeleton* SourceSkeletonAsset;
-	mutable const USkeleton* TargetSkeletonAsset;
-	/** allow settings to provide a skeleton for any given FBoneReference widget. Defaults to the target skeleton. */
-	virtual USkeleton* GetSkeleton(const FName InPropertyName) { return const_cast<USkeleton*>(TargetSkeletonAsset); };
+	/** assigned by the processor when an editor instance is initialized, used for UI reflection for bone selectors
+	 * NOTE: if your op has sub-settings structs that inherit from FIKRetargetOpSkeletonProvider, override this to give them pointers to the skeletons */
+	UE_API virtual void AssignSkeletonAssets(const USkeleton* InSourceSkeleton, const USkeleton* InTargetSkeleton);
 #endif
 };
 
@@ -331,7 +355,7 @@ public:
 	virtual void DebugDraw(
 		FPrimitiveDrawInterface* InPDI,
 		const FTransform& InSourceTransform,
-		const FTransform& InComponentTransform,
+		const FTransform& InTargetTransform,
 		const double InComponentScale,
 		const FIKRetargetDebugDrawState& InEditorState) const {};
 
diff --git a/Engine/Plugins/Animation/IKRig/Source/IKRigEditor/Private/RetargetEditor/IKRetargetOpDetails.cpp b/Engine/Plugins/Animation/IKRig/Source/IKRigEditor/Private/RetargetEditor/IKRetargetOpDetails.cpp
--- a/Engine/Plugins/Animation/IKRig/Source/IKRigEditor/Private/RetargetEditor/IKRetargetOpDetails.cpp
+++ b/Engine/Plugins/Animation/IKRig/Source/IKRigEditor/Private/RetargetEditor/IKRetargetOpDetails.cpp
@@ -227,12 +227,12 @@ void FChainsFKOpCustomization::CustomizeChildren(
 	ChainListConfig.Controller = AssetController;
 	ChainListConfig.bEnableGoalColumn = false;
 	ChainListConfig.bEnableChainMapping = true;
-	ChainListConfig.ChainSettingsGetterFunc = [this](const FName InTargetChainName) -> UObject*
+	ChainListConfig.ChainSettingsGetterFunc = [Controller=AssetController, ThisOpName=OpName](const FName InTargetChainName) -> UObject*
 	{
 		// this lambda fetches the FK chain settings struct for the given target chain in the given op (by name)
-		auto MemoryProvider = [this, InTargetChainName]() -> uint8*
+		auto MemoryProvider = [Controller, ThisOpName, InTargetChainName]() -> uint8*
 		{
-			if (FIKRetargetOpBase* Op = AssetController->GetRetargetOpByName(OpName))
+			if (FIKRetargetOpBase* Op = Controller->GetRetargetOpByName(ThisOpName))
 			{
 				FIKRetargetFKChainsOpSettings* OpSettings = reinterpret_cast<FIKRetargetFKChainsOpSettings*>(Op->GetSettings());
 				for (FRetargetFKChainSettings& ChainToRetarget : OpSettings->ChainsToRetarget)
@@ -248,14 +248,14 @@ void FChainsFKOpCustomization::CustomizeChildren(
 		};
 			
 		FIKRigStructToView StructToView;
-		StructToView.Owner = AssetController->GetAsset();
+		StructToView.Owner = Controller->GetAsset();
 		StructToView.Type = FRetargetFKChainSettings::StaticStruct();
 		StructToView.MemoryProvider = MemoryProvider;
-		StructToView.UniqueName = OpName;
+		StructToView.UniqueName = ThisOpName;
 		
 		URetargetFKChainSettingsWrapper* Wrapper = NewObject<URetargetFKChainSettingsWrapper>(GetTransientPackage(), NAME_None, RF_Standalone);
 		const FName SettingsPropertyName = GET_MEMBER_NAME_CHECKED(URetargetFKChainSettingsWrapper, Settings);
-		Wrapper->InitializeWithRetargeter(StructToView, SettingsPropertyName, AssetController);
+		Wrapper->InitializeWithRetargeter(StructToView, SettingsPropertyName, Controller);
 		Wrapper->SetPropertyHidden(GET_MEMBER_NAME_CHECKED(FRetargetFKChainSettings, TargetChainName), true);
 		return Wrapper;
 	};
@@ -340,12 +340,12 @@ void FIKChainOpCustomization::CustomizeChildren(
 	ChainListConfig.bEnableGoalColumn = true;
 	ChainListConfig.bEnableChainMapping = false;
 	ChainListConfig.Filter = ShowOnlyIKFilter;
-	ChainListConfig.ChainSettingsGetterFunc = [this](const FName InTargetChainName) -> UObject*
+	ChainListConfig.ChainSettingsGetterFunc = [Controller=AssetController, ThisOpName=OpName](const FName InTargetChainName) -> UObject*
 	{
 		// this lambda fetches the chain settings struct for the given target chain in the given op (by name)
-		auto MemoryProvider = [this, InTargetChainName]() -> uint8*
+		auto MemoryProvider = [Controller, ThisOpName, InTargetChainName]() -> uint8*
 		{
-			if (FIKRetargetOpBase* Op = AssetController->GetRetargetOpByName(OpName))
+			if (FIKRetargetOpBase* Op = Controller->GetRetargetOpByName(ThisOpName))
 			{
 				FIKRetargetIKChainsOpSettings* OpSettings = reinterpret_cast<FIKRetargetIKChainsOpSettings*>(Op->GetSettings());
 				for (FRetargetIKChainSettings& ChainToRetarget : OpSettings->ChainsToRetarget)
@@ -360,14 +360,14 @@ void FIKChainOpCustomization::CustomizeChildren(
 		};
 			
 		FIKRigStructToView StructToView;
-		StructToView.Owner = AssetController->GetAsset();
+		StructToView.Owner = Controller->GetAsset();
 		StructToView.Type = FRetargetIKChainSettings::StaticStruct();
 		StructToView.MemoryProvider = MemoryProvider;
-		StructToView.UniqueName = OpName;
+		StructToView.UniqueName = ThisOpName;
 		
 		URetargetIKChainSettingsWrapper* Wrapper = NewObject<URetargetIKChainSettingsWrapper>(GetTransientPackage(), NAME_None, RF_Standalone);
 		const FName SettingsPropertyName = GET_MEMBER_NAME_CHECKED(URetargetIKChainSettingsWrapper, Settings);
-		Wrapper->InitializeWithRetargeter(StructToView, SettingsPropertyName, AssetController);
+		Wrapper->InitializeWithRetargeter(StructToView, SettingsPropertyName, Controller);
 		Wrapper->SetPropertyHidden(GET_MEMBER_NAME_CHECKED(FRetargetIKChainSettings, TargetChainName), true);
 		return Wrapper;
 	};
@@ -409,12 +409,12 @@ void FStrideWarpOpCustomization::CustomizeChildren(
 	ChainListConfig.bEnableGoalColumn = true;
 	ChainListConfig.bEnableChainMapping = false;
 	ChainListConfig.Filter = ShowOnlyIKFilter;
-	ChainListConfig.ChainSettingsGetterFunc = [this](const FName InTargetChainName) -> UObject*
+	ChainListConfig.ChainSettingsGetterFunc = [Controller=AssetController, ThisOpName=OpName](const FName InTargetChainName) -> UObject*
 	{
 		// this lambda fetches the chain settings struct for the given target chain in the given op (by name)
-		auto MemoryProvider = [this, InTargetChainName]() -> uint8*
+		auto MemoryProvider = [Controller, ThisOpName, InTargetChainName]() -> uint8*
 		{
-			if (FIKRetargetOpBase* Op = AssetController->GetRetargetOpByName(OpName))
+			if (FIKRetargetOpBase* Op = Controller->GetRetargetOpByName(ThisOpName))
 			{
 				FIKRetargetStrideWarpingOpSettings* OpSettings = reinterpret_cast<FIKRetargetStrideWarpingOpSettings*>(Op->GetSettings());
 				for (FRetargetStrideWarpChainSettings& ChainToWarp : OpSettings->ChainSettings)
@@ -429,14 +429,14 @@ void FStrideWarpOpCustomization::CustomizeChildren(
 		};
 			
 		FIKRigStructToView StructToView;
-		StructToView.Owner = AssetController->GetAsset();
+		StructToView.Owner = Controller->GetAsset();
 		StructToView.Type = FRetargetStrideWarpChainSettings::StaticStruct();
 		StructToView.MemoryProvider = MemoryProvider;
-		StructToView.UniqueName = OpName;
+		StructToView.UniqueName = ThisOpName;
 		
 		URetargetStrideWarpSettingsWrapper* Wrapper = NewObject<URetargetStrideWarpSettingsWrapper>(GetTransientPackage(), NAME_None, RF_Standalone);
 		const FName SettingsPropertyName = GET_MEMBER_NAME_CHECKED(URetargetStrideWarpSettingsWrapper, Settings);
-		Wrapper->InitializeWithRetargeter(StructToView, SettingsPropertyName, AssetController);
+		Wrapper->InitializeWithRetargeter(StructToView, SettingsPropertyName, Controller);
 		Wrapper->SetPropertyHidden(GET_MEMBER_NAME_CHECKED(FRetargetStrideWarpChainSettings, TargetChainName), true);
 		return Wrapper;
 	};
@@ -479,12 +479,12 @@ void FSpeedPlantOpCustomization::CustomizeChildren(
 	ChainListConfig.bEnableGoalColumn = true;
 	ChainListConfig.bEnableChainMapping = false;
 	ChainListConfig.Filter = ShowOnlyIKFilter;
-	ChainListConfig.ChainSettingsGetterFunc = [this](const FName InTargetChainName) -> UObject*
+	ChainListConfig.ChainSettingsGetterFunc = [Controller=AssetController, ThisOpName=OpName](const FName InTargetChainName) -> UObject*
 	{
 		// this lambda fetches the chain settings struct for the given target chain in the given op (by name)
-		auto MemoryProvider = [this, InTargetChainName]() -> uint8*
+		auto MemoryProvider = [Controller, ThisOpName, InTargetChainName]() -> uint8*
 		{
-			if (FIKRetargetOpBase* Op = AssetController->GetRetargetOpByName(OpName))
+			if (FIKRetargetOpBase* Op = Controller->GetRetargetOpByName(ThisOpName))
 			{
 				FIKRetargetSpeedPlantingOpSettings* OpSettings = reinterpret_cast<FIKRetargetSpeedPlantingOpSettings*>(Op->GetSettings());
 				for (FRetargetSpeedPlantingSettings& ChainToWarp : OpSettings->ChainsToSpeedPlant)
@@ -499,14 +499,14 @@ void FSpeedPlantOpCustomization::CustomizeChildren(
 		};
 			
 		FIKRigStructToView StructToView;
-		StructToView.Owner = AssetController->GetAsset();
+		StructToView.Owner = Controller->GetAsset();
 		StructToView.Type = FRetargetSpeedPlantingSettings::StaticStruct();
 		StructToView.MemoryProvider = MemoryProvider;
-		StructToView.UniqueName = OpName;
+		StructToView.UniqueName = ThisOpName;
 		
 		URetargetSpeedPlantSettingsWrapper* Wrapper = NewObject<URetargetSpeedPlantSettingsWrapper>(GetTransientPackage(), NAME_None, RF_Standalone);
 		const FName SettingsPropertyName = GET_MEMBER_NAME_CHECKED(URetargetSpeedPlantSettingsWrapper, Settings);
-		Wrapper->InitializeWithRetargeter(StructToView, SettingsPropertyName, AssetController);
+		Wrapper->InitializeWithRetargeter(StructToView, SettingsPropertyName, Controller);
 		Wrapper->SetPropertyHidden(GET_MEMBER_NAME_CHECKED(FRetargetSpeedPlantingSettings, TargetChainName), true);
 		return Wrapper;
 	};
@@ -546,12 +546,12 @@ void FPoleVectorOpCustomization::CustomizeChildren(
 	ChainListConfig.Controller = AssetController;
 	ChainListConfig.bEnableGoalColumn = true;
 	ChainListConfig.bEnableChainMapping = true;
-	ChainListConfig.ChainSettingsGetterFunc = [this](const FName InTargetChainName) -> UObject*
+	ChainListConfig.ChainSettingsGetterFunc = [Controller=AssetController, ThisOpName=OpName](const FName InTargetChainName) -> UObject*
 	{
 		// this lambda fetches the chain settings struct for the given target chain in the given op (by name)
-		auto MemoryProvider = [this, InTargetChainName]() -> uint8*
+		auto MemoryProvider = [Controller, ThisOpName, InTargetChainName]() -> uint8*
 		{
-			if (FIKRetargetOpBase* Op = AssetController->GetRetargetOpByName(OpName))
+			if (FIKRetargetOpBase* Op = Controller->GetRetargetOpByName(ThisOpName))
 			{
 				FIKRetargetAlignPoleVectorOpSettings* OpSettings = reinterpret_cast<FIKRetargetAlignPoleVectorOpSettings*>(Op->GetSettings());
 				for (FRetargetPoleVectorSettings& ChainToAlign : OpSettings->ChainsToAlign)
@@ -566,14 +566,14 @@ void FPoleVectorOpCustomization::CustomizeChildren(
 		};
 			
 		FIKRigStructToView StructToView;
-		StructToView.Owner = AssetController->GetAsset();
+		StructToView.Owner = Controller->GetAsset();
 		StructToView.Type = FRetargetPoleVectorSettings::StaticStruct();
 		StructToView.MemoryProvider = MemoryProvider;
-		StructToView.UniqueName = OpName;
+		StructToView.UniqueName = ThisOpName;
 		
 		UPoleVectorSettingsWrapper* Wrapper = NewObject<UPoleVectorSettingsWrapper>(GetTransientPackage(), NAME_None, RF_Standalone);
 		const FName SettingsPropertyName = GET_MEMBER_NAME_CHECKED(UPoleVectorSettingsWrapper, Settings);
-		Wrapper->InitializeWithRetargeter(StructToView, SettingsPropertyName, AssetController);
+		Wrapper->InitializeWithRetargeter(StructToView, SettingsPropertyName, Controller);
 		Wrapper->SetPropertyHidden(GET_MEMBER_NAME_CHECKED(FRetargetPoleVectorSettings, TargetChainName), true);
 		return Wrapper;
 	};
@@ -699,12 +699,12 @@ void FStretchChainOpCustomization::CustomizeChildren(
 	ChainListConfig.Controller = AssetController;
 	ChainListConfig.bEnableGoalColumn = true;
 	ChainListConfig.bEnableChainMapping = true;
-	ChainListConfig.ChainSettingsGetterFunc = [this](const FName InTargetChainName) -> UObject*
+	ChainListConfig.ChainSettingsGetterFunc = [Controller=AssetController, ThisOpName=OpName](const FName InTargetChainName) -> UObject*
 	{
 		// this lambda fetches the chain settings struct for the given target chain in the given op (by name)
-		auto MemoryProvider = [this, InTargetChainName]() -> uint8*
+		auto MemoryProvider = [Controller, ThisOpName, InTargetChainName]() -> uint8*
 		{
-			if (FIKRetargetOpBase* Op = AssetController->GetRetargetOpByName(OpName))
+			if (FIKRetargetOpBase* Op = Controller->GetRetargetOpByName(ThisOpName))
 			{
 				FIKRetargetStretchChainOpSettings* OpSettings = reinterpret_cast<FIKRetargetStretchChainOpSettings*>(Op->GetSettings());
 				for (FRetargetStretchChainSettings& ChainToStretch : OpSettings->ChainsToStretch)
@@ -719,14 +719,14 @@ void FStretchChainOpCustomization::CustomizeChildren(
 		};
 			
 		FIKRigStructToView StructToView;
-		StructToView.Owner = AssetController->GetAsset();
+		StructToView.Owner = Controller->GetAsset();
 		StructToView.Type = FRetargetStretchChainSettings::StaticStruct();
 		StructToView.MemoryProvider = MemoryProvider;
-		StructToView.UniqueName = OpName;
+		StructToView.UniqueName = ThisOpName;
 		
 		UStretchChainSettingsWrapper* Wrapper = NewObject<UStretchChainSettingsWrapper>(GetTransientPackage(), NAME_None, RF_Standalone);
 		const FName SettingsPropertyName = GET_MEMBER_NAME_CHECKED(UStretchChainSettingsWrapper, Settings);
-		Wrapper->InitializeWithRetargeter(StructToView, SettingsPropertyName, AssetController);
+		Wrapper->InitializeWithRetargeter(StructToView, SettingsPropertyName, Controller);
 		Wrapper->SetPropertyHidden(GET_MEMBER_NAME_CHECKED(FRetargetStretchChainSettings, TargetChainName), true);
 		return Wrapper;
 	};
@@ -765,12 +765,12 @@ void FFloorConstraintOpCustomization::CustomizeChildren(
 	ChainListConfig.bEnableGoalColumn = true;
 	ChainListConfig.bEnableChainMapping = false;
 	ChainListConfig.Filter = ShowOnlyIKFilter;
-	ChainListConfig.ChainSettingsGetterFunc = [this](const FName InTargetChainName) -> UObject*
+	ChainListConfig.ChainSettingsGetterFunc = [Controller=AssetController, ThisOpName=OpName](const FName InTargetChainName) -> UObject*
 	{
 		// this lambda fetches the chain settings struct for the given target chain in the given op (by name)
-		auto MemoryProvider = [this, InTargetChainName]() -> uint8*
+		auto MemoryProvider = [Controller, ThisOpName, InTargetChainName]() -> uint8*
 		{
-			if (FIKRetargetOpBase* Op = AssetController->GetRetargetOpByName(OpName))
+			if (FIKRetargetOpBase* Op = Controller->GetRetargetOpByName(ThisOpName))
 			{
 				FIKRetargetFloorConstraintOpSettings* OpSettings = reinterpret_cast<FIKRetargetFloorConstraintOpSettings*>(Op->GetSettings());
 				for (FFloorConstraintChainSettings& ChainSettings : OpSettings->ChainsToAffect)
@@ -785,14 +785,14 @@ void FFloorConstraintOpCustomization::CustomizeChildren(
 		};
 			
 		FIKRigStructToView StructToView;
-		StructToView.Owner = AssetController->GetAsset();
+		StructToView.Owner = Controller->GetAsset();
 		StructToView.Type = FFloorConstraintChainSettings::StaticStruct();
 		StructToView.MemoryProvider = MemoryProvider;
-		StructToView.UniqueName = OpName;
+		StructToView.UniqueName = ThisOpName;
 		
 		UFloorConstraintSettingsWrapper* Wrapper = NewObject<UFloorConstraintSettingsWrapper>(GetTransientPackage(), NAME_None, RF_Standalone);
 		const FName SettingsPropertyName = GET_MEMBER_NAME_CHECKED(UFloorConstraintSettingsWrapper, Settings);
-		Wrapper->InitializeWithRetargeter(StructToView, SettingsPropertyName, AssetController);
+		Wrapper->InitializeWithRetargeter(StructToView, SettingsPropertyName, Controller);
 		Wrapper->SetPropertyHidden(GET_MEMBER_NAME_CHECKED(FFloorConstraintChainSettings, TargetChainName), true);
 		return Wrapper;
 	};
diff --git a/Engine/Plugins/Animation/IKRig/Source/IKRigEditor/Private/RigEditor/IKRigStructViewer.cpp b/Engine/Plugins/Animation/IKRig/Source/IKRigEditor/Private/RigEditor/IKRigStructViewer.cpp
--- a/Engine/Plugins/Animation/IKRig/Source/IKRigEditor/Private/RigEditor/IKRigStructViewer.cpp
+++ b/Engine/Plugins/Animation/IKRig/Source/IKRigEditor/Private/RigEditor/IKRigStructViewer.cpp
@@ -84,12 +84,18 @@ USkeleton* UIKRigStructViewer::GetSkeleton(bool& bInvalidSkeletonIsError, const
 	
 	// NOTE: it's not ideal that we are hardcoding supported types here, but because UStruct's do not support multiple
 	// inheritance we cannot use an interface to identify skeleton providers as we normally would.
-	if (StructToView.Type->IsChildOf(FIKRetargetOpSettingsBase::StaticStruct()))
+	if (StructToView.Type->IsChildOf(FIKRetargetOpSkeletonProvider::StaticStruct()))
 	{
 		FName PropertyName = PropertyHandle->GetProperty()->GetFName();
 		uint8* StructMemory = StructToView.MemoryProvider();
-		FIKRetargetOpSettingsBase* SkeletonProvider = reinterpret_cast<FIKRetargetOpSettingsBase*>(StructMemory);
-		return SkeletonProvider->GetSkeleton(PropertyName);
+		FIKRetargetOpSkeletonProvider* SkeletonProvider = reinterpret_cast<FIKRetargetOpSkeletonProvider*>(StructMemory);
+		USkeleton* Skeleton = SkeletonProvider->GetSkeleton(PropertyName);
+		if (!Skeleton)
+		{
+			// if GetSkeleton(FName) returns null, try the alternative function that takes a property handle
+			Skeleton = SkeletonProvider->GetSkeleton(PropertyHandle);
+		}
+		return Skeleton;
 	}
 	
 	return nullptr;
