[IK Retarget] Adding foot constraint to floor constraint op.

diff --git a/Engine/Plugins/Animation/IKRig/Source/IKRig/Private/Retargeter/RetargetOps/FloorConstraintOp.cpp b/Engine/Plugins/Animation/IKRig/Source/IKRig/Private/Retargeter/RetargetOps/FloorConstraintOp.cpp
--- a/Engine/Plugins/Animation/IKRig/Source/IKRig/Private/Retargeter/RetargetOps/FloorConstraintOp.cpp
+++ b/Engine/Plugins/Animation/IKRig/Source/IKRig/Private/Retargeter/RetargetOps/FloorConstraintOp.cpp
@@ -7,6 +7,12 @@
 #include "Retargeter/RetargetOps/IKChainsOp.h"
 #include "Retargeter/RetargetOps/RunIKRigOp.h"
 
+#if WITH_EDITOR
+#include "PrimitiveDrawInterface.h"
+#include "PrimitiveDrawingUtils.h"
+#include "Engine/EngineTypes.h"
+#endif
+
 #include UE_INLINE_GENERATED_CPP_BY_NAME(FloorConstraintOp)
 
 #define LOCTEXT_NAMESPACE "FloorConstraintOp"
@@ -39,14 +45,515 @@ void FIKRetargetFloorConstraintOpSettings::CopySettingsAtRuntime(const FIKRetarg
 	{
 		for (FFloorConstraintChainSettings& ChainSettings : ChainsToAffect)
 		{
-			if (ChainSettings.TargetChainName == NewChainSettings.TargetChainName)
+			if (ChainSettings.TargetChainName != NewChainSettings.TargetChainName)
+			{
+				continue;
+			}
+			
+			// copy all chain settings except the foot definition
+			const TArray<FName> IgnoreFoot = {GET_MEMBER_NAME_CHECKED(FFloorConstraintChainSettings, Foot)};
+			FIKRetargetOpBase::CopyStructProperties(
+			FFloorConstraintChainSettings::StaticStruct(),
+			&NewChainSettings,
+			&ChainSettings,
+			IgnoreFoot);
+
+			// copy the foot definition
+			FIKRetargetOpBase::CopyStructProperties(
+			FFloorConstraintFootDefinition::StaticStruct(),
+			&NewChainSettings.Foot,
+			&ChainSettings.Foot,
+			TArray<FName>());
+
+			// copy the toes definition except the toes array
+			const TArray<FName> IgnoreToes = {GET_MEMBER_NAME_CHECKED(FFloorConstraintToesDefinition, AllToes)};
+			FIKRetargetOpBase::CopyStructProperties(
+			FFloorConstraintToesDefinition::StaticStruct(),
+			&NewChainSettings.Toes,
+			&ChainSettings.Toes,
+			IgnoreToes);
+
+			// copy only the toes that match
+			for (const FFloorConstraintToeDefinition& NewToe : NewChainSettings.Toes.AllToes)
 			{
-				ChainSettings = NewChainSettings;
+				for (FFloorConstraintToeDefinition& Toe : ChainSettings.Toes.AllToes)
+				{
+					if (Toe.ToeBone == NewToe.ToeBone)
+					{
+						FIKRetargetOpBase::CopyStructProperties(
+						FFloorConstraintToeDefinition::StaticStruct(),
+						&NewToe,
+						&Toe,
+						TArray<FName>());
+					}
+				}
 			}
 		}
 	}
 }
 
+#if WITH_EDITOR
+void FIKRetargetFloorConstraintOpSettings::AssignSkeletonAssets(const USkeleton* InSourceSkeleton, const USkeleton* InTargetSkeleton)
+{
+	FIKRetargetOpSettingsBase::AssignSkeletonAssets(InSourceSkeleton, InTargetSkeleton);
+	for (FFloorConstraintChainSettings& Chain : ChainsToAffect)
+	{
+		Chain.SourceSkeletonAsset = InSourceSkeleton;
+		Chain.TargetSkeletonAsset = InTargetSkeleton;
+	}
+}
+#endif
+
+void FScalarOneEuroFilter::Tune(
+	const double InFrequency,
+	const double InMinCutoff,
+	const double InBeta,
+	const double InDerivativeCutoff)
+{
+	Frequency = InFrequency;
+	MinCuttoff = InMinCutoff;
+	Beta = InBeta;
+	DerivativeCutoff = InDerivativeCutoff;
+}
+
+double FScalarOneEuroFilter::Run(double InValue, double InDeltaTime)
+{
+	auto LowPass = [](const double InPrev, const double InCurrent, const double InAlpha) -> double
+		{
+			return InAlpha * InCurrent + (1.0 - InAlpha) * InPrev;
+		};
+
+	auto Alpha = [](const double InInvDeltaTime, const double InCutoff) -> double
+		{
+			const double Tau = 1.0 / (2.0 * UE_PI * InCutoff);
+			return 1.0 / (1.0 + Tau * InInvDeltaTime);
+		};
+		
+	if (bIsFirstRun)
+	{
+		PrevValue = InValue;
+		PrevValueDerivative = 0.0;
+		bIsFirstRun = false;
+	}
+
+	// compute the derivative of the signal
+	const double InvDeltaTime = 1.0 / FMath::Max(InDeltaTime, UE_KINDA_SMALL_NUMBER);
+	const double Derivative = (InValue - PrevValue) * InvDeltaTime;
+	const double LowPassDerivative = LowPass(PrevValueDerivative, Derivative, Alpha(InvDeltaTime, DerivativeCutoff));
+	PrevValueDerivative = LowPassDerivative;
+
+	// compute the cutoff frequency based on signal speed
+	const double Cutoff = MinCuttoff + Beta * FMath::Abs(LowPassDerivative);
+
+	// filter the signal
+	const double Result = LowPass(PrevValue, InValue, Alpha(InvDeltaTime, Cutoff));
+	PrevValue = Result;
+
+	return Result;
+}
+
+void FFootConstraint::Initialize(
+	const FTargetSkeleton& InTargetSkeleton,
+	const FFloorConstraintFootDefinition& InFootSettings,
+	const int32 InAnkleBoneIndex,
+	FIKRigLogger& InLog)
+{
+	bReadyToRun = false;
+	
+	AnkleBoneIndex = InAnkleBoneIndex;
+	if (!InTargetSkeleton.BoneNames.IsValidIndex(InAnkleBoneIndex))
+	{
+		InLog.LogWarning(LOCTEXT("FootConstraintMissingAnkleBone", "Foot Constraint: missing target ankle bone."));
+		return;
+	}
+	
+	const TArray<FTransform>& TargetRetargetPose = InTargetSkeleton.RetargetPoses.GetGlobalRetargetPose();
+	CurrentAnkleGlobal = TargetRetargetPose[AnkleBoneIndex];
+
+	const FVector MedialDirection = FVector::XAxisVector * (CurrentAnkleGlobal.GetLocation().X > 0.0 ? -1.0 : 1.0);
+	const FVector LateralDirection = MedialDirection * -1.0;
+	const FVector ToeDirection = FVector(0,1,0);
+	const FVector HeelDirection = -ToeDirection;
+	const FVector MedialOffset = MedialDirection * InFootSettings.MedialOffset;
+	const FVector LateralOffset = LateralDirection * InFootSettings.LateralOffset;
+	const FVector ToeOffset = ToeDirection * InFootSettings.ToeOffset;
+	const FVector HeelOffset = HeelDirection * InFootSettings.HeelOffset;
+	const FVector VerticalOffset = FVector(0,0, InFootSettings.VerticalOffset);
+	const FVector AnkleOnFloor = (CurrentAnkleGlobal.GetLocation() * FVector(1,1,0)) + VerticalOffset;
+
+	// offset corner points from position of ankle projected on the floor
+	PointsOrig[FrontMedial] = AnkleOnFloor + ToeOffset + MedialOffset;
+	PointsOrig[FrontLateral] = AnkleOnFloor + ToeOffset + LateralOffset;
+	PointsOrig[RearMedial] = AnkleOnFloor + HeelOffset + MedialOffset;
+	PointsOrig[RearLateral] = AnkleOnFloor + HeelOffset + LateralOffset;
+
+	// put in space of ankle
+	PointsOrig[FrontMedial] = CurrentAnkleGlobal.InverseTransformPosition(PointsOrig[FrontMedial]);
+	PointsOrig[FrontLateral] = CurrentAnkleGlobal.InverseTransformPosition(PointsOrig[FrontLateral]);
+	PointsOrig[RearMedial] = CurrentAnkleGlobal.InverseTransformPosition(PointsOrig[RearMedial]);
+	PointsOrig[RearLateral] = CurrentAnkleGlobal.InverseTransformPosition(PointsOrig[RearLateral]);
+
+	bReadyToRun = true;
+}
+
+void FFootConstraint::UpdateBasePoints(const TArray<FTransform>& InGlobalPose, const FIKRigGoal* InGoal)
+{
+	// get latest foot/ball transforms
+	if (InGoal)
+	{
+		CurrentAnkleGlobal.SetLocation(InGoal->Position);
+		CurrentAnkleGlobal.SetRotation(InGoal->Rotation.Quaternion());
+	}
+	else
+	{
+		CurrentAnkleGlobal = InGlobalPose[AnkleBoneIndex];
+	}
+	
+	Points[FrontMedial] = CurrentAnkleGlobal.TransformPosition(PointsOrig[FrontMedial]);
+	Points[FrontLateral] = CurrentAnkleGlobal.TransformPosition(PointsOrig[FrontLateral]);
+	Points[RearMedial] = CurrentAnkleGlobal.TransformPosition(PointsOrig[RearMedial]);
+	Points[RearLateral] = CurrentAnkleGlobal.TransformPosition(PointsOrig[RearLateral]);
+}
+
+int32 FFootConstraint::GetIndexOfLowestPoint(const FVector* InPoints)
+{
+	int32 MinIndex = 0;
+	float MinValue = TNumericLimits<float>::Max();
+
+	for (int32 i = 0; i <= Count; ++i)
+	{
+		const float Z = InPoints[i].Z;
+		if (Z < MinValue)
+		{
+			MinValue = Z;
+			MinIndex = i;
+		}
+	}
+
+	return MinIndex;
+}
+
+void FToeConstraint::Initialize(
+	const FTargetSkeleton& InTargetSkeleton,
+	const FFloorConstraintToesDefinition& InAllToeSettings,
+	const FFloorConstraintToeDefinition& InToeSettings,
+	const int32 InAnkleBoneIndex,
+	FIKRigLogger& InLog)
+{
+	bToeReady = false;
+	
+	const FName ToeBoneName = InToeSettings.ToeBone.BoneName;
+	ToeBoneIndex = InTargetSkeleton.FindBoneIndexByName(ToeBoneName);
+	if (ToeBoneIndex == INDEX_NONE)
+	{
+		InLog.LogWarning(
+			FText::Format(LOCTEXT("FootConstraintMissingToeBone", "Floor Constraint: missing target toe bone, '{0}'."),
+				FText::FromName(ToeBoneName)));
+		return;
+	}
+
+	AnkleBoneIndex = InAnkleBoneIndex;
+	if (InAnkleBoneIndex == INDEX_NONE)
+	{
+		return; // logged inside Foot.Init
+	}
+
+	Settings = &InToeSettings;
+	AllToeSettings = &InAllToeSettings;
+
+	// initialize local offset and start/end points based on retarget pose
+	const TArray<FTransform>& GlobalRetargetPose = InTargetSkeleton.RetargetPoses.GetGlobalRetargetPose();
+	InitialToeGlobal = GlobalRetargetPose[ToeBoneIndex];
+	constexpr bool bUpdateLocalOffset = true;
+	UpdateFramePoints(GlobalRetargetPose,  nullptr, bUpdateLocalOffset);
+
+	bToeReady = true;
+}
+
+void FToeConstraint::RunAfterParent(
+	const FTargetSkeleton& InTargetSkeleton,
+	TArray<FTransform>& OutGlobalPose,
+	const FFloorConstraintChainSettings* InChainSettings,
+	const double InDeltaTime)
+{
+	if (!bToeReady)
+	{
+		return;
+	}
+
+	constexpr FIKRigGoal* NullGoal = nullptr; // this is running after the IK solve, so use the foot at the input pose transform
+	bool bUpdateLocalOffset = false;
+	UpdateFramePoints(OutGlobalPose, NullGoal, bUpdateLocalOffset);
+
+	// solve for the angle to put toe on floor
+	double ThetaRadians = CalcToeCorrectionAngle();
+
+	// temporal filtering of the toe adjustment
+	AngleFilter.Tune(
+		AllToeSettings->FilterFrequency,
+		AllToeSettings->FilterMinCutoff,
+		AllToeSettings->FilterBeta,
+		AllToeSettings->FilterDerivCutoff);
+	ThetaRadians = AngleFilter.Run(ThetaRadians, InDeltaTime);
+
+	// alpha the whole effect
+	ThetaRadians = FMath::Lerp(0, ThetaRadians, AllToeSettings->Alpha * InChainSettings->Alpha);
+
+	// convert to a quaternion that rotates around the pitch axis by theta
+	const FQuat DeltaRotation = FQuat(PitchAxis, ThetaRadians);
+	
+	// apply rotation to toe bone and update children
+	FTransform NewToeGlobal = OutGlobalPose[ToeBoneIndex];
+	NewToeGlobal.SetRotation(DeltaRotation * NewToeGlobal.GetRotation());
+	InTargetSkeleton.SetGlobalTransformAndUpdateChildren(ToeBoneIndex, NewToeGlobal, OutGlobalPose);
+
+#if WITH_EDITOR
+	// update frame points after solve while in editor for debug drawing
+	bUpdateLocalOffset = true;
+	UpdateFramePoints(OutGlobalPose, NullGoal, bUpdateLocalOffset);
+#endif
+}
+
+void FToeConstraint::UpdateFramePoints(
+	const TArray<FTransform>& InGlobalPose,
+	const FIKRigGoal* InGoal,
+	bool bUpdateLocalOffset)
+{
+	FTransform CurrentAnkleGlobal;
+	if (InGoal)
+	{
+		CurrentAnkleGlobal.SetLocation(InGoal->Position);
+		CurrentAnkleGlobal.SetRotation(InGoal->Rotation.Quaternion());
+	}
+	else
+	{
+		CurrentAnkleGlobal = InGlobalPose[AnkleBoneIndex];
+	}
+
+	const FTransform CurrentLocalToe = InGlobalPose[ToeBoneIndex].GetRelativeTransform(InGlobalPose[AnkleBoneIndex]);
+	CurrentToeGlobal = CurrentLocalToe * CurrentAnkleGlobal;
+
+	if (bUpdateLocalOffset && Settings)
+	{
+		// calc toe end relative to the toe
+		FVector EndOffset = FVector::UnitY() * Settings->Length;
+		const FQuat YawRotation = FQuat(FVector::ZAxisVector, FMath::DegreesToRadians(Settings->YawOffset));
+		EndOffset = YawRotation.RotateVector(EndOffset);
+		FVector EndPoint = InitialToeGlobal.GetLocation() + EndOffset;
+		EndPoint.Z = Settings->VerticalOffset;
+		EndRelativeToToe = InitialToeGlobal.InverseTransformPosition(EndPoint);
+
+		// calc pitch axis
+		const FVector InitialGlobalPitchAxis = YawRotation.RotateVector(FVector::UnitX());
+		PitchAxisRelativeToToe = InitialToeGlobal.InverseTransformVector(InitialGlobalPitchAxis);
+	}
+	
+	Start = CurrentToeGlobal.GetLocation();
+	End = CurrentToeGlobal.TransformPosition(EndRelativeToToe);
+	PitchAxis = CurrentToeGlobal.TransformVector(PitchAxisRelativeToToe);
+}
+
+double FToeConstraint::CalcToeCorrectionAngle()
+{
+	if (End.Z >= 0)
+	{
+		return 0; // toe is above the floor already
+	}
+
+	const FVector r = End - Start;
+	if (r.IsNearlyZero())
+	{
+		return 0; // toe has degenerate length
+	}
+
+	// NOTE:
+	// The following solver code may, at first, look a bit magic, so I include this derivation.
+	//
+	// The overall idea is to treat this as a circle/plane collision check where the toe point (End) rotates
+	// around the toe joint (Start) in the plane perpendicular to the PitchAxis (k).
+	//
+	// We need to find the rotation angle theta around axis k such that the rotated toe endpoint lies on the floor (Z = 0).
+	//
+	// Rodrigues' formula gives vector r' after rotating r by theta:
+	//
+	//     r'(theta) = r * cos(theta) + cross(k, r) * sin(theta) + k * dot(k, r) * (1 - cos(theta))
+	//
+	// Isolating the Z component gives a function:
+	//
+	//     z(theta) = A * cos(theta) + B * sin(theta) + C
+	//
+	// Explanation:
+	// Each term in Rodrigues' formula contributes to the Z height of the rotated toe tip:
+	//   - r * cos(theta): projects the original vector r onto its rotated position (scaled by cos(theta)).
+	//     Its Z component varies as cos(theta), giving coefficient A.
+	//   - cross(k, r) * sin(theta): represents rotation perpendicular to both k and r.
+	//     Its Z component varies as sin(theta), giving coefficient B.
+	//   - k * dot(k, r) * (1 - cos(theta)): accounts for the component of r already aligned with the axis k.
+	//     This adds a constant offset to Z (independent of theta) captured by coefficient C.
+	//
+	// Summing these gives a trig dependence of Z on theta:
+	//
+	//     z(theta) = A * cos(theta) + B * sin(theta) + C
+	//
+	// Meaning the vertical position of the toe tip moves up and down as the vector r rotates in a circular arc about axis k.
+	//
+	// We want z(theta) + Start.Z = 0, which expands to:
+	//
+	//     A * cos(theta) + B * sin(theta) + C = 0
+	//
+	// Any linear combination of sine and cosine can be written as a single cosine with a phase shift:
+	//
+	//     A * cos(theta) + B * sin(theta) = R * cos(theta - phi)
+	//
+	// where:
+	//     R = sqrt(A*A + B*B)   // amplitude of the combined wave
+	//     phi = atan2(B, A)     // phase offset (angle of vector [A, B])
+	//
+	// Substituting back:
+	//
+	//     R * cos(theta - phi) = -C
+	//
+	// Now with a single cosine, we can isolate theta:
+	//
+	//     theta = phi +/- acos(-C / R)
+	//
+	// These are the two possible rotations that put the toe on the floor.
+	// We clamp X = -C / R to [-1, 1] for safety and pick the smallest positive theta so the toe only bends upward.
+	
+	// a is z component of first term
+	const double a = r.Z;
+	// b is z component of second term
+	const FVector& k = PitchAxis;
+	const FVector kxr = FVector::CrossProduct(k, r);
+	const double b = kxr.Z;
+	// c is z component of third term
+	const double c = k.Z * FVector::DotProduct(k, r);
+
+	// now we need to solve for: Start.Z + z(theta) = 0
+	// (a-c)*cos(theta) + b*sin(theta) + (c+Start.Z) = 0
+	const double A = a - c;
+	const double B = b;
+	const double C = c + Start.Z;
+
+	// if A and B are orthogonal, then rotation about k cannot change Z, bail.
+	const double R = FMath::Sqrt(A*A + B*B);
+	if (R < UE_SMALL_NUMBER) 
+	{
+		return 0; // toe Z coord is effectively independent of theta
+	}
+	
+	const double Phi = FMath::Atan2(B, A);
+	double X = -C / R;
+	X = FMath::Clamp(X, -1.0, 1.0);
+	const double AcosX = FMath::Acos(X);
+
+	// two candidate angles: theta = phi +/- acos(x)
+	const double ThetaA = FMath::UnwindRadians(Phi + AcosX); // wrap to (-pi,pi)
+	const double ThetaB = FMath::UnwindRadians(Phi - AcosX);
+	
+	// pick the smallest positive theta
+	// (this ensures toes only bend up, never backwards)
+	if (ThetaA >= 0.0 && ThetaB >= 0.0)
+	{
+		// both up: pick smaller magnitude
+		return (ThetaA <= ThetaB) ? ThetaA : ThetaB;
+	}
+	if (ThetaA >= 0.0)
+	{
+		return ThetaA;
+	}
+	if (ThetaB >= 0.0)
+	{
+		return ThetaB;
+	}
+	
+	return 0; // both would rotate downward, bail out
+}
+
+void FFootConstraint::Run(
+	FIKRigGoal& InOutGoal,
+	const TArray<FTransform>& OutTargetGlobalPose,
+	const FFloorConstraintChainSettings& InSettings)
+{
+	// calc vertical delta to push all base points above the floor
+	UpdateBasePoints(OutTargetGlobalPose, &InOutGoal);
+	const int32 LowestBasePointIndex = GetIndexOfLowestPoint(Points);
+	const double BaseHeight = Points[LowestBasePointIndex].Z;
+	const double BaseDelta = (BaseHeight < 0.0f ? FMath::Abs(BaseHeight) : 0.0) * InSettings.Alpha;
+	
+	// move the goal up
+	InOutGoal.Position.Z += BaseDelta;
+}
+
+void FFloorConstraint::Initialize(
+	const FResolvedBoneChain& InSourceBoneChain,
+	const FResolvedBoneChain& InTargetBoneChain,
+	const FFloorConstraintChainSettings& InSettings)
+{
+	IKGoalName = InTargetBoneChain.IKGoalName;
+	SourceEndBoneIndex = InSourceBoneChain.BoneIndices.Last();
+	TargetEndBoneIndex = InTargetBoneChain.BoneIndices.Last();
+	Settings = &InSettings;
+
+	// cache ref pose offset
+	const FTransform& SourceEndBoneRefPose = InSourceBoneChain.RefPoseGlobalTransforms.Last();
+	const FTransform& TargetEndBoneRefPose = InTargetBoneChain.RefPoseGlobalTransforms.Last();
+	HeightOffsetInRefPose = TargetEndBoneRefPose.GetTranslation().Z - SourceEndBoneRefPose.GetTranslation().Z;
+}
+
+void FFloorConstraint::Run(
+	const FTargetSkeleton& InTargetSkeleton,
+	FIKRigGoalContainer& InGoalContainer,
+	const TArray<FTransform>& InSourceGlobalPose,
+	TArray<FTransform>& OutTargetGlobalPose,
+	const FIKRetargetFloorConstraintOpSettings& InOpSettings)
+{
+	if (!Settings->EnableFloorConstraint)
+	{
+		return; // constraint disabled
+	}
+		
+	if (FMath::IsNearlyZero(Settings->Alpha))
+	{
+		return; // constraint weight near zero
+	}
+
+	// get the goal to adjust
+	FIKRigGoal* GoalPtr = InGoalContainer.FindGoalByName(IKGoalName);
+	if (!ensure(GoalPtr))
+	{
+		return; // goal missing (should not happen at runtime)
+	}
+	FIKRigGoal& Goal = *GoalPtr;
+
+	// adjust height of IK Goal to match the source relative to floor
+		
+	// calculate falloff range
+	const float FalloffStartHeight = InOpSettings.HeightFalloffOffset;
+	float FalloffEndHeight = FalloffStartHeight + FMath::Max(0.0, InOpSettings.HeightFalloffDistance);
+
+	// get the current height of the source bone
+	const FTransform& CurrentSourceEndGlobal = InSourceGlobalPose[SourceEndBoneIndex];
+	const float SourceHeight = FMath::Abs(CurrentSourceEndGlobal.GetLocation().Z);
+
+	// constraint on while source is near ground
+	if (SourceHeight < FalloffEndHeight)
+	{
+		// calculate the target height of the goal by blending towards the source height when close to floor
+		const double FalloffRange = FMath::Max(1.0,FalloffEndHeight - FalloffStartHeight);
+		const double Falloff = SourceHeight < FalloffStartHeight ? 1.0f : 1.0f - (SourceHeight - FalloffStartHeight) / FalloffRange;
+		const double Weight = Falloff * Settings->Alpha;
+		const double HeightOffset = FMath::Lerp(0.0f, HeightOffsetInRefPose, Settings->MaintainHeightOffset);
+		const double TargetHeight = SourceHeight + HeightOffset;
+		
+		// lerp goal height from its current height to the target height by the weight
+		Goal.Position.Z = FMath::Lerp(Goal.Position.Z, TargetHeight, Weight);
+	}
+
+	// run the foot/floor constraint based on new goal location
+	FootConstraint.Run(Goal, OutTargetGlobalPose, *Settings);
+}
+
 bool FIKRetargetFloorConstraintOp::Initialize(
     const FIKRetargetProcessor& InProcessor,
     const FRetargetSkeleton& InSourceSkeleton,
@@ -65,11 +572,12 @@ bool FIKRetargetFloorConstraintOp::Initialize(
 		return false;
 	}
 
-	FloorConstraints.Reset();
+	// initialize / load floor constraints
 	const FRetargeterBoneChains& BoneChains = InProcessor.GetBoneChains();
-	for (const FFloorConstraintChainSettings& Chain : Settings.ChainsToAffect)
+	FloorConstraints.Reset();
+	for (const FFloorConstraintChainSettings& ChainSettings : Settings.ChainsToAffect)
 	{
-		const FName TargetChainName = Chain.TargetChainName;
+		const FName TargetChainName = ChainSettings.TargetChainName;
 		const FResolvedBoneChain* TargetBoneChain = BoneChains.GetResolvedBoneChainByName(TargetChainName, ERetargetSourceOrTarget::Target,ParentOp->Settings.IKRigAsset);
 		if (!TargetBoneChain)
 		{
@@ -95,19 +603,28 @@ bool FIKRetargetFloorConstraintOp::Initialize(
 			continue;
 		}
 
-		// all prerequisites met, create a new floor constraint
-		FFloorConstraint FloorConstraint;
-		FloorConstraint.IKGoalName = TargetBoneChain->IKGoalName;
-		FloorConstraint.SourceEndBoneIndex = SourceBoneChain->BoneIndices.Last();
-		FloorConstraint.Settings = &Chain;
-
-		// cache ref pose offset
-		const FTransform& SourceEndBoneRefPose = SourceBoneChain->RefPoseGlobalTransforms.Last();
-		const FTransform& TargetEndBoneRefPose = TargetBoneChain->RefPoseGlobalTransforms.Last();
-		FloorConstraint.HeightOffsetInRefPose = TargetEndBoneRefPose.GetTranslation().Z - SourceEndBoneRefPose.GetTranslation().Z;
+		// create a FLOOR constraint
+		if (ChainSettings.EnableFloorConstraint)
+		{
+			FFloorConstraint& FloorConstraint = FloorConstraints.Emplace_GetRef();
+			FloorConstraint.Initialize(*SourceBoneChain, *TargetBoneChain, ChainSettings);
+			
+			// initialize foot constraint
+			if (ChainSettings.bUseFoot || ChainSettings.bUseToes)
+			{
+				FloorConstraint.FootConstraint.Initialize(InTargetSkeleton, ChainSettings.Foot, FloorConstraint.TargetEndBoneIndex, InLog);
+			}
 
-		// store floor constraint data for runtime
-		FloorConstraints.Add(FloorConstraint);
+			// initialize toe constraints
+			if (ChainSettings.bUseToes)
+			{
+				for (const FFloorConstraintToeDefinition& ToeSettings : ChainSettings.Toes.AllToes)
+				{
+					FToeConstraint& Toe = FloorConstraint.ToeConstraints.Emplace_GetRef();
+					Toe.Initialize(InTargetSkeleton, ChainSettings.Toes, ToeSettings, FloorConstraint.TargetEndBoneIndex, InLog);
+				}
+			}
+		}
 	}
 	
 	bIsInitialized = true;
@@ -120,47 +637,43 @@ void FIKRetargetFloorConstraintOp::Run(
     const TArray<FTransform>& InSourceGlobalPose,
     TArray<FTransform>& OutTargetGlobalPose)
 {
+	// run all the floor constraints (1 per chain)
+	// NOTE: in the main pass, all we do is affect the IK Goal to:
+	// A. match the source height relative to the ground
+	// B. if foot dimensions are available, push the goal upwards to keep foot out of the ground
+	const FTargetSkeleton& TargetSkeleton = InProcessor.GetTargetSkeleton();
 	FIKRigGoalContainer& GoalContainer = InProcessor.GetIKRigGoalContainer();
 	for (FFloorConstraint& FloorConstraint : FloorConstraints)
 	{
-		if (!FloorConstraint.Settings->EnableFloorConstraint)
-		{
-			continue; // constraint disabled
-		}
-		
-		if (FMath::IsNearlyZero(FloorConstraint.Settings->Alpha))
-		{
-			continue; // constraint weight near zero
-		}
+		FloorConstraint.Run(TargetSkeleton,GoalContainer, InSourceGlobalPose, OutTargetGlobalPose, Settings);
+	}
+}
 
-		// get the goal to adjust
-		FIKRigGoal* Goal = GoalContainer.FindGoalByName(FloorConstraint.IKGoalName);
-		if (!ensure(Goal))
+void FIKRetargetFloorConstraintOp::RunAfterParent(
+	FIKRetargetProcessor& InProcessor,
+	const double InDeltaTime,
+	const TArray<FTransform>& InSourceGlobalPose,
+	TArray<FTransform>& OutTargetGlobalPose)
+{
+	// NOTE: this pass runs AFTER the IK has solved.
+	// At this point, we can adjust the toes on the post-IK pose to keep them above the floor
+	const FTargetSkeleton& TargetSkeleton = InProcessor.GetTargetSkeleton();
+	for (FFloorConstraint& FloorConstraint : FloorConstraints)
+	{
+		if (!FloorConstraint.FootConstraint.bReadyToRun)
 		{
-			continue; // goal missing (should not happen at runtime)
+			continue;
 		}
 
-		// calculate falloff range
-		const float FalloffStartHeight = Settings.HeightFalloffOffset;
-		float FalloffEndHeight = FalloffStartHeight + FMath::Max(0.0, Settings.HeightFalloffDistance);
-
-		// get the current height of the source bone
-		const FTransform& CurrentSourceEndGlobal = InSourceGlobalPose[FloorConstraint.SourceEndBoneIndex];
-		const float SourceHeight = FMath::Abs(CurrentSourceEndGlobal.GetLocation().Z);
-		if (SourceHeight >= FalloffEndHeight)
+#if WITH_EDITOR
+		// for accurate debug drawing, we update the foot base points AFTER IK solve
+		FloorConstraint.FootConstraint.UpdateBasePoints(OutTargetGlobalPose, nullptr);
+#endif
+		
+		for (FToeConstraint& Toe : FloorConstraint.ToeConstraints)
 		{
-			continue; // constraint off while source is not near ground
+			Toe.RunAfterParent(TargetSkeleton, OutTargetGlobalPose, FloorConstraint.Settings, InDeltaTime);
 		}
-		
-		// calculate the target height of the goal by blending towards the source height when close to floor
-		const double FalloffRange = FMath::Max(1.0,FalloffEndHeight - FalloffStartHeight);
-		const double Falloff = SourceHeight < FalloffStartHeight ? 1.0f : 1.0f - (SourceHeight - FalloffStartHeight) / FalloffRange;
-		const double Weight = Falloff * FloorConstraint.Settings->Alpha;
-		const double HeightOffset = FMath::Lerp(0.0f, FloorConstraint.HeightOffsetInRefPose, FloorConstraint.Settings->MaintainHeightOffset);
-		const double TargetHeight = SourceHeight + HeightOffset;
-		
-		// lerp goal height from it's current height to the target height by the weight
-		Goal->Position.Z = FMath::Lerp(Goal->Position.Z, TargetHeight, Weight);
 	}
 }
 
@@ -188,14 +701,6 @@ FIKRetargetOpSettingsBase* FIKRetargetFloorConstraintOp::GetSettings()
     return &Settings;
 }
 
-void FIKRetargetFloorConstraintOp::SetSettings(const FIKRetargetOpSettingsBase* InSettings)
-{
-    if (InSettings)
-    {
-        Settings = *static_cast<const FIKRetargetFloorConstraintOpSettings*>(InSettings);
-    }
-}
-
 const UScriptStruct* FIKRetargetFloorConstraintOp::GetSettingsType() const
 {
     return FIKRetargetFloorConstraintOpSettings::StaticStruct();
@@ -257,6 +762,55 @@ bool FIKRetargetFloorConstraintOp::AreChainSettingsAtDefault(const FName& InChai
 	return true;
 }
 
+void FIKRetargetFloorConstraintOp::DebugDraw(
+	FPrimitiveDrawInterface* InPDI,
+	const FTransform& InSourceTransform,
+	const FTransform& InTargetTransform,
+	const double InComponentScale,
+	const FIKRetargetDebugDrawState& InEditorState) const
+{
+	static FLinearColor FootColor = FLinearColor::Green;
+	static FLinearColor ToeColor = FLinearColor::Blue;
+	constexpr float PointSize = 5.0f;
+	constexpr float LineThickness = 0.5f;
+	
+	auto DrawFootFrame = [&](const FFootConstraint& InFoot, const FTransform& InComponentTransform)
+		{
+			const FVector FrontMedial = InComponentTransform.TransformPosition(InFoot.Points[FFootConstraint::FrontMedial]);
+			const FVector FrontLateral = InComponentTransform.TransformPosition(InFoot.Points[FFootConstraint::FrontLateral]);
+			const FVector RearMedial = InComponentTransform.TransformPosition(InFoot.Points[FFootConstraint::RearMedial]);
+			const FVector RearLateral = InComponentTransform.TransformPosition(InFoot.Points[FFootConstraint::RearLateral]);
+			
+			InPDI->DrawPoint(RearMedial, FootColor, PointSize, SDPG_Foreground);
+			InPDI->DrawPoint(RearLateral, FootColor, PointSize, SDPG_Foreground);
+			InPDI->DrawPoint(FrontMedial, FootColor, PointSize, SDPG_Foreground);
+			InPDI->DrawPoint(FrontLateral, FootColor, PointSize, SDPG_Foreground);
+			
+			DrawDashedLine(InPDI, RearMedial, RearLateral, FootColor, 1, SDPG_Foreground);
+			DrawDashedLine(InPDI, RearMedial, FrontMedial, FootColor, 1, SDPG_Foreground);
+			DrawDashedLine(InPDI, RearLateral, FrontLateral, FootColor, 1, SDPG_Foreground);
+			DrawDashedLine(InPDI, FrontMedial, FrontLateral, FootColor, 1, SDPG_Foreground);
+		};
+	
+	for (const FFloorConstraint& FloorConstraint : FloorConstraints)
+	{
+		if (FloorConstraint.FootConstraint.bReadyToRun)
+		{
+			DrawFootFrame(FloorConstraint.FootConstraint, InTargetTransform);
+			
+			for (const FToeConstraint& Toe : FloorConstraint.ToeConstraints)
+			{
+				const FVector ToeStart = InTargetTransform.TransformPosition(Toe.Start);
+				const FVector ToeEnd = InTargetTransform.TransformPosition(Toe.End);
+				InPDI->DrawPoint(ToeEnd, ToeColor, PointSize, SDPG_Foreground);
+				InPDI->DrawLine(ToeStart, ToeEnd, ToeColor, SDPG_Foreground, LineThickness);
+				// draw pitch axis
+				InPDI->DrawLine(ToeStart, ToeStart + Toe.PitchAxis, FLinearColor::Red, SDPG_Foreground, LineThickness);
+			}
+		}
+	}
+}
+
 #endif // WITH_EDITOR
 
 FIKRetargetIKChainsOpSettings UIKRetargetFloorGoalsController::GetSettings()
diff --git a/Engine/Plugins/Animation/IKRig/Source/IKRig/Public/Retargeter/RetargetOps/FloorConstraintOp.h b/Engine/Plugins/Animation/IKRig/Source/IKRig/Public/Retargeter/RetargetOps/FloorConstraintOp.h
--- a/Engine/Plugins/Animation/IKRig/Source/IKRig/Public/Retargeter/RetargetOps/FloorConstraintOp.h
+++ b/Engine/Plugins/Animation/IKRig/Source/IKRig/Public/Retargeter/RetargetOps/FloorConstraintOp.h
@@ -3,15 +3,102 @@
 
 #include "Retargeter/IKRetargetChainMapping.h"
 #include "Retargeter/IKRetargetOps.h"
+#include "Animation/BoneReference.h"
 #include "FloorConstraintOp.generated.h"
 
 #define UE_API IKRIG_API
 
 #define LOCTEXT_NAMESPACE "FloorConstraintOp"
 
+struct FResolvedBoneChain;
+struct FIKRigGoalContainer;
+struct FIKRigGoal;
+
+struct FFloorConstraint;
+
+USTRUCT(BlueprintType)
+struct FFloorConstraintToeDefinition
+{
+	GENERATED_BODY()
+
+	/** The bone located at the base of the toe (the first knuckle). */
+	UPROPERTY(EditAnywhere, Category=Setup, meta=(ReinitializeOnEdit))
+	FBoneReference ToeBone;
+
+	/** The length of the toe from the ToeBone to the tip of the toe mesh. */
+	UPROPERTY(EditAnywhere, BlueprintReadWrite, Category = "Toe Dimensions", meta = (ReinitializeOnEdit, ClampMin = "0.0", UIMin = "0.0", UIMax = "50.0"))
+	double Length = 4.0f;
+
+	/** Rotate the end of the toe side-to-side. Adjust this to align the length of the toe with the model. */
+	UPROPERTY(EditAnywhere, BlueprintReadWrite, Category = "Toe Dimensions", meta = (ReinitializeOnEdit, ClampMin = "-179.0", ClampMax = "179.0", UIMin = "-179.0", UIMax = "179.0"))
+	double YawOffset = 0.0f;
+
+	/** Push the tip of the toes up/down. Adjust this to place the tip of the toe on the bottom of the toe mesh. */
+	UPROPERTY(EditAnywhere, BlueprintReadWrite, Category = "Toe Dimensions", meta = (ReinitializeOnEdit, UIMin = "-10.0", UIMax = "10.0"))
+	double VerticalOffset = 0.0f;
+};
 
 USTRUCT(BlueprintType)
-struct FFloorConstraintChainSettings
+struct FFloorConstraintToesDefinition
+{
+	GENERATED_BODY()
+
+	/** Default 1.0. Range 0 to 1. Blend the effect of the toe solver on/off. */
+	UPROPERTY(EditAnywhere, BlueprintReadWrite, Category = "Toe Settings", meta = (ReinitializeOnEdit, ClampMin = "0.0", ClampMax = "1.0", UIMin = "0.0", UIMax = "1.0"))
+	double Alpha = 1.0;
+	
+	/** Sampling frequency (Hz) used by the One Euro Filter. */
+	UPROPERTY(EditAnywhere, BlueprintReadWrite, Category = "Toe Filter", meta = (ClampMin = "1.0", UIMin = "1.0", UIMax = "240.0"))
+	float FilterFrequency = 120.0f;
+
+	/** Minimum cutoff frequency for the One Euro Filter (higher = less smoothing). */
+	UPROPERTY(EditAnywhere, BlueprintReadWrite, Category = "Toe Filter", meta = (ClampMin = "0.0", UIMin = "0.0", UIMax = "10.0"))
+	float FilterMinCutoff = 8.0f;
+
+	/** Speed coefficient; increases cutoff when signal changes quickly (reactivity). */
+	UPROPERTY(EditAnywhere, BlueprintReadWrite, Category = "Toe Filter", meta = (ClampMin = "0.0", UIMin = "0.0", UIMax = "5.0"))
+	float FilterBeta = 0.02f;
+
+	/** Cutoff frequency for the derivative smoothing step. */
+	UPROPERTY(EditAnywhere, BlueprintReadWrite, Category = "Toe Filter", meta = (ClampMin = "0.0", UIMin = "0.0", UIMax = "10.0"))
+	float FilterDerivCutoff = 1.0f;
+	
+	/** An array of toes to bend when contacting the floor */
+	UPROPERTY(EditAnywhere, Category="Toe Definitions", meta=(ReinitializeOnEdit))
+	TArray<FFloorConstraintToeDefinition> AllToes;
+};
+
+USTRUCT(BlueprintType)
+struct FFloorConstraintFootDefinition
+{
+	GENERATED_BODY()
+	
+	/** Default 10. The offset in cm towards the interior (middle) of the body. */
+	UPROPERTY(EditAnywhere, BlueprintReadWrite, Category = "Foot Dimensions", meta = (ReinitializeOnEdit, ClampMin = "0.0", UIMin = "0.0", UIMax = "50.0"))
+	double MedialOffset = 10.0f;
+
+	/** Default 10. The offset in cm towards the exterior (outside) of the body. */
+	UPROPERTY(EditAnywhere, BlueprintReadWrite, Category = "Foot Dimensions", meta = (ReinitializeOnEdit, ClampMin = "0.0", UIMin = "0.0", UIMax = "50.0"))
+	double LateralOffset = 10.0f;
+
+	/** Default 10. The offset in the direction of the heel (rear) of the foot. */
+	UPROPERTY(EditAnywhere, BlueprintReadWrite, Category = "Foot Dimensions", meta = (ReinitializeOnEdit, ClampMin = "0.0", UIMin = "0.0", UIMax = "50.0"))
+	double HeelOffset = 10.0f;
+
+	/** Default 10. The offset towards the toes (front) of the foot.
+	 * NOTE: if you are using toe constraints, this offset should extend to the ball of the foot.
+	 * Otherwise the toes will never be below the floor and thus never bend. */
+	UPROPERTY(EditAnywhere, BlueprintReadWrite, Category = "Foot Dimensions", meta = (ReinitializeOnEdit, ClampMin = "0.0", UIMin = "0.0", UIMax = "50.0"))
+	double ToeOffset = 10.0f;
+
+	/** Default 0. The vertical offset used to define the "bottom" of the foot mesh.
+	 * NOTE: adjust this up/down if the floor constraint puts the foot mesh below/above the ground level. */
+	UPROPERTY(EditAnywhere, BlueprintReadWrite, Category = "Foot Dimensions", meta = (ReinitializeOnEdit, UIMin = "-10.0", UIMax = "10.0"))
+	double VerticalOffset = 0.0f;
+};
+
+USTRUCT(BlueprintType)
+struct FFloorConstraintChainSettings : public FIKRetargetOpSkeletonProvider
 {
 	GENERATED_BODY()
 
@@ -40,6 +127,22 @@ struct FFloorConstraintChainSettings
 	UPROPERTY(EditAnywhere, BlueprintReadWrite, Category = "Floor Constraint Chain Settings", meta = (ClampMin = "0.0", ClampMax = "1.0", UIMin = "0.0", UIMax = "1.0"))
 	double MaintainHeightOffset = 0.0f;
 
+	/** If true, will use the foot definition for more accurate foot/ground retargeting. If false, treats the foot as a point. */
+	UPROPERTY(EditAnywhere, BlueprintReadWrite, Category = "Foot Settings", meta=(ReinitializeOnEdit))
+	bool bUseFoot = false;
+	
+	/** If true, will use the toes definition for more accurate toe/ground collision. */
+	UPROPERTY(EditAnywhere, BlueprintReadWrite, Category = "Foot Settings", meta=(ReinitializeOnEdit))
+	bool bUseToes = false;
+
+	/** Define the proportions of the target foot. */
+	UPROPERTY(EditAnywhere, BlueprintReadWrite, Category="Foot Settings", meta=(EditCondition="bUseFoot"))
+	FFloorConstraintFootDefinition Foot;
+
+	/** Define the proportions of the target foot. */
+	UPROPERTY(EditAnywhere, BlueprintReadWrite, Category="Foot Settings", meta=(EditCondition="bUseToes"))
+	FFloorConstraintToesDefinition Toes;
+
 	UE_API bool operator==(const FFloorConstraintChainSettings& Other) const;
 
 	FName GetName() const { return TargetChainName; };
@@ -69,14 +172,130 @@ struct FIKRetargetFloorConstraintOpSettings : public FIKRetargetOpSettingsBase
 	UE_API virtual const UClass* GetControllerType() const override;
 
 	UE_API virtual void CopySettingsAtRuntime(const FIKRetargetOpSettingsBase* InSettingsToCopyFrom) override;
+
+#if WITH_EDITOR
+	UE_API virtual void AssignSkeletonAssets(const USkeleton* InSourceSkeleton, const USkeleton* InTargetSkeleton) override;
+#endif
+};
+
+// filter a double over time using the "One Euro" method
+struct FScalarOneEuroFilter
+{
+	void Tune(
+		const double InFrequency,
+		const double InMinCutoff,
+		const double InBeta,
+		const double InDerivativeCutoff);
+	
+	double Run(double InValue, double InDeltaTime);
+
+private:
+	
+	double Frequency = 1.0;
+	double MinCuttoff = 1.0;
+	double Beta = 0.0;
+	double DerivativeCutoff = 1.0;
+	double PrevValue = 0.0;
+	double PrevValueDerivative = 0.0;
+	bool bIsFirstRun = true;
+};
+
+struct FToeConstraint
+{
+	// constant
+	int32 ToeBoneIndex = INDEX_NONE;
+	int32 AnkleBoneIndex = INDEX_NONE;
+	bool bToeReady = false;
+	FTransform InitialToeGlobal;
+	FVector EndRelativeToToe;
+	FVector PitchAxisRelativeToToe;
+	const FFloorConstraintToeDefinition* Settings;
+	const FFloorConstraintToesDefinition* AllToeSettings;
+
+	// updated every frame
+	FTransform CurrentToeGlobal;
+	FVector Start;
+	FVector End;
+	FVector PitchAxis;
+	FScalarOneEuroFilter AngleFilter;
+
+	void Initialize(
+		const FTargetSkeleton& InTargetSkeleton,
+		const FFloorConstraintToesDefinition& InAllToeSettings,
+		const FFloorConstraintToeDefinition& InToeSettings,
+		const int32 InAnkleBoneIndex,
+		FIKRigLogger& InLog);
+
+	void RunAfterParent(
+		const FTargetSkeleton& InTargetSkeleton,
+		TArray<FTransform>& OutGlobalPose,
+		const FFloorConstraintChainSettings* InChainSettings,
+		const double InDeltaTime);
+
+	void UpdateFramePoints(
+		const TArray<FTransform>& InGlobalPose,
+		const FIKRigGoal* InGoal,
+		bool bUpdateLocalOffset);
+
+	double CalcToeCorrectionAngle();
+};
+
+struct FFootConstraint
+{
+	// the foot frame is made of a 4-point quad representing the base of the foot from the heel to the start of the toes
+	static constexpr int32 FrontMedial   = 0;
+	static constexpr int32 RearMedial    = 1;
+	static constexpr int32 FrontLateral  = 2;
+	static constexpr int32 RearLateral   = 3;
+	static constexpr int32 Count         = 4;
+	
+	// constants
+	bool bReadyToRun = false;
+	int32 AnkleBoneIndex = INDEX_NONE;
+	FVector PointsOrig[Count] = {};
+
+	// updated every frame
+	FVector Points[Count] = {};
+	FTransform CurrentAnkleGlobal;
+
+	void Initialize(
+		const FTargetSkeleton& InTargetSkeleton,
+		const FFloorConstraintFootDefinition& InFootSettings,
+		const int32 InAnkleBoneIndex,
+		FIKRigLogger& InLog);
+
+	void Run(
+		FIKRigGoal& InOutGoal,
+		const TArray<FTransform>& OutTargetGlobalPose,
+		const FFloorConstraintChainSettings& InSettings);
+	
+	void UpdateBasePoints(const TArray<FTransform>& InGlobalPose, const FIKRigGoal* InGoal);
+	
+	static int32 GetIndexOfLowestPoint(const FVector* InPoints);
 };
 
 struct FFloorConstraint
 {
 	FName IKGoalName = NAME_None;
 	int32 SourceEndBoneIndex = INDEX_NONE;
+	int32 TargetEndBoneIndex = INDEX_NONE;
 	double HeightOffsetInRefPose = 0.0;
 	const FFloorConstraintChainSettings* Settings = nullptr;
+
+	FFootConstraint FootConstraint;
+	TArray<FToeConstraint> ToeConstraints;
+
+	void Initialize(
+		const FResolvedBoneChain& InSourceBoneChain,
+		const FResolvedBoneChain& InTargetBoneChain,
+		const FFloorConstraintChainSettings& InSettings);
+
+	void Run(
+		const FTargetSkeleton& InTargetSkeleton,
+		FIKRigGoalContainer& InGoalContainer,
+		const TArray<FTransform>& InSourceGlobalPose,
+		TArray<FTransform>& OutTargetGlobalPose,
+		const FIKRetargetFloorConstraintOpSettings& InOpSettings);
 };
 
 USTRUCT(BlueprintType, meta = (DisplayName = "Floor Constraint"))
@@ -97,14 +316,18 @@ struct FIKRetargetFloorConstraintOp : public FIKRetargetOpBase
 		const TArray<FTransform>& InSourceGlobalPose,
 		TArray<FTransform>& OutTargetGlobalPose) override;
 
+	UE_API virtual void RunAfterParent(
+		FIKRetargetProcessor& InProcessor,
+		const double InDeltaTime,
+		const TArray<FTransform>& InSourceGlobalPose,
+		TArray<FTransform>& OutTargetGlobalPose) override;
+
 	UE_API virtual void OnAddedToStack(const UIKRetargeter* InRetargetAsset, const FIKRetargetOpBase* InParentOp) override;
 
 	UE_API virtual void OnAssignIKRig(const ERetargetSourceOrTarget SourceOrTarget, const UIKRigDefinition* InIKRig, const FIKRetargetOpBase* InParentOp) override;
 
 	UE_API virtual FIKRetargetOpSettingsBase* GetSettings() override;
 	
-	UE_API virtual void SetSettings(const FIKRetargetOpSettingsBase* InSettings) override;
-	
 	UE_API virtual const UScriptStruct* GetSettingsType() const override;
 
 	UE_API virtual const UScriptStruct* GetType() const override;
@@ -121,8 +344,18 @@ struct FIKRetargetFloorConstraintOp : public FIKRetargetOpBase
 #if WITH_EDITOR
 	UE_API virtual void ResetChainSettingsToDefault(const FName& InChainName) override;
 	UE_API virtual bool AreChainSettingsAtDefault(const FName& InChainName) override;
+	
+	virtual void DebugDraw(
+		FPrimitiveDrawInterface* InPDI,
+		const FTransform& InSourceTransform,
+		const FTransform& InTargetTransform,
+		const double InComponentScale,
+		const FIKRetargetDebugDrawState& InEditorState) const override;
+	virtual bool HasDebugDrawing() override { return true; }
 #endif
 
+private:
+	
 	TArray<FFloorConstraint> FloorConstraints;
 };
 
