[IK Retarget] Bughawk issue in FootConstraintOp

diff --git a/Engine/Plugins/Animation/IKRig/Source/IKRig/Private/Retargeter/RetargetOps/FloorConstraintOp.cpp b/Engine/Plugins/Animation/IKRig/Source/IKRig/Private/Retargeter/RetargetOps/FloorConstraintOp.cpp
--- a/Engine/Plugins/Animation/IKRig/Source/IKRig/Private/Retargeter/RetargetOps/FloorConstraintOp.cpp
+++ b/Engine/Plugins/Animation/IKRig/Source/IKRig/Private/Retargeter/RetargetOps/FloorConstraintOp.cpp
@@ -215,14 +215,14 @@ void FFootConstraint::UpdateBasePoints(const TArray<FTransform>& InGlobalPose, c
 	Points[RearLateral] = CurrentAnkleGlobal.TransformPosition(PointsOrig[RearLateral]);
 }
 
-int32 FFootConstraint::GetIndexOfLowestPoint(const FVector* InPoints)
+int32 FFootConstraint::GetIndexOfLowestPoint()
 {
 	int32 MinIndex = 0;
 	float MinValue = TNumericLimits<float>::Max();
 
-	for (int32 i = 0; i <= Count; ++i)
+	for (int32 i=0; i<Points.Num(); ++i)
 	{
-		const float Z = InPoints[i].Z;
+		const float Z = Points[i].Z;
 		if (Z < MinValue)
 		{
 			MinValue = Z;
@@ -477,7 +477,7 @@ void FFootConstraint::Run(
 {
 	// calc vertical delta to push all base points above the floor
 	UpdateBasePoints(OutTargetGlobalPose, &InOutGoal);
-	const int32 LowestBasePointIndex = GetIndexOfLowestPoint(Points);
+	const int32 LowestBasePointIndex = GetIndexOfLowestPoint();
 	const double BaseHeight = Points[LowestBasePointIndex].Z;
 	const double BaseDelta = (BaseHeight < 0.0f ? FMath::Abs(BaseHeight) : 0.0) * InSettings.Alpha;
 	
diff --git a/Engine/Plugins/Animation/IKRig/Source/IKRig/Public/Retargeter/RetargetOps/FloorConstraintOp.h b/Engine/Plugins/Animation/IKRig/Source/IKRig/Public/Retargeter/RetargetOps/FloorConstraintOp.h
--- a/Engine/Plugins/Animation/IKRig/Source/IKRig/Public/Retargeter/RetargetOps/FloorConstraintOp.h
+++ b/Engine/Plugins/Animation/IKRig/Source/IKRig/Public/Retargeter/RetargetOps/FloorConstraintOp.h
@@ -252,10 +252,10 @@ struct FFootConstraint
 	// constants
 	bool bReadyToRun = false;
 	int32 AnkleBoneIndex = INDEX_NONE;
-	FVector PointsOrig[Count] = {};
+	TStaticArray<FVector, Count> PointsOrig;
 
 	// updated every frame
-	FVector Points[Count] = {};
+	TStaticArray<FVector, Count> Points;
 	FTransform CurrentAnkleGlobal;
 
 	void Initialize(
@@ -271,7 +271,7 @@ struct FFootConstraint
 	
 	void UpdateBasePoints(const TArray<FTransform>& InGlobalPose, const FIKRigGoal* InGoal);
 	
-	static int32 GetIndexOfLowestPoint(const FVector* InPoints);
+	int32 GetIndexOfLowestPoint();
 };
 
 struct FFloorConstraint
