[IK Retarget] Consolidating duplicate chain settings logic across all ops.

#Synced-CL 48403841

diff --git a/Engine/Plugins/Animation/IKRig/Source/IKRig/Private/Retargeter/IKRetargetProfile.cpp b/Engine/Plugins/Animation/IKRig/Source/IKRig/Private/Retargeter/IKRetargetProfile.cpp
--- a/Engine/Plugins/Animation/IKRig/Source/IKRig/Private/Retargeter/IKRetargetProfile.cpp
+++ b/Engine/Plugins/Animation/IKRig/Source/IKRig/Private/Retargeter/IKRetargetProfile.cpp
@@ -47,9 +47,12 @@ bool FRetargetOpProfile::CopySettingsToOp(FInstancedStruct& InOutOpStruct, ECopy
 	switch (InApplyContext)
 	{
 		case ECopyOpSettingsContext::Runtime:
-			// apply the profile settings to an op in the stack using the Op::SetSettings() virtual
+			// apply the profile settings to an op in the stack through the CopySettingsAtRuntime interface
 			// this gives each op a chance to digest the new settings in a way that won't require reinitialization
-			Op->SetSettings(&OpSettingsFromProfile);
+			if (FIKRetargetOpSettingsBase* OpSettings = Op->GetSettings())
+			{
+				OpSettings->CopySettingsAtRuntime(&OpSettingsFromProfile);
+			}
 			break;
 		
 		case ECopyOpSettingsContext::PreInitialize:
diff --git a/Engine/Plugins/Animation/IKRig/Source/IKRig/Private/Retargeter/RetargetOps/AlignPoleVectorOp.cpp b/Engine/Plugins/Animation/IKRig/Source/IKRig/Private/Retargeter/RetargetOps/AlignPoleVectorOp.cpp
--- a/Engine/Plugins/Animation/IKRig/Source/IKRig/Private/Retargeter/RetargetOps/AlignPoleVectorOp.cpp
+++ b/Engine/Plugins/Animation/IKRig/Source/IKRig/Private/Retargeter/RetargetOps/AlignPoleVectorOp.cpp
@@ -3,6 +3,7 @@
 #include "Retargeter/RetargetOps/AlignPoleVectorOp.h"
 
 #include "Engine/SkeletalMesh.h"
+#include "Retargeter/IKRetargetOpUtils.h"
 #include "Retargeter/IKRetargetProcessor.h"
 #include "Retargeter/RetargetOps/StrideWarpingOp.h"
 
@@ -189,18 +190,17 @@ const UClass* FIKRetargetAlignPoleVectorOpSettings::GetControllerType() const
 
 void FIKRetargetAlignPoleVectorOpSettings::CopySettingsAtRuntime(const FIKRetargetOpSettingsBase* InSettingsToCopyFrom)
 {
-	// copy settings only for chains that the op has initialized
+	// copy everything except the chains
+	const TArray<FName> PropertiesToIgnore = {GET_MEMBER_NAME_CHECKED(FIKRetargetAlignPoleVectorOpSettings, ChainsToAlign)};
+	FIKRetargetOpBase::CopyStructProperties(
+		FIKRetargetAlignPoleVectorOpSettings::StaticStruct(),
+		InSettingsToCopyFrom,
+		this,
+		PropertiesToIgnore);
+
+	// copy chains
 	const FIKRetargetAlignPoleVectorOpSettings* NewSettings = reinterpret_cast<const FIKRetargetAlignPoleVectorOpSettings*>(InSettingsToCopyFrom);
-	for (const FRetargetPoleVectorSettings& NewChainSettings : NewSettings->ChainsToAlign)
-	{
-		for (FRetargetPoleVectorSettings& ChainSettings : ChainsToAlign)
-		{
-			if (ChainSettings.TargetChainName == NewChainSettings.TargetChainName)
-			{
-				ChainSettings = NewChainSettings;
-			}
-		}
-	}
+	IKRetargetOpUtils::CopyChainSettingsAtRuntime(NewSettings->ChainsToAlign, ChainsToAlign);
 }
 
 bool FIKRetargetAlignPoleVectorOp::Initialize(
@@ -368,27 +368,17 @@ FText FIKRetargetAlignPoleVectorOp::GetWarningMessage() const
 
 void FIKRetargetAlignPoleVectorOp::ResetChainSettingsToDefault(const FName& InChainName)
 {
-	for (FRetargetPoleVectorSettings& ChainToAlign : Settings.ChainsToAlign)
-	{
-		if (ChainToAlign.TargetChainName == InChainName)
-		{
-			ChainToAlign = FRetargetPoleVectorSettings(InChainName);
-			return;
-		}
-	}
+	IKRetargetOpUtils::ResetChainSettingsToDefault(Settings.ChainsToAlign, InChainName);
 }
 
 bool FIKRetargetAlignPoleVectorOp::AreChainSettingsAtDefault(const FName& InChainName)
 {
-	for (FRetargetPoleVectorSettings& ChainToAlign : Settings.ChainsToAlign)
-	{
-		if (ChainToAlign.TargetChainName == InChainName)
-		{
-			return ChainToAlign == FRetargetPoleVectorSettings(InChainName);
-		}
-	}
-	
-	return true;
+	return IKRetargetOpUtils::AreChainSettingsAtDefault(Settings.ChainsToAlign, InChainName);
+}
+
+uint8* FIKRetargetAlignPoleVectorOp::GetChainSettingsMemory(const FName InChainName)
+{
+	return IKRetargetOpUtils::GetChainSettingsMemory(Settings.ChainsToAlign, InChainName);
 }
 #endif
 
diff --git a/Engine/Plugins/Animation/IKRig/Source/IKRig/Private/Retargeter/RetargetOps/FKChainsOp.cpp b/Engine/Plugins/Animation/IKRig/Source/IKRig/Private/Retargeter/RetargetOps/FKChainsOp.cpp
--- a/Engine/Plugins/Animation/IKRig/Source/IKRig/Private/Retargeter/RetargetOps/FKChainsOp.cpp
+++ b/Engine/Plugins/Animation/IKRig/Source/IKRig/Private/Retargeter/RetargetOps/FKChainsOp.cpp
@@ -3,6 +3,7 @@
 #include "Retargeter/RetargetOps/FKChainsOp.h"
 
 #include "IKRigObjectVersion.h"
+#include "Retargeter/IKRetargetOpUtils.h"
 #include "Retargeter/RetargetOps/PelvisMotionOp.h"
 
 #include UE_INLINE_GENERATED_CPP_BY_NAME(FKChainsOp)
@@ -785,26 +786,17 @@ const UClass* FIKRetargetFKChainsOpSettings::GetControllerType() const
 
 void FIKRetargetFKChainsOpSettings::CopySettingsAtRuntime(const FIKRetargetOpSettingsBase* InSettingsToCopyFrom)
 {
-	// copies everything except the ChainsToRetarget array (those are copied below, only for already existing chains)
-	const TArray<FName> PropertiesToIgnore = {"ChainsToRetarget"};
+	// copy everything except the chains
+	const TArray<FName> PropertiesToIgnore = {GET_MEMBER_NAME_CHECKED(FIKRetargetFKChainsOpSettings, ChainsToRetarget)};
 	FIKRetargetOpBase::CopyStructProperties(
 		FIKRetargetFKChainsOpSettings::StaticStruct(),
 		InSettingsToCopyFrom,
 		this,
 		PropertiesToIgnore);
-	
-	// copy settings only for chains that the op has initialized
+
+	// copy chains
 	const FIKRetargetFKChainsOpSettings* NewSettings = reinterpret_cast<const FIKRetargetFKChainsOpSettings*>(InSettingsToCopyFrom);
-	for (const FRetargetFKChainSettings& NewChainSettings : NewSettings->ChainsToRetarget)
-	{
-		for (FRetargetFKChainSettings& ChainSettings : ChainsToRetarget)
-		{
-			if (ChainSettings.TargetChainName == NewChainSettings.TargetChainName)
-			{
-				ChainSettings = NewChainSettings;
-			}
-		}
-	}
+	IKRetargetOpUtils::CopyChainSettingsAtRuntime(NewSettings->ChainsToRetarget, ChainsToRetarget);
 }
 
 bool FIKRetargetFKChainsOp::Initialize(
@@ -1164,28 +1156,17 @@ void FIKRetargetFKChainsOp::SaveDebugData(const TArray<FTransform>& OutTargetGlo
 
 void FIKRetargetFKChainsOp::ResetChainSettingsToDefault(const FName& InChainName)
 {
-	for (FRetargetFKChainSettings& ChainToRetarget : Settings.ChainsToRetarget)
-	{
-		if (ChainToRetarget.TargetChainName == InChainName)
-		{
-			ChainToRetarget = FRetargetFKChainSettings(InChainName);
-			return;
-		}
-	}
+	IKRetargetOpUtils::ResetChainSettingsToDefault(Settings.ChainsToRetarget, InChainName);
 }
 
 bool FIKRetargetFKChainsOp::AreChainSettingsAtDefault(const FName& InChainName)
 {
-	for (FRetargetFKChainSettings& ChainToRetarget : Settings.ChainsToRetarget)
-	{
-		if (ChainToRetarget.TargetChainName == InChainName)
-		{
-			FRetargetFKChainSettings DefaultSettings = FRetargetFKChainSettings();
-			return ChainToRetarget == DefaultSettings;
-		}
-	}
+	return IKRetargetOpUtils::AreChainSettingsAtDefault(Settings.ChainsToRetarget, InChainName);
+}
 
-	return true;
+uint8* FIKRetargetFKChainsOp::GetChainSettingsMemory(const FName InChainName)
+{
+	return IKRetargetOpUtils::GetChainSettingsMemory(Settings.ChainsToRetarget, InChainName);
 }
 
 #endif
diff --git a/Engine/Plugins/Animation/IKRig/Source/IKRig/Private/Retargeter/RetargetOps/FilterBoneOp.cpp b/Engine/Plugins/Animation/IKRig/Source/IKRig/Private/Retargeter/RetargetOps/FilterBoneOp.cpp
--- a/Engine/Plugins/Animation/IKRig/Source/IKRig/Private/Retargeter/RetargetOps/FilterBoneOp.cpp
+++ b/Engine/Plugins/Animation/IKRig/Source/IKRig/Private/Retargeter/RetargetOps/FilterBoneOp.cpp
@@ -75,7 +75,7 @@ const UClass* FIKRetargetFilterBoneOpSettings::GetControllerType() const
 void FIKRetargetFilterBoneOpSettings::CopySettingsAtRuntime(const FIKRetargetOpSettingsBase* InSettingsToCopyFrom)
 {
 	// copies everything except the bones we are operating on (those require reinit)
-	const TArray<FName> PropertiesToIgnore = {"BonesToFilter"};
+	const TArray<FName> PropertiesToIgnore = {GET_MEMBER_NAME_CHECKED(FIKRetargetFilterBoneOpSettings, BonesToFilter)};
 	FIKRetargetOpBase::CopyStructProperties(
 		FIKRetargetFilterBoneOpSettings::StaticStruct(),
 		InSettingsToCopyFrom,
diff --git a/Engine/Plugins/Animation/IKRig/Source/IKRig/Private/Retargeter/RetargetOps/FloorConstraintOp.cpp b/Engine/Plugins/Animation/IKRig/Source/IKRig/Private/Retargeter/RetargetOps/FloorConstraintOp.cpp
--- a/Engine/Plugins/Animation/IKRig/Source/IKRig/Private/Retargeter/RetargetOps/FloorConstraintOp.cpp
+++ b/Engine/Plugins/Animation/IKRig/Source/IKRig/Private/Retargeter/RetargetOps/FloorConstraintOp.cpp
@@ -31,8 +31,8 @@ const UClass* FIKRetargetFloorConstraintOpSettings::GetControllerType() const
 
 void FIKRetargetFloorConstraintOpSettings::CopySettingsAtRuntime(const FIKRetargetOpSettingsBase* InSettingsToCopyFrom)
 {
-	// copies everything except the ChainsToAffect array (those are copied below, only for already existing chains)
-	const TArray<FName> PropertiesToIgnore = {"ChainsToAffect"};
+	// copy everything except the chains
+	const TArray<FName> PropertiesToIgnore = {GET_MEMBER_NAME_CHECKED(FIKRetargetFloorConstraintOpSettings, ChainsToAffect)};
 	FIKRetargetOpBase::CopyStructProperties(
 		FIKRetargetFloorConstraintOpSettings::StaticStruct(),
 		InSettingsToCopyFrom,
@@ -743,28 +743,17 @@ void FIKRetargetFloorConstraintOp::OnParentReinitPropertyEdited(
 
 void FIKRetargetFloorConstraintOp::ResetChainSettingsToDefault(const FName& InChainName)
 {
-	for (FFloorConstraintChainSettings& Chain : Settings.ChainsToAffect)
-	{
-		if (Chain.TargetChainName == InChainName)
-		{
-			Chain = FFloorConstraintChainSettings(InChainName);
-			return;
-		}
-	}
+	IKRetargetOpUtils::ResetChainSettingsToDefault(Settings.ChainsToAffect, InChainName);
 }
 
 bool FIKRetargetFloorConstraintOp::AreChainSettingsAtDefault(const FName& InChainName)
 {
-	for (FFloorConstraintChainSettings& ChainToRetarget : Settings.ChainsToAffect)
-	{
-		if (ChainToRetarget.TargetChainName == InChainName)
-		{
-			FFloorConstraintChainSettings DefaultSettings = FFloorConstraintChainSettings();
-			return ChainToRetarget == DefaultSettings;
-		}
-	}
+	return IKRetargetOpUtils::AreChainSettingsAtDefault(Settings.ChainsToAffect, InChainName);
+}
 
-	return true;
+uint8* FIKRetargetFloorConstraintOp::GetChainSettingsMemory(const FName InChainName)
+{
+	return IKRetargetOpUtils::GetChainSettingsMemory(Settings.ChainsToAffect, InChainName);
 }
 
 void FIKRetargetFloorConstraintOp::DebugDraw(
diff --git a/Engine/Plugins/Animation/IKRig/Source/IKRig/Private/Retargeter/RetargetOps/IKChainsOp.cpp b/Engine/Plugins/Animation/IKRig/Source/IKRig/Private/Retargeter/RetargetOps/IKChainsOp.cpp
--- a/Engine/Plugins/Animation/IKRig/Source/IKRig/Private/Retargeter/RetargetOps/IKChainsOp.cpp
+++ b/Engine/Plugins/Animation/IKRig/Source/IKRig/Private/Retargeter/RetargetOps/IKChainsOp.cpp
@@ -6,6 +6,7 @@
 #include "Retargeter/RetargetOps/FKChainsOp.h"
 #include "Retargeter/RetargetOps/PelvisMotionOp.h"
 #include "Retargeter/RetargetOps/RunIKRigOp.h"
+#include "Retargeter/IKRetargetOpUtils.h"
 #include "IKRigObjectVersion.h"
 
 #if WITH_EDITOR
@@ -202,26 +203,17 @@ const UClass* FIKRetargetIKChainsOpSettings::GetControllerType() const
 
 void FIKRetargetIKChainsOpSettings::CopySettingsAtRuntime(const FIKRetargetOpSettingsBase* InSettingsToCopyFrom)
 {
-	// copies everything except the ChainsToRetarget array (those are copied below, only for already existing chains)
-	const TArray<FName> PropertiesToIgnore = {"ChainsToRetarget"};
+	// copy everything except the chains
+	const TArray<FName> PropertiesToIgnore = {GET_MEMBER_NAME_CHECKED(FIKRetargetIKChainsOpSettings, ChainsToRetarget)};
 	FIKRetargetOpBase::CopyStructProperties(
 		FIKRetargetIKChainsOpSettings::StaticStruct(),
 		InSettingsToCopyFrom,
 		this,
 		PropertiesToIgnore);
-	
-	// copy ChainsToRetarget only for chains are in-common
+
+	// copy chains
 	const FIKRetargetIKChainsOpSettings* NewSettings = reinterpret_cast<const FIKRetargetIKChainsOpSettings*>(InSettingsToCopyFrom);
-	for (const FRetargetIKChainSettings& NewChainSettings : NewSettings->ChainsToRetarget)
-	{
-		for (FRetargetIKChainSettings& ChainSettings : ChainsToRetarget)
-		{
-			if (ChainSettings.TargetChainName == NewChainSettings.TargetChainName)
-			{
-				ChainSettings = NewChainSettings;
-			}
-		}
-	}
+	IKRetargetOpUtils::CopyChainSettingsAtRuntime(NewSettings->ChainsToRetarget, ChainsToRetarget);
 }
 
 void FIKRetargetIKChainsOpSettings::PostLoad(const FIKRigObjectVersion::Type InVersion)
@@ -411,27 +403,6 @@ FIKRetargetOpSettingsBase* FIKRetargetIKChainsOp::GetSettings()
 	return &Settings;
 }
 
-void FIKRetargetIKChainsOp::SetSettings(const FIKRetargetOpSettingsBase* InSettings)
-{
-	const FIKRetargetIKChainsOpSettings* NewSettings = reinterpret_cast<const FIKRetargetIKChainsOpSettings*>(InSettings);
-	
-	// copies everything except the ChainsToRetarget array (those are copied below, only for already existing chains)
-	const TArray<FName> PropertiesToIgnore = {"ChainsToRetarget"};
-	CopySettingsRaw(InSettings, PropertiesToIgnore);
-	
-	// copy ChainsToRetarget only for chains are in-common
-	for (const FRetargetIKChainSettings& NewChainSettings : NewSettings->ChainsToRetarget)
-	{
-		for (FRetargetIKChainSettings& ChainSettings : Settings.ChainsToRetarget)
-		{
-			if (ChainSettings.TargetChainName == NewChainSettings.TargetChainName)
-			{
-				ChainSettings = NewChainSettings;
-			}
-		}
-	}
-}
-
 const UScriptStruct* FIKRetargetIKChainsOp::GetSettingsType() const
 {
 	return FIKRetargetIKChainsOpSettings::StaticStruct();
@@ -576,28 +547,17 @@ void FIKRetargetIKChainsOp::SaveDebugData(
 
 void FIKRetargetIKChainsOp::ResetChainSettingsToDefault(const FName& InChainName)
 {
-	for (FRetargetIKChainSettings& ChainToRetarget : Settings.ChainsToRetarget)
-	{
-		if (ChainToRetarget.TargetChainName == InChainName)
-		{
-			ChainToRetarget = FRetargetIKChainSettings(InChainName);
-			return;
-		}
-	}
+	IKRetargetOpUtils::ResetChainSettingsToDefault(Settings.ChainsToRetarget, InChainName);
 }
 
 bool FIKRetargetIKChainsOp::AreChainSettingsAtDefault(const FName& InChainName)
 {
-	for (FRetargetIKChainSettings& ChainToRetarget : Settings.ChainsToRetarget)
-	{
-		if (ChainToRetarget.TargetChainName == InChainName)
-		{
-			FRetargetIKChainSettings DefaultSettings = FRetargetIKChainSettings();
-			return ChainToRetarget == DefaultSettings;
-		}
-	}
+	return IKRetargetOpUtils::AreChainSettingsAtDefault(Settings.ChainsToRetarget, InChainName);
+}
 
-	return true;
+uint8* FIKRetargetIKChainsOp::GetChainSettingsMemory(const FName InChainName)
+{
+	return IKRetargetOpUtils::GetChainSettingsMemory(Settings.ChainsToRetarget, InChainName);
 }
 
 #endif
diff --git a/Engine/Plugins/Animation/IKRig/Source/IKRig/Private/Retargeter/RetargetOps/PelvisMotionOp.cpp b/Engine/Plugins/Animation/IKRig/Source/IKRig/Private/Retargeter/RetargetOps/PelvisMotionOp.cpp
--- a/Engine/Plugins/Animation/IKRig/Source/IKRig/Private/Retargeter/RetargetOps/PelvisMotionOp.cpp
+++ b/Engine/Plugins/Animation/IKRig/Source/IKRig/Private/Retargeter/RetargetOps/PelvisMotionOp.cpp
@@ -25,7 +25,9 @@ const UClass* FIKRetargetPelvisMotionOpSettings::GetControllerType() const
 void FIKRetargetPelvisMotionOpSettings::CopySettingsAtRuntime(const FIKRetargetOpSettingsBase* InSettingsToCopyFrom)
 {
 	// copies everything except the bones we are operating on (those require reinit)
-	const TArray<FName> PropertiesToIgnore = {"SourcePelvisBone", "TargetPelvisBone"};
+	const TArray<FName> PropertiesToIgnore = {
+		GET_MEMBER_NAME_CHECKED(FIKRetargetPelvisMotionOpSettings, SourcePelvisBone),
+		GET_MEMBER_NAME_CHECKED(FIKRetargetPelvisMotionOpSettings, TargetPelvisBone)};
 	FIKRetargetOpBase::CopyStructProperties(
 		FIKRetargetPelvisMotionOpSettings::StaticStruct(),
 		InSettingsToCopyFrom,
diff --git a/Engine/Plugins/Animation/IKRig/Source/IKRig/Private/Retargeter/RetargetOps/PinBoneOp.cpp b/Engine/Plugins/Animation/IKRig/Source/IKRig/Private/Retargeter/RetargetOps/PinBoneOp.cpp
--- a/Engine/Plugins/Animation/IKRig/Source/IKRig/Private/Retargeter/RetargetOps/PinBoneOp.cpp
+++ b/Engine/Plugins/Animation/IKRig/Source/IKRig/Private/Retargeter/RetargetOps/PinBoneOp.cpp
@@ -59,7 +59,7 @@ const UClass* FIKRetargetPinBoneOpSettings::GetControllerType() const
 void FIKRetargetPinBoneOpSettings::CopySettingsAtRuntime(const FIKRetargetOpSettingsBase* InSettingsToCopyFrom)
 {
 	// copies everything except the bones we are operating on (those require reinit)
-	const TArray<FName> PropertiesToIgnore = {"BonesToPin"};
+	const TArray<FName> PropertiesToIgnore = {GET_MEMBER_NAME_CHECKED(FIKRetargetPinBoneOpSettings, BonesToPin)};
 	FIKRetargetOpBase::CopyStructProperties(
 		FIKRetargetPinBoneOpSettings::StaticStruct(),
 		InSettingsToCopyFrom,
diff --git a/Engine/Plugins/Animation/IKRig/Source/IKRig/Private/Retargeter/RetargetOps/RetargetPoseOp.cpp b/Engine/Plugins/Animation/IKRig/Source/IKRig/Private/Retargeter/RetargetOps/RetargetPoseOp.cpp
--- a/Engine/Plugins/Animation/IKRig/Source/IKRig/Private/Retargeter/RetargetOps/RetargetPoseOp.cpp
+++ b/Engine/Plugins/Animation/IKRig/Source/IKRig/Private/Retargeter/RetargetOps/RetargetPoseOp.cpp
@@ -48,11 +48,6 @@ FIKRetargetOpSettingsBase* FIKRetargetAdditivePoseOp::GetSettings()
 	return &Settings;
 }
 
-void FIKRetargetAdditivePoseOp::SetSettings(const FIKRetargetOpSettingsBase* InSettings)
-{
-	Settings.CopySettingsAtRuntime(InSettings);
-}
-
 const UScriptStruct* FIKRetargetAdditivePoseOp::GetSettingsType() const
 {
 	return FIKRetargetAdditivePoseOpSettings::StaticStruct();
diff --git a/Engine/Plugins/Animation/IKRig/Source/IKRig/Private/Retargeter/RetargetOps/RootMotionGeneratorOp.cpp b/Engine/Plugins/Animation/IKRig/Source/IKRig/Private/Retargeter/RetargetOps/RootMotionGeneratorOp.cpp
--- a/Engine/Plugins/Animation/IKRig/Source/IKRig/Private/Retargeter/RetargetOps/RootMotionGeneratorOp.cpp
+++ b/Engine/Plugins/Animation/IKRig/Source/IKRig/Private/Retargeter/RetargetOps/RootMotionGeneratorOp.cpp
@@ -25,7 +25,10 @@ const UClass* FIKRetargetRootMotionOpSettings::GetControllerType() const
 void FIKRetargetRootMotionOpSettings::CopySettingsAtRuntime(const FIKRetargetOpSettingsBase* InSettingsToCopyFrom)
 {
 	// copies everything except the bones we are operating on (those require reinit)
-	const TArray<FName> PropertiesToIgnore = {"SourceRootBone", "TargetRootBone", "TargetPelvisBone"};
+	const TArray<FName> PropertiesToIgnore = {
+		GET_MEMBER_NAME_CHECKED(FIKRetargetRootMotionOpSettings, SourceRoot),
+		GET_MEMBER_NAME_CHECKED(FIKRetargetRootMotionOpSettings, TargetRoot),
+		GET_MEMBER_NAME_CHECKED(FIKRetargetRootMotionOpSettings, TargetPelvis)};
 	FIKRetargetOpBase::CopyStructProperties(
 		FIKRetargetRootMotionOpSettings::StaticStruct(),
 		InSettingsToCopyFrom,
diff --git a/Engine/Plugins/Animation/IKRig/Source/IKRig/Private/Retargeter/RetargetOps/RunIKRigOp.cpp b/Engine/Plugins/Animation/IKRig/Source/IKRig/Private/Retargeter/RetargetOps/RunIKRigOp.cpp
--- a/Engine/Plugins/Animation/IKRig/Source/IKRig/Private/Retargeter/RetargetOps/RunIKRigOp.cpp
+++ b/Engine/Plugins/Animation/IKRig/Source/IKRig/Private/Retargeter/RetargetOps/RunIKRigOp.cpp
@@ -18,7 +18,7 @@ const UClass* FIKRetargetRunIKRigOpSettings::GetControllerType() const
 void FIKRetargetRunIKRigOpSettings::CopySettingsAtRuntime(const FIKRetargetOpSettingsBase* InSettingsToCopyFrom)
 {
 	// copies everything except "IKRigAsset"
-	const TArray<FName> PropertiesToIgnore = {"IKRigAsset"};
+	const TArray<FName> PropertiesToIgnore = {GET_MEMBER_NAME_CHECKED(FIKRetargetRunIKRigOpSettings, IKRigAsset)};
 	FIKRetargetOpBase::CopyStructProperties(
 		FIKRetargetRunIKRigOpSettings::StaticStruct(),
 		InSettingsToCopyFrom,
diff --git a/Engine/Plugins/Animation/IKRig/Source/IKRig/Private/Retargeter/RetargetOps/SpeedPlantingOp.cpp b/Engine/Plugins/Animation/IKRig/Source/IKRig/Private/Retargeter/RetargetOps/SpeedPlantingOp.cpp
--- a/Engine/Plugins/Animation/IKRig/Source/IKRig/Private/Retargeter/RetargetOps/SpeedPlantingOp.cpp
+++ b/Engine/Plugins/Animation/IKRig/Source/IKRig/Private/Retargeter/RetargetOps/SpeedPlantingOp.cpp
@@ -6,6 +6,7 @@
 #include "Animation/AnimInstance.h"
 #include "Components/SkeletalMeshComponent.h"
 #include "Animation/AnimNodeBase.h"
+#include "Retargeter/IKRetargetOpUtils.h"
 #include "Retargeter/RetargetOps/RunIKRigOp.h"
 
 #include UE_INLINE_GENERATED_CPP_BY_NAME(SpeedPlantingOp)
@@ -20,25 +21,16 @@ const UClass* FIKRetargetSpeedPlantingOpSettings::GetControllerType() const
 void FIKRetargetSpeedPlantingOpSettings::CopySettingsAtRuntime(const FIKRetargetOpSettingsBase* InSettingsToCopyFrom)
 {
 	// copies everything except "ChainsToSpeedPlant"
-	static TArray<FName> PropertiesToIgnore = {"ChainsToSpeedPlant"};
+	const TArray<FName> PropertiesToIgnore = {GET_MEMBER_NAME_CHECKED(FIKRetargetSpeedPlantingOpSettings, ChainsToSpeedPlant)};
 	FIKRetargetOpBase::CopyStructProperties(
 		FIKRetargetSpeedPlantingOpSettings::StaticStruct(),
 		InSettingsToCopyFrom,
 		this,
 		PropertiesToIgnore);
-	
-	// copy settings only for chains that the op has initialized
+
+	// copy chains
 	const FIKRetargetSpeedPlantingOpSettings* NewSettings = reinterpret_cast<const FIKRetargetSpeedPlantingOpSettings*>(InSettingsToCopyFrom);
-	for (const FRetargetSpeedPlantingSettings& NewChainSettings : NewSettings->ChainsToSpeedPlant)
-	{
-		for (FRetargetSpeedPlantingSettings& ChainSettings : ChainsToSpeedPlant)
-		{
-			if (ChainSettings.TargetChainName == NewChainSettings.TargetChainName)
-			{
-				ChainSettings = NewChainSettings;
-			}
-		}
-	}
+	IKRetargetOpUtils::CopyChainSettingsAtRuntime(NewSettings->ChainsToSpeedPlant, ChainsToSpeedPlant);
 }
 
 bool FIKRetargetSpeedPlantingOp::Initialize(
@@ -290,6 +282,21 @@ FText FIKRetargetSpeedPlantingOp::GetWarningMessage() const
 
 	return FIKRetargetOpBase::GetWarningMessage();
 }
+
+void FIKRetargetSpeedPlantingOp::ResetChainSettingsToDefault(const FName& InChainName)
+{
+	IKRetargetOpUtils::ResetChainSettingsToDefault(Settings.ChainsToSpeedPlant, InChainName);
+}
+
+bool FIKRetargetSpeedPlantingOp::AreChainSettingsAtDefault(const FName& InChainName)
+{
+	return IKRetargetOpUtils::AreChainSettingsAtDefault(Settings.ChainsToSpeedPlant, InChainName);
+}
+
+uint8* FIKRetargetSpeedPlantingOp::GetChainSettingsMemory(const FName InChainName)
+{
+	return IKRetargetOpUtils::GetChainSettingsMemory(Settings.ChainsToSpeedPlant, InChainName);
+}
 #endif
 
 void FIKRetargetSpeedPlantingOp::RegenerateChainSettings(const FIKRetargetOpBase* InParentOp)
diff --git a/Engine/Plugins/Animation/IKRig/Source/IKRig/Private/Retargeter/RetargetOps/StretchChainOp.cpp b/Engine/Plugins/Animation/IKRig/Source/IKRig/Private/Retargeter/RetargetOps/StretchChainOp.cpp
--- a/Engine/Plugins/Animation/IKRig/Source/IKRig/Private/Retargeter/RetargetOps/StretchChainOp.cpp
+++ b/Engine/Plugins/Animation/IKRig/Source/IKRig/Private/Retargeter/RetargetOps/StretchChainOp.cpp
@@ -3,6 +3,7 @@
 #include "Retargeter/RetargetOps/StretchChainOp.h"
 
 #include "Engine/SkeletalMesh.h"
+#include "Retargeter/IKRetargetOpUtils.h"
 #include "Retargeter/IKRetargetProcessor.h"
 
 #include UE_INLINE_GENERATED_CPP_BY_NAME(StretchChainOp)
@@ -92,18 +93,9 @@ const UClass* FIKRetargetStretchChainOpSettings::GetControllerType() const
 
 void FIKRetargetStretchChainOpSettings::CopySettingsAtRuntime(const FIKRetargetOpSettingsBase* InSettingsToCopyFrom)
 {
-	// copy settings only for chains that the op has initialized
+	// copy chains
 	const FIKRetargetStretchChainOpSettings* NewSettings = reinterpret_cast<const FIKRetargetStretchChainOpSettings*>(InSettingsToCopyFrom);
-	for (const FRetargetStretchChainSettings& NewChainSettings : NewSettings->ChainsToStretch)
-	{
-		for (FRetargetStretchChainSettings& ChainSettings : ChainsToStretch)
-		{
-			if (ChainSettings.TargetChainName == NewChainSettings.TargetChainName)
-			{
-				ChainSettings = NewChainSettings;
-			}
-		}
-	}
+	IKRetargetOpUtils::CopyChainSettingsAtRuntime(NewSettings->ChainsToStretch, ChainsToStretch);
 }
 
 bool FIKRetargetStretchChainOp::Initialize(
@@ -272,27 +264,17 @@ FText FIKRetargetStretchChainOp::GetWarningMessage() const
 
 void FIKRetargetStretchChainOp::ResetChainSettingsToDefault(const FName& InChainName)
 {
-	for (FRetargetStretchChainSettings& ChainToStretch : Settings.ChainsToStretch)
-	{
-		if (ChainToStretch.TargetChainName == InChainName)
-		{
-			ChainToStretch = FRetargetStretchChainSettings(InChainName);
-			return;
-		}
-	}
+	IKRetargetOpUtils::ResetChainSettingsToDefault(Settings.ChainsToStretch, InChainName);
 }
 
 bool FIKRetargetStretchChainOp::AreChainSettingsAtDefault(const FName& InChainName)
 {
-	for (FRetargetStretchChainSettings& ChainToStretch : Settings.ChainsToStretch)
-	{
-		if (ChainToStretch.TargetChainName == InChainName)
-		{
-			return ChainToStretch == FRetargetStretchChainSettings(InChainName);
-		}
-	}
-	
-	return true;
+	return IKRetargetOpUtils::AreChainSettingsAtDefault(Settings.ChainsToStretch, InChainName);
+}
+
+uint8* FIKRetargetStretchChainOp::GetChainSettingsMemory(const FName InChainName)
+{
+	return IKRetargetOpUtils::GetChainSettingsMemory(Settings.ChainsToStretch, InChainName);
 }
 #endif
 
diff --git a/Engine/Plugins/Animation/IKRig/Source/IKRig/Private/Retargeter/RetargetOps/StrideWarpingOp.cpp b/Engine/Plugins/Animation/IKRig/Source/IKRig/Private/Retargeter/RetargetOps/StrideWarpingOp.cpp
--- a/Engine/Plugins/Animation/IKRig/Source/IKRig/Private/Retargeter/RetargetOps/StrideWarpingOp.cpp
+++ b/Engine/Plugins/Animation/IKRig/Source/IKRig/Private/Retargeter/RetargetOps/StrideWarpingOp.cpp
@@ -12,6 +12,8 @@
 #include "Engine/EngineTypes.h"
 #endif
 
+#include "Retargeter/IKRetargetOpUtils.h"
+
 #include UE_INLINE_GENERATED_CPP_BY_NAME(StrideWarpingOp)
 
 #define LOCTEXT_NAMESPACE "StrideWarpingOp"
@@ -29,13 +31,17 @@ const UClass* FIKRetargetStrideWarpingOpSettings::GetControllerType() const
 
 void FIKRetargetStrideWarpingOpSettings::CopySettingsAtRuntime(const FIKRetargetOpSettingsBase* InSettingsToCopyFrom)
 {
-	// copies everything except "ChainsToStrideWarp"
-	static TArray<FName> PropertiesToIgnore = {"ChainsToStrideWarp"};
+	// copy everything except the chains
+	const TArray<FName> PropertiesToIgnore = {GET_MEMBER_NAME_CHECKED(FIKRetargetStrideWarpingOpSettings, ChainSettings)};
 	FIKRetargetOpBase::CopyStructProperties(
 		FIKRetargetStrideWarpingOpSettings::StaticStruct(),
 		InSettingsToCopyFrom,
 		this,
 		PropertiesToIgnore);
+
+	// copy chains
+	const FIKRetargetStrideWarpingOpSettings* NewSettings = reinterpret_cast<const FIKRetargetStrideWarpingOpSettings*>(InSettingsToCopyFrom);
+	IKRetargetOpUtils::CopyChainSettingsAtRuntime(NewSettings->ChainSettings, ChainSettings);
 }
 
 FVector FIKRetargetStrideWarpingOpSettings::GetAxisVector(const EBasicAxis& Axis)
@@ -433,34 +439,17 @@ void FIKRetargetStrideWarpingOp::DebugDraw(
 
 void FIKRetargetStrideWarpingOp::ResetChainSettingsToDefault(const FName& InChainName)
 {
-	for (FRetargetStrideWarpChainSettings& ChainToStrideWarp : Settings.ChainSettings)
-	{
-		if (ChainToStrideWarp.TargetChainName != InChainName)
-		{
-			continue;
-		}
-
-		// reset
-		const FName ChainName = ChainToStrideWarp.TargetChainName;
-		ChainToStrideWarp = FRetargetStrideWarpChainSettings(ChainName);
-		return;
-	}
+	IKRetargetOpUtils::ResetChainSettingsToDefault(Settings.ChainSettings, InChainName);
 }
 
 bool FIKRetargetStrideWarpingOp::AreChainSettingsAtDefault(const FName& InChainName)
 {
-	for (FRetargetStrideWarpChainSettings& ChainToStrideWarp : Settings.ChainSettings)
-	{
-		if (ChainToStrideWarp.TargetChainName != InChainName)
-		{
-			continue;
-		}
-		
-		FRetargetStrideWarpChainSettings DefaultSettings = FRetargetStrideWarpChainSettings(ChainToStrideWarp.TargetChainName);
-		return ChainToStrideWarp == DefaultSettings;
-	}
+	return IKRetargetOpUtils::AreChainSettingsAtDefault(Settings.ChainSettings, InChainName);
+}
 
-	return true;
+uint8* FIKRetargetStrideWarpingOp::GetChainSettingsMemory(const FName InChainName)
+{
+	return IKRetargetOpUtils::GetChainSettingsMemory(Settings.ChainSettings, InChainName);
 }
 #endif
 
diff --git a/Engine/Plugins/Animation/IKRig/Source/IKRig/Public/Retargeter/IKRetargetOpUtils.h b/Engine/Plugins/Animation/IKRig/Source/IKRig/Public/Retargeter/IKRetargetOpUtils.h
--- a/Engine/Plugins/Animation/IKRig/Source/IKRig/Public/Retargeter/IKRetargetOpUtils.h
+++ b/Engine/Plugins/Animation/IKRig/Source/IKRig/Public/Retargeter/IKRetargetOpUtils.h
@@ -10,16 +10,26 @@
 namespace IKRetargetOpUtils
 {
 	/**
-	 * This is a convenience function for retarget ops that need to keep an
-	 * array of custom chain settings in-sync with the chains in and IK Rig asset.
+	 * Synchronizes a custom chain settings array with the chains in an IK Rig asset.
 	 *
-	 * NOTE: to use this, pass in an array of custom types T and ensure they
+	 * Convenience function for retarget ops that need to keep an array of custom chain
+	 * settings in sync with the chains defined in an IK Rig asset.
 	 *
-	 * 1. Implement these functions:
-	 * FName GetName()
-	 * void SetName(FName InChainName)
+	 * @tparam T  Custom struct type representing chain settings.
 	 *
-	 * 2. Provide a constructor that takes the chain name
+	 * @param InOutChainSettings   Array of chain settings to synchronize.
+	 * @param InRetargetOp         Retarget operation providing access to the IK Rig.
+	 * @param bSkipUnmappedChains  If true, unmapped chains are ignored.
+	 * @param bSkipNonIKChains     If true, chains without IK goals are ignored.
+	 *
+	 * @note  To use this function, your type T must:
+	 *        1. Implement:
+	 *           - FName GetName()
+	 *           - void SetName(FName InChainName)
+	 *        2. Provide a constructor that accepts the chain name.
+	 *
+	 * @note  This function will early-return if no valid chains are found. It does NOT clear the array,
+	 *        allowing users to reassign a different rig and preserve compatible settings.
 	 */
 	template<typename T>
 	void SynchronizeChainSettingsWithIKRig(
@@ -106,11 +116,20 @@ namespace IKRetargetOpUtils
 	};
 
 	/**
-	 * This is a convenience function for retarget ops that need to respond to an IK Rig chain being renamed.
+	 * Responds to an IK Rig chain being renamed by updating the corresponding entry in a settings array.
+	 *
+	 * Convenience function for retarget ops that maintain custom per-chain settings.
+	 *
+	 * @tparam T  Custom struct type representing chain settings.
+	 *
+	 * @param InOutChainSettings  Array of chain settings to update.
+	 * @param InOldChainName      Previous name of the chain.
+	 * @param InNewChainName      New name of the chain.
 	 *
-	 * NOTE: to use this, pass in an array of custom types T and ensure they implement these functions:
-	 * FName GetName()
-	 * void SetName(FName InChainName) */
+	 * @note  Type T must implement:
+	 *        - FName GetName()
+	 *        - void SetName(FName InChainName)
+	 */
 	template<typename T>
 	void OnRetargetChainRenamed(TArray<T>& InOutChainSettings, const FName InOldChainName, const FName InNewChainName)
 	{
@@ -122,5 +141,129 @@ namespace IKRetargetOpUtils
 			}
 		}
 	};
-	
+
+	/**
+	 * Copies matching chain settings from one array to another at runtime.
+	 *
+	 * Convenience function for retarget ops that need to copy per-chain settings
+	 * between runtime instances of the same operation.
+	 *
+	 * @tparam T  Custom struct type representing chain settings.
+	 *
+	 * @param InSettingsToCopyFrom  Source array of chain settings.
+	 * @param InSettingsToCopyTo    Destination array to update.
+	 *
+	 * @note  Only copies settings for chains that exist in both arrays.
+	 * @note  Type T must implement:
+	 *        - FName GetName()
+	 *        - void SetName(FName InChainName)
+	 */
+	template<typename T>
+	void CopyChainSettingsAtRuntime(const TArray<T>& InSettingsToCopyFrom, TArray<T>& InSettingsToCopyTo)
+	{
+		for (const T& SettingsToCopyFrom : InSettingsToCopyFrom)
+		{
+			for (T& SettingsToCopyTo : InSettingsToCopyTo)
+			{
+				if (SettingsToCopyFrom.GetName() == SettingsToCopyTo.GetName())
+				{
+					SettingsToCopyTo = SettingsToCopyFrom;
+				}
+			}
+		}
+	};
+
+	/**
+	 * Resets the settings for a specific IK chain to its default values.
+	 *
+	 * Convenience function for retarget ops that manage per-chain settings arrays.
+	 * Finds the matching chain by name and replaces its settings with a new default-constructed instance.
+	 *
+	 * @tparam T  Custom struct type representing chain settings.
+	 *
+	 * @param InChainSettingsArray  Array of chain settings to search and modify.
+	 * @param InChainName           Name of the chain to reset.
+	 *
+	 * @note  Type T must:
+	 *        - Implement FName GetName()
+	 *        - Provide a constructor that accepts a chain name (FName)
+	 *
+	 * @note  If no matching chain is found, this function does nothing.
+	 */
+	template<typename T>
+	void ResetChainSettingsToDefault(TArray<T>& InChainSettingsArray, const FName& InChainName)
+	{
+		for (T& ChainSettings : InChainSettingsArray)
+		{
+			if (ChainSettings.GetName() == InChainName)
+			{
+				// reset
+				ChainSettings = T(InChainName);
+				return;
+			}			
+		}
+	}
+
+	/**
+	 * Checks whether the settings for a specific chain are at their default values.
+	 *
+	 * Generic utility function for retarget ops that manage per-chain settings arrays.
+	 * Compares the specified chain’s current settings to a default-constructed instance.
+	 *
+	 * @tparam T  Custom struct type representing chain settings.
+	 *
+	 * @param InChainSettingsArray  Array of chain settings to search.
+	 * @param InChainName           Name of the chain to compare.
+	 *
+	 * @return True if the chain’s settings match the default values, or if the chain is not found.
+	 *
+	 * @note  Type T must:
+	 *        - Implement FName GetName()
+	 *        - Provide a constructor that accepts a chain name (FName)
+	 *        - Implement operator== for equality comparison
+	 */
+	template<typename T>
+	bool AreChainSettingsAtDefault(const TArray<T>& InChainSettingsArray, const FName& InChainName)
+	{
+		for (const T& ChainSettings : InChainSettingsArray)
+		{
+			if (ChainSettings.GetName() == InChainName)
+			{
+				const T DefaultSettings(InChainName);
+				return ChainSettings == DefaultSettings;
+			}
+		}
+
+		return true;
+	}
+
+	/**
+	 * Returns a raw pointer to the memory of a specific chain’s settings instance.
+	 *
+	 * @tparam T  Custom struct type representing chain settings.
+	 *
+	 * @param InChainSettingsArray  Array of chain settings to search.
+	 * @param InChainName           Name of the chain whose settings memory address should be returned.
+	 *
+	 * @return Pointer to the chain’s settings in memory, or nullptr if the chain is not found.
+	 *
+	 * @note  Type T must implement:
+	 *        - FName GetName()
+	 *
+	 * @warning  The returned pointer refers to memory owned by the array. It becomes invalid if
+	 *           the array is modified or reallocated. Use with caution.
+	 */
+	template<typename T>
+	uint8* GetChainSettingsMemory(TArray<T>& InChainSettingsArray, const FName& InChainName)
+	{
+		for (T& ChainSettings : InChainSettingsArray)
+		{
+			if (ChainSettings.GetName() == InChainName)
+			{
+				return reinterpret_cast<uint8*>(&ChainSettings);
+			}
+		}
+
+		return nullptr;
+	}
 }
diff --git a/Engine/Plugins/Animation/IKRig/Source/IKRig/Public/Retargeter/IKRetargetOps.h b/Engine/Plugins/Animation/IKRig/Source/IKRig/Public/Retargeter/IKRetargetOps.h
--- a/Engine/Plugins/Animation/IKRig/Source/IKRig/Public/Retargeter/IKRetargetOps.h
+++ b/Engine/Plugins/Animation/IKRig/Source/IKRig/Public/Retargeter/IKRetargetOps.h
@@ -51,7 +51,7 @@ public:
 #endif
 };
 
-/** This is the base class for defining editable settings for your custom retargeting operation.
+/** This is the base class for defining editable settings for a custom retargeting operation.
  * All user-configurable properties for an "op" should be stored within a single instance of a subclass of this type.
  * These settings will automatically be:
  * 1. Displayed in the details panel when the op is selected
@@ -252,12 +252,9 @@ public:
 	/** const access to the settings */
 	const FIKRetargetOpSettingsBase* GetSettingsConst() const { return const_cast<FIKRetargetOpBase*>(this)->GetSettings(); }
 	
-	/** (optional) override to react when settings are applied at runtime
-	 * NOTE:
-	 * This is called while the Op is running AFTER Initialize() so it up to the Op author
-	 * to copy only those settings which are safe to be updated while the op is running.
-	 * The default implementation uses CopySettingsAtRuntime() which must be implemented on the settings type.
-	 * @param InSettings - pointer to the settings to copy, can be safely cast to the type returned by GetSettingsType() */ 
+	// Old API (deprecated):
+	UE_DEPRECATED(5.7, "SetSettings(const FIKRetargetOpSettingsBase*) is deprecated. "
+					"Use GetSettings()->CopySettingsAtRuntime(const FIKRetargetOpSettingsBase*) instead")
 	UE_API virtual void SetSettings(const FIKRetargetOpSettingsBase* InSettings);
 	
 	/** (required) override and return the type used to house the settings for this operation */
@@ -368,6 +365,10 @@ public:
 	/** (optional) override to support "Reset to Default" UI for chains */
 	virtual bool AreChainSettingsAtDefault(const FName& InChainName) { return true; };
 
+	/** (required for any op that stores settings per-chain) return the memory for the settings struct associated with a particular chain
+	 * NOTE: implementing this can automate some boilerplate UI setup, allows systems to display chain settings and get/set them. */
+	virtual uint8* GetChainSettingsMemory(const FName InChainName) { checkNoEntry(); return NULL; };
+
 	/** the instance of this op currently running in the editor viewport */
 	mutable FIKRetargetOpBase* EditorInstance;
 
diff --git a/Engine/Plugins/Animation/IKRig/Source/IKRig/Public/Retargeter/RetargetOps/AlignPoleVectorOp.h b/Engine/Plugins/Animation/IKRig/Source/IKRig/Public/Retargeter/RetargetOps/AlignPoleVectorOp.h
--- a/Engine/Plugins/Animation/IKRig/Source/IKRig/Public/Retargeter/RetargetOps/AlignPoleVectorOp.h
+++ b/Engine/Plugins/Animation/IKRig/Source/IKRig/Public/Retargeter/RetargetOps/AlignPoleVectorOp.h
@@ -87,6 +87,8 @@ struct FRetargetPoleVectorSettings
 	bool MaintainOffset = false;
 	
 	UE_API bool operator==(const FRetargetPoleVectorSettings& Other) const;
+	UE_API FName GetName() const { return TargetChainName; };
+	UE_API void SetName(const FName InName) { TargetChainName = InName; };
 };
 
 USTRUCT(BlueprintType, meta = (DisplayName = "Pole Vector Alignment Settings"))
@@ -146,9 +148,9 @@ struct FIKRetargetAlignPoleVectorOp : public FIKRetargetOpBase
 
 #if WITH_EDITOR
 	UE_API virtual FText GetWarningMessage() const override;
-
 	UE_API virtual void ResetChainSettingsToDefault(const FName& InChainName) override;
 	UE_API virtual bool AreChainSettingsAtDefault(const FName& InChainName) override;
+	UE_API virtual uint8* GetChainSettingsMemory(const FName InChainName) override;
 #endif
 
 	UPROPERTY()
diff --git a/Engine/Plugins/Animation/IKRig/Source/IKRig/Public/Retargeter/RetargetOps/FKChainsOp.h b/Engine/Plugins/Animation/IKRig/Source/IKRig/Public/Retargeter/RetargetOps/FKChainsOp.h
--- a/Engine/Plugins/Animation/IKRig/Source/IKRig/Public/Retargeter/RetargetOps/FKChainsOp.h
+++ b/Engine/Plugins/Animation/IKRig/Source/IKRig/Public/Retargeter/RetargetOps/FKChainsOp.h
@@ -225,6 +225,8 @@ struct FRetargetFKChainSettings
 	double TranslationAlpha = 1.0f;
 
 	IKRIG_API bool operator==(const FRetargetFKChainSettings& Other) const;
+	IKRIG_API FName GetName() const { return TargetChainName; };
+	IKRIG_API void SetName(const FName InName) { TargetChainName = InName; };
 };
 
 USTRUCT(BlueprintType, meta = (DisplayName = "Retarget Chains Settings"))
@@ -344,6 +346,7 @@ public:
 
 	IKRIG_API virtual void ResetChainSettingsToDefault(const FName& InChainName) override;
 	IKRIG_API virtual bool AreChainSettingsAtDefault(const FName& InChainName) override;
+	IKRIG_API virtual uint8* GetChainSettingsMemory(const FName InChainName) override;
 #endif
 
 private:
diff --git a/Engine/Plugins/Animation/IKRig/Source/IKRig/Public/Retargeter/RetargetOps/FloorConstraintOp.h b/Engine/Plugins/Animation/IKRig/Source/IKRig/Public/Retargeter/RetargetOps/FloorConstraintOp.h
--- a/Engine/Plugins/Animation/IKRig/Source/IKRig/Public/Retargeter/RetargetOps/FloorConstraintOp.h
+++ b/Engine/Plugins/Animation/IKRig/Source/IKRig/Public/Retargeter/RetargetOps/FloorConstraintOp.h
@@ -345,6 +345,7 @@ struct FIKRetargetFloorConstraintOp : public FIKRetargetOpBase
 #if WITH_EDITOR
 	UE_API virtual void ResetChainSettingsToDefault(const FName& InChainName) override;
 	UE_API virtual bool AreChainSettingsAtDefault(const FName& InChainName) override;
+	UE_API virtual uint8* GetChainSettingsMemory(const FName InChainName) override;
 	
 	virtual void DebugDraw(
 		FPrimitiveDrawInterface* InPDI,
diff --git a/Engine/Plugins/Animation/IKRig/Source/IKRig/Public/Retargeter/RetargetOps/IKChainsOp.h b/Engine/Plugins/Animation/IKRig/Source/IKRig/Public/Retargeter/RetargetOps/IKChainsOp.h
--- a/Engine/Plugins/Animation/IKRig/Source/IKRig/Public/Retargeter/RetargetOps/IKChainsOp.h
+++ b/Engine/Plugins/Animation/IKRig/Source/IKRig/Public/Retargeter/RetargetOps/IKChainsOp.h
@@ -134,6 +134,8 @@ struct FRetargetIKChainSettings
 	double Extension = 1.0f;
 
 	UE_API bool operator==(const FRetargetIKChainSettings& Other) const;
+	UE_API FName GetName() const { return TargetChainName; };
+	UE_API void SetName(const FName InName) { TargetChainName = InName; };
 
 	/** Custom serialization to handle loading old assets */
 	bool Serialize(FArchive& Ar);
@@ -218,8 +220,6 @@ struct FIKRetargetIKChainsOp : public FIKRetargetOpBase
 
 	UE_API virtual FIKRetargetOpSettingsBase* GetSettings() override;
 	
-	UE_API virtual void SetSettings(const FIKRetargetOpSettingsBase* InSettings) override;
-	
 	UE_API virtual const UScriptStruct* GetSettingsType() const override;
 
 	UE_API virtual const UScriptStruct* GetType() const override;
@@ -255,6 +255,7 @@ public:
 
 	UE_API virtual void ResetChainSettingsToDefault(const FName& InChainName) override;
 	UE_API virtual bool AreChainSettingsAtDefault(const FName& InChainName) override;
+	UE_API virtual uint8* GetChainSettingsMemory(const FName InChainName) override;
 #endif
 
 private:
diff --git a/Engine/Plugins/Animation/IKRig/Source/IKRig/Public/Retargeter/RetargetOps/RetargetPoseOp.h b/Engine/Plugins/Animation/IKRig/Source/IKRig/Public/Retargeter/RetargetOps/RetargetPoseOp.h
--- a/Engine/Plugins/Animation/IKRig/Source/IKRig/Public/Retargeter/RetargetOps/RetargetPoseOp.h
+++ b/Engine/Plugins/Animation/IKRig/Source/IKRig/Public/Retargeter/RetargetOps/RetargetPoseOp.h
@@ -47,8 +47,6 @@ struct FIKRetargetAdditivePoseOp : public FIKRetargetOpBase
 		TArray<FTransform>& OutTargetGlobalPose) override;
 
 	UE_API virtual FIKRetargetOpSettingsBase* GetSettings() override;
-
-	UE_API virtual void SetSettings(const FIKRetargetOpSettingsBase* InSettings) override;
 	
 	UE_API virtual const UScriptStruct* GetSettingsType() const override;
 	
diff --git a/Engine/Plugins/Animation/IKRig/Source/IKRig/Public/Retargeter/RetargetOps/SpeedPlantingOp.h b/Engine/Plugins/Animation/IKRig/Source/IKRig/Public/Retargeter/RetargetOps/SpeedPlantingOp.h
--- a/Engine/Plugins/Animation/IKRig/Source/IKRig/Public/Retargeter/RetargetOps/SpeedPlantingOp.h
+++ b/Engine/Plugins/Animation/IKRig/Source/IKRig/Public/Retargeter/RetargetOps/SpeedPlantingOp.h
@@ -26,6 +26,10 @@ struct FRetargetSpeedPlantingSettings
 	/** The name of the curve on the source animation that contains the speed of the end effector bone.*/
 	UPROPERTY(EditAnywhere, BlueprintReadWrite, BlueprintReadWrite, Category = "Plant IK by Speed", meta = (ClampMin = "0.0", ClampMax = "100.0", UIMin = "0.0", UIMax = "100.0"))
 	FName SpeedCurveName;
+
+	UE_API bool operator==(const FRetargetSpeedPlantingSettings& Other) const { return SpeedCurveName == Other.SpeedCurveName; };
+	UE_API FName GetName() const { return TargetChainName; };
+	UE_API void SetName(const FName InName) { TargetChainName = InName; };
 };
 
 USTRUCT(BlueprintType, meta = (DisplayName = "Speed Plant Settings"))
@@ -125,6 +129,9 @@ struct FIKRetargetSpeedPlantingOp : public FIKRetargetOpBase
 #if WITH_EDITOR
 public:
 	UE_API virtual FText GetWarningMessage() const override;
+	UE_API virtual void ResetChainSettingsToDefault(const FName& InChainName) override;
+	UE_API virtual bool AreChainSettingsAtDefault(const FName& InChainName) override;
+	UE_API virtual uint8* GetChainSettingsMemory(const FName InChainName) override;
 #endif
 
 	bool ResetThisTick;
diff --git a/Engine/Plugins/Animation/IKRig/Source/IKRig/Public/Retargeter/RetargetOps/StretchChainOp.h b/Engine/Plugins/Animation/IKRig/Source/IKRig/Public/Retargeter/RetargetOps/StretchChainOp.h
--- a/Engine/Plugins/Animation/IKRig/Source/IKRig/Public/Retargeter/RetargetOps/StretchChainOp.h
+++ b/Engine/Plugins/Animation/IKRig/Source/IKRig/Public/Retargeter/RetargetOps/StretchChainOp.h
@@ -62,6 +62,8 @@ struct FRetargetStretchChainSettings
 	double ScaleChainLength = 1.0f;
 	
 	bool operator==(const FRetargetStretchChainSettings& Other) const;
+	UE_API FName GetName() const { return TargetChainName; };
+	UE_API void SetName(const FName InName) { TargetChainName = InName; };
 };
 
 USTRUCT(BlueprintType, meta = (DisplayName = "Stretch Chain Settings"))
@@ -119,9 +121,9 @@ struct FIKRetargetStretchChainOp : public FIKRetargetOpBase
 
 #if WITH_EDITOR
 	UE_API virtual FText GetWarningMessage() const override;
-
 	UE_API virtual void ResetChainSettingsToDefault(const FName& InChainName) override;
 	UE_API virtual bool AreChainSettingsAtDefault(const FName& InChainName) override;
+	UE_API virtual uint8* GetChainSettingsMemory(const FName InChainName) override;
 #endif
 
 	UPROPERTY()
diff --git a/Engine/Plugins/Animation/IKRig/Source/IKRig/Public/Retargeter/RetargetOps/StrideWarpingOp.h b/Engine/Plugins/Animation/IKRig/Source/IKRig/Public/Retargeter/RetargetOps/StrideWarpingOp.h
--- a/Engine/Plugins/Animation/IKRig/Source/IKRig/Public/Retargeter/RetargetOps/StrideWarpingOp.h
+++ b/Engine/Plugins/Animation/IKRig/Source/IKRig/Public/Retargeter/RetargetOps/StrideWarpingOp.h
@@ -36,6 +36,8 @@ struct FRetargetStrideWarpChainSettings
 	bool EnableStrideWarping = true;
 
 	UE_API bool operator==(const FRetargetStrideWarpChainSettings& Other) const;
+	UE_API FName GetName() const { return TargetChainName; };
+	UE_API void SetName(const FName InName) { TargetChainName = InName; };
 };
 
 USTRUCT(BlueprintType, meta = (DisplayName = "Stride Warp Settings"))
@@ -175,6 +177,7 @@ struct FIKRetargetStrideWarpingOp : public FIKRetargetOpBase
 
 	UE_API virtual void ResetChainSettingsToDefault(const FName& InChainName) override;
 	UE_API virtual bool AreChainSettingsAtDefault(const FName& InChainName) override;
+	UE_API virtual uint8* GetChainSettingsMemory(const FName InChainName) override;
 #endif
 
 	UE_API void RegenerateChainSettings(const FIKRetargetOpBase* InParentOp);
diff --git a/Engine/Plugins/Animation/IKRig/Source/IKRigEditor/Private/RetargetEditor/IKRetargetOpDetails.cpp b/Engine/Plugins/Animation/IKRig/Source/IKRigEditor/Private/RetargetEditor/IKRetargetOpDetails.cpp
--- a/Engine/Plugins/Animation/IKRig/Source/IKRigEditor/Private/RetargetEditor/IKRetargetOpDetails.cpp
+++ b/Engine/Plugins/Animation/IKRig/Source/IKRigEditor/Private/RetargetEditor/IKRetargetOpDetails.cpp
@@ -227,38 +227,10 @@ void FChainsFKOpCustomization::CustomizeChildren(
 	ChainListConfig.Controller = AssetController;
 	ChainListConfig.bEnableGoalColumn = false;
 	ChainListConfig.bEnableChainMapping = true;
-	ChainListConfig.ChainSettingsGetterFunc = [Controller=AssetController, ThisOpName=OpName](const FName InTargetChainName) -> UObject*
-	{
-		// this lambda fetches the FK chain settings struct for the given target chain in the given op (by name)
-		auto MemoryProvider = [Controller, ThisOpName, InTargetChainName]() -> uint8*
-		{
-			if (FIKRetargetOpBase* Op = Controller->GetRetargetOpByName(ThisOpName))
-			{
-				FIKRetargetFKChainsOpSettings* OpSettings = reinterpret_cast<FIKRetargetFKChainsOpSettings*>(Op->GetSettings());
-				for (FRetargetFKChainSettings& ChainToRetarget : OpSettings->ChainsToRetarget)
-				{
-					if (ChainToRetarget.TargetChainName == InTargetChainName)
-					{
-						return reinterpret_cast<uint8*>(&ChainToRetarget);
-					}
-				}
-					
-			}
-			return nullptr;
-		};
-			
-		FIKRigStructToView StructToView;
-		StructToView.Owner = Controller->GetAsset();
-		StructToView.Type = FRetargetFKChainSettings::StaticStruct();
-		StructToView.MemoryProvider = MemoryProvider;
-		StructToView.UniqueName = ThisOpName;
-		
-		URetargetFKChainSettingsWrapper* Wrapper = NewObject<URetargetFKChainSettingsWrapper>(GetTransientPackage(), NAME_None, RF_Standalone);
-		const FName SettingsPropertyName = GET_MEMBER_NAME_CHECKED(URetargetFKChainSettingsWrapper, Settings);
-		Wrapper->InitializeWithRetargeter(StructToView, SettingsPropertyName, Controller);
-		Wrapper->SetPropertyHidden(GET_MEMBER_NAME_CHECKED(FRetargetFKChainSettings, TargetChainName), true);
-		return Wrapper;
-	};
+	const FName SettingsPropertyName = GET_MEMBER_NAME_CHECKED(URetargetFKChainSettingsWrapper, Settings);
+	const TArray<FName> HiddenSettings = {GET_MEMBER_NAME_CHECKED(FRetargetFKChainSettings, TargetChainName)};
+	ChainListConfig.ChainSettingsGetterFunc =
+		TMakeChainSettingsGetter<FRetargetFKChainSettings,URetargetFKChainSettingsWrapper>(OpName, SettingsPropertyName, HiddenSettings);
 
 	// add the chain mapping list
 	const FText ChainMapRowName = FText::FromString("Chain Map");
@@ -340,37 +312,10 @@ void FIKChainOpCustomization::CustomizeChildren(
 	ChainListConfig.bEnableGoalColumn = true;
 	ChainListConfig.bEnableChainMapping = false;
 	ChainListConfig.Filter = ShowOnlyIKFilter;
-	ChainListConfig.ChainSettingsGetterFunc = [Controller=AssetController, ThisOpName=OpName](const FName InTargetChainName) -> UObject*
-	{
-		// this lambda fetches the chain settings struct for the given target chain in the given op (by name)
-		auto MemoryProvider = [Controller, ThisOpName, InTargetChainName]() -> uint8*
-		{
-			if (FIKRetargetOpBase* Op = Controller->GetRetargetOpByName(ThisOpName))
-			{
-				FIKRetargetIKChainsOpSettings* OpSettings = reinterpret_cast<FIKRetargetIKChainsOpSettings*>(Op->GetSettings());
-				for (FRetargetIKChainSettings& ChainToRetarget : OpSettings->ChainsToRetarget)
-				{
-					if (ChainToRetarget.TargetChainName == InTargetChainName)
-					{
-						return reinterpret_cast<uint8*>(&ChainToRetarget);
-					}
-				}
-			}
-			return nullptr;
-		};
-			
-		FIKRigStructToView StructToView;
-		StructToView.Owner = Controller->GetAsset();
-		StructToView.Type = FRetargetIKChainSettings::StaticStruct();
-		StructToView.MemoryProvider = MemoryProvider;
-		StructToView.UniqueName = ThisOpName;
-		
-		URetargetIKChainSettingsWrapper* Wrapper = NewObject<URetargetIKChainSettingsWrapper>(GetTransientPackage(), NAME_None, RF_Standalone);
-		const FName SettingsPropertyName = GET_MEMBER_NAME_CHECKED(URetargetIKChainSettingsWrapper, Settings);
-		Wrapper->InitializeWithRetargeter(StructToView, SettingsPropertyName, Controller);
-		Wrapper->SetPropertyHidden(GET_MEMBER_NAME_CHECKED(FRetargetIKChainSettings, TargetChainName), true);
-		return Wrapper;
-	};
+	const FName SettingsPropertyName = GET_MEMBER_NAME_CHECKED(URetargetIKChainSettingsWrapper, Settings);
+	const TArray<FName> HiddenSettings = {GET_MEMBER_NAME_CHECKED(FRetargetIKChainSettings, TargetChainName)};
+	ChainListConfig.ChainSettingsGetterFunc =
+		TMakeChainSettingsGetter<FRetargetIKChainSettings,URetargetIKChainSettingsWrapper>(OpName, SettingsPropertyName, HiddenSettings);
 
 	// add the chain mapping list
 	FText ChainMapRowName = FText::FromString("Chain Map");
@@ -409,37 +354,10 @@ void FStrideWarpOpCustomization::CustomizeChildren(
 	ChainListConfig.bEnableGoalColumn = true;
 	ChainListConfig.bEnableChainMapping = false;
 	ChainListConfig.Filter = ShowOnlyIKFilter;
-	ChainListConfig.ChainSettingsGetterFunc = [Controller=AssetController, ThisOpName=OpName](const FName InTargetChainName) -> UObject*
-	{
-		// this lambda fetches the chain settings struct for the given target chain in the given op (by name)
-		auto MemoryProvider = [Controller, ThisOpName, InTargetChainName]() -> uint8*
-		{
-			if (FIKRetargetOpBase* Op = Controller->GetRetargetOpByName(ThisOpName))
-			{
-				FIKRetargetStrideWarpingOpSettings* OpSettings = reinterpret_cast<FIKRetargetStrideWarpingOpSettings*>(Op->GetSettings());
-				for (FRetargetStrideWarpChainSettings& ChainToWarp : OpSettings->ChainSettings)
-				{
-					if (ChainToWarp.TargetChainName == InTargetChainName)
-					{
-						return reinterpret_cast<uint8*>(&ChainToWarp);
-					}
-				}
-			}
-			return nullptr;
-		};
-			
-		FIKRigStructToView StructToView;
-		StructToView.Owner = Controller->GetAsset();
-		StructToView.Type = FRetargetStrideWarpChainSettings::StaticStruct();
-		StructToView.MemoryProvider = MemoryProvider;
-		StructToView.UniqueName = ThisOpName;
-		
-		URetargetStrideWarpSettingsWrapper* Wrapper = NewObject<URetargetStrideWarpSettingsWrapper>(GetTransientPackage(), NAME_None, RF_Standalone);
-		const FName SettingsPropertyName = GET_MEMBER_NAME_CHECKED(URetargetStrideWarpSettingsWrapper, Settings);
-		Wrapper->InitializeWithRetargeter(StructToView, SettingsPropertyName, Controller);
-		Wrapper->SetPropertyHidden(GET_MEMBER_NAME_CHECKED(FRetargetStrideWarpChainSettings, TargetChainName), true);
-		return Wrapper;
-	};
+	const FName SettingsPropertyName = GET_MEMBER_NAME_CHECKED(URetargetStrideWarpSettingsWrapper, Settings);
+	const TArray<FName> HiddenSettings = {GET_MEMBER_NAME_CHECKED(FRetargetStrideWarpChainSettings, TargetChainName)};
+	ChainListConfig.ChainSettingsGetterFunc =
+		TMakeChainSettingsGetter<FRetargetStrideWarpChainSettings,URetargetStrideWarpSettingsWrapper>(OpName, SettingsPropertyName, HiddenSettings);
 
 	// add the chain mapping list
 	FText ChainMapRowName = FText::FromString("Chain Map");
@@ -479,37 +397,10 @@ void FSpeedPlantOpCustomization::CustomizeChildren(
 	ChainListConfig.bEnableGoalColumn = true;
 	ChainListConfig.bEnableChainMapping = false;
 	ChainListConfig.Filter = ShowOnlyIKFilter;
-	ChainListConfig.ChainSettingsGetterFunc = [Controller=AssetController, ThisOpName=OpName](const FName InTargetChainName) -> UObject*
-	{
-		// this lambda fetches the chain settings struct for the given target chain in the given op (by name)
-		auto MemoryProvider = [Controller, ThisOpName, InTargetChainName]() -> uint8*
-		{
-			if (FIKRetargetOpBase* Op = Controller->GetRetargetOpByName(ThisOpName))
-			{
-				FIKRetargetSpeedPlantingOpSettings* OpSettings = reinterpret_cast<FIKRetargetSpeedPlantingOpSettings*>(Op->GetSettings());
-				for (FRetargetSpeedPlantingSettings& ChainToWarp : OpSettings->ChainsToSpeedPlant)
-				{
-					if (ChainToWarp.TargetChainName == InTargetChainName)
-					{
-						return reinterpret_cast<uint8*>(&ChainToWarp);
-					}
-				}
-			}
-			return nullptr;
-		};
-			
-		FIKRigStructToView StructToView;
-		StructToView.Owner = Controller->GetAsset();
-		StructToView.Type = FRetargetSpeedPlantingSettings::StaticStruct();
-		StructToView.MemoryProvider = MemoryProvider;
-		StructToView.UniqueName = ThisOpName;
-		
-		URetargetSpeedPlantSettingsWrapper* Wrapper = NewObject<URetargetSpeedPlantSettingsWrapper>(GetTransientPackage(), NAME_None, RF_Standalone);
-		const FName SettingsPropertyName = GET_MEMBER_NAME_CHECKED(URetargetSpeedPlantSettingsWrapper, Settings);
-		Wrapper->InitializeWithRetargeter(StructToView, SettingsPropertyName, Controller);
-		Wrapper->SetPropertyHidden(GET_MEMBER_NAME_CHECKED(FRetargetSpeedPlantingSettings, TargetChainName), true);
-		return Wrapper;
-	};
+	const FName SettingsPropertyName = GET_MEMBER_NAME_CHECKED(URetargetSpeedPlantSettingsWrapper, Settings);
+	const TArray<FName> HiddenSettings = {GET_MEMBER_NAME_CHECKED(FRetargetSpeedPlantingSettings, TargetChainName)};
+	ChainListConfig.ChainSettingsGetterFunc =
+		TMakeChainSettingsGetter<FRetargetSpeedPlantingSettings,URetargetSpeedPlantSettingsWrapper>(OpName, SettingsPropertyName, HiddenSettings);
 
 	// add the chain mapping list
 	FText ChainMapRowName = FText::FromString("Chain Map");
@@ -546,37 +437,10 @@ void FPoleVectorOpCustomization::CustomizeChildren(
 	ChainListConfig.Controller = AssetController;
 	ChainListConfig.bEnableGoalColumn = true;
 	ChainListConfig.bEnableChainMapping = true;
-	ChainListConfig.ChainSettingsGetterFunc = [Controller=AssetController, ThisOpName=OpName](const FName InTargetChainName) -> UObject*
-	{
-		// this lambda fetches the chain settings struct for the given target chain in the given op (by name)
-		auto MemoryProvider = [Controller, ThisOpName, InTargetChainName]() -> uint8*
-		{
-			if (FIKRetargetOpBase* Op = Controller->GetRetargetOpByName(ThisOpName))
-			{
-				FIKRetargetAlignPoleVectorOpSettings* OpSettings = reinterpret_cast<FIKRetargetAlignPoleVectorOpSettings*>(Op->GetSettings());
-				for (FRetargetPoleVectorSettings& ChainToAlign : OpSettings->ChainsToAlign)
-				{
-					if (ChainToAlign.TargetChainName == InTargetChainName)
-					{
-						return reinterpret_cast<uint8*>(&ChainToAlign);
-					}
-				}
-			}
-			return nullptr;
-		};
-			
-		FIKRigStructToView StructToView;
-		StructToView.Owner = Controller->GetAsset();
-		StructToView.Type = FRetargetPoleVectorSettings::StaticStruct();
-		StructToView.MemoryProvider = MemoryProvider;
-		StructToView.UniqueName = ThisOpName;
-		
-		UPoleVectorSettingsWrapper* Wrapper = NewObject<UPoleVectorSettingsWrapper>(GetTransientPackage(), NAME_None, RF_Standalone);
-		const FName SettingsPropertyName = GET_MEMBER_NAME_CHECKED(UPoleVectorSettingsWrapper, Settings);
-		Wrapper->InitializeWithRetargeter(StructToView, SettingsPropertyName, Controller);
-		Wrapper->SetPropertyHidden(GET_MEMBER_NAME_CHECKED(FRetargetPoleVectorSettings, TargetChainName), true);
-		return Wrapper;
-	};
+	const FName SettingsPropertyName = GET_MEMBER_NAME_CHECKED(UPoleVectorSettingsWrapper, Settings);
+	const TArray<FName> HiddenSettings = {GET_MEMBER_NAME_CHECKED(FRetargetPoleVectorSettings, TargetChainName)};
+	ChainListConfig.ChainSettingsGetterFunc =
+		TMakeChainSettingsGetter<FRetargetPoleVectorSettings,UPoleVectorSettingsWrapper>(OpName, SettingsPropertyName, HiddenSettings);
 
 	// add the chain mapping list
 	FText ChainMapRowName = FText::FromString("Chain Map");
@@ -699,37 +563,10 @@ void FStretchChainOpCustomization::CustomizeChildren(
 	ChainListConfig.Controller = AssetController;
 	ChainListConfig.bEnableGoalColumn = true;
 	ChainListConfig.bEnableChainMapping = true;
-	ChainListConfig.ChainSettingsGetterFunc = [Controller=AssetController, ThisOpName=OpName](const FName InTargetChainName) -> UObject*
-	{
-		// this lambda fetches the chain settings struct for the given target chain in the given op (by name)
-		auto MemoryProvider = [Controller, ThisOpName, InTargetChainName]() -> uint8*
-		{
-			if (FIKRetargetOpBase* Op = Controller->GetRetargetOpByName(ThisOpName))
-			{
-				FIKRetargetStretchChainOpSettings* OpSettings = reinterpret_cast<FIKRetargetStretchChainOpSettings*>(Op->GetSettings());
-				for (FRetargetStretchChainSettings& ChainToStretch : OpSettings->ChainsToStretch)
-				{
-					if (ChainToStretch.TargetChainName == InTargetChainName)
-					{
-						return reinterpret_cast<uint8*>(&ChainToStretch);
-					}
-				}
-			}
-			return nullptr;
-		};
-			
-		FIKRigStructToView StructToView;
-		StructToView.Owner = Controller->GetAsset();
-		StructToView.Type = FRetargetStretchChainSettings::StaticStruct();
-		StructToView.MemoryProvider = MemoryProvider;
-		StructToView.UniqueName = ThisOpName;
-		
-		UStretchChainSettingsWrapper* Wrapper = NewObject<UStretchChainSettingsWrapper>(GetTransientPackage(), NAME_None, RF_Standalone);
-		const FName SettingsPropertyName = GET_MEMBER_NAME_CHECKED(UStretchChainSettingsWrapper, Settings);
-		Wrapper->InitializeWithRetargeter(StructToView, SettingsPropertyName, Controller);
-		Wrapper->SetPropertyHidden(GET_MEMBER_NAME_CHECKED(FRetargetStretchChainSettings, TargetChainName), true);
-		return Wrapper;
-	};
+	const FName SettingsPropertyName = GET_MEMBER_NAME_CHECKED(UStretchChainSettingsWrapper, Settings);
+	const TArray<FName> HiddenSettings = {GET_MEMBER_NAME_CHECKED(FRetargetStretchChainSettings, TargetChainName)};
+	ChainListConfig.ChainSettingsGetterFunc =
+		TMakeChainSettingsGetter<FRetargetStretchChainSettings,UStretchChainSettingsWrapper>(OpName, SettingsPropertyName, HiddenSettings);
 
 	// add the chain mapping list
 	FText ChainMapRowName = FText::FromString("Chain Map");
@@ -765,38 +602,11 @@ void FFloorConstraintOpCustomization::CustomizeChildren(
 	ChainListConfig.bEnableGoalColumn = true;
 	ChainListConfig.bEnableChainMapping = false;
 	ChainListConfig.Filter = ShowOnlyIKFilter;
-	ChainListConfig.ChainSettingsGetterFunc = [Controller=AssetController, ThisOpName=OpName](const FName InTargetChainName) -> UObject*
-	{
-		// this lambda fetches the chain settings struct for the given target chain in the given op (by name)
-		auto MemoryProvider = [Controller, ThisOpName, InTargetChainName]() -> uint8*
-		{
-			if (FIKRetargetOpBase* Op = Controller->GetRetargetOpByName(ThisOpName))
-			{
-				FIKRetargetFloorConstraintOpSettings* OpSettings = reinterpret_cast<FIKRetargetFloorConstraintOpSettings*>(Op->GetSettings());
-				for (FFloorConstraintChainSettings& ChainSettings : OpSettings->ChainsToAffect)
-				{
-					if (ChainSettings.TargetChainName == InTargetChainName)
-					{
-						return reinterpret_cast<uint8*>(&ChainSettings);
-					}
-				}
-			}
-			return nullptr;
-		};
-			
-		FIKRigStructToView StructToView;
-		StructToView.Owner = Controller->GetAsset();
-		StructToView.Type = FFloorConstraintChainSettings::StaticStruct();
-		StructToView.MemoryProvider = MemoryProvider;
-		StructToView.UniqueName = ThisOpName;
-		
-		UFloorConstraintSettingsWrapper* Wrapper = NewObject<UFloorConstraintSettingsWrapper>(GetTransientPackage(), NAME_None, RF_Standalone);
-		const FName SettingsPropertyName = GET_MEMBER_NAME_CHECKED(UFloorConstraintSettingsWrapper, Settings);
-		Wrapper->InitializeWithRetargeter(StructToView, SettingsPropertyName, Controller);
-		Wrapper->SetPropertyHidden(GET_MEMBER_NAME_CHECKED(FFloorConstraintChainSettings, TargetChainName), true);
-		return Wrapper;
-	};
-
+	const FName SettingsPropertyName = GET_MEMBER_NAME_CHECKED(UFloorConstraintSettingsWrapper, Settings);
+	const TArray<FName> HiddenSettings = {GET_MEMBER_NAME_CHECKED(FFloorConstraintChainSettings, TargetChainName)};
+	ChainListConfig.ChainSettingsGetterFunc =
+		TMakeChainSettingsGetter<FFloorConstraintChainSettings,UFloorConstraintSettingsWrapper>(OpName, SettingsPropertyName, HiddenSettings);
+	
 	// add the chain mapping list
 	FText ChainMapRowName = FText::FromString("Chain Map");
 	ChildBuilder.AddCustomRow(ChainMapRowName)
@@ -812,4 +622,5 @@ void FFloorConstraintOpCustomization::CustomizeChildren(
 	AddBaseOpProperties(StructPropertyHandle, ChildBuilder);
 }
 
+
 #undef LOCTEXT_NAMESPACE
diff --git a/Engine/Plugins/Animation/IKRig/Source/IKRigEditor/Private/RetargetEditor/IKRetargeterController.cpp b/Engine/Plugins/Animation/IKRig/Source/IKRigEditor/Private/RetargetEditor/IKRetargeterController.cpp
--- a/Engine/Plugins/Animation/IKRig/Source/IKRigEditor/Private/RetargetEditor/IKRetargeterController.cpp
+++ b/Engine/Plugins/Animation/IKRig/Source/IKRigEditor/Private/RetargetEditor/IKRetargeterController.cpp
@@ -527,6 +527,17 @@ void UIKRetargeterController::OnOpPropertyChanged(const FName& InOpName, const F
 	RetargeterNeedsInitialized.Broadcast();
 }
 
+uint8* UIKRetargeterController::GetChainSettingsMemory(const FName& InOpName, const FName InChainName) const
+{
+	FIKRetargetOpBase* Op = GetRetargetOpByName(InOpName);
+	if (!Op)
+	{
+		return nullptr;
+	}
+
+	return Op->GetChainSettingsMemory(InChainName);
+}
+
 FInstancedStruct* UIKRetargeterController::GetRetargetOpStructAtIndex(int32 Index) const
 {
 	check(Asset);
diff --git a/Engine/Plugins/Animation/IKRig/Source/IKRigEditor/Public/RetargetEditor/IKRetargetOpDetails.h b/Engine/Plugins/Animation/IKRig/Source/IKRigEditor/Public/RetargetEditor/IKRetargetOpDetails.h
--- a/Engine/Plugins/Animation/IKRig/Source/IKRigEditor/Public/RetargetEditor/IKRetargetOpDetails.h
+++ b/Engine/Plugins/Animation/IKRig/Source/IKRigEditor/Public/RetargetEditor/IKRetargetOpDetails.h
@@ -11,6 +11,7 @@
 #include "Retargeter/RetargetOps/StretchChainOp.h"
 #include "RigEditor/IKRigStructViewer.h"
 #include "Templates/SharedPointer.h"
+#include "RetargetEditor/IKRetargeterController.h"
 #include "IKRetargetOpDetails.generated.h"
 
 enum class EFKChainRotationMode : uint8;
@@ -53,6 +54,41 @@ public:
 
 protected:
 
+	using FChainSettingsGetter = TFunction<UObject*(const FName)>;
+	
+	template<typename TSettings, typename TWrapper>
+	FChainSettingsGetter TMakeChainSettingsGetter(
+	const FName InOpName,
+	const FName SettingsMemberPropertyName,
+	const TArray<FName>& HiddenProps = {})
+	{
+		static_assert(TIsDerivedFrom<TWrapper, UIKRigStructWrapperBase>::IsDerived,"TWrapper must derive from UIKRigStructWrapperBase");
+
+		UIKRetargeterController* Controller = AssetController;
+		
+		return [Controller, InOpName, SettingsMemberPropertyName, HiddenProps](const FName InTargetChainName) -> UObject*
+			{
+				FIKRigStructToView StructToView;
+				StructToView.Owner = Controller->GetAsset();
+				StructToView.Type = TSettings::StaticStruct();
+				StructToView.UniqueName = InOpName;
+				StructToView.MemoryProvider = [Controller, InOpName, InTargetChainName]() -> uint8*
+					{
+						return Controller->GetChainSettingsMemory(InOpName, InTargetChainName);
+					};
+
+				TWrapper* Wrapper = NewObject<TWrapper>(GetTransientPackage(), NAME_None, RF_Standalone);
+				Wrapper->InitializeWithRetargeter(StructToView, SettingsMemberPropertyName, Controller);
+
+				for (const FName Prop : HiddenProps)
+				{
+					Wrapper->SetPropertyHidden(Prop, true);
+				}
+
+				return Wrapper;
+			};
+	}
+
 	// the op that owns the settings being customized / edited
 	FName OpName;
 	TArray<TWeakObjectPtr<UObject>> SelectedObjects;
@@ -81,7 +117,7 @@ public:
 	// End IPropertyTypeCustomization
 };
 
-UCLASS(BlueprintType)
+UCLASS()
 class URetargetFKChainSettingsWrapper : public UIKRigStructWrapperBase
 {
 	GENERATED_BODY()
@@ -130,7 +166,7 @@ public:
 	// End IPropertyTypeCustomization
 };
 
-UCLASS(BlueprintType)
+UCLASS()
 class URetargetIKChainSettingsWrapper : public UIKRigStructWrapperBase
 {
 	GENERATED_BODY()
@@ -159,7 +195,7 @@ public:
 	// End IPropertyTypeCustomization
 };
 
-UCLASS(BlueprintType)
+UCLASS()
 class URetargetStrideWarpSettingsWrapper : public UIKRigStructWrapperBase
 {
 	GENERATED_BODY()
@@ -188,7 +224,7 @@ public:
 	// End IPropertyTypeCustomization
 };
 
-UCLASS(BlueprintType)
+UCLASS()
 class URetargetSpeedPlantSettingsWrapper : public UIKRigStructWrapperBase
 {
 	GENERATED_BODY()
@@ -217,7 +253,7 @@ public:
 	// End IPropertyTypeCustomization
 };
 
-UCLASS(BlueprintType)
+UCLASS()
 class UPoleVectorSettingsWrapper : public UIKRigStructWrapperBase
 {
 	GENERATED_BODY()
@@ -255,7 +291,7 @@ private:
 	TArray<TSharedPtr<FName>> PoseNameOptions;
 };
 
-UCLASS(BlueprintType)
+UCLASS()
 class URetargetPoseOpSettingsWrapper : public UIKRigStructWrapperBase
 {
 	GENERATED_BODY()
@@ -284,7 +320,7 @@ public:
 	// End IPropertyTypeCustomization
 };
 
-UCLASS(BlueprintType)
+UCLASS()
 class UStretchChainSettingsWrapper : public UIKRigStructWrapperBase
 {
 	GENERATED_BODY()
@@ -313,7 +349,7 @@ public:
 	// End IPropertyTypeCustomization
 };
 
-UCLASS(BlueprintType)
+UCLASS()
 class UFloorConstraintSettingsWrapper : public UIKRigStructWrapperBase
 {
 	GENERATED_BODY()
diff --git a/Engine/Plugins/Animation/IKRig/Source/IKRigEditor/Public/RetargetEditor/IKRetargeterController.h b/Engine/Plugins/Animation/IKRig/Source/IKRigEditor/Public/RetargetEditor/IKRetargeterController.h
--- a/Engine/Plugins/Animation/IKRig/Source/IKRigEditor/Public/RetargetEditor/IKRetargeterController.h
+++ b/Engine/Plugins/Animation/IKRig/Source/IKRigEditor/Public/RetargetEditor/IKRetargeterController.h
@@ -240,6 +240,9 @@ public:
 	// A callback whenever the property of an op is modified
 	UE_API void OnOpPropertyChanged(const FName& InOpName, const FPropertyChangedEvent& InPropertyChangedEvent) const;
 
+	// Return pointer to the memory associated with the given chain on the given op.
+	UE_API uint8* GetChainSettingsMemory(const FName& InOpName, const FName InChainName) const;
+
 	//
 	// GENERAL C++ ONLY API
 	//
