[IK Retarget] Split IK Goals Op into three separate ops.

#Synced-CL 48671149

diff --git a/Engine/Plugins/Animation/IKRig/Source/IKRig/Private/Retargeter/IKRetargetOps.cpp b/Engine/Plugins/Animation/IKRig/Source/IKRig/Private/Retargeter/IKRetargetOps.cpp
--- a/Engine/Plugins/Animation/IKRig/Source/IKRig/Private/Retargeter/IKRetargetOps.cpp
+++ b/Engine/Plugins/Animation/IKRig/Source/IKRig/Private/Retargeter/IKRetargetOps.cpp
@@ -79,21 +79,29 @@ FText FIKRetargetOpBase::GetWarningMessage() const
 	return LOCTEXT("OpNotReady", "Not initialized. See output log.");
 }
 
-void FIKRetargetOpBase::ResetExecutionTime()
+void FOpExecutionTimer::Reset()
 {
-	AverageExecutionTime = 0.0f;
+	TimeThisFrame = 0.0f;
+	RunningAverageTime = 0.0f;
 }
 
-void FIKRetargetOpBase::SetLatestExecutionTime(const double InSeconds)
+void FOpExecutionTimer::AddTimeThisFrame(const double InSeconds)
+{
+	TimeThisFrame += InSeconds;
+}
+
+void FOpExecutionTimer::StartNewFrame()
 {
 	constexpr double FrameWindow = 30.0;
 	constexpr double Alpha = 2.0 / (FrameWindow + 1.0);
-	AverageExecutionTime += Alpha * (InSeconds - AverageExecutionTime);
+	RunningAverageTime += Alpha * (TimeThisFrame - RunningAverageTime);
+
+	TimeThisFrame = 0.0f;
 }
 
-double FIKRetargetOpBase::GetAverageExecutionTime() const
+double FOpExecutionTimer::GetRunningAverage() const
 {
-	return AverageExecutionTime;
+	return RunningAverageTime;
 }
 
 #endif
diff --git a/Engine/Plugins/Animation/IKRig/Source/IKRig/Private/Retargeter/IKRetargetProcessor.cpp b/Engine/Plugins/Animation/IKRig/Source/IKRig/Private/Retargeter/IKRetargetProcessor.cpp
--- a/Engine/Plugins/Animation/IKRig/Source/IKRig/Private/Retargeter/IKRetargetProcessor.cpp
+++ b/Engine/Plugins/Animation/IKRig/Source/IKRig/Private/Retargeter/IKRetargetProcessor.cpp
@@ -1190,7 +1190,7 @@ void FIKRetargetProcessor::InitialOpStackSetup(
 		// asset instance gets reference to running editor instance
 		AssetOpSettings->EditorInstance = NewOp.GetSettings();
 		AssetOp->EditorInstance = &NewOp;
-		NewOp.ResetExecutionTime();
+		NewOp.Timer.Reset();
 		// references the skeletons are used for bone selector widgets in the details panel
 		AssetOpSettings->AssignSkeletonAssets(
 			SourceSkeleton.SkeletalMesh->GetSkeleton(),
@@ -1240,6 +1240,9 @@ void FIKRetargetProcessor::InitializeRetargetOps()
 		
 		Op.Initialize(*this, SourceSkeleton, TargetSkeleton, ParentOp, Log);
 	}
+
+	// record meta-data to build execution order
+	InitializeOpEntries();
 	
 	// gather retargeted bones from each op
 	TSet<int32> AllRetargetedBones;
@@ -1330,78 +1333,204 @@ void FIKRetargetProcessor::RunRetargetOps(
 	const double InDeltaTime,
 	const int32 InLOD)
 {
-	auto ShouldExecuteOp = [&InLOD](FIKRetargetOpBase& InOp)
+#if WITH_EDITOR
+	// reset op timers for this frame (editor only)
+	for (FInstancedStruct& OpStruct : OpStack)
 	{
-		if (!InOp.IsEnabled())
+		FIKRetargetOpBase& Op = OpStruct.GetMutable<FIKRetargetOpBase>();
+		Op.Timer.StartNewFrame();
+	}
+#endif
+
+	// build a stack of op calls in the correct order
+	TArray<FRetargetOpCall> ExecutionStack;
+	BuildExecutionStack(InLOD, ExecutionStack);
+
+	// run the ops
+	for (const FRetargetOpCall& Call : ExecutionStack)
+	{
+		
+#if WITH_EDITOR
+		const double StartTime = FPlatformTime::Seconds();
+#endif
+
+		FIKRetargetOpBase* Op = Call.Op;
+
+		switch (Call.Phase)
 		{
-			return false; // disabled
+		case ERetargetOpPhase::ParentBeforeChildren:
+			Op->RunBeforeChildren(*this, InDeltaTime, InSourceGlobalPose, OutTargetGlobalPose);
+			break;
+
+		case ERetargetOpPhase::Main:
+			Op->Run(*this, InDeltaTime, InSourceGlobalPose, OutTargetGlobalPose);
+			break;
+
+		case ERetargetOpPhase::ChildAfterParent:
+			Op->RunAfterParent(*this, InDeltaTime, InSourceGlobalPose, OutTargetGlobalPose);
+			break;
+
+		default:
+			checkNoEntry();
 		}
-		if (!InOp.IsInitialized())
+
+#if WITH_EDITOR
+		Op->Timer.AddTimeThisFrame(FPlatformTime::Seconds() - StartTime);
+#endif
+	}
+}
+
+bool FIKRetargetProcessor::ShouldExecuteOp(const FIKRetargetOpBase& InOp, int32 InLOD)
+{
+	if (!InOp.IsEnabled())
+	{
+		return false;
+	}
+	if (!InOp.IsInitialized())
+	{
+		return false;
+	}
+
+	const int32 LODThreshold = InOp.GetSettingsConst()->LODThreshold;
+	if (LODThreshold != INDEX_NONE && InLOD > LODThreshold)
+	{
+		return false;
+	}
+
+	return true;
+}
+
+void FIKRetargetProcessor::BuildExecutionStack(int32 InLOD, TArray<FRetargetOpCall>& OutStack)
+{
+	OutStack.Reset();
+
+	// 1st pass update bExecute
+	for (FOpEntry& Entry : OpEntries)
+	{
+		Entry.bExecute = ShouldExecuteOp(*Entry.Op, InLOD);
+	}
+
+	// 2nd pass, generate execution order and store in stack
+	for (FOpEntry& OpEntry : OpEntries)
+	{
+		if (!OpEntry.bExecute)
 		{
-			return false; // not initialized
+			// op not running this frame
+			// NOTE: this implies it's children will not run either (by design)
+			continue;
 		}
-			
-		const int32 LODThreshold = InOp.GetSettings()->LODThreshold;
-		if (LODThreshold != INDEX_NONE && InLOD > LODThreshold)
+
+		const bool bHasParent = OpEntry.Parent != nullptr;
+		const bool bIsParent = OpEntry.Op->CanHaveChildOps();
+
+		if (bHasParent)
 		{
-			return false; // LOD'd off
+			continue; // children ops managed by parent
 		}
 
-		return true;
-	};
+		// neither a child or a parent, ie a simple op, just run Main once.
+		if (!bIsParent && !bHasParent)
+		{
+			OutStack.Add({ ERetargetOpPhase::Main, OpEntry.Op });
+			continue;
+		}
 
-	TArray<FIKRetargetOpBase*> OpsRunBeforeChildrenThisTick;
+		// Parent with children execution order:
+		// NOTE: parents run all steps even if they don't have children
+		// but children don't run at all if the parent is disabled.
+		// 1. Parent->RunBeforeChildren
+		// 2. Children->Run
+		// 3. Parent->Run
+		// 4. Children->RunAfterParent
 
-	// TODO: This won't work if child ops can have children
-	FIKRetargetOpBase* CurrentParentOp = nullptr;
-	TArray<FIKRetargetOpBase*> OpsToRunAfterParent;
-	
-	for (FInstancedStruct& OpStruct : OpStack)
-	{
-		// collect profiling data per op
-		#if WITH_EDITOR
-		double StartTime = FPlatformTime::Seconds();
-		#endif
+		// 1. Parent->RunBeforeChildren()
+		OutStack.Emplace(ERetargetOpPhase::ParentBeforeChildren, OpEntry.Op);
 
-		// execute the op
-		FIKRetargetOpBase& Op = OpStruct.GetMutable<FIKRetargetOpBase>();
-
-		// give parent ops a chance to run right before their children
-		FIKRetargetOpBase* ParentOp = GetRetargetOpByName(Op.GetParentOpName());
-		if (ParentOp && !OpsRunBeforeChildrenThisTick.Contains(ParentOp))
+		// 2. Children->Run()
+		for (FOpEntry* Child : OpEntry.Children)
 		{
-			ParentOp->RunBeforeChildren(*this, InDeltaTime, InSourceGlobalPose, OutTargetGlobalPose);
-			OpsRunBeforeChildrenThisTick.Add(ParentOp);
+			if (Child->bExecute)
+			{
+				OutStack.Emplace(ERetargetOpPhase::Main, Child->Op);
+			}
 		}
-		
-		if (ShouldExecuteOp(Op))
+
+		// 3. Parent->Run()
+		OutStack.Emplace(ERetargetOpPhase::Main, OpEntry.Op);
+
+		// 4. Children->RunAfterParent()
+		for (FOpEntry* Child : OpEntry.Children)
 		{
-			// Give child ops entry point to run after parent op finishes
-			if (ParentOp)
+			if (Child->bExecute)
 			{
-				ensure(!CurrentParentOp || CurrentParentOp == ParentOp);
-				CurrentParentOp = ParentOp;
-				OpsToRunAfterParent.Add(&Op);
+				OutStack.Emplace(ERetargetOpPhase::ChildAfterParent, Child->Op);
 			}
-			
-			// run the op
-			Op.Run(*this, InDeltaTime, InSourceGlobalPose, OutTargetGlobalPose);
+		}
+	}
+}
 
-			// Run children list after parent
-			if (CurrentParentOp == &Op)
+void FIKRetargetProcessor::InitializeOpEntries()
+{
+	auto FindOpEntry = [this](const FIKRetargetOpBase* InOpToFind) -> FOpEntry*
+	{
+		for (FOpEntry& Entry : OpEntries)
+		{
+			if (Entry.Op == InOpToFind)
 			{
-				CurrentParentOp = nullptr;
-				for (FIKRetargetOpBase* ChildOp : OpsToRunAfterParent)
-				{
-					ChildOp->RunAfterParent(*this, InDeltaTime, InSourceGlobalPose, OutTargetGlobalPose);
-				}
+				return &Entry;
 			}
 		}
+		return nullptr;
+	};
+	
+	OpEntries.Reset(OpStack.Num());
+
+	// create entries for each op in the stack
+	for (int32 Index = 0; Index < OpStack.Num(); ++Index)
+	{
+		FIKRetargetOpBase& Op = OpStack[Index].GetMutable<FIKRetargetOpBase>();
+		
+		FOpEntry& Entry = OpEntries.AddDefaulted_GetRef();
+		Entry.Op = &Op;
+		Entry.StackIndex = Index;
+		Entry.bExecute = true; // default to true, set during runtime
+	}
+
+	// wire up parent/children on the entries
+	for (FOpEntry& Entry : OpEntries)
+	{
+		const FName ParentName = Entry.Op->GetParentOpName();
+		if (ParentName.IsNone())
+		{
+			continue; // no parent
+		}
 
-		// END op profiler scope
-		#if WITH_EDITOR
-		Op.SetLatestExecutionTime( FPlatformTime::Seconds() - StartTime);
-		#endif
+		FIKRetargetOpBase* ParentOp = GetRetargetOpByName(ParentName);
+		if (!ensure(ParentOp))
+		{
+			continue; // should not happen, ill formed op stack referencing missing parent
+		}
+		
+		FOpEntry* ParentEntry = FindOpEntry(ParentOp);
+		if (!ensure(ParentEntry))
+		{
+			continue; // should not happen, entries out of sync with op stack
+		}
+
+		Entry.Parent = ParentEntry;
+		ParentEntry->Children.Add(&Entry);
+	}
+
+	// keep children ordered by stack index
+	for (FOpEntry& Entry : OpEntries)
+	{
+		if (Entry.Children.Num() > 1)
+		{
+			Entry.Children.Sort([](const FOpEntry& A, const FOpEntry& B)
+			{
+				return A.StackIndex < B.StackIndex;
+			});
+		}
 	}
 }
 
@@ -1756,15 +1885,22 @@ void FIKRetargetProcessor::DebugDrawAllOps(
 	const FTransform SourceTransform = FTransform(FQuat::Identity, RetargetAsset->SourceMeshOffset, FVector::OneVector);
 
 	// debug draw each op in the stack
-	for (int32 OpIndex=0; OpIndex<OpStack.Num(); ++OpIndex)
+	for (const FOpEntry& OpEntry : OpEntries)
 	{
+		const FIKRetargetOpBase* Op = OpEntry.Op;
+
 		// filter disabled or uninitialized ops
-		const FIKRetargetOpBase* Op = OpStack[OpIndex].GetPtr<FIKRetargetOpBase>();
 		if (!Op->IsEnabled() || !Op->IsInitialized() || !Op->GetSettingsConst()->bDebugDraw)
 		{
 			continue;
 		}
 
+		// skip if parent is disabled
+		if (OpEntry.Parent != nullptr && !OpEntry.Parent->Op->IsEnabled())
+		{
+			continue;
+		}
+
 		Op->DebugDraw(InPDI, SourceTransform, TargetTransform, TargetScale, EditorState);
 	}
 }
diff --git a/Engine/Plugins/Animation/IKRig/Source/IKRig/Private/Retargeter/IKRetargeter.cpp b/Engine/Plugins/Animation/IKRig/Source/IKRig/Private/Retargeter/IKRetargeter.cpp
--- a/Engine/Plugins/Animation/IKRig/Source/IKRig/Private/Retargeter/IKRetargeter.cpp
+++ b/Engine/Plugins/Animation/IKRig/Source/IKRig/Private/Retargeter/IKRetargeter.cpp
@@ -7,10 +7,13 @@
 #include "Retargeter/IKRetargetProfile.h"
 #include "Engine/SkeletalMesh.h"
 #include "Retargeter/RetargetOps/AlignPoleVectorOp.h"
+#include "Retargeter/RetargetOps/BlendToSourceOp.h"
 #include "Retargeter/RetargetOps/FKChainsOp.h"
 #include "Retargeter/RetargetOps/IKChainsOp.h"
+#include "Retargeter/RetargetOps/OffsetGoalsOp.h"
 #include "Retargeter/RetargetOps/PelvisMotionOp.h"
 #include "Retargeter/RetargetOps/RunIKRigOp.h"
+#include "Retargeter/RetargetOps/ScaleGoalsOp.h"
 #include "Retargeter/RetargetOps/ScaleSourceOp.h"
 #include "Retargeter/RetargetOps/SpeedPlantingOp.h"
 #include "Retargeter/RetargetOps/StrideWarpingOp.h"
@@ -130,20 +133,25 @@ void UIKRetargeter::PostLoad()
 	}
 
 #if WITH_EDITOR
+	// post-load patches in order of oldest to newest asset upgrades
 	PostLoadOldSettingsToNew();
-	
 	PostLoadOldOpsToNewStructOps();
-
 	PostLoadConvertEverythingToOps();
-
 	PostLoadPutChainMappingInOps();
 #endif
 
+	// give ops a chance to post load patch themselves
+	// plugin architecture here allows ops in other plugins to load-patch
+	PostLoadOpStack();
+
+#if WITH_EDITOR
+	// NOTE: must run after PostLoadOpStack() to properly patch blend to source "Affect IK"
+	PostLoadSplitIKChainsOp();	
+#endif
+
 	CleanRetargetPoses();
 	
 	CleanOpStack();
-
-	PostLoadOpStack();
 }
 
 #if WITH_EDITOR
@@ -495,6 +503,248 @@ void UIKRetargeter::PostLoadPutChainMappingInOps()
 	
 	PRAGMA_ENABLE_DEPRECATION_WARNINGS
 }
+
+void UIKRetargeter::PostLoadSplitIKChainsOp()
+{
+	// only perform this if loading older package version
+	if (GetLinkerCustomVersion(FIKRigObjectVersion::GUID) >= FIKRigObjectVersion::IKChainsOpSplitIntoSeparateOps)
+	{
+		return;
+	}
+
+	auto AllChainSettingsAtDefault = [](FIKRetargetIKChainsOp* InIKOp)
+		{
+			const FRetargetIKChainSettings DefaultSettings;
+			for (const FRetargetIKChainSettings& Chain : InIKOp->Settings.ChainsToRetarget)
+			{
+				if (Chain != DefaultSettings)
+				{
+					return false;
+				}
+			}
+			return true;
+		};
+
+	auto CopySettingsToRunIKRigOp = [](FIKRetargetIKChainsOp& InIKOp, FIKRetargetRunIKRigOp& InRunIKRigOp)
+		{
+			for (const FRetargetIKChainSettings& CopyFromChain : InIKOp.Settings.ChainsToRetarget)
+			{
+				for (FIKRetargetRunIKRigChainSettings& CopyToChain : InRunIKRigOp.Settings.Chains)
+				{
+					if (CopyFromChain.TargetChainName != CopyToChain.TargetChainName)
+					{
+						continue;
+					}
+
+					// copy only the enable flag to this op
+					CopyToChain.bEnableIK = CopyFromChain.EnableIK;
+				}
+			}
+		};
+
+	auto UsingBlendToSource = [](FIKRetargetIKChainsOp* InIKOp)
+		{
+			for (const FRetargetIKChainSettings& Chain : InIKOp->Settings.ChainsToRetarget)
+			{
+				if (!FMath::IsNearlyEqual(Chain.BlendToSource, 0.0))
+				{
+					return true;
+				}
+			}
+			return false;
+		};
+	
+	auto CopySettingsToBlendToSourceOp = [](FIKRetargetIKChainsOp& InIKOp, FIKRetargetBlendToSourceOp& InNewOp)
+		{
+			for (const FRetargetIKChainSettings& CopyFromChain : InIKOp.Settings.ChainsToRetarget)
+			{
+				for (FIKRetargetBlendToSourceChainSettings& CopyToChain : InNewOp.Settings.Chains)
+				{
+					if (CopyFromChain.TargetChainName != CopyToChain.TargetChainName)
+					{
+						continue;
+					}
+
+					CopyToChain.BlendToSource = CopyFromChain.BlendToSource;
+					CopyToChain.TranslationAlpha = CopyFromChain.BlendToSourceTranslation;
+					CopyToChain.RotationAlpha = CopyFromChain.BlendToSourceRotation;
+					CopyToChain.TranslationPerAxisAlpha = CopyFromChain.BlendToSourceWeights;
+					CopyToChain.ApplyPelvisOffset = CopyFromChain.ApplyPelvisOffsetToSourceGoals;
+				}
+			}
+		};
+	
+	auto UsingStaticOffsets = [](FIKRetargetIKChainsOp* InIKOp)
+		{
+			for (const FRetargetIKChainSettings& Chain : InIKOp->Settings.ChainsToRetarget)
+			{
+				if (!Chain.StaticOffset.Equals(FVector::ZeroVector))
+				{
+					return true;
+				}
+				if (!Chain.StaticLocalOffset.Equals(FVector::ZeroVector))
+				{
+					return true;
+				}
+				if (!Chain.StaticRotationOffset.Equals(FRotator::ZeroRotator))
+				{
+					return true;
+				}
+			}
+			return false;
+		};
+
+	auto CopySettingsToOffsetOp = [](FIKRetargetIKChainsOp& InIKOp, FIKRetargetOffsetGoalsOp& InNewOp)
+		{
+			for (const FRetargetIKChainSettings& CopyFromChain : InIKOp.Settings.ChainsToRetarget)
+			{
+				for (FIKRetargetOffsetGoalsChainSettings& CopyToChain : InNewOp.Settings.Chains)
+				{
+					if (CopyFromChain.TargetChainName != CopyToChain.TargetChainName)
+					{
+						continue;
+					}
+
+					// if any of the offsets on this chain are non-default, set alpha to 1 to enable them
+					if (!CopyFromChain.StaticOffset.Equals(FVector::ZeroVector) ||
+						!CopyFromChain.StaticLocalOffset.Equals(FVector::ZeroVector) ||
+						!CopyFromChain.StaticRotationOffset.Equals(FRotator::ZeroRotator))
+					{
+						CopyToChain.Alpha = 1.0;
+						CopyToChain.GlobalTranslationOffset = CopyFromChain.StaticOffset;
+						CopyToChain.LocalTranslationOffset = CopyFromChain.StaticLocalOffset;
+						CopyToChain.LocalRotationOffset = CopyFromChain.StaticRotationOffset;
+					}
+				}
+			}
+		};
+
+	auto UsingScaling = [](FIKRetargetIKChainsOp* InIKOp)
+		{
+			for (const FRetargetIKChainSettings& Chain : InIKOp->Settings.ChainsToRetarget)
+			{
+				if (!FMath::IsNearlyEqual(Chain.ScaleVertical, 1.0))
+				{
+					return true;
+				}
+				if (!FMath::IsNearlyEqual(Chain.Extension, 1.0))
+				{
+					return true;
+				}
+			}
+			return false;
+		};
+
+	auto CopySettingsToScaleGoalsOp = [](FIKRetargetIKChainsOp& InIKOp, FIKRetargetScaleGoalsOp& InNewOp)
+		{
+			for (const FRetargetIKChainSettings& CopyFromChain : InIKOp.Settings.ChainsToRetarget)
+			{
+				for (FIKRetargetScaleGoalsChainSettings& CopyToChain : InNewOp.Settings.Chains)
+				{
+					if (CopyFromChain.TargetChainName != CopyToChain.TargetChainName)
+					{
+						continue;
+					}
+
+					CopyToChain.ScaleVertical = CopyFromChain.ScaleVertical;
+					CopyToChain.ScaleAlongChain = CopyFromChain.Extension;
+				}
+			}
+		};
+
+	auto RemoveOp = [this](FIKRetargetOpBase* InOp)
+		{
+			int32 OpIndex = GetIndexOfRetargetOp(InOp);
+			if (OpIndex != INDEX_NONE)
+			{
+				RetargetOps.RemoveAt(OpIndex);	
+			}
+		};
+
+	// 1st pass:
+	// initialize chain settings on any RunIKRig ops so they have settings for all the chains in the rig
+	// prior version of this op had no chain settings, so we must forcibly generate them now
+	TArray<FIKRetargetRunIKRigOp*> RunIKRigOps = GetAllRetargetOpsOfType<FIKRetargetRunIKRigOp>();
+	for (FIKRetargetRunIKRigOp* Op : RunIKRigOps)
+	{
+		Op->RegenerateChainSettings();
+	}
+
+	// 2nd pass
+	// go through all the "Retarget IK Goal" ops and split them out into the new ops as needed
+	TArray<FIKRetargetIKChainsOp*> IKChainOps = GetAllRetargetOpsOfType<FIKRetargetIKChainsOp>();
+	for (FIKRetargetIKChainsOp* Op : IKChainOps)
+	{
+		FIKRetargetOpBase* ParentOp = GetRetargetOpByNameNonConst(Op->GetParentOpName());
+		if (!ParentOp)
+		{
+			// should not happen
+			// orphaned IK Chain op is invalid, skip upgrade
+			continue;
+		}
+
+		if (ParentOp->GetType() != FIKRetargetRunIKRigOp::StaticStruct())
+		{
+			// should not happen
+			// IK Chain op parented incorrectly, skip upgrade
+			continue;
+		}
+		
+		FIKRetargetRunIKRigOp* RunIKRigOp = reinterpret_cast<FIKRetargetRunIKRigOp*>(ParentOp);
+
+		// copy chain settings to the run ik rig op
+		CopySettingsToRunIKRigOp(*Op, *RunIKRigOp);
+		
+		// with default settings, the op behaves identically to the new default RunIKRigOp
+		// so we can safely delete it and let the parent op do the work
+		if (AllChainSettingsAtDefault(Op))
+		{
+			continue;
+		}
+
+		FIKRetargetOpBase* NextOpInStack = Op;
+
+		// if using any blend to source features, create an identically configured Blend to Source op
+		if (UsingBlendToSource(Op))
+		{
+			FInstancedStruct& NewOpStruct = RetargetOps.InsertDefaulted_GetRef(GetIndexOfRetargetOp(NextOpInStack) + 1);
+			NewOpStruct.InitializeAs<FIKRetargetBlendToSourceOp>();
+			FIKRetargetBlendToSourceOp& NewOp = NewOpStruct.GetMutable<FIKRetargetBlendToSourceOp>();
+			NewOp.RegenerateChainSettings(ParentOp);
+			CopySettingsToBlendToSourceOp(*Op, NewOp);
+			NextOpInStack = &NewOp;
+		}
+
+		// if using any offset features, create an identically configured Offset Goals op
+		if (UsingStaticOffsets(Op))
+		{
+			FInstancedStruct& NewOpStruct = RetargetOps.InsertDefaulted_GetRef(GetIndexOfRetargetOp(NextOpInStack) + 1);
+			NewOpStruct.InitializeAs<FIKRetargetOffsetGoalsOp>();
+			FIKRetargetOffsetGoalsOp& NewOp = NewOpStruct.GetMutable<FIKRetargetOffsetGoalsOp>();
+			NewOp.RegenerateChainSettings(ParentOp);
+			CopySettingsToOffsetOp(*Op, NewOp);
+			NextOpInStack = &NewOp;
+		}
+
+		// if using any scaling features, create an identically configured Scale Goals op
+		if (UsingScaling(Op))
+		{
+			FInstancedStruct& NewOpStruct = RetargetOps.InsertDefaulted_GetRef(GetIndexOfRetargetOp(NextOpInStack) + 1);
+			NewOpStruct.InitializeAs<FIKRetargetScaleGoalsOp>();
+			FIKRetargetScaleGoalsOp& NewOp = NewOpStruct.GetMutable<FIKRetargetScaleGoalsOp>();
+			NewOp.RegenerateChainSettings(ParentOp);
+			CopySettingsToScaleGoalsOp(*Op, NewOp);
+			NextOpInStack = &NewOp;
+		}
+	}
+
+	// 3rd pass
+	// remove all the IKGoal ops
+	for (FIKRetargetIKChainsOp* Op : IKChainOps)
+	{
+		RemoveOp(Op);
+	}
+}
 #endif
 
 FName UIKRetargeter::GetCleanAndUniqueOpName(const FName& InOpName, const int32 InIndexOfOp)
@@ -832,16 +1082,36 @@ const FIKRetargetOpBase* UIKRetargeter::GetRetargetOpByName(const FName& InOpNam
 {
 	for (const FInstancedStruct& OpStruct : RetargetOps)
 	{
-		const FIKRetargetOpBase& Op = OpStruct.Get<FIKRetargetOpBase>();
-		if (Op.GetName() == InOpName)
+		const FIKRetargetOpBase* Op = OpStruct.GetPtr<FIKRetargetOpBase>();
+		if (Op->GetName() == InOpName)
 		{
-			return &Op;
+			return Op;
 		}
 	}
 	
 	return nullptr;
 }
 
+FIKRetargetOpBase* UIKRetargeter::GetRetargetOpByNameNonConst(const FName& InOpName)
+{
+	return const_cast<FIKRetargetOpBase*>(GetRetargetOpByName(InOpName));
+}
+
+int32 UIKRetargeter::GetIndexOfRetargetOp(const FIKRetargetOpBase* InOp) const
+{
+	for (int32 OpIndex=0; OpIndex < RetargetOps.Num(); OpIndex++)
+	{
+		const FInstancedStruct& OpStruct = RetargetOps[OpIndex];
+		const FIKRetargetOpBase* Op = OpStruct.GetPtr<FIKRetargetOpBase>();
+		if (Op==InOp)
+		{
+			return OpIndex;
+		}
+	}
+
+	return INDEX_NONE;
+}
+
 const FIKRetargetPose* UIKRetargeter::GetCurrentRetargetPose(const ERetargetSourceOrTarget& SourceOrTarget) const
 {
 	return SourceOrTarget == ERetargetSourceOrTarget::Source ? &SourceRetargetPoses[CurrentSourceRetargetPose] : &TargetRetargetPoses[CurrentTargetRetargetPose];
diff --git a/Engine/Plugins/Animation/IKRig/Source/IKRig/Private/Retargeter/RetargetOps/BlendToSourceOp.cpp b/Engine/Plugins/Animation/IKRig/Source/IKRig/Private/Retargeter/RetargetOps/BlendToSourceOp.cpp
new file mode 100644
--- /dev/null
+++ b/Engine/Plugins/Animation/IKRig/Source/IKRig/Private/Retargeter/RetargetOps/BlendToSourceOp.cpp
@@ -0,0 +1,304 @@
+﻿// Copyright Epic Games, Inc. All Rights Reserved.
+
+#include "Retargeter/RetargetOps/BlendToSourceOp.h"
+
+#include "Retargeter/IKRetargetOpUtils.h"
+#include "Retargeter/IKRetargetProcessor.h"
+#include "Retargeter/RetargetOps/PelvisMotionOp.h"
+#include "Retargeter/RetargetOps/RunIKRigOp.h"
+
+#if WITH_EDITOR
+#include "PrimitiveDrawingUtils.h"
+#include "Engine/EngineTypes.h"
+#endif
+
+bool FIKRetargetBlendToSourceChainSettings::operator==(const FIKRetargetBlendToSourceChainSettings& Other) const
+{
+	return FMath::IsNearlyEqualByULP(BlendToSource, Other.BlendToSource)
+		&& FMath::IsNearlyEqualByULP(ApplyPelvisOffset, Other.ApplyPelvisOffset)
+		&& TranslationPerAxisAlpha.Equals(Other.TranslationPerAxisAlpha);
+}
+
+const UClass* FIKRetargetBlendToSourceOpSettings::GetControllerType() const
+{
+	return UIKRetargetBlendToSourceController::StaticClass();
+}
+
+void FIKRetargetBlendToSourceOpSettings::CopySettingsAtRuntime(const FIKRetargetOpSettingsBase* InSettingsToCopyFrom)
+{
+	// copy everything except the chains
+	const TArray<FName> PropertiesToIgnore = {GET_MEMBER_NAME_CHECKED(FIKRetargetBlendToSourceOpSettings, Chains)};
+	FIKRetargetOpBase::CopyStructProperties(
+		FIKRetargetBlendToSourceOpSettings::StaticStruct(),
+		InSettingsToCopyFrom,
+		this,
+		PropertiesToIgnore);
+
+	// copy chains
+	const FIKRetargetBlendToSourceOpSettings* NewSettings = reinterpret_cast<const FIKRetargetBlendToSourceOpSettings*>(InSettingsToCopyFrom);
+	IKRetargetOpUtils::CopyChainSettingsAtRuntime(NewSettings->Chains, Chains);
+}
+
+void FIKChainToBlend::Run(
+	FIKRigGoalContainer& InGoalContainer,
+	const TArray<FTransform>& InSourceGlobalPose,
+	const FIKRetargetBlendToSourceOpSettings& InOpSettings,
+	const FIKRetargetPelvisMotionOp* InPelvisMotionOp) const
+{
+	if (Settings->BlendToSource <= UE_KINDA_SMALL_NUMBER)
+	{
+		return; // disabled on this chain
+	}
+	
+	if (InOpSettings.Alpha <= UE_KINDA_SMALL_NUMBER)
+	{
+		return; // whole op disabled
+	}
+
+	FIKRigGoal* Goal = InGoalContainer.FindGoalByName(TargetBoneChain->IKGoalName);
+	if (!ensure(Goal))
+	{
+		return; // goal should be validated at initialization, should not happen
+	}
+
+	// blend to source position
+	const double FinalPositionAlpha = Settings->BlendToSource * Settings->TranslationAlpha * InOpSettings.Alpha;
+	if (FinalPositionAlpha > UE_KINDA_SMALL_NUMBER)
+	{
+		// get source location to blend towards
+		const int32 SourceEndBoneIndex = SourceBoneChain->BoneIndices.Last();
+		FVector SourceLocation = InSourceGlobalPose[SourceEndBoneIndex].GetLocation();
+		if (InPelvisMotionOp && Settings->ApplyPelvisOffset > UE_KINDA_SMALL_NUMBER)
+		{
+			const FVector WeightedPelvisOffset = InPelvisMotionOp->GetPelvisTranslationOffset() * InPelvisMotionOp->GetAffectIKWeightAsVector();
+			SourceLocation += WeightedPelvisOffset * Settings->ApplyPelvisOffset;
+		}
+
+		// blend and apply to goal
+		const FVector PerAxisWeight = FinalPositionAlpha * Settings->TranslationPerAxisAlpha;
+		FVector& GoalPosition = Goal->Position;
+		GoalPosition.X = FMath::Lerp(GoalPosition.X, SourceLocation.X, PerAxisWeight.X);
+		GoalPosition.Y = FMath::Lerp(GoalPosition.Y, SourceLocation.Y, PerAxisWeight.Y);
+		GoalPosition.Z = FMath::Lerp(GoalPosition.Z, SourceLocation.Z, PerAxisWeight.Z);
+	}
+
+	// blend to source rotation
+	const double FinalRotationAlpha = InOpSettings.Alpha * Settings->RotationAlpha * Settings->BlendToSource;
+	if (FinalRotationAlpha > UE_KINDA_SMALL_NUMBER)
+	{
+		// get relevant bone transforms
+		const int32 SourceEndBoneIndex = SourceBoneChain->BoneIndices.Last();
+		const FTransform& InitialSourceEndGlobal = SourceBoneChain->RefPoseGlobalTransforms.Last();
+		const FTransform& InitialTargetEndGlobal = TargetBoneChain->RefPoseGlobalTransforms.Last();
+		const FTransform& CurrentSourceEndGlobal = InSourceGlobalPose[SourceEndBoneIndex];
+		
+		// apply delta rotation from source input
+		const FQuat CurrentSourceEndRotation = CurrentSourceEndGlobal.GetRotation();
+		const FQuat SourceDeltaRotation = (CurrentSourceEndRotation * InitialSourceEndGlobal.GetRotation().Inverse()).GetNormalized();
+		FQuat NewGoalRotation = SourceDeltaRotation * InitialTargetEndGlobal.GetRotation();
+
+		// blend to source rotation
+		NewGoalRotation = FQuat::FastLerp(Goal->Rotation.Quaternion(), NewGoalRotation, FinalRotationAlpha);
+		NewGoalRotation.Normalize();
+
+		// apply to goal
+		Goal->Rotation = NewGoalRotation.Rotator();
+	}
+
+#if WITH_EDITOR
+	{
+		DebugOutputGoal = FTransform(Goal->Rotation.Quaternion(), Goal->Position);
+		DebugSourceGoal = InSourceGlobalPose[SourceBoneChain->BoneIndices.Last()];
+	}
+#endif
+}
+
+#if WITH_EDITOR
+void FIKChainToBlend::DebugDraw(FPrimitiveDrawInterface* InPDI, const FTransform& InComponentTransform, const double InDrawSize) const
+{
+	const FTransform Source = DebugSourceGoal * InComponentTransform;
+	const FTransform Output = DebugOutputGoal * InComponentTransform;
+	
+	DrawWireSphere(
+				InPDI,
+				Source,
+				FLinearColor::Blue,
+				InDrawSize * 0.5f,
+				12,
+				SDPG_MAX,
+				0.0f,
+				0.001f,
+				false);
+
+	DrawWireSphere(
+			InPDI,
+			Output,
+			FLinearColor::Yellow,
+			InDrawSize,
+			12,
+			SDPG_MAX,
+			0.0f,
+			0.001f,
+			false);
+
+	DrawDashedLine(
+					InPDI,
+					Source.GetLocation(),
+					Output.GetLocation(),
+					FLinearColor::Black,
+					1.0f,
+					SDPG_MAX);
+}
+#endif
+
+bool FIKRetargetBlendToSourceOp::Initialize(
+	const FIKRetargetProcessor& InProcessor,
+	const FRetargetSkeleton& InSourceSkeleton,
+	const FTargetSkeleton& InTargetSkeleton,
+	const FIKRetargetOpBase* InParentOp,
+	FIKRigLogger& InLog)
+{
+	bIsInitialized = false;
+	
+	// go through all chains to retarget and load them
+	ChainsToBlend.Reset();
+	const FRetargeterBoneChains& BoneChains = InProcessor.GetBoneChains();
+	const FIKRigGoalContainer& GoalContainer = InProcessor.GetIKRigGoalContainer();
+	for (FIKRetargetBlendToSourceChainSettings& ChainSettings : Settings.Chains)
+	{
+		FName TargetChainName = ChainSettings.TargetChainName;
+		const FResolvedBoneChain* SourceBoneChain = nullptr;
+		const FResolvedBoneChain* TargetBoneChain = nullptr;
+		if (!FIKRetargetRunIKRigOp::ValidateIKChain(
+			TargetChainName,
+			GoalContainer,
+			BoneChains,
+			this,
+			InParentOp,
+			InLog,
+			SourceBoneChain /*out*/,
+			TargetBoneChain /*out*/))
+		{
+			continue;
+		}
+
+		ChainsToBlend.Emplace(*SourceBoneChain, *TargetBoneChain, ChainSettings);
+	}
+
+	// cache pelvis op
+	PelvisMotionOp = InProcessor.GetFirstRetargetOpOfType<FIKRetargetPelvisMotionOp>();
+
+	bIsInitialized = true;
+	return true;
+}
+
+void FIKRetargetBlendToSourceOp::Run(
+	FIKRetargetProcessor& InProcessor,
+	const double InDeltaTime,
+	const TArray<FTransform>& InSourceGlobalPose,
+	TArray<FTransform>& OutTargetGlobalPose)
+{
+	if (Settings.Alpha <= UE_KINDA_SMALL_NUMBER)
+	{
+		return; // fast skip whole op
+	}
+	
+	FIKRigGoalContainer& GoalContainer = InProcessor.GetIKRigGoalContainer();
+	for (const FIKChainToBlend& Chain : ChainsToBlend)
+	{
+		Chain.Run(GoalContainer, InSourceGlobalPose, Settings, PelvisMotionOp);
+	}
+}
+
+FIKRetargetOpSettingsBase* FIKRetargetBlendToSourceOp::GetSettings()
+{
+	return &Settings;
+}
+
+const UScriptStruct* FIKRetargetBlendToSourceOp::GetSettingsType() const
+{
+	return FIKRetargetBlendToSourceOpSettings::StaticStruct();
+}
+
+const UScriptStruct* FIKRetargetBlendToSourceOp::GetType() const
+{
+	return FIKRetargetBlendToSourceOp::StaticStruct();
+}
+
+const UScriptStruct* FIKRetargetBlendToSourceOp::GetParentOpType() const
+{
+	return FIKRetargetRunIKRigOp::StaticStruct();
+}
+
+void FIKRetargetBlendToSourceOp::OnAddedToStack(const UIKRetargeter* InRetargetAsset, const FIKRetargetOpBase* InParentOp)
+{
+	RegenerateChainSettings(InParentOp);
+}
+
+void FIKRetargetBlendToSourceOp::OnAssignIKRig(
+	const ERetargetSourceOrTarget SourceOrTarget,
+	const UIKRigDefinition* InIKRig,
+	const FIKRetargetOpBase* InParentOp)
+{
+	RegenerateChainSettings(InParentOp);
+}
+
+void FIKRetargetBlendToSourceOp::OnTargetChainRenamed(const FName InOldChainName, const FName InNewChainName)
+{
+	IKRetargetOpUtils::OnRetargetChainRenamed(Settings.Chains, InOldChainName, InNewChainName);
+}
+
+void FIKRetargetBlendToSourceOp::OnParentReinitPropertyEdited(
+	const FIKRetargetOpBase& InParentOp,
+	const FPropertyChangedEvent* InPropertyChangedEvent)
+{
+	RegenerateChainSettings(&InParentOp);	
+}
+
+void FIKRetargetBlendToSourceOp::RegenerateChainSettings(const FIKRetargetOpBase* InParentOp)
+{
+	// regenerate chain settings
+	constexpr bool bSkipUnmappedChains = false;
+	constexpr bool bSkipNonIKChains = true;
+	IKRetargetOpUtils::SynchronizeChainSettingsWithIKRig(Settings.Chains, InParentOp, bSkipUnmappedChains, bSkipNonIKChains);
+}
+
+FIKRetargetBlendToSourceOpSettings UIKRetargetBlendToSourceController::GetSettings()
+{
+	return *reinterpret_cast<FIKRetargetBlendToSourceOpSettings*>(OpSettingsToControl);
+}
+
+void UIKRetargetBlendToSourceController::SetSettings(FIKRetargetBlendToSourceOpSettings InSettings)
+{
+	reinterpret_cast<FIKRetargetBlendToSourceOpSettings*>(OpSettingsToControl)->CopySettingsAtRuntime(&InSettings);
+}
+
+#if WITH_EDITOR
+void FIKRetargetBlendToSourceOp::DebugDraw(
+	FPrimitiveDrawInterface* InPDI,
+	const FTransform& InSourceTransform,
+	const FTransform& InComponentTransform,
+	const double InComponentScale,
+	const FIKRetargetDebugDrawState& InEditorState) const
+{
+	for (const FIKChainToBlend& Chain : ChainsToBlend)
+	{
+		Chain.DebugDraw(InPDI, InComponentTransform, Settings.DebugDrawSize);
+	}
+}
+
+void FIKRetargetBlendToSourceOp::ResetChainSettingsToDefault(const FName& InChainName)
+{
+	IKRetargetOpUtils::ResetChainSettingsToDefault(Settings.Chains, InChainName);
+}
+
+bool FIKRetargetBlendToSourceOp::AreChainSettingsAtDefault(const FName& InChainName)
+{
+	return IKRetargetOpUtils::AreChainSettingsAtDefault(Settings.Chains, InChainName);
+}
+
+uint8* FIKRetargetBlendToSourceOp::GetChainSettingsMemory(const FName InChainName)
+{
+	return IKRetargetOpUtils::GetChainSettingsMemory(Settings.Chains, InChainName);
+}
+#endif
diff --git a/Engine/Plugins/Animation/IKRig/Source/IKRig/Private/Retargeter/RetargetOps/IKChainsOp.cpp b/Engine/Plugins/Animation/IKRig/Source/IKRig/Private/Retargeter/RetargetOps/IKChainsOp.cpp
--- a/Engine/Plugins/Animation/IKRig/Source/IKRig/Private/Retargeter/RetargetOps/IKChainsOp.cpp
+++ b/Engine/Plugins/Animation/IKRig/Source/IKRig/Private/Retargeter/RetargetOps/IKChainsOp.cpp
@@ -244,88 +244,38 @@ bool FIKRetargetIKChainsOp::Initialize(
 	bIsInitialized = false;
 	IKChainRetargeters.Reset();
 
-	// this op requires a parent to supply an IK Rig
-	if (!ensure(InParentOp))
-	{
-		return false;
-	}
-
-	// validate that an IK rig has been assigned
-	const FIKRetargetRunIKRigOp* ParentOp = reinterpret_cast<const FIKRetargetRunIKRigOp*>(InParentOp);
-	if (ParentOp->Settings.IKRigAsset == nullptr)
-	{
-		InLog.LogWarning( FText::Format(
-		LOCTEXT("MissingIKRig", "{0}, is missing an IK rig. No chains can be retargeted."), FText::FromName(GetName())));
-		return false;
-	}
-
 	// go through all chains to retarget and load them
 	const FRetargeterBoneChains& BoneChains = InProcessor.GetBoneChains();
+	const FIKRigGoalContainer& GoalContainer = InProcessor.GetIKRigGoalContainer();
 	for (FRetargetIKChainSettings& ChainSettings : Settings.ChainsToRetarget)
 	{
 		FName TargetChainName = ChainSettings.TargetChainName;
-		const FResolvedBoneChain* TargetBoneChain = BoneChains.GetResolvedBoneChainByName(
+		const FResolvedBoneChain* OutSourceBoneChain = nullptr;
+		const FResolvedBoneChain* OutTargetBoneChain = nullptr;
+		if (!FIKRetargetRunIKRigOp::ValidateIKChain(
 			TargetChainName,
-			ERetargetSourceOrTarget::Target,
-			ParentOp->Settings.IKRigAsset);
-
-		// validate that the chain even exists
-		if (TargetBoneChain == nullptr)
+			GoalContainer,
+			BoneChains,
+			this,
+			InParentOp,
+			InLog,
+			OutSourceBoneChain,
+			OutTargetBoneChain))
 		{
-			InLog.LogWarning( FText::Format(
-			LOCTEXT("IKChainOpMissingChain", "IK Chain Op: chain data is out of sync with IK Rig. Missing target chain, '{0}."),
-			FText::FromName(TargetChainName)));
 			continue;
 		}
 
-		// validate that the chain has IK applied to it
-		if (TargetBoneChain->IKGoalName == NAME_None)
-		{
-			InLog.LogWarning( FText::Format(
-			LOCTEXT("IKChainOpChainHasNoIK", "IK Chain Op: an IK chain was found with no IK goal assigned to it, '{0}."),
-			FText::FromName(TargetChainName)));
-			continue;
-		}
-		
-		// which source chain was this target chain mapped to?
-		const FName SourceChainName = ParentOp->ChainMapping.GetChainMappedTo(TargetChainName, ERetargetSourceOrTarget::Target);
-		const FResolvedBoneChain* SourceBoneChain = BoneChains.GetResolvedBoneChainByName(SourceChainName, ERetargetSourceOrTarget::Source);
-		if (SourceBoneChain == nullptr)
-		{
-			InLog.LogWarning( FText::Format(
-			LOCTEXT("IKChainOpChainNotMapped", "IK Chain Op: found IK chain that was not mapped to a source chain, '{0}."),
-			FText::FromName(TargetChainName)));
-			continue;
-		}
-		
 		// initialize the mapped pair of source/target bone chains
 		FIKChainRetargeter IKChainRetargeter;
-		const bool bChainInitialized = IKChainRetargeter.Initialize(*SourceBoneChain, *TargetBoneChain, ChainSettings, InLog);
+		const bool bChainInitialized = IKChainRetargeter.Initialize(*OutSourceBoneChain, *OutTargetBoneChain, ChainSettings, InLog);
 		if (!bChainInitialized)
 		{
 			InLog.LogWarning( FText::Format(
-			LOCTEXT("IKChainOpBadChain", "IK Chain Op: could not initialize a mapped retarget chain for IK, '{0}."), FText::FromName(TargetChainName)));
+			LOCTEXT("IKChainOpBadChain", "IK Chain Op: could not initialize a mapped retarget chain for IK, '{0}."),
+			FText::FromName(TargetChainName)));
 			continue;
 		}
-
-		// warn user if IK goal is not on the END bone of the target chain. It will still work, but may produce bad results.
-		const TArray<UIKRigEffectorGoal*>& AllGoals = ParentOp->Settings.IKRigAsset->GetGoalArray();
-		for (const UIKRigEffectorGoal* Goal : AllGoals)
-		{
-			if (Goal->GoalName != TargetBoneChain->IKGoalName)
-			{
-				continue;
-			}
-			
-			if (Goal->BoneName != TargetBoneChain->EndBone)
-			{
-				InLog.LogWarning( FText::Format(
-			LOCTEXT("TargetIKNotOnEndBone", "IK Chain Op: Retarget chain, '{0}' has an IK goal that is not on the End Bone of the chain."),
-				FText::FromString(TargetBoneChain->ChainName.ToString())));
-			}
-			break;
-		}
-
+		
 		// store valid chain pair to be retargeted
 		IKChainRetargeters.Add(IKChainRetargeter);
 	}
diff --git a/Engine/Plugins/Animation/IKRig/Source/IKRig/Private/Retargeter/RetargetOps/OffsetGoalsOp.cpp b/Engine/Plugins/Animation/IKRig/Source/IKRig/Private/Retargeter/RetargetOps/OffsetGoalsOp.cpp
new file mode 100644
--- /dev/null
+++ b/Engine/Plugins/Animation/IKRig/Source/IKRig/Private/Retargeter/RetargetOps/OffsetGoalsOp.cpp
@@ -0,0 +1,213 @@
+﻿// Copyright Epic Games, Inc. All Rights Reserved.
+
+#include "Retargeter/RetargetOps/OffsetGoalsOp.h"
+
+#include "Retargeter/IKRetargetOpUtils.h"
+#include "Retargeter/IKRetargetProcessor.h"
+#include "Retargeter/RetargetOps/RunIKRigOp.h"
+
+bool FIKRetargetOffsetGoalsChainSettings::operator==(const FIKRetargetOffsetGoalsChainSettings& Other) const
+{
+	return GlobalTranslationOffset.Equals(Other.GlobalTranslationOffset)
+		&& LocalTranslationOffset.Equals(Other.LocalTranslationOffset)
+		&& GlobalRotationOffset.Equals(Other.GlobalRotationOffset)
+		&& LocalRotationOffset.Equals(Other.LocalRotationOffset);
+}
+
+const UClass* FIKRetargetOffsetGoalsOpSettings::GetControllerType() const
+{
+	return UIKRetargetOffsetGoalsController::StaticClass();
+}
+
+void FIKRetargetOffsetGoalsOpSettings::CopySettingsAtRuntime(const FIKRetargetOpSettingsBase* InSettingsToCopyFrom)
+{
+	// copy everything except the chains
+	const TArray<FName> PropertiesToIgnore = {GET_MEMBER_NAME_CHECKED(FIKRetargetOffsetGoalsOpSettings, Chains)};
+	FIKRetargetOpBase::CopyStructProperties(
+		FIKRetargetOffsetGoalsOpSettings::StaticStruct(),
+		InSettingsToCopyFrom,
+		this,
+		PropertiesToIgnore);
+
+	// copy chains
+	const FIKRetargetOffsetGoalsOpSettings* NewSettings = reinterpret_cast<const FIKRetargetOffsetGoalsOpSettings*>(InSettingsToCopyFrom);
+	IKRetargetOpUtils::CopyChainSettingsAtRuntime(NewSettings->Chains, Chains);
+}
+
+void FIKChainToOffset::Run(FIKRigGoalContainer& InGoalContainer, const FIKRetargetOffsetGoalsOpSettings& InOpSettings) const
+{
+	FIKRigGoal* Goal = InGoalContainer.FindGoalByName(TargetBoneChain->IKGoalName);
+	if (!ensure(Goal))
+	{
+		return; // goal should have been validated inside Op::Initialize()
+	}
+
+	if (Settings->Alpha < UE_KINDA_SMALL_NUMBER)
+	{
+		return; // blended off
+	}
+
+	// modify rotation of goal
+	{
+		FQuat GoalRotation = Goal->Rotation.Quaternion();
+		// apply global rotation offset
+		GoalRotation = Settings->GlobalRotationOffset.Quaternion() * GoalRotation;
+		// apply local rotation offset
+		GoalRotation = GoalRotation * Settings->LocalRotationOffset.Quaternion();
+		// alpha and apply
+		const float Alpha = Settings->Alpha * InOpSettings.Alpha * InOpSettings.RotationAlpha;
+		GoalRotation = FQuat::FastLerp(Goal->Rotation.Quaternion(), GoalRotation, Alpha);
+		GoalRotation.Normalize();
+		Goal->Rotation = GoalRotation.Rotator();
+	}
+	
+	// modify position of goal
+	{
+		FVector GoalPosition = Goal->Position;
+		// apply global static offset
+		GoalPosition += Settings->GlobalTranslationOffset;
+		// apply local static offset
+		GoalPosition += Goal->Rotation.Quaternion().RotateVector(Settings->LocalTranslationOffset);
+		// alpha and apply
+		const float Alpha = Settings->Alpha * InOpSettings.Alpha * InOpSettings.TranslationAlpha;
+		GoalPosition = FMath::Lerp(Goal->Position, GoalPosition, Alpha);
+		Goal->Position = GoalPosition;
+	}
+}
+
+bool FIKRetargetOffsetGoalsOp::Initialize(
+	const FIKRetargetProcessor& InProcessor,
+	const FRetargetSkeleton& InSourceSkeleton,
+	const FTargetSkeleton& InTargetSkeleton,
+	const FIKRetargetOpBase* InParentOp,
+	FIKRigLogger& InLog)
+{
+	bIsInitialized = false;
+	
+	// go through all chains to retarget and load them
+	ChainsToOffset.Reset();
+	const FRetargeterBoneChains& BoneChains = InProcessor.GetBoneChains();
+	const FIKRigGoalContainer& GoalContainer = InProcessor.GetIKRigGoalContainer();
+	for (FIKRetargetOffsetGoalsChainSettings& ChainSettings : Settings.Chains)
+	{
+		FName TargetChainName = ChainSettings.TargetChainName;
+		const FResolvedBoneChain* SourceBoneChain = nullptr;
+		const FResolvedBoneChain* TargetBoneChain = nullptr;
+		if (!FIKRetargetRunIKRigOp::ValidateIKChain(
+			TargetChainName,
+			GoalContainer,
+			BoneChains,
+			this,
+			InParentOp,
+			InLog,
+			SourceBoneChain /*out*/,
+			TargetBoneChain /*out*/))
+		{
+			continue;
+		}
+
+		ChainsToOffset.Emplace(*TargetBoneChain, ChainSettings);
+	}
+
+	bIsInitialized = true;
+	return true;
+}
+
+void FIKRetargetOffsetGoalsOp::Run(
+	FIKRetargetProcessor& InProcessor,
+	const double InDeltaTime,
+	const TArray<FTransform>& InSourceGlobalPose,
+	TArray<FTransform>& OutTargetGlobalPose)
+{
+	if (Settings.Alpha <= UE_KINDA_SMALL_NUMBER)
+	{
+		return; // fast skip whole op
+	}
+	
+	FIKRigGoalContainer& GoalContainer = InProcessor.GetIKRigGoalContainer();
+	for (const FIKChainToOffset& Chain : ChainsToOffset)
+	{
+		Chain.Run(GoalContainer, Settings);
+	}
+}
+
+FIKRetargetOpSettingsBase* FIKRetargetOffsetGoalsOp::GetSettings()
+{
+	return &Settings;
+}
+
+const UScriptStruct* FIKRetargetOffsetGoalsOp::GetSettingsType() const
+{
+	return FIKRetargetOffsetGoalsOpSettings::StaticStruct();
+}
+
+const UScriptStruct* FIKRetargetOffsetGoalsOp::GetType() const
+{
+	return FIKRetargetOffsetGoalsOp::StaticStruct();
+}
+
+const UScriptStruct* FIKRetargetOffsetGoalsOp::GetParentOpType() const
+{
+	return FIKRetargetRunIKRigOp::StaticStruct();
+}
+
+void FIKRetargetOffsetGoalsOp::OnAddedToStack(const UIKRetargeter* InRetargetAsset, const FIKRetargetOpBase* InParentOp)
+{
+	RegenerateChainSettings(InParentOp);
+}
+
+void FIKRetargetOffsetGoalsOp::OnAssignIKRig(
+	const ERetargetSourceOrTarget SourceOrTarget,
+	const UIKRigDefinition* InIKRig,
+	const FIKRetargetOpBase* InParentOp)
+{
+	RegenerateChainSettings(InParentOp);
+}
+
+void FIKRetargetOffsetGoalsOp::OnTargetChainRenamed(const FName InOldChainName, const FName InNewChainName)
+{
+	IKRetargetOpUtils::OnRetargetChainRenamed(Settings.Chains, InOldChainName, InNewChainName);
+}
+
+void FIKRetargetOffsetGoalsOp::OnParentReinitPropertyEdited(
+	const FIKRetargetOpBase& InParentOp,
+	const FPropertyChangedEvent* InPropertyChangedEvent)
+{
+	RegenerateChainSettings(&InParentOp);	
+}
+
+void FIKRetargetOffsetGoalsOp::RegenerateChainSettings(const FIKRetargetOpBase* InParentOp)
+{
+	// regenerate chain settings
+	constexpr bool bSkipUnmappedChains = false;
+	constexpr bool bSkipNonIKChains = true;
+	IKRetargetOpUtils::SynchronizeChainSettingsWithIKRig(Settings.Chains, InParentOp, bSkipUnmappedChains, bSkipNonIKChains);
+}
+
+FIKRetargetOffsetGoalsOpSettings UIKRetargetOffsetGoalsController::GetSettings()
+{
+	return *reinterpret_cast<FIKRetargetOffsetGoalsOpSettings*>(OpSettingsToControl);
+}
+
+void UIKRetargetOffsetGoalsController::SetSettings(FIKRetargetOffsetGoalsOpSettings InSettings)
+{
+	reinterpret_cast<FIKRetargetOffsetGoalsOpSettings*>(OpSettingsToControl)->CopySettingsAtRuntime(&InSettings);
+}
+
+#if WITH_EDITOR
+
+void FIKRetargetOffsetGoalsOp::ResetChainSettingsToDefault(const FName& InChainName)
+{
+	IKRetargetOpUtils::ResetChainSettingsToDefault(Settings.Chains, InChainName);
+}
+
+bool FIKRetargetOffsetGoalsOp::AreChainSettingsAtDefault(const FName& InChainName)
+{
+	return IKRetargetOpUtils::AreChainSettingsAtDefault(Settings.Chains, InChainName);
+}
+
+uint8* FIKRetargetOffsetGoalsOp::GetChainSettingsMemory(const FName InChainName)
+{
+	return IKRetargetOpUtils::GetChainSettingsMemory(Settings.Chains, InChainName);
+}
+#endif
diff --git a/Engine/Plugins/Animation/IKRig/Source/IKRig/Private/Retargeter/RetargetOps/RunIKRigOp.cpp b/Engine/Plugins/Animation/IKRig/Source/IKRig/Private/Retargeter/RetargetOps/RunIKRigOp.cpp
--- a/Engine/Plugins/Animation/IKRig/Source/IKRig/Private/Retargeter/RetargetOps/RunIKRigOp.cpp
+++ b/Engine/Plugins/Animation/IKRig/Source/IKRig/Private/Retargeter/RetargetOps/RunIKRigOp.cpp
@@ -5,11 +5,23 @@
 #include "IKRigDebugRendering.h"
 #include "Retargeter/IKRetargetProcessor.h"
 #include "Engine/SkeletalMesh.h"
+#include "Retargeter/IKRetargetOpUtils.h"
+#include "Retargeter/RetargetOps/BlendToSourceOp.h"
+#include "Retargeter/RetargetOps/PelvisMotionOp.h"
 
 #include UE_INLINE_GENERATED_CPP_BY_NAME(RunIKRigOp)
 
 #define LOCTEXT_NAMESPACE "RunIKRigSolversOp"
 
+struct FIKRigGoal;
+
+bool FIKRetargetRunIKRigChainSettings::operator==(const FIKRetargetRunIKRigChainSettings& Other) const
+{
+	return bEnableIK == Other.bEnableIK
+		&& FMath::IsNearlyEqualByULP(PositionAlpha, Other.PositionAlpha)
+		&& FMath::IsNearlyEqualByULP(RotationAlpha, Other.RotationAlpha);
+}
+
 const UClass* FIKRetargetRunIKRigOpSettings::GetControllerType() const
 {
 	return UIKRetargetRunIKRigController::StaticClass();
@@ -26,6 +38,120 @@ void FIKRetargetRunIKRigOpSettings::CopySettingsAtRuntime(const FIKRetargetOpSet
 		PropertiesToIgnore);
 }
 
+FRunIKRigChain::FRunIKRigChain(
+	const FResolvedBoneChain& InSourceBoneChain,
+	const FResolvedBoneChain& InTargetBoneChain,
+	const FIKRetargetRunIKRigChainSettings& InSettings,
+	FIKRigLogger& InLog)
+{
+	bIsReady = false;
+	
+	if (InSourceBoneChain.BoneIndices.Num() < 2)
+	{
+		InLog.LogWarning(FText::Format(
+			LOCTEXT("SourceChainLessThanTwo","Run IK Rig Op: found source bone chain, {0} with less than 2 joints. Cannot retarget goal."),
+			FText::FromName(InSourceBoneChain.ChainName)));
+		return;
+	}
+
+	// store pointers to required data
+	SourceBoneChain = &InSourceBoneChain;
+	TargetBoneChain = &InTargetBoneChain;
+	Settings = &InSettings;
+	
+	// cache initial length of SOURCE chain
+	const FTransform& InitialSourceStartGlobal = InSourceBoneChain.RefPoseGlobalTransforms[0];
+	const FTransform& InitialSourceEndGlobal = InSourceBoneChain.RefPoseGlobalTransforms.Last();
+	const float InitialStartLength = static_cast<float>((InitialSourceStartGlobal.GetLocation() - InitialSourceEndGlobal.GetLocation()).Size());
+	if (InitialStartLength <= UE_KINDA_SMALL_NUMBER)
+	{
+		InLog.LogWarning(FText::Format(
+			LOCTEXT("SourceZeroLengthIK", "Run IK Rig Op: found source bone chain {0} with IK, but it is zero length. Cannot retarget goal."),
+			FText::FromName(InSourceBoneChain.ChainName)));
+		return;
+	}
+	InvInitialLengthSource = 1.0f / InitialStartLength;
+
+	// cache initial length of TARGET chain
+	const FTransform& InitialTargetStartGlobal = InTargetBoneChain.RefPoseGlobalTransforms[0];
+	const FTransform& InitialTargetEndGlobal = InTargetBoneChain.RefPoseGlobalTransforms.Last();
+	InitialTargetChainLength = static_cast<float>((InitialTargetStartGlobal.GetTranslation() - InitialTargetEndGlobal.GetTranslation()).Size());
+	if (InitialTargetChainLength <= UE_KINDA_SMALL_NUMBER)
+	{
+		InLog.LogWarning(FText::Format(
+			LOCTEXT("TargetZeroLengthIK", "Run IK Rig Op: found target bone chain {0} with IK, but it is zero length. Cannot retarget goal."),
+			FText::FromName(InTargetBoneChain.ChainName)));
+		return;
+	}
+
+	bIsReady = true;
+}
+
+void FRunIKRigChain::Run(
+	FIKRigGoalContainer& InGoalContainer,
+	const TArray<FTransform>& InSourceGlobalPose,
+	const TArray<FTransform>& InTargetGlobalPose,
+	const FIKRetargetRunIKRigOpSettings& InOpSettings,
+	const FIKRetargetPelvisMotionOp* InPelvisMotionOp) const
+{
+	if (!bIsReady)
+	{
+		return;
+	}
+	
+	FIKRigGoal* Goal = InGoalContainer.FindGoalByName(TargetBoneChain->IKGoalName);
+	if (!ensure(Goal))
+	{
+		return;
+	}
+
+	if (Settings->PositionAlpha > UE_KINDA_SMALL_NUMBER)
+	{
+		// get relevant bone transforms
+		const int32 SourceStartBoneIndex = SourceBoneChain->BoneIndices[0];
+		const int32 SourceEndBoneIndex = SourceBoneChain->BoneIndices.Last();
+		const FTransform& CurrentSourceStartGlobal = InSourceGlobalPose[SourceStartBoneIndex];
+		const FTransform& CurrentSourceEndGlobal = InSourceGlobalPose[SourceEndBoneIndex];
+
+		// get the current normalized direction / length of the IK limb (how extended it is as percentage of original length)
+		const FVector CurrentSourceChainVector = CurrentSourceEndGlobal.GetLocation() - CurrentSourceStartGlobal.GetLocation();
+		double CurrentSourceChainLength;
+		FVector CurrentSourceChainDirection;
+		CurrentSourceChainVector.ToDirectionAndLength(CurrentSourceChainDirection, CurrentSourceChainLength);
+		const double NormalizedLimbLength = CurrentSourceChainLength * InvInitialLengthSource;
+		const FVector CurrentSourceEndDirectionNormalized = CurrentSourceChainDirection * NormalizedLimbLength;
+	
+		// set position to length-scaled direction from source limb
+		const FVector PelvisTranslationDelta = InPelvisMotionOp ? InPelvisMotionOp->GetPelvisTranslationOffset() : FVector::ZeroVector;
+		const FVector AffectIKWeights = InPelvisMotionOp ? InPelvisMotionOp->GetAffectIKWeightAsVector() : FVector::ZeroVector;
+		const FVector InvAffectIKWeights = FVector::OneVector - AffectIKWeights;
+		const FVector InvRootModification = PelvisTranslationDelta * InvAffectIKWeights;
+		const int32 TargetStartBoneIndex = TargetBoneChain->BoneIndices[0];
+		const FVector CurrentTargetStartLocation = InTargetGlobalPose[TargetStartBoneIndex].GetTranslation() - InvRootModification;
+		FVector GoalPosition = CurrentTargetStartLocation + (CurrentSourceEndDirectionNormalized * InitialTargetChainLength);
+		
+		// alpha and apply to the goal
+		Goal->Position = FMath::Lerp(Goal->Position, GoalPosition, Settings->PositionAlpha);
+	}
+
+	if (Settings->RotationAlpha > UE_KINDA_SMALL_NUMBER)
+	{
+		// calc delta rotation of source relative to source retarget pose
+		const int32 SourceEndBoneIndex = SourceBoneChain->BoneIndices.Last();
+		const FTransform& CurrentSourceEndGlobal = InSourceGlobalPose[SourceEndBoneIndex];
+		const FTransform& InitialSourceEndGlobal = SourceBoneChain->RefPoseGlobalTransforms.Last();
+		const FTransform& InitialTargetEndGlobal = TargetBoneChain->RefPoseGlobalTransforms.Last();
+		const FQuat CurrentSourceEndRotation = CurrentSourceEndGlobal.GetRotation();
+		const FQuat SourceDeltaRotation = (CurrentSourceEndRotation * InitialSourceEndGlobal.GetRotation().Inverse()).GetNormalized();
+
+		// apply delta rotation to the target
+		const FQuat NewGoalRotation = SourceDeltaRotation * InitialTargetEndGlobal.GetRotation();
+		
+		// alpha blend rotation and apply to goal
+		Goal->Rotation = FQuat::FastLerp(Goal->Rotation.Quaternion(), NewGoalRotation, Settings->RotationAlpha).GetNormalized().Rotator();
+	}
+}
+
 bool FIKRetargetRunIKRigOp::Initialize(
 	const FIKRetargetProcessor& InProcessor,
 	const FRetargetSkeleton& InSourceSkeleton,
@@ -63,6 +189,62 @@ bool FIKRetargetRunIKRigOp::Initialize(
 			GoalBoneIndicesMap.Add(Goal->BoneName, BoneIndex);
 		}
 	}
+
+	// warn user if IK goal is not on the END bone of the target chain.
+	// NOTE: this will still work but could produce bad results so we throw a warning instead of an error
+	const TArray<UIKRigEffectorGoal*>& AllGoals = Settings.IKRigAsset->GetGoalArray();
+	const TMap<const UIKRigDefinition*, TArray<FResolvedBoneChain>>& AllTargetBoneChains = InProcessor.GetBoneChains().GetAllResolvedTargetBoneChains();
+	if (const TArray<FResolvedBoneChain>* ChainsInRig = AllTargetBoneChains.Find(Settings.IKRigAsset))
+	{
+		for (const FResolvedBoneChain& Chain : *ChainsInRig)
+		{
+			if (Chain.IKGoalName == NAME_None)
+			{
+				continue;
+			}
+			
+			for (UIKRigEffectorGoal* Goal : AllGoals)
+			{
+				if (Goal->GoalName != Chain.IKGoalName)
+				{
+					continue;
+				}
+
+				if (Goal->BoneName != Chain.EndBone)
+				{
+					InLog.LogWarning( FText::Format(
+						LOCTEXT("TargetIKNotOnEndBone", "IK chain, '{0}' has an IK goal that is not on the End Bone of the chain."),
+						FText::FromName(Chain.ChainName)));
+					break;
+				}
+			}
+		}
+	}
+
+	// go through all chains to retarget and load them
+	Chains.Reset();
+	const FRetargeterBoneChains& BoneChains = InProcessor.GetBoneChains();
+	const FIKRigGoalContainer& GoalContainer = InProcessor.GetIKRigGoalContainer();
+	for (FIKRetargetRunIKRigChainSettings& ChainSettings : Settings.Chains)
+	{
+		FName TargetChainName = ChainSettings.TargetChainName;
+		const FResolvedBoneChain* SourceBoneChain = nullptr;
+		const FResolvedBoneChain* TargetBoneChain = nullptr;
+		if (!FIKRetargetRunIKRigOp::ValidateIKChain(
+			TargetChainName,
+			GoalContainer,
+			BoneChains,
+			this,
+			this,
+			InLog,
+			SourceBoneChain /*out*/,
+			TargetBoneChain /*out*/))
+		{
+			continue;
+		}
+
+		Chains.Emplace(*SourceBoneChain, *TargetBoneChain, ChainSettings, InLog);
+	}
 	
 	bIsInitialized = IKRigProcessor.IsInitialized();
 	return bIsInitialized;
@@ -88,13 +270,15 @@ void FIKRetargetRunIKRigOp::Run(
 	}
 #endif
 
+	FIKRigGoalContainer& GoalContainer = InProcessor.GetIKRigGoalContainer();
+
 	// set the goals
-	IKRigProcessor.ApplyGoalsFromOtherContainer(InProcessor.GetIKRigGoalContainer());
+	IKRigProcessor.ApplyGoalsFromOtherContainer(GoalContainer);
 	// trigger reinitialization if the goal container was modified such that it needs it
-	if (IKRigProcessor.GetGoalContainer().NeedsInitialized())
+	if (GoalContainer.NeedsInitialized())
 	{
 		const USkeletalMesh* TargetSkeletalMesh = InProcessor.GetSkeleton(ERetargetSourceOrTarget::Target).SkeletalMesh;
-		IKRigProcessor.Initialize(Settings.IKRigAsset, TargetSkeletalMesh, InProcessor.GetIKRigGoalContainer());
+		IKRigProcessor.Initialize(Settings.IKRigAsset, TargetSkeletalMesh, GoalContainer);
 	}
 	
 #if WITH_EDITOR
@@ -137,6 +321,13 @@ void FIKRetargetRunIKRigOp::RunBeforeChildren(
 	// reset the goal container for this IK Rig
 	FIKRigGoalContainer& GoalContainer = InProcessor.GetIKRigGoalContainer();
 	ResetGoalContainer(OutTargetGlobalPose, GoalContainer);
+	
+	// do basic chain-based goal retargeting
+	const FIKRetargetPelvisMotionOp* PelvisMotionOp = InProcessor.GetFirstRetargetOpOfType<FIKRetargetPelvisMotionOp>();
+	for (const FRunIKRigChain& Chain : Chains)
+	{
+		Chain.Run(GoalContainer, InSourceGlobalPose, OutTargetGlobalPose, Settings, PelvisMotionOp);
+	}
 }
 
 void FIKRetargetRunIKRigOp::OnAddedToStack(const UIKRetargeter* InRetargetAsset, const FIKRetargetOpBase* InParentOp)
@@ -148,6 +339,9 @@ void FIKRetargetRunIKRigOp::OnAddedToStack(const UIKRetargeter* InRetargetAsset,
 	ChainMapping.ReinitializeWithIKRigs(InRetargetAsset->GetIKRig(ERetargetSourceOrTarget::Source), Settings.IKRigAsset);
 	static bool bForceRemap = true;
 	ChainMapping.AutoMapChains(EAutoMapChainType::Fuzzy, bForceRemap);
+
+	// regenerate set of chain settings
+	RegenerateChainSettings();
 }
 
 void FIKRetargetRunIKRigOp::OnAssignIKRig(
@@ -162,6 +356,9 @@ void FIKRetargetRunIKRigOp::OnAssignIKRig(
 	ChainMapping.ReinitializeWithIKRigs(ChainMapping.GetIKRig(ERetargetSourceOrTarget::Source), Settings.IKRigAsset);
 	static bool bForceRemap = false; // don't force mapping, keeps existing mappings
 	ChainMapping.AutoMapChains(EAutoMapChainType::Fuzzy, bForceRemap);
+
+	// regenerate set of chain settings based on new IK Rig
+	RegenerateChainSettings();
 }
 
 FIKRetargetOpSettingsBase* FIKRetargetRunIKRigOp::GetSettings()
@@ -192,8 +389,88 @@ FRetargetChainMapping* FIKRetargetRunIKRigOp::GetChainMapping()
 void FIKRetargetRunIKRigOp::OnReinitPropertyEdited(const FPropertyChangedEvent* InPropertyChangedEvent)
 {
 	ChainMapping.ReinitializeWithIKRigs(ChainMapping.GetIKRig(ERetargetSourceOrTarget::Source), Settings.IKRigAsset);
+
+	// regenerate set of chain settings based on new IK Rig
+	RegenerateChainSettings();
 }
 
+bool FIKRetargetRunIKRigOp::ValidateIKChain(
+	const FName& InChainToValidate,
+	const FIKRigGoalContainer& InGoalContainer,
+	const FRetargeterBoneChains& InBoneChains,
+	const FIKRetargetOpBase* InOp,
+	const FIKRetargetOpBase* InParentOp,
+	FIKRigLogger& InLog,
+	const FResolvedBoneChain*& OutSourceBoneChain,
+	const FResolvedBoneChain*& OutTargetBoneChain)
+{
+	OutSourceBoneChain = nullptr;
+	OutTargetBoneChain = nullptr;
+	
+	// requires a parent Run IK Rig op
+	if (!ensure(InParentOp))
+	{
+		return false;
+	}
+
+	// validate that an IK rig has been assigned
+	const FIKRetargetRunIKRigOp* ParentOp = reinterpret_cast<const FIKRetargetRunIKRigOp*>(InParentOp);
+	if (ParentOp->Settings.IKRigAsset == nullptr)
+	{
+		InLog.LogWarning( FText::Format(
+		LOCTEXT("MissingIKRig", "{0}, is missing an IK rig. No chains can be retargeted."), FText::FromName(InOp->GetName())));
+		return false;
+	}
+	
+	OutTargetBoneChain = InBoneChains.GetResolvedBoneChainByName(
+		InChainToValidate,
+		ERetargetSourceOrTarget::Target,
+		ParentOp->Settings.IKRigAsset);
+
+	// validate that the chain even exists
+	if (OutTargetBoneChain == nullptr)
+	{
+		InLog.LogWarning( FText::Format(
+		LOCTEXT("IKChainOpMissingChain", "IK chain data is out of sync with IK Rig. Missing target chain, '{0}."),
+		FText::FromName(InChainToValidate)));
+		return false;
+	}
+
+	// validate that the chain has IK applied to it
+	if (OutTargetBoneChain->IKGoalName == NAME_None)
+	{
+		InLog.LogWarning( FText::Format(
+		LOCTEXT("IKChainOpChainHasNoIK", "An IK chain was found with no IK goal assigned to it, '{0}."),
+		FText::FromName(InChainToValidate)));
+		return false;
+	}
+	
+	// which source chain was this target chain mapped to?
+	const FName SourceChainName = ParentOp->ChainMapping.GetChainMappedTo(InChainToValidate, ERetargetSourceOrTarget::Target);
+	OutSourceBoneChain = InBoneChains.GetResolvedBoneChainByName(SourceChainName, ERetargetSourceOrTarget::Source);
+	if (OutSourceBoneChain == nullptr)
+	{
+		InLog.LogWarning( FText::Format(
+		LOCTEXT("IKChainOpChainNotMapped", "Found IK chain that was not mapped to a source chain, '{0}."),
+		FText::FromName(InChainToValidate)));
+		return false;
+	}
+
+	// does the goal actually exist?
+	const FIKRigGoal* Goal = InGoalContainer.FindGoalByName(OutTargetBoneChain->IKGoalName);
+	if (!Goal)
+	{
+		InLog.LogWarning( FText::Format(
+		LOCTEXT("IKChainWithMissingGoal", "Found a retarget chain, '{0}' referring to a goal that doesn't exist, '{1}'."),
+		FText::FromName(InChainToValidate),
+		FText::FromName(OutTargetBoneChain->IKGoalName)));
+		return false;
+	}
+
+	// chain is valid!
+	return true;
+};
+
 #if WITH_EDITOR
 
 FCriticalSection FIKRetargetRunIKRigOp::DebugDataMutex;
@@ -257,6 +534,21 @@ void FIKRetargetRunIKRigOp::DebugDraw(
 	}
 }
 
+void FIKRetargetRunIKRigOp::ResetChainSettingsToDefault(const FName& InChainName)
+{
+	IKRetargetOpUtils::ResetChainSettingsToDefault(Settings.Chains, InChainName);
+}
+
+bool FIKRetargetRunIKRigOp::AreChainSettingsAtDefault(const FName& InChainName)
+{
+	return IKRetargetOpUtils::AreChainSettingsAtDefault(Settings.Chains, InChainName);
+}
+
+uint8* FIKRetargetRunIKRigOp::GetChainSettingsMemory(const FName InChainName)
+{
+	return IKRetargetOpUtils::GetChainSettingsMemory(Settings.Chains, InChainName);
+}
+
 void FIKRetargetRunIKRigOp::SaveInitialGoalTransformsIntoDebugData(
 	const FIKRetargetProcessor& InProcessor,
 	const TArray<FTransform>& InTargetGlobalPose)
@@ -364,6 +656,26 @@ void FIKRetargetRunIKRigOp::ResetGoalContainer(
 			Goal.Rotation = FRotator::ZeroRotator;
 		}
 	}
+
+	// optionally enable/disable IK goals completely (removes them from the IK Rig solve)
+	for (const FRunIKRigChain& Chain : Chains)
+	{
+		if (!Chain.TargetBoneChain)
+		{
+			continue;
+		}
+
+		FIKRigGoal* Goal = InOutGoalContainer.FindGoalByName(Chain.TargetBoneChain->IKGoalName);
+		Goal->bEnabled = Chain.Settings->bEnableIK;
+	}
+}
+
+void FIKRetargetRunIKRigOp::RegenerateChainSettings()
+{
+	// regenerate chain settings
+	constexpr bool bSkipUnmappedChains = false;
+	constexpr bool bSkipNonIKChains = true;
+	IKRetargetOpUtils::SynchronizeChainSettingsWithIKRig(Settings.Chains, this, bSkipUnmappedChains, bSkipNonIKChains);
 }
 
 FIKRetargetRunIKRigOpSettings UIKRetargetRunIKRigController::GetSettings()
diff --git a/Engine/Plugins/Animation/IKRig/Source/IKRig/Private/Retargeter/RetargetOps/ScaleGoalsOp.cpp b/Engine/Plugins/Animation/IKRig/Source/IKRig/Private/Retargeter/RetargetOps/ScaleGoalsOp.cpp
new file mode 100644
--- /dev/null
+++ b/Engine/Plugins/Animation/IKRig/Source/IKRig/Private/Retargeter/RetargetOps/ScaleGoalsOp.cpp
@@ -0,0 +1,219 @@
+﻿// Copyright Epic Games, Inc. All Rights Reserved.
+
+#include "Retargeter/RetargetOps/ScaleGoalsOp.h"
+
+#include "Retargeter/IKRetargetOpUtils.h"
+#include "Retargeter/IKRetargetProcessor.h"
+#include "Retargeter/RetargetOps/RunIKRigOp.h"
+
+bool FIKRetargetScaleGoalsChainSettings::operator==(const FIKRetargetScaleGoalsChainSettings& Other) const
+{
+	return FMath::IsNearlyEqualByULP(ScaleVertical,    Other.ScaleVertical)
+		&& FMath::IsNearlyEqualByULP(ScaleHorizontal, Other.ScaleHorizontal)
+		&& FMath::IsNearlyEqualByULP(ScaleAlongChain, Other.ScaleAlongChain);
+}
+
+const UClass* FIKRetargetScaleGoalsOpSettings::GetControllerType() const
+{
+	return UIKRetargetScaleGoalsController::StaticClass();
+}
+
+void FIKRetargetScaleGoalsOpSettings::CopySettingsAtRuntime(const FIKRetargetOpSettingsBase* InSettingsToCopyFrom)
+{
+	// copy everything except the chains
+	const TArray<FName> PropertiesToIgnore = {GET_MEMBER_NAME_CHECKED(FIKRetargetScaleGoalsOpSettings, Chains)};
+	FIKRetargetOpBase::CopyStructProperties(
+		FIKRetargetScaleGoalsOpSettings::StaticStruct(),
+		InSettingsToCopyFrom,
+		this,
+		PropertiesToIgnore);
+
+	// copy chains
+	const FIKRetargetScaleGoalsOpSettings* NewSettings = reinterpret_cast<const FIKRetargetScaleGoalsOpSettings*>(InSettingsToCopyFrom);
+	IKRetargetOpUtils::CopyChainSettingsAtRuntime(NewSettings->Chains, Chains);
+}
+
+void FIKChainToScale::Run(
+	FIKRigGoalContainer& InGoalContainer,
+	const TArray<FTransform>& InTargetGlobalPose,
+	const FIKRetargetScaleGoalsOpSettings& InOpSettings) const
+{
+	FIKRigGoal* Goal = InGoalContainer.FindGoalByName(TargetBoneChain->IKGoalName);
+	if (!ensure(Goal))
+	{
+		return;
+	}
+
+	const bool bScalingGoal = !FMath::IsNearlyEqual(Settings->ScaleVertical, 1.0f) || !FMath::IsNearlyEqual(Settings->ScaleHorizontal, 1.0f);
+	const bool bExtendingGoal = !FMath::IsNearlyEqual(Settings->ScaleAlongChain, 1.0f);
+	if (!bScalingGoal && !bExtendingGoal)
+	{
+		return; // nothing to do
+	}
+
+	// modify position of goal
+	FVector GoalPosition = Goal->Position;
+	
+	// apply horizontal and vertical scaling
+	if (bScalingGoal)
+	{
+		GoalPosition *= FVector(Settings->ScaleHorizontal, Settings->ScaleHorizontal, Settings->ScaleVertical);
+	}
+	
+	// apply extension
+	if (bExtendingGoal)
+	{
+		const int32 TargetStartBoneIndex = TargetBoneChain->BoneIndices[0];
+		const FVector CurrentTargetStartLocation = InTargetGlobalPose[TargetStartBoneIndex].GetTranslation();
+		GoalPosition = CurrentTargetStartLocation + (GoalPosition - CurrentTargetStartLocation) * Settings->ScaleAlongChain;	
+	}
+
+	// alpha and apply
+	GoalPosition = FMath::Lerp(Goal->Position, GoalPosition, InOpSettings.Alpha);
+	Goal->Position = GoalPosition;
+}
+
+bool FIKRetargetScaleGoalsOp::Initialize(
+	const FIKRetargetProcessor& InProcessor,
+	const FRetargetSkeleton& InSourceSkeleton,
+	const FTargetSkeleton& InTargetSkeleton,
+	const FIKRetargetOpBase* InParentOp,
+	FIKRigLogger& InLog)
+{
+	bIsInitialized = false;
+	
+	// go through all chains to retarget and load them
+	ChainsToScale.Reset();
+	const FRetargeterBoneChains& BoneChains = InProcessor.GetBoneChains();
+	const FIKRigGoalContainer& GoalContainer = InProcessor.GetIKRigGoalContainer();
+	for (FIKRetargetScaleGoalsChainSettings& ChainSettings : Settings.Chains)
+	{
+		FName TargetChainName = ChainSettings.TargetChainName;
+		const FResolvedBoneChain* SourceBoneChain = nullptr;
+		const FResolvedBoneChain* TargetBoneChain = nullptr;
+		if (!FIKRetargetRunIKRigOp::ValidateIKChain(
+			TargetChainName,
+			GoalContainer,
+			BoneChains,
+			this,
+			InParentOp,
+			InLog,
+			SourceBoneChain /*out*/,
+			TargetBoneChain /*out*/))
+		{
+			continue;
+		}
+
+		ChainsToScale.Emplace(*TargetBoneChain, ChainSettings);
+	}
+
+	bIsInitialized = true;
+	return true;
+}
+
+void FIKRetargetScaleGoalsOp::Run(
+	FIKRetargetProcessor& InProcessor,
+	const double InDeltaTime,
+	const TArray<FTransform>& InSourceGlobalPose,
+	TArray<FTransform>& OutTargetGlobalPose)
+{
+	if (Settings.Alpha <= UE_KINDA_SMALL_NUMBER)
+	{
+		return; // fast skip whole op
+	}
+	
+	FIKRigGoalContainer& GoalContainer = InProcessor.GetIKRigGoalContainer();
+	for (const FIKChainToScale& Chain : ChainsToScale)
+	{
+		Chain.Run(GoalContainer, OutTargetGlobalPose, Settings);
+	}
+}
+
+FIKRetargetOpSettingsBase* FIKRetargetScaleGoalsOp::GetSettings()
+{
+	return &Settings;
+}
+
+const UScriptStruct* FIKRetargetScaleGoalsOp::GetSettingsType() const
+{
+	return FIKRetargetScaleGoalsOpSettings::StaticStruct();
+}
+
+const UScriptStruct* FIKRetargetScaleGoalsOp::GetType() const
+{
+	return FIKRetargetScaleGoalsOp::StaticStruct();
+}
+
+const UScriptStruct* FIKRetargetScaleGoalsOp::GetParentOpType() const
+{
+	return FIKRetargetRunIKRigOp::StaticStruct();
+}
+
+void FIKRetargetScaleGoalsOp::OnAddedToStack(const UIKRetargeter* InRetargetAsset, const FIKRetargetOpBase* InParentOp)
+{
+	RegenerateChainSettings(InParentOp);
+}
+
+void FIKRetargetScaleGoalsOp::OnAssignIKRig(
+	const ERetargetSourceOrTarget SourceOrTarget,
+	const UIKRigDefinition* InIKRig,
+	const FIKRetargetOpBase* InParentOp)
+{
+	RegenerateChainSettings(InParentOp);
+}
+
+void FIKRetargetScaleGoalsOp::OnTargetChainRenamed(const FName InOldChainName, const FName InNewChainName)
+{
+	IKRetargetOpUtils::OnRetargetChainRenamed(Settings.Chains, InOldChainName, InNewChainName);
+}
+
+void FIKRetargetScaleGoalsOp::OnParentReinitPropertyEdited(
+	const FIKRetargetOpBase& InParentOp,
+	const FPropertyChangedEvent* InPropertyChangedEvent)
+{
+	RegenerateChainSettings(&InParentOp);	
+}
+
+void FIKRetargetScaleGoalsOp::RegenerateChainSettings(const FIKRetargetOpBase* InParentOp)
+{
+	// regenerate chain settings
+	constexpr bool bSkipUnmappedChains = false;
+	constexpr bool bSkipNonIKChains = true;
+	IKRetargetOpUtils::SynchronizeChainSettingsWithIKRig(Settings.Chains, InParentOp, bSkipUnmappedChains, bSkipNonIKChains);
+}
+
+FIKRetargetScaleGoalsOpSettings UIKRetargetScaleGoalsController::GetSettings()
+{
+	return *reinterpret_cast<FIKRetargetScaleGoalsOpSettings*>(OpSettingsToControl);
+}
+
+void UIKRetargetScaleGoalsController::SetSettings(FIKRetargetScaleGoalsOpSettings InSettings)
+{
+	reinterpret_cast<FIKRetargetScaleGoalsOpSettings*>(OpSettingsToControl)->CopySettingsAtRuntime(&InSettings);
+}
+
+#if WITH_EDITOR
+void FIKRetargetScaleGoalsOp::DebugDraw(
+	FPrimitiveDrawInterface* InPDI,
+	const FTransform& InSourceTransform,
+	const FTransform& InComponentTransform,
+	const double InComponentScale,
+	const FIKRetargetDebugDrawState& InEditorState) const
+{
+}
+
+void FIKRetargetScaleGoalsOp::ResetChainSettingsToDefault(const FName& InChainName)
+{
+	IKRetargetOpUtils::ResetChainSettingsToDefault(Settings.Chains, InChainName);
+}
+
+bool FIKRetargetScaleGoalsOp::AreChainSettingsAtDefault(const FName& InChainName)
+{
+	return IKRetargetOpUtils::AreChainSettingsAtDefault(Settings.Chains, InChainName);
+}
+
+uint8* FIKRetargetScaleGoalsOp::GetChainSettingsMemory(const FName InChainName)
+{
+	return IKRetargetOpUtils::GetChainSettingsMemory(Settings.Chains, InChainName);
+}
+#endif
diff --git a/Engine/Plugins/Animation/IKRig/Source/IKRig/Public/IKRigObjectVersion.h b/Engine/Plugins/Animation/IKRig/Source/IKRig/Public/IKRigObjectVersion.h
--- a/Engine/Plugins/Animation/IKRig/Source/IKRig/Public/IKRigObjectVersion.h
+++ b/Engine/Plugins/Animation/IKRig/Source/IKRig/Public/IKRigObjectVersion.h
@@ -50,6 +50,9 @@ struct FIKRigObjectVersion
 		// added local/global offsets to pelvis motion op
 		AddedLocalGlobalOffsetsToPelvisOp,
 
+		// the "Retarget IK Goals" op is deprecated and features are split into separate ops
+		IKChainsOpSplitIntoSeparateOps,
+
 		// -----<new versions can be added above this line>-------------------------------------------------
 		VersionPlusOne,
 		LatestVersion = VersionPlusOne - 1
diff --git a/Engine/Plugins/Animation/IKRig/Source/IKRig/Public/Retargeter/IKRetargetOpUtils.h b/Engine/Plugins/Animation/IKRig/Source/IKRig/Public/Retargeter/IKRetargetOpUtils.h
--- a/Engine/Plugins/Animation/IKRig/Source/IKRig/Public/Retargeter/IKRetargetOpUtils.h
+++ b/Engine/Plugins/Animation/IKRig/Source/IKRig/Public/Retargeter/IKRetargetOpUtils.h
@@ -6,6 +6,9 @@
 #include "Chaos/Character/CharacterGroundConstraintContainer.h"
 #include "Rig/IKRigDefinition.h"
 #include "Retargeter/IKRetargetChainMapping.h"
+#include "RetargetOps/RunIKRigOp.h"
+
+struct FIKRetargetRunIKRigOp;
 
 namespace IKRetargetOpUtils
 {
@@ -73,10 +76,10 @@ namespace IKRetargetOpUtils
 				continue; // skip non-IK chains
 			}
 
-			if (!bSkipUnmappedChains)
+			if (bSkipUnmappedChains && ChainMapping)
 			{
 				const FName SourceChain = ChainMapping->GetChainMappedTo(TargetChain.ChainName, ERetargetSourceOrTarget::Target);
-				if (bSkipUnmappedChains && SourceChain == NAME_None)
+				if (SourceChain == NAME_None)
 				{
 					continue; // skip unmapped chains
 				}
diff --git a/Engine/Plugins/Animation/IKRig/Source/IKRig/Public/Retargeter/IKRetargetOps.h b/Engine/Plugins/Animation/IKRig/Source/IKRig/Public/Retargeter/IKRetargetOps.h
--- a/Engine/Plugins/Animation/IKRig/Source/IKRig/Public/Retargeter/IKRetargetOps.h
+++ b/Engine/Plugins/Animation/IKRig/Source/IKRig/Public/Retargeter/IKRetargetOps.h
@@ -150,6 +150,30 @@ public:
 #endif
 };
 
+#if WITH_EDITOR
+struct FOpExecutionTimer
+{
+	/** reset the running average execution time */
+	void Reset();
+
+	/** reset the per-frame time and update the running average*/
+	void StartNewFrame();
+	
+	/** call this to add some execution time for this frame
+	 * NOTE: this may be called multiple times per frame */
+	void AddTimeThisFrame(const double InSeconds);
+
+	/** returns the current running average execution time using a 30 frame exponential moving average */
+	UE_API double GetRunningAverage() const;
+
+private:
+	/** running averrage of past n frames */
+	double RunningAverageTime = 0.0f;
+	/** profiling data; time in seconds of last execution */
+	double TimeThisFrame = 0.0f;
+};
+#endif
+
 /**
  * This is the base class for defining operations that live in the retargeter "op" stack.
  * These operations are executed in order by calling the virtual Run() function on each one in order.
@@ -372,14 +396,8 @@ public:
 	/** the instance of this op currently running in the editor viewport */
 	mutable FIKRetargetOpBase* EditorInstance;
 
-	/** reset the running average execution time */
-	UE_API void ResetExecutionTime();
-	
-	/** set the latest execution time for profiling */
-	UE_API void SetLatestExecutionTime(const double InSeconds);
-
-	/** returns the current average execution time using a 30 frame exponential moving average */
-	UE_API double GetAverageExecutionTime() const;
+	/** a timer for the op profiler */
+	FOpExecutionTimer Timer;
 #endif
 
 	/** Get the name of this op (may be customized by user) */
@@ -410,13 +428,6 @@ public:
 		const TArray<FName>& InPropertiesToIgnore);
 
 private:
-
-#if WITH_EDITOR
-	/** profiling data; time in seconds of last execution */
-	double AverageExecutionTime = 0.0f;
-
-	
-#endif
 	
 	/** when false, execution of this op is skipped */
 	UPROPERTY()
diff --git a/Engine/Plugins/Animation/IKRig/Source/IKRig/Public/Retargeter/IKRetargetProcessor.h b/Engine/Plugins/Animation/IKRig/Source/IKRig/Public/Retargeter/IKRetargetProcessor.h
--- a/Engine/Plugins/Animation/IKRig/Source/IKRig/Public/Retargeter/IKRetargetProcessor.h
+++ b/Engine/Plugins/Animation/IKRig/Source/IKRig/Public/Retargeter/IKRetargetProcessor.h
@@ -583,6 +583,38 @@ private:
 		TArray<FTransform>& OutTargetGlobalPose,
 		const double InDeltaTime,
 		const int32 InLOD);
+
+	/** Helper function to determine if an op should run this frame */
+	static bool ShouldExecuteOp(const FIKRetargetOpBase& InOp, int32 InLOD);
+
+	struct FOpEntry
+	{
+		FIKRetargetOpBase* Op = nullptr;
+		FOpEntry* Parent = nullptr;
+		TArray<FOpEntry*> Children;
+
+		int32 StackIndex = INDEX_NONE;
+		bool bExecute = false;
+	};
+	
+	enum class ERetargetOpPhase : uint8
+	{
+		ParentBeforeChildren,
+		Main,
+		ChildAfterParent
+	};
+
+	struct FRetargetOpCall
+	{
+		FRetargetOpCall(const ERetargetOpPhase InPhase,	FIKRetargetOpBase* InOp) : Phase(InPhase), Op(InOp) {}
+		ERetargetOpPhase Phase;
+		FIKRetargetOpBase* Op;
+	};
+
+	void BuildExecutionStack(int32 InLOD, TArray<FRetargetOpCall>& OutStack);
+
+	void InitializeOpEntries();
+	TArray<FOpEntry> OpEntries;
 };
 
 //
diff --git a/Engine/Plugins/Animation/IKRig/Source/IKRig/Public/Retargeter/IKRetargeter.h b/Engine/Plugins/Animation/IKRig/Source/IKRig/Public/Retargeter/IKRetargeter.h
--- a/Engine/Plugins/Animation/IKRig/Source/IKRig/Public/Retargeter/IKRetargeter.h
+++ b/Engine/Plugins/Animation/IKRig/Source/IKRig/Public/Retargeter/IKRetargeter.h
@@ -111,6 +111,11 @@ public:
 	}
 	// Get a retarget op by name
 	UE_API const FIKRetargetOpBase* GetRetargetOpByName(const FName& InOpName) const;
+	// Get a retarget op by name non-const
+	// NOTE: be careful with this, general op-editing should be done through a controller
+	UE_API FIKRetargetOpBase* GetRetargetOpByNameNonConst(const FName& InOpName);
+	// Get index of a retarget op
+	UE_API int32 GetIndexOfRetargetOp(const FIKRetargetOpBase* InOp) const;
 	// Get read-only access to a retarget pose 
 	UE_API const FIKRetargetPose* GetCurrentRetargetPose(const ERetargetSourceOrTarget& SourceOrTarget) const;
 	// Get name of the current retarget pose 
@@ -307,6 +312,7 @@ private:
 	UE_API void PostLoadOldOpsToNewStructOps();
 	UE_API void PostLoadConvertEverythingToOps();
 	UE_API void PostLoadPutChainMappingInOps();
+	UE_API void PostLoadSplitIKChainsOp();
 #endif
 	UE_API void CleanRetargetPoses();
 	UE_API void CleanOpStack();
diff --git a/Engine/Plugins/Animation/IKRig/Source/IKRig/Public/Retargeter/RetargetOps/BlendToSourceOp.h b/Engine/Plugins/Animation/IKRig/Source/IKRig/Public/Retargeter/RetargetOps/BlendToSourceOp.h
new file mode 100644
--- /dev/null
+++ b/Engine/Plugins/Animation/IKRig/Source/IKRig/Public/Retargeter/RetargetOps/BlendToSourceOp.h
@@ -0,0 +1,194 @@
+﻿// Copyright Epic Games, Inc. All Rights Reserved.
+#pragma once
+
+
+#include "Retargeter/IKRetargetOps.h"
+#include "BlendToSourceOp.generated.h"
+
+#define UE_API IKRIG_API
+
+#define LOCTEXT_NAMESPACE "BlendToSourceOp"
+
+struct FResolvedBoneChain;
+struct FIKRigGoalContainer;
+struct FIKRetargetPelvisMotionOp;
+
+USTRUCT(BlueprintType)
+struct FIKRetargetBlendToSourceChainSettings
+{
+	GENERATED_BODY()
+
+	FIKRetargetBlendToSourceChainSettings() = default;
+	FIKRetargetBlendToSourceChainSettings(const FName InTargetChainName) : TargetChainName(InTargetChainName) {}
+	bool operator==(const FIKRetargetBlendToSourceChainSettings& Other) const;
+	FName GetName() const { return TargetChainName; };
+	void SetName(const FName InName) { TargetChainName = InName; };
+
+	/** The name of the TARGET chain to blend to source on. */
+	UPROPERTY(EditAnywhere, BlueprintReadWrite, Category = "Chain Settings", meta=(ReinitializeOnEdit))
+	FName TargetChainName;
+	
+	/** Range 0 to 1. Default 0. Blend the goal on this chain to the location of the equivalent goal on the source chain. */
+	UPROPERTY(EditAnywhere, BlueprintReadWrite, Category = "Chain Settings", meta=(ReinitializeOnEdit, ClampMin = "0.0", ClampMax = "1.0", UIMin = "0.0", UIMax = "1.0"))
+	double BlendToSource = 0.0;
+
+	/** Range 0 to 1. Default 1. Blends the translational component of BlendToSource on/off.
+	*  At 0 the goal is placed at the input location.
+	*  At 1 the goal is placed at the location of the source chain's end bone. */
+	UPROPERTY(EditAnywhere, BlueprintReadWrite, Category = "Chain Settings", meta = (ClampMin = "0.0", ClampMax = "1.0", UIMin = "0.0", UIMax = "1.0"))
+	double TranslationAlpha = 1.0f;
+
+	/** Range 0 to 1. Default 0. Blends the rotational component of BlendToSource on/off.
+	*  At 0 the goal is oriented to the input rotation.
+	*  At 1 the goal is oriented to the source chain's end bone rotation. */
+	UPROPERTY(EditAnywhere, BlueprintReadWrite, Category = "Chain Settings", meta = (ClampMin = "0.0", ClampMax = "1.0", UIMin = "0.0", UIMax = "1.0"))
+	double RotationAlpha = 0.0f;
+
+	/** Range 0 to 1. Default 1. Weight each axis separately when using Blend To Source.
+	*  At 0 the goal is placed at the input location.
+	*  At 1 the goal is placed at the location of the source chain's end bone. */
+	UPROPERTY(EditAnywhere, BlueprintReadWrite, Category = "Chain Settings", meta = (ClampMin = "0.0", ClampMax = "1.0", UIMin = "0.0", UIMax = "1.0"))
+	FVector TranslationPerAxisAlpha = FVector::OneVector;
+
+	/** Range 0 to 1. Default 0. At 1, the source goal locations are fully affected by any offsets applied to the pelvis from a Pelvis Motion Op.
+	* NOTE: If no Pelvis Motion Op is present, or if the "Affect IK" weights in the Pelvis Motion Op are zero, then this setting has no effect. */
+	UPROPERTY(EditAnywhere, BlueprintReadWrite, Category = "Chain Settings", meta = (ClampMin = "0.0", ClampMax = "1.0", UIMin = "0.0", UIMax = "1.0"))
+	double ApplyPelvisOffset = 0.0;
+};
+
+USTRUCT(BlueprintType, meta = (DisplayName = "Blend to Source Settings"))
+struct FIKRetargetBlendToSourceOpSettings : public FIKRetargetOpSettingsBase
+{
+	GENERATED_BODY()
+	
+	UPROPERTY(EditAnywhere, BlueprintReadWrite, Category = Chains, meta=(ReinitializeOnEdit))
+	TArray<FIKRetargetBlendToSourceChainSettings> Chains;
+
+	/** Range 0 to 1. Default 1. Blends IK goal transform from retargeted transform (0) to source bone transform (1).
+	*  At 0 the goal is placed at the input location and rotation.
+	*  At 1 the goal is placed at the location and rotation of the source chain's end bone. */
+	UPROPERTY(EditAnywhere, BlueprintReadWrite, Category = "Op Settings", meta = (ClampMin = "0.0", ClampMax = "1.0", UIMin = "0.0", UIMax = "1.0"))
+	double Alpha = 1.0f;
+
+	// Adjust size of goal debug drawing in viewport
+	UPROPERTY(EditAnywhere, Category = Debug)
+	double DebugDrawSize = 5.0;
+
+	virtual const UClass* GetControllerType() const override;
+
+	virtual void CopySettingsAtRuntime(const FIKRetargetOpSettingsBase* InSettingsToCopyFrom) override;
+};
+
+struct FIKChainToBlend
+{
+	FIKChainToBlend() = delete;
+	FIKChainToBlend(
+		const FResolvedBoneChain& InSourceBoneChain,
+		const FResolvedBoneChain& InTargetBoneChain,
+		const FIKRetargetBlendToSourceChainSettings& InSettings)
+		: SourceBoneChain(&InSourceBoneChain), TargetBoneChain(&InTargetBoneChain), Settings(&InSettings){}
+
+	void Run(
+		FIKRigGoalContainer& InGoalContainer,
+		const TArray<FTransform>& InSourceGlobalPose,
+		const FIKRetargetBlendToSourceOpSettings& InOpSettings,
+		const FIKRetargetPelvisMotionOp* InPelvisMotionOp) const;
+
+#if WITH_EDITOR
+	void DebugDraw(FPrimitiveDrawInterface* InPDI, const FTransform& InComponentTransform, const double InDrawSize) const;
+#endif
+	
+
+private:
+	
+	// cached data
+	const FResolvedBoneChain* SourceBoneChain;
+	const FResolvedBoneChain* TargetBoneChain;
+	const FIKRetargetBlendToSourceChainSettings* Settings;
+
+#if WITH_EDITOR
+	static UE_API FCriticalSection DebugDataMutex;
+	mutable FTransform DebugOutputGoal = FTransform::Identity;
+	mutable FTransform DebugSourceGoal = FTransform::Identity;
+#endif
+};
+
+USTRUCT(BlueprintType, meta = (DisplayName = "Blend to Source"))
+struct FIKRetargetBlendToSourceOp : public FIKRetargetOpBase
+{
+	GENERATED_BODY()
+	
+	virtual bool Initialize(
+		const FIKRetargetProcessor& InProcessor,
+		const FRetargetSkeleton& InSourceSkeleton,
+		const FTargetSkeleton& InTargetSkeleton,
+		const FIKRetargetOpBase* InParentOp,
+		FIKRigLogger& InLog) override;
+	
+	virtual void Run(
+		FIKRetargetProcessor& InProcessor,
+		const double InDeltaTime,
+		const TArray<FTransform>& InSourceGlobalPose,
+		TArray<FTransform>& OutTargetGlobalPose) override;
+
+	virtual void OnAddedToStack(const UIKRetargeter* InRetargetAsset, const FIKRetargetOpBase* InParentOp) override;
+
+	virtual void OnAssignIKRig(const ERetargetSourceOrTarget SourceOrTarget, const UIKRigDefinition* InIKRig, const FIKRetargetOpBase* InParentOp) override;
+
+	virtual FIKRetargetOpSettingsBase* GetSettings() override;
+	
+	virtual const UScriptStruct* GetSettingsType() const override;
+
+	virtual const UScriptStruct* GetType() const override;
+
+	virtual const UScriptStruct* GetParentOpType() const override;
+
+	virtual void OnTargetChainRenamed(const FName InOldChainName, const FName InNewChainName) override;
+
+	virtual void OnParentReinitPropertyEdited(const FIKRetargetOpBase& InParentOp, const FPropertyChangedEvent* InPropertyChangedEvent) override;
+	
+	UPROPERTY()
+	FIKRetargetBlendToSourceOpSettings Settings;
+
+#if WITH_EDITOR
+public:
+	virtual void DebugDraw(
+		FPrimitiveDrawInterface* InPDI,
+		const FTransform& InSourceTransform,
+		const FTransform& InComponentTransform,
+		const double InComponentScale,
+		const FIKRetargetDebugDrawState& InEditorState) const override;
+	virtual bool HasDebugDrawing() override { return true; };
+	virtual void ResetChainSettingsToDefault(const FName& InChainName) override;
+	virtual bool AreChainSettingsAtDefault(const FName& InChainName) override;
+	virtual uint8* GetChainSettingsMemory(const FName InChainName) override;
+#endif
+
+	void RegenerateChainSettings(const FIKRetargetOpBase* InParentOp);
+	
+private:
+
+	TArray<FIKChainToBlend> ChainsToBlend;
+	const FIKRetargetPelvisMotionOp* PelvisMotionOp;
+};
+
+/* The blueprint/python API for editing a Blend to Source Op */
+UCLASS(MinimalAPI, BlueprintType)
+class UIKRetargetBlendToSourceController : public UIKRetargetOpControllerBase
+{
+	GENERATED_BODY()
+	
+public:
+	/* Get the current op settings as a struct.
+	 * @return FIKBlendToSourceChainSettings struct with the current settings used by the op. */
+	UFUNCTION(BlueprintCallable, Category = Settings)
+	FIKRetargetBlendToSourceOpSettings GetSettings();
+
+	/* Set the op settings. Input is a custom struct type for this op.
+	 * @param InSettings a FIKBlendToSourceChainSettings struct containing all the settings to apply to this op */
+	UFUNCTION(BlueprintCallable, Category = Settings)
+	void SetSettings(FIKRetargetBlendToSourceOpSettings InSettings);
+};
+
+#undef UE_API
+#undef LOCTEXT_NAMESPACE
diff --git a/Engine/Plugins/Animation/IKRig/Source/IKRig/Public/Retargeter/RetargetOps/IKChainsOp.h b/Engine/Plugins/Animation/IKRig/Source/IKRig/Public/Retargeter/RetargetOps/IKChainsOp.h
--- a/Engine/Plugins/Animation/IKRig/Source/IKRig/Public/Retargeter/RetargetOps/IKChainsOp.h
+++ b/Engine/Plugins/Animation/IKRig/Source/IKRig/Public/Retargeter/RetargetOps/IKChainsOp.h
@@ -196,7 +196,7 @@ struct FChainDebugData
 };
 #endif
 
-USTRUCT(BlueprintType, meta = (DisplayName = "Retarget IK Goals"))
+USTRUCT(BlueprintType, meta = (DisplayName = "Retarget IK Goals", Deprecated))
 struct FIKRetargetIKChainsOp : public FIKRetargetOpBase
 {
 	GENERATED_BODY()
diff --git a/Engine/Plugins/Animation/IKRig/Source/IKRig/Public/Retargeter/RetargetOps/OffsetGoalsOp.h b/Engine/Plugins/Animation/IKRig/Source/IKRig/Public/Retargeter/RetargetOps/OffsetGoalsOp.h
new file mode 100644
--- /dev/null
+++ b/Engine/Plugins/Animation/IKRig/Source/IKRig/Public/Retargeter/RetargetOps/OffsetGoalsOp.h
@@ -0,0 +1,166 @@
+﻿// Copyright Epic Games, Inc. All Rights Reserved.
+#pragma once
+
+
+#include "Retargeter/IKRetargetOps.h"
+#include "OffsetGoalsOp.generated.h"
+
+#define UE_API IKRIG_API
+
+#define LOCTEXT_NAMESPACE "BlendToSourceOp"
+
+struct FResolvedBoneChain;
+struct FIKRigGoalContainer;
+
+USTRUCT(BlueprintType)
+struct FIKRetargetOffsetGoalsChainSettings
+{
+	GENERATED_BODY()
+
+	FIKRetargetOffsetGoalsChainSettings() = default;
+	FIKRetargetOffsetGoalsChainSettings(const FName InTargetChainName) : TargetChainName(InTargetChainName) {}
+	bool operator==(const FIKRetargetOffsetGoalsChainSettings& Other) const;
+	FName GetName() const { return TargetChainName; };
+	void SetName(const FName InName) { TargetChainName = InName; };
+
+	/** The name of the TARGET chain for these settings. */
+	UPROPERTY(EditAnywhere, BlueprintReadWrite, Category = "Setup", meta=(ReinitializeOnEdit))
+	FName TargetChainName;
+
+	/** Range 0 to 1. Default 0. Blend the offsets on this chain on/off. */
+	UPROPERTY(EditAnywhere, BlueprintReadWrite, Category = "Op Settings", meta = (ClampMin = "0.0", ClampMax = "1.0", UIMin = "0.0", UIMax = "1.0"))
+	double Alpha = 0.0f;
+	
+	/** Default 0,0,0. Apply a global-space offset to the IK goal position. */
+	UPROPERTY(EditAnywhere, BlueprintReadWrite, Category="Chain Settings")
+	FVector GlobalTranslationOffset = FVector::ZeroVector;
+
+	/** Default 0,0,0. Apply a local-space offset to the IK goal position. */
+	UPROPERTY(EditAnywhere, BlueprintReadWrite, Category="Chain Settings")
+	FVector LocalTranslationOffset = FVector::ZeroVector;
+
+	/** Default 0,0,0. Apply a global-space offset to the IK goal rotation. */
+	UPROPERTY(EditAnywhere, BlueprintReadWrite, Category="Chain Settings")
+	FRotator GlobalRotationOffset = FRotator::ZeroRotator;
+
+	/** Default 0,0,0. Apply a local-space offset to the IK goal rotation. */
+	UPROPERTY(EditAnywhere, BlueprintReadWrite, Category="Chain Settings")
+	FRotator LocalRotationOffset = FRotator::ZeroRotator;
+};
+
+USTRUCT(BlueprintType, meta = (DisplayName = "Blend to Source Settings"))
+struct FIKRetargetOffsetGoalsOpSettings : public FIKRetargetOpSettingsBase
+{
+	GENERATED_BODY()
+	
+	UPROPERTY(EditAnywhere, BlueprintReadWrite, Category = Chains, meta=(ReinitializeOnEdit))
+	TArray<FIKRetargetOffsetGoalsChainSettings> Chains;
+
+	/** Range 0 to 1. Default 1. Blend all the offsets on/off. */
+	UPROPERTY(EditAnywhere, BlueprintReadWrite, Category = "Op Settings", meta = (ClampMin = "0.0", ClampMax = "1.0", UIMin = "0.0", UIMax = "1.0"))
+	double Alpha = 1.0f;
+
+	/** Range 0 to 1. Default 1. Blend the translation offsets on/off. */
+	UPROPERTY(EditAnywhere, BlueprintReadWrite, Category = "Op Settings", meta = (ClampMin = "0.0", ClampMax = "1.0", UIMin = "0.0", UIMax = "1.0"))
+	double TranslationAlpha = 1.0f;
+
+	/** Range 0 to 1. Default 1. Blend the rotation offsets on/off. */
+	UPROPERTY(EditAnywhere, BlueprintReadWrite, Category = "Op Settings", meta = (ClampMin = "0.0", ClampMax = "1.0", UIMin = "0.0", UIMax = "1.0"))
+	double RotationAlpha = 1.0f;
+
+	virtual const UClass* GetControllerType() const override;
+
+	virtual void CopySettingsAtRuntime(const FIKRetargetOpSettingsBase* InSettingsToCopyFrom) override;
+};
+
+struct FIKChainToOffset
+{
+	FIKChainToOffset() = delete;
+	FIKChainToOffset(
+		const FResolvedBoneChain& InTargetBoneChain,
+		const FIKRetargetOffsetGoalsChainSettings& InSettings)
+		: TargetBoneChain(&InTargetBoneChain)
+		, Settings(&InSettings){}
+
+	void Run(
+		FIKRigGoalContainer& InGoalContainer,
+		const FIKRetargetOffsetGoalsOpSettings& InOpSettings) const;
+
+private:
+	
+	// cached data
+	const FResolvedBoneChain* TargetBoneChain;
+	const FIKRetargetOffsetGoalsChainSettings* Settings;
+};
+
+USTRUCT(BlueprintType, meta = (DisplayName = "Offset Goals"))
+struct FIKRetargetOffsetGoalsOp : public FIKRetargetOpBase
+{
+	GENERATED_BODY()
+	
+	virtual bool Initialize(
+		const FIKRetargetProcessor& InProcessor,
+		const FRetargetSkeleton& InSourceSkeleton,
+		const FTargetSkeleton& InTargetSkeleton,
+		const FIKRetargetOpBase* InParentOp,
+		FIKRigLogger& InLog) override;
+	
+	virtual void Run(
+		FIKRetargetProcessor& InProcessor,
+		const double InDeltaTime,
+		const TArray<FTransform>& InSourceGlobalPose,
+		TArray<FTransform>& OutTargetGlobalPose) override;
+
+	virtual void OnAddedToStack(const UIKRetargeter* InRetargetAsset, const FIKRetargetOpBase* InParentOp) override;
+
+	virtual void OnAssignIKRig(const ERetargetSourceOrTarget SourceOrTarget, const UIKRigDefinition* InIKRig, const FIKRetargetOpBase* InParentOp) override;
+
+	virtual FIKRetargetOpSettingsBase* GetSettings() override;
+	
+	virtual const UScriptStruct* GetSettingsType() const override;
+
+	virtual const UScriptStruct* GetType() const override;
+
+	virtual const UScriptStruct* GetParentOpType() const override;
+
+	virtual void OnTargetChainRenamed(const FName InOldChainName, const FName InNewChainName) override;
+
+	virtual void OnParentReinitPropertyEdited(const FIKRetargetOpBase& InParentOp, const FPropertyChangedEvent* InPropertyChangedEvent) override;
+	
+	UPROPERTY()
+	FIKRetargetOffsetGoalsOpSettings Settings;
+
+#if WITH_EDITOR
+public:
+	virtual void ResetChainSettingsToDefault(const FName& InChainName) override;
+	virtual bool AreChainSettingsAtDefault(const FName& InChainName) override;
+	virtual uint8* GetChainSettingsMemory(const FName InChainName) override;
+#endif
+
+	void RegenerateChainSettings(const FIKRetargetOpBase* InParentOp);
+
+private:
+
+	TArray<FIKChainToOffset> ChainsToOffset;
+};
+
+/* The blueprint/python API for editing a Offset IK Goals Op */
+UCLASS(MinimalAPI, BlueprintType)
+class UIKRetargetOffsetGoalsController : public UIKRetargetOpControllerBase
+{
+	GENERATED_BODY()
+	
+public:
+	/* Get the current op settings as a struct.
+	 * @return FIKRetargetOffsetGoalsOpSettings struct with the current settings used by the op. */
+	UFUNCTION(BlueprintCallable, Category = Settings)
+	FIKRetargetOffsetGoalsOpSettings GetSettings();
+
+	/* Set the op settings. Input is a custom struct type for this op.
+	 * @param InSettings a FIKRetargetOffsetGoalsOpSettings struct containing all the settings to apply to this op */
+	UFUNCTION(BlueprintCallable, Category = Settings)
+	void SetSettings(FIKRetargetOffsetGoalsOpSettings InSettings);
+};
+
+#undef UE_API
+#undef LOCTEXT_NAMESPACE
diff --git a/Engine/Plugins/Animation/IKRig/Source/IKRig/Public/Retargeter/RetargetOps/RunIKRigOp.h b/Engine/Plugins/Animation/IKRig/Source/IKRig/Public/Retargeter/RetargetOps/RunIKRigOp.h
--- a/Engine/Plugins/Animation/IKRig/Source/IKRig/Public/Retargeter/RetargetOps/RunIKRigOp.h
+++ b/Engine/Plugins/Animation/IKRig/Source/IKRig/Public/Retargeter/RetargetOps/RunIKRigOp.h
@@ -11,18 +11,54 @@
 
 #define LOCTEXT_NAMESPACE "RunIKRigOp"
 
+struct FRetargeterBoneChains;
+struct FResolvedBoneChain;
+struct FIKRetargetPelvisMotionOp;
+
+USTRUCT(BlueprintType)
+struct FIKRetargetRunIKRigChainSettings
+{
+	GENERATED_BODY()
+
+	FIKRetargetRunIKRigChainSettings() = default;
+	FIKRetargetRunIKRigChainSettings(const FName InTargetChainName) : TargetChainName(InTargetChainName) {}
+	bool operator==(const FIKRetargetRunIKRigChainSettings& Other) const;
+	FName GetName() const { return TargetChainName; };
+	void SetName(const FName InName) { TargetChainName = InName; };
+
+	/** The name of the TARGET chain to run IK on. */
+	UPROPERTY(EditAnywhere, BlueprintReadWrite, Category = "Setup", meta=(ReinitializeOnEdit))
+	FName TargetChainName;
+	
+	/** Enable/disable IK on this chain.
+	 * NOTE: This actually removes the goal from the internal IK Rig solve and triggers a reinitialization of the solver stack when toggled.*/
+	UPROPERTY(EditAnywhere, BlueprintReadWrite, Category = "Chain Settings", meta=(ReinitializeOnEdit))
+	bool bEnableIK = true;
+	
+	/** Range 0 to 1. Default 1. Retargets the goal position based on the source chain length and direction.
+	 * At 0, the goal is left at the input pose bone location.
+	 * At 1, the goal is moved in the direction of the source chain, at a distance from the base of the chain proportional to how much the chain is compressed or stretched. */
+	UPROPERTY(EditAnywhere, BlueprintReadWrite, Category = "Chain Settings", meta = (ClampMin = "0.0", ClampMax = "1.0", UIMin = "0.0", UIMax = "1.0"))
+	double PositionAlpha = 1.0f;
+
+	/** Range 0 to 1. Default 1. Retargets the goal rotation based on the source chain end bone rotation.
+	 * At 0, the goal is left at the target input pose bone orientation.
+	 * At 1, the goal is rotated by the source chain end bone, while maintaining the offset from the retarget pose.. */
+	UPROPERTY(EditAnywhere, BlueprintReadWrite, Category = "Chain Settings", meta = (ClampMin = "0.0", ClampMax = "1.0", UIMin = "0.0", UIMax = "1.0"))
+	double RotationAlpha = 1.0f;
+};
+
 USTRUCT(BlueprintType, meta = (DisplayName = "Solve IK Goal Settings"))
 struct FIKRetargetRunIKRigOpSettings : public FIKRetargetOpSettingsBase
 {
 	GENERATED_BODY()
 
+	UPROPERTY(EditAnywhere, BlueprintReadWrite, Category = Chains, meta=(ReinitializeOnEdit))
+	TArray<FIKRetargetRunIKRigChainSettings> Chains;
+
 	/** The IK Rig asset to run when this op is executed. */
 	UPROPERTY(EditAnywhere, BlueprintReadWrite, Category = "IK Rig Asset", meta=(ReinitializeOnEdit))
 	TObjectPtr<const UIKRigDefinition> IKRigAsset = nullptr;
-
-	/** Goals in this list will be excluded from the rig */
-	UPROPERTY(EditAnywhere, BlueprintReadWrite, Category = "IK Rig Asset")
-	TArray<FName> ExcludedGoals;
 	
 	// Draw IK goal locations. 
 	UPROPERTY(EditAnywhere, Category = Debug)
@@ -45,6 +81,30 @@ struct FIKRetargetRunIKRigOpSettings : public FIKRetargetOpSettingsBase
 	UE_API virtual void CopySettingsAtRuntime(const FIKRetargetOpSettingsBase* InSettingsToCopyFrom) override;
 };
 
+struct FRunIKRigChain
+{
+	FRunIKRigChain(
+		const FResolvedBoneChain& InSourceBoneChain,
+		const FResolvedBoneChain& InTargetBoneChain,
+		const FIKRetargetRunIKRigChainSettings& InSettings,
+		FIKRigLogger& InLog);
+
+	void Run(
+		FIKRigGoalContainer& InGoalContainer,
+		const TArray<FTransform>& InSourceGlobalPose,
+		const TArray<FTransform>& InTargetGlobalPose,
+		const FIKRetargetRunIKRigOpSettings& InOpSettings,
+		const FIKRetargetPelvisMotionOp* InPelvisMotionOp) const;
+	
+	// cached data
+	bool bIsReady = false;
+	const FResolvedBoneChain* SourceBoneChain = nullptr;
+	const FResolvedBoneChain* TargetBoneChain = nullptr;
+	const FIKRetargetRunIKRigChainSettings* Settings = nullptr;
+	double InvInitialLengthSource = 1.0f;
+	double InitialTargetChainLength = 1.0f;
+};
+
 #if WITH_EDITOR
 struct FRunIKRigOpGoalDebugData
 {
@@ -102,6 +162,17 @@ struct FIKRetargetRunIKRigOp : public FIKRetargetOpBase
 	
 	UE_API virtual void OnReinitPropertyEdited(const FPropertyChangedEvent* InPropertyChangedEvent) override;
 
+	// a helper function that sub-ops can use to validate IK chains before use
+	static bool ValidateIKChain(
+		const FName& InChainToValidate,
+		const FIKRigGoalContainer& InGoalContainer,
+		const FRetargeterBoneChains& InBoneChains,
+		const FIKRetargetOpBase* InOp,
+		const FIKRetargetOpBase* InParentOp,
+		FIKRigLogger& InLog,
+		const FResolvedBoneChain*& OutSourceBoneChain,
+		const FResolvedBoneChain*& OutTargetBoneChain);
+
 	UPROPERTY()
 	FIKRetargetRunIKRigOpSettings Settings;
 	
@@ -118,6 +189,10 @@ struct FIKRetargetRunIKRigOp : public FIKRetargetOpBase
 		const double InComponentScale,
 		const FIKRetargetDebugDrawState& InEditorState) const override;
 
+	virtual void ResetChainSettingsToDefault(const FName& InChainName) override;
+	virtual bool AreChainSettingsAtDefault(const FName& InChainName) override;
+	virtual uint8* GetChainSettingsMemory(const FName InChainName) override;
+
 	virtual bool HasDebugDrawing() override { return true; };
 	
 	UE_API void SaveInitialGoalTransformsIntoDebugData(const FIKRetargetProcessor& InProcessor, const TArray<FTransform>& InTargetGlobalPose);
@@ -129,11 +204,14 @@ struct FIKRetargetRunIKRigOp : public FIKRetargetOpBase
 #endif
 
 	UE_API TArray<FName> GetRequiredTargetChains() const;
+	void RegenerateChainSettings();
 	
 private:
 
 	void ResetGoalContainer(const TArray<FTransform> &InTargetGlobalPose, FIKRigGoalContainer& InOutGoalContainer);
 
+	TArray<FRunIKRigChain> Chains;
+	
 	/** The IK Rig processor for running IK on the target */
 	FIKRigProcessor IKRigProcessor;
 
diff --git a/Engine/Plugins/Animation/IKRig/Source/IKRig/Public/Retargeter/RetargetOps/ScaleGoalsOp.h b/Engine/Plugins/Animation/IKRig/Source/IKRig/Public/Retargeter/RetargetOps/ScaleGoalsOp.h
new file mode 100644
--- /dev/null
+++ b/Engine/Plugins/Animation/IKRig/Source/IKRig/Public/Retargeter/RetargetOps/ScaleGoalsOp.h
@@ -0,0 +1,160 @@
+﻿// Copyright Epic Games, Inc. All Rights Reserved.
+#pragma once
+
+#include "Retargeter/IKRetargetOps.h"
+#include "ScaleGoalsOp.generated.h"
+
+#define UE_API IKRIG_API
+
+#define LOCTEXT_NAMESPACE "ScaleGoalsOp"
+
+struct FResolvedBoneChain;
+struct FIKRigGoalContainer;
+
+USTRUCT(BlueprintType)
+struct FIKRetargetScaleGoalsChainSettings
+{
+	GENERATED_BODY()
+
+	FIKRetargetScaleGoalsChainSettings() = default;
+	FIKRetargetScaleGoalsChainSettings(const FName InTargetChainName) : TargetChainName(InTargetChainName) {}
+	bool operator==(const FIKRetargetScaleGoalsChainSettings& Other) const;
+	FName GetName() const { return TargetChainName; };
+	void SetName(const FName InName) { TargetChainName = InName; };
+
+	/** The name of the TARGET chain for these settings. */
+	UPROPERTY(EditAnywhere, BlueprintReadWrite, Category = "Setup", meta=(ReinitializeOnEdit))
+	FName TargetChainName;
+
+	/** Range 0 to infinity. Default 1. Scales the vertical component of the IK goal's position. */
+	UPROPERTY(EditAnywhere, BlueprintReadWrite, Category = "Chain Settings", meta = (ClampMin = "0.0", UIMin = "0.0", UIMax = "2.0"))
+	double ScaleVertical = 1.0f;
+
+	/** Range 0 to infinity. Default 1. Scales the horizontal component of the IK goal's position. */
+	UPROPERTY(EditAnywhere, BlueprintReadWrite, Category = "Chain Settings", meta = (ClampMin = "0.0", UIMin = "0.0", UIMax = "2.0"))
+	double ScaleHorizontal = 1.0f;
+	
+	/** Range 0 to infinity. Default 1. Brings IK goal closer (0) or further (1+) from origin of chain.
+	*  At 0 the effector is placed at the origin of the chain (ie Shoulder, Hip etc).
+	*  At 1 the effector is left at the end of the chain (ie Wrist, Foot etc)
+	*  Values in-between 0-1 will slide the effector along the vector from the start to the end of the chain.
+	*  Values greater than 1 will stretch the chain beyond the retargeted length. */
+	UPROPERTY(EditAnywhere, BlueprintReadWrite, Category = "Chain Settings", meta = (ClampMin = "0.0", UIMin = "0.0", UIMax = "2.0"))
+	double ScaleAlongChain = 1.0f;
+};
+
+USTRUCT(BlueprintType, meta = (DisplayName = "Scale Goals Settings"))
+struct FIKRetargetScaleGoalsOpSettings : public FIKRetargetOpSettingsBase
+{
+	GENERATED_BODY()
+	
+	UPROPERTY(EditAnywhere, BlueprintReadWrite, Category = Chains, meta=(ReinitializeOnEdit))
+	TArray<FIKRetargetScaleGoalsChainSettings> Chains;
+
+	/** Range 0 to 1. Default 1. Blend the effect on all chains on/off. */
+	UPROPERTY(EditAnywhere, BlueprintReadWrite, Category = "Alpha", meta = (ClampMin = "0.0", ClampMax = "1.0", UIMin = "0.0", UIMax = "1.0"))
+	double Alpha = 1.0f;
+
+	virtual const UClass* GetControllerType() const override;
+
+	virtual void CopySettingsAtRuntime(const FIKRetargetOpSettingsBase* InSettingsToCopyFrom) override;
+};
+
+struct FIKChainToScale
+{
+	FIKChainToScale() = delete;
+	FIKChainToScale(
+		const FResolvedBoneChain& InTargetBoneChain,
+		const FIKRetargetScaleGoalsChainSettings& InSettings)
+		: TargetBoneChain(&InTargetBoneChain)
+		, Settings(&InSettings){}
+
+	void Run(
+		FIKRigGoalContainer& InGoalContainer,
+		const TArray<FTransform>& InTargetGlobalPose,
+		const FIKRetargetScaleGoalsOpSettings& InOpSettings) const;
+
+private:
+	
+	// cached data
+	const FResolvedBoneChain* TargetBoneChain;
+	const FIKRetargetScaleGoalsChainSettings* Settings;
+};
+
+USTRUCT(BlueprintType, meta = (DisplayName = "Scale Goals"))
+struct FIKRetargetScaleGoalsOp : public FIKRetargetOpBase
+{
+	GENERATED_BODY()
+	
+	virtual bool Initialize(
+		const FIKRetargetProcessor& InProcessor,
+		const FRetargetSkeleton& InSourceSkeleton,
+		const FTargetSkeleton& InTargetSkeleton,
+		const FIKRetargetOpBase* InParentOp,
+		FIKRigLogger& InLog) override;
+	
+	virtual void Run(
+		FIKRetargetProcessor& InProcessor,
+		const double InDeltaTime,
+		const TArray<FTransform>& InSourceGlobalPose,
+		TArray<FTransform>& OutTargetGlobalPose) override;
+
+	virtual void OnAddedToStack(const UIKRetargeter* InRetargetAsset, const FIKRetargetOpBase* InParentOp) override;
+
+	virtual void OnAssignIKRig(const ERetargetSourceOrTarget SourceOrTarget, const UIKRigDefinition* InIKRig, const FIKRetargetOpBase* InParentOp) override;
+
+	virtual FIKRetargetOpSettingsBase* GetSettings() override;
+	
+	virtual const UScriptStruct* GetSettingsType() const override;
+
+	virtual const UScriptStruct* GetType() const override;
+
+	virtual const UScriptStruct* GetParentOpType() const override;
+
+	virtual void OnTargetChainRenamed(const FName InOldChainName, const FName InNewChainName) override;
+
+	virtual void OnParentReinitPropertyEdited(const FIKRetargetOpBase& InParentOp, const FPropertyChangedEvent* InPropertyChangedEvent) override;
+	
+	UPROPERTY()
+	FIKRetargetScaleGoalsOpSettings Settings;
+
+#if WITH_EDITOR
+public:
+	virtual void DebugDraw(
+		FPrimitiveDrawInterface* InPDI,
+		const FTransform& InSourceTransform,
+		const FTransform& InComponentTransform,
+		const double InComponentScale,
+		const FIKRetargetDebugDrawState& InEditorState) const override;
+	virtual void ResetChainSettingsToDefault(const FName& InChainName) override;
+	virtual bool AreChainSettingsAtDefault(const FName& InChainName) override;
+	virtual uint8* GetChainSettingsMemory(const FName InChainName) override;
+#endif
+
+	void RegenerateChainSettings(const FIKRetargetOpBase* InParentOp);
+
+private:
+
+	TArray<FIKChainToScale> ChainsToScale;
+};
+
+/* The blueprint/python API for editing a Scale Goals Op */
+UCLASS(MinimalAPI, BlueprintType)
+class UIKRetargetScaleGoalsController : public UIKRetargetOpControllerBase
+{
+	GENERATED_BODY()
+	
+public:
+	/* Get the current op settings as a struct.
+	 * @return FIKRetargetScaleGoalsOpSettings struct with the current settings used by the op. */
+	UFUNCTION(BlueprintCallable, Category = Settings)
+	FIKRetargetScaleGoalsOpSettings GetSettings();
+
+	/* Set the op settings. Input is a custom struct type for this op.
+	 * @param InSettings a FIKRetargetScaleGoalsOpSettings struct containing all the settings to apply to this op */
+	UFUNCTION(BlueprintCallable, Category = Settings)
+	void SetSettings(FIKRetargetScaleGoalsOpSettings InSettings);
+};
+
+#undef UE_API
+#undef LOCTEXT_NAMESPACE
diff --git a/Engine/Plugins/Animation/IKRig/Source/IKRigEditor/Private/IKRigEditor.cpp b/Engine/Plugins/Animation/IKRig/Source/IKRigEditor/Private/IKRigEditor.cpp
--- a/Engine/Plugins/Animation/IKRig/Source/IKRigEditor/Private/IKRigEditor.cpp
+++ b/Engine/Plugins/Animation/IKRig/Source/IKRigEditor/Private/IKRigEditor.cpp
@@ -110,6 +110,15 @@ void FIKRigEditor::StartupModule()
 	PropertyEditorModule.RegisterCustomPropertyTypeLayout(
 			"IKRetargetFloorConstraintOpSettings",
 			FOnGetPropertyTypeCustomizationInstance::CreateStatic(&FFloorConstraintOpCustomization::MakeInstance));
+	PropertyEditorModule.RegisterCustomPropertyTypeLayout(
+				"IKRetargetBlendToSourceOpSettings",
+				FOnGetPropertyTypeCustomizationInstance::CreateStatic(&FBlendToSourceOpCustomization::MakeInstance));
+	PropertyEditorModule.RegisterCustomPropertyTypeLayout(
+				"IKRetargetOffsetGoalsOpSettings",
+				FOnGetPropertyTypeCustomizationInstance::CreateStatic(&FOffsetGoalsOpCustomization::MakeInstance));
+	PropertyEditorModule.RegisterCustomPropertyTypeLayout(
+				"IKRetargetScaleGoalsOpSettings",
+				FOnGetPropertyTypeCustomizationInstance::CreateStatic(&FScaleGoalsOpCustomization::MakeInstance));
 
 	// register a thumbnail renderer for the assets
 	UThumbnailManager::Get().RegisterCustomRenderer(UIKRigDefinition::StaticClass(), UIKRigThumbnailRenderer::StaticClass());
diff --git a/Engine/Plugins/Animation/IKRig/Source/IKRigEditor/Private/RetargetEditor/IKRetargetOpDetails.cpp b/Engine/Plugins/Animation/IKRig/Source/IKRigEditor/Private/RetargetEditor/IKRetargetOpDetails.cpp
--- a/Engine/Plugins/Animation/IKRig/Source/IKRigEditor/Private/RetargetEditor/IKRetargetOpDetails.cpp
+++ b/Engine/Plugins/Animation/IKRig/Source/IKRigEditor/Private/RetargetEditor/IKRetargetOpDetails.cpp
@@ -266,14 +266,17 @@ void FRunIKRigOpCustomization::CustomizeChildren(
 	ShowOnlyIKFilter.bNeverShowChainsWithoutIK = true;
 
 	// configure chain map list for the IK op
-	// NOTE: no settings memory provider because this op only maps chains, it doesn't store any settings for them
 	FChainMapListConfig ChainListConfig;
-	ChainListConfig.OpWithChainSettings = NAME_None;
+	ChainListConfig.OpWithChainSettings = OpName;
 	ChainListConfig.OpWithChainMapping = OpName;
 	ChainListConfig.Controller = AssetController;
 	ChainListConfig.bEnableGoalColumn = true;
 	ChainListConfig.bEnableChainMapping = true;
 	ChainListConfig.Filter = ShowOnlyIKFilter;
+	const FName SettingsPropertyName = GET_MEMBER_NAME_CHECKED(URunIKRigSettingsWrapper, Settings);
+	const TArray<FName> HiddenSettings = {GET_MEMBER_NAME_CHECKED(FIKRetargetRunIKRigChainSettings, TargetChainName)};
+	ChainListConfig.ChainSettingsGetterFunc =
+		TMakeChainSettingsGetter<FIKRetargetRunIKRigChainSettings,URunIKRigSettingsWrapper>(OpName, SettingsPropertyName, HiddenSettings);
 
 	// add the chain mapping list
 	FText ChainMapRowName = FText::FromString("Chain Map");
@@ -283,7 +286,7 @@ void FRunIKRigOpCustomization::CustomizeChildren(
 		SNew(SIKRetargetChainMapList).InChainMapListConfig(ChainListConfig)
 	];
 
-	// add the debug properties
+	// add the extra properties
 	AddChildPropertiesInCategory(StructPropertyHandle, ChildBuilder, FName("Debug"));
 
 	// add the properties from the base op
@@ -622,5 +625,132 @@ void FFloorConstraintOpCustomization::CustomizeChildren(
 	AddBaseOpProperties(StructPropertyHandle, ChildBuilder);
 }
 
+void FBlendToSourceOpCustomization::CustomizeChildren(
+	TSharedRef<IPropertyHandle> StructPropertyHandle,
+	IDetailChildrenBuilder& ChildBuilder,
+	IPropertyTypeCustomizationUtils& StructCustomizationUtils)
+{
+	if (!LoadAndValidateStructToCustomize(StructPropertyHandle, StructCustomizationUtils))
+	{
+		return;
+	}
+	
+	// show only IK chains
+	FChainMapFilterOptions ShowOnlyIKFilter;
+	ShowOnlyIKFilter.bNeverShowChainsWithoutIK = true;
+
+	// configure chain map list for the IK op
+	FChainMapListConfig ChainListConfig;
+	ChainListConfig.OpWithChainSettings = OpName;
+	ChainListConfig.OpWithChainMapping = AssetController->GetParentOpByName(OpName); // use chain mapping from parent
+	ChainListConfig.Controller = AssetController;
+	ChainListConfig.bEnableGoalColumn = true;
+	ChainListConfig.bEnableChainMapping = false;
+	ChainListConfig.Filter = ShowOnlyIKFilter;
+	const FName SettingsPropertyName = GET_MEMBER_NAME_CHECKED(UBlendToSourceSettingsWrapper, Settings);
+	const TArray<FName> HiddenSettings = {GET_MEMBER_NAME_CHECKED(FIKRetargetBlendToSourceChainSettings, TargetChainName)};
+	ChainListConfig.ChainSettingsGetterFunc =
+		TMakeChainSettingsGetter<FIKRetargetBlendToSourceChainSettings,UBlendToSourceSettingsWrapper>(OpName, SettingsPropertyName, HiddenSettings);
+	
+	// add the chain mapping list
+	FText ChainMapRowName = FText::FromString("Chain Map");
+	ChildBuilder.AddCustomRow(ChainMapRowName)
+	.WholeRowContent()
+	[
+		SNew(SIKRetargetChainMapList).InChainMapListConfig(ChainListConfig)
+	];
+
+	// add op properties
+	AddChildPropertiesInCategory(StructPropertyHandle, ChildBuilder, FName("Op Settings"));
+	AddChildPropertiesInCategory(StructPropertyHandle, ChildBuilder, FName("Debug"));
+	
+	// add the properties from the base op
+	AddBaseOpProperties(StructPropertyHandle, ChildBuilder);
+}
+
+void FOffsetGoalsOpCustomization::CustomizeChildren(
+	TSharedRef<IPropertyHandle> StructPropertyHandle,
+	IDetailChildrenBuilder& ChildBuilder,
+	IPropertyTypeCustomizationUtils& StructCustomizationUtils)
+{
+	if (!LoadAndValidateStructToCustomize(StructPropertyHandle, StructCustomizationUtils))
+	{
+		return;
+	}
+	
+	// show only IK chains
+	FChainMapFilterOptions ShowOnlyIKFilter;
+	ShowOnlyIKFilter.bNeverShowChainsWithoutIK = true;
+
+	// configure chain map list for the IK op
+	FChainMapListConfig ChainListConfig;
+	ChainListConfig.OpWithChainSettings = OpName;
+	ChainListConfig.OpWithChainMapping = AssetController->GetParentOpByName(OpName); // use chain mapping from parent
+	ChainListConfig.Controller = AssetController;
+	ChainListConfig.bEnableGoalColumn = true;
+	ChainListConfig.bEnableChainMapping = false;
+	ChainListConfig.Filter = ShowOnlyIKFilter;
+	const FName SettingsPropertyName = GET_MEMBER_NAME_CHECKED(UOffsetGoalsSettingsWrapper, Settings);
+	const TArray<FName> HiddenSettings = {GET_MEMBER_NAME_CHECKED(FIKRetargetOffsetGoalsChainSettings, TargetChainName)};
+	ChainListConfig.ChainSettingsGetterFunc =
+		TMakeChainSettingsGetter<FIKRetargetOffsetGoalsChainSettings,UOffsetGoalsSettingsWrapper>(OpName, SettingsPropertyName, HiddenSettings);
+	
+	// add the chain mapping list
+	FText ChainMapRowName = FText::FromString("Chain Map");
+	ChildBuilder.AddCustomRow(ChainMapRowName)
+	.WholeRowContent()
+	[
+		SNew(SIKRetargetChainMapList).InChainMapListConfig(ChainListConfig)
+	];
+
+	// add op properties
+	AddChildPropertiesInCategory(StructPropertyHandle, ChildBuilder, FName("Op Settings"));
+	
+	// add the properties from the base op
+	AddBaseOpProperties(StructPropertyHandle, ChildBuilder);
+}
+
+void FScaleGoalsOpCustomization::CustomizeChildren(
+	TSharedRef<IPropertyHandle> StructPropertyHandle,
+	IDetailChildrenBuilder& ChildBuilder,
+	IPropertyTypeCustomizationUtils& StructCustomizationUtils)
+{
+	if (!LoadAndValidateStructToCustomize(StructPropertyHandle, StructCustomizationUtils))
+	{
+		return;
+	}
+	
+	// show only IK chains
+	FChainMapFilterOptions ShowOnlyIKFilter;
+	ShowOnlyIKFilter.bNeverShowChainsWithoutIK = true;
+
+	// configure chain map list for the IK op
+	FChainMapListConfig ChainListConfig;
+	ChainListConfig.OpWithChainSettings = OpName;
+	ChainListConfig.OpWithChainMapping = AssetController->GetParentOpByName(OpName); // use chain mapping from parent
+	ChainListConfig.Controller = AssetController;
+	ChainListConfig.bEnableGoalColumn = true;
+	ChainListConfig.bEnableChainMapping = false;
+	ChainListConfig.Filter = ShowOnlyIKFilter;
+	const FName SettingsPropertyName = GET_MEMBER_NAME_CHECKED(UScaleGoalsSettingsWrapper, Settings);
+	const TArray<FName> HiddenSettings = {GET_MEMBER_NAME_CHECKED(FIKRetargetScaleGoalsChainSettings, TargetChainName)};
+	ChainListConfig.ChainSettingsGetterFunc =
+		TMakeChainSettingsGetter<FIKRetargetScaleGoalsChainSettings,UScaleGoalsSettingsWrapper>(OpName, SettingsPropertyName, HiddenSettings);
+	
+	// add the chain mapping list
+	FText ChainMapRowName = FText::FromString("Chain Map");
+	ChildBuilder.AddCustomRow(ChainMapRowName)
+	.WholeRowContent()
+	[
+		SNew(SIKRetargetChainMapList).InChainMapListConfig(ChainListConfig)
+	];
+
+	// add op properties
+	AddChildPropertiesInCategory(StructPropertyHandle, ChildBuilder, FName("Alpha"));
+	
+	// add the properties from the base op
+	AddBaseOpProperties(StructPropertyHandle, ChildBuilder);
+}
+
 
 #undef LOCTEXT_NAMESPACE
diff --git a/Engine/Plugins/Animation/IKRig/Source/IKRigEditor/Private/RetargetEditor/SIKRetargetChainMapList.cpp b/Engine/Plugins/Animation/IKRig/Source/IKRigEditor/Private/RetargetEditor/SIKRetargetChainMapList.cpp
--- a/Engine/Plugins/Animation/IKRig/Source/IKRigEditor/Private/RetargetEditor/SIKRetargetChainMapList.cpp
+++ b/Engine/Plugins/Animation/IKRig/Source/IKRigEditor/Private/RetargetEditor/SIKRetargetChainMapList.cpp
@@ -345,7 +345,11 @@ void SIKRetargetChainMapList::Construct(const FArguments& InArgs)
 				.ComboButtonStyle(&FAppStyle::Get().GetWidgetStyle<FComboButtonStyle>("SimpleComboButton"))
 				.ForegroundColor(FSlateColor::UseStyle())
 				.ContentPadding(2.0f)
-				.IsEnabled(this, &SIKRetargetChainMapList::IsChainMappingEnabled)
+				.Visibility_Lambda([this]()
+				{
+					const bool bIsVisible = IsChainMappingEnabled();
+					return bIsVisible ? EVisibility::Visible : EVisibility::Collapsed;
+				})
 				.ToolTipText(LOCTEXT("AutoMapButtonToolTip", "Automatically assign source chains based on matching rule."))
 				.OnGetMenuContent( this, &SIKRetargetChainMapList::CreateChainMapMenuWidget )
 				.HasDownArrow(true)
diff --git a/Engine/Plugins/Animation/IKRig/Source/IKRigEditor/Private/RetargetEditor/SRetargetOpStack.cpp b/Engine/Plugins/Animation/IKRig/Source/IKRigEditor/Private/RetargetEditor/SRetargetOpStack.cpp
--- a/Engine/Plugins/Animation/IKRig/Source/IKRigEditor/Private/RetargetEditor/SRetargetOpStack.cpp
+++ b/Engine/Plugins/Animation/IKRig/Source/IKRigEditor/Private/RetargetEditor/SRetargetOpStack.cpp
@@ -56,9 +56,15 @@ void SRetargetOpItem::Construct(const FArguments& InArgs)
         SNew(SHorizontalBox)
     	.ToolTipText(TAttribute<FText>::CreateLambda([this]()
 		{
+			const UScriptStruct* OpType = Element.Pin()->GetType();
+    		if (!OpType)
+    		{
+    			return FText::FromString(TEXT("Unknown Op type. Resave asset to clean."));
+    		}
+    			
     		return FText::Format(
 			FText::FromString(TEXT("Op Type: {0} \nIndex in Stack: {1}\n")),
-			FText::FromString(Element.Pin()->GetType()->GetStructCPPName()),
+			FText::FromString(OpType->GetStructCPPName()),
 			FText::AsNumber(Element.Pin()->GetIndexInStack()) );
 		}))
 	    
@@ -186,7 +192,7 @@ void SRetargetOpItem::Construct(const FArguments& InArgs)
 			{
 				FIKRetargetEditorController* EditorController = ListView.Pin()->EditorController.Pin().Get();
 				FIKRetargetOpBase* Op = EditorController->AssetController->GetRetargetOpByIndex(Element.Pin()->GetIndexInStack());
-				const double Milliseconds = (Op && Op->EditorInstance ? Op->EditorInstance->GetAverageExecutionTime() : 0.0f) * 1000.0;
+				const double Milliseconds = (Op && Op->EditorInstance ? Op->EditorInstance->Timer.GetRunningAverage() : 0.0f) * 1000.0;
 				FString Formatted = FString::Printf(TEXT("%.3f ms"), Milliseconds);
 				return FText::FromString(Formatted);
 			})
@@ -364,6 +370,13 @@ void SRetargetOpList::CacheOpTypeMetaData()
 		{
 			continue;
 		}
+
+		// filter deprecated ops
+		static const FName DeprecatedStruct = TEXT("Deprecated");
+		if (TempOp->GetType()->HasMetaData(DeprecatedStruct))
+		{
+			continue;
+		}
 		
 		FIKRetargetOpMetaData OpMetaData;
 		OpMetaData.NiceName = TempOp->GetDefaultName();
diff --git a/Engine/Plugins/Animation/IKRig/Source/IKRigEditor/Public/RetargetEditor/IKRetargetOpDetails.h b/Engine/Plugins/Animation/IKRig/Source/IKRigEditor/Public/RetargetEditor/IKRetargetOpDetails.h
--- a/Engine/Plugins/Animation/IKRig/Source/IKRigEditor/Public/RetargetEditor/IKRetargetOpDetails.h
+++ b/Engine/Plugins/Animation/IKRig/Source/IKRigEditor/Public/RetargetEditor/IKRetargetOpDetails.h
@@ -12,6 +12,10 @@
 #include "RigEditor/IKRigStructViewer.h"
 #include "Templates/SharedPointer.h"
 #include "RetargetEditor/IKRetargeterController.h"
+#include "Retargeter/RetargetOps/BlendToSourceOp.h"
+#include "Retargeter/RetargetOps/OffsetGoalsOp.h"
+#include "Retargeter/RetargetOps/RunIKRigOp.h"
+#include "Retargeter/RetargetOps/ScaleGoalsOp.h"
 #include "IKRetargetOpDetails.generated.h"
 
 enum class EFKChainRotationMode : uint8;
@@ -146,6 +150,15 @@ public:
 	// End IPropertyTypeCustomization
 };
 
+UCLASS()
+class URunIKRigSettingsWrapper : public UIKRigStructWrapperBase
+{
+	GENERATED_BODY()
+public:
+	UPROPERTY(EditAnywhere, Category = "IK Chain Settings")
+	FIKRetargetRunIKRigChainSettings Settings;
+};
+
 //
 // IK Chains Op Customization
 //
@@ -357,3 +370,90 @@ public:
 	UPROPERTY(EditAnywhere, Category = "Floor Constraint Chain Settings")
 	FFloorConstraintChainSettings Settings;
 };
+
+//
+// Blend to Source Op Customization
+//
+
+class FBlendToSourceOpCustomization : public FIKRetargetOpBaseSettingsCustomization
+{
+public:
+	static TSharedRef<IPropertyTypeCustomization> MakeInstance()
+	{
+		return MakeShared<FBlendToSourceOpCustomization>();
+	}
+
+	// IPropertyTypeCustomization
+	virtual void CustomizeChildren(
+		TSharedRef<IPropertyHandle> StructPropertyHandle,
+		IDetailChildrenBuilder& ChildBuilder,
+		IPropertyTypeCustomizationUtils& StructCustomizationUtils) override;
+	// End IPropertyTypeCustomization
+};
+
+UCLASS()
+class UBlendToSourceSettingsWrapper : public UIKRigStructWrapperBase
+{
+	GENERATED_BODY()
+public:
+	UPROPERTY(EditAnywhere, Category = "Blend to Source Chain Settings")
+	FIKRetargetBlendToSourceChainSettings Settings;
+};
+
+//
+// Offset Goals Op Customization
+//
+
+class FOffsetGoalsOpCustomization : public FIKRetargetOpBaseSettingsCustomization
+{
+public:
+	static TSharedRef<IPropertyTypeCustomization> MakeInstance()
+	{
+		return MakeShared<FOffsetGoalsOpCustomization>();
+	}
+
+	// IPropertyTypeCustomization
+	virtual void CustomizeChildren(
+		TSharedRef<IPropertyHandle> StructPropertyHandle,
+		IDetailChildrenBuilder& ChildBuilder,
+		IPropertyTypeCustomizationUtils& StructCustomizationUtils) override;
+	// End IPropertyTypeCustomization
+};
+
+UCLASS()
+class UOffsetGoalsSettingsWrapper : public UIKRigStructWrapperBase
+{
+	GENERATED_BODY()
+public:
+	UPROPERTY(EditAnywhere, Category = "Offset Goals Chain Settings")
+	FIKRetargetOffsetGoalsChainSettings Settings;
+};
+
+//
+// Scale Goals Op Customization
+//
+
+class FScaleGoalsOpCustomization : public FIKRetargetOpBaseSettingsCustomization
+{
+public:
+	static TSharedRef<IPropertyTypeCustomization> MakeInstance()
+	{
+		return MakeShared<FScaleGoalsOpCustomization>();
+	}
+
+	// IPropertyTypeCustomization
+	virtual void CustomizeChildren(
+		TSharedRef<IPropertyHandle> StructPropertyHandle,
+		IDetailChildrenBuilder& ChildBuilder,
+		IPropertyTypeCustomizationUtils& StructCustomizationUtils) override;
+	// End IPropertyTypeCustomization
+};
+
+UCLASS()
+class UScaleGoalsSettingsWrapper : public UIKRigStructWrapperBase
+{
+	GENERATED_BODY()
+public:
+	UPROPERTY(EditAnywhere, Category = "Scale Goals Chain Settings")
+	FIKRetargetScaleGoalsChainSettings Settings;
+};
