[IK Retarget] Fixing initialization order. Removed "Retarget IK Goals" from default ops.

#Synced-CL 48768247

diff --git a/Engine/Plugins/Animation/IKRig/Source/IKRig/Private/Retargeter/IKRetargetProcessor.cpp b/Engine/Plugins/Animation/IKRig/Source/IKRig/Private/Retargeter/IKRetargetProcessor.cpp
--- a/Engine/Plugins/Animation/IKRig/Source/IKRig/Private/Retargeter/IKRetargetProcessor.cpp
+++ b/Engine/Plugins/Animation/IKRig/Source/IKRig/Private/Retargeter/IKRetargetProcessor.cpp
@@ -1220,51 +1220,59 @@ void FIKRetargetProcessor::InitialOpStackSetup(
 	
 	// resolve all bone chains onto the skeletons
 	AllBoneChains.Initialize(RetargeterAsset, AllIKRigs, SourceSkeleton, TargetSkeleton, Log);
+
+	// record meta-data to build execution order
+	InitializeOpEntries();
 }
 
 void FIKRetargetProcessor::InitializeRetargetOps()
 {
-	// initialize the ops
-	TArray<FIKRetargetOpBase*> OpsInitBeforeChildrenThisTick;
-	for (FInstancedStruct& OpStruct : OpStack)
-	{
-		FIKRetargetOpBase& Op = OpStruct.GetMutable<FIKRetargetOpBase>();
-		FIKRetargetOpBase* ParentOp = GetRetargetOpByName(Op.GetParentOpName());
-
-		// give parent ops a chance to initialize right before their children
-		if (ParentOp && !OpsInitBeforeChildrenThisTick.Contains(ParentOp))
+	auto PostInitializeProcessor = [this]()
 		{
-			ParentOp->InitializeBeforeChildren(*this, SourceSkeleton, TargetSkeleton, Log);
-			OpsInitBeforeChildrenThisTick.Add(ParentOp);
-		}
-		
-		Op.Initialize(*this, SourceSkeleton, TargetSkeleton, ParentOp, Log);
-	}
-
-	// record meta-data to build execution order
-	InitializeOpEntries();
+			// gather retargeted bones from each op
+			TSet<int32> AllRetargetedBones;
+			for (FInstancedStruct& OpStruct : OpStack)
+			{
+				FIKRetargetOpBase& Op = OpStruct.GetMutable<FIKRetargetOpBase>();
+				if (Op.IsInitialized())
+				{
+					Op.CollectRetargetedBones(AllRetargetedBones);
+				}
+			}
 	
-	// gather retargeted bones from each op
-	TSet<int32> AllRetargetedBones;
-	for (FInstancedStruct& OpStruct : OpStack)
-	{
-		FIKRetargetOpBase& Op = OpStruct.GetMutable<FIKRetargetOpBase>();
-		if (Op.IsEnabled() && Op.IsInitialized())
-		{
-			Op.CollectRetargetedBones(AllRetargetedBones);
-		}
-	}
+			// store retargeted bone mask on target skeleton
+			TargetSkeleton.SetRetargetedBones(AllRetargetedBones);
+		};
 	
-	// store retargeted bone mask on target skeleton
-	TargetSkeleton.SetRetargetedBones(AllRetargetedBones);
+	// build a stack of op calls in the correct order for initialization
+	TArray<FRetargetOpCall> InitializationStack;
+	BuildInitializationStack(InitializationStack);
 
-	// run PostInitialize() on all ops
-	for (FInstancedStruct& OpStruct : OpStack)
+	// initialize the ops
+	for (FRetargetOpCall& Call : InitializationStack)
 	{
-		FIKRetargetOpBase& Op = OpStruct.GetMutable<FIKRetargetOpBase>();
-		if (Op.IsEnabled() && Op.IsInitialized())
+		FIKRetargetOpBase* Op = Call.GetOp();
+		
+		switch (Call.Phase)
 		{
-			Op.PostInitialize(*this, SourceSkeleton, TargetSkeleton, Log);
+		case ERetargetOpPhase::InitializeBeforeChildren:
+			Op->InitializeBeforeChildren(*this, SourceSkeleton, TargetSkeleton, Log);
+			break;
+
+		case ERetargetOpPhase::Initialize:
+			Op->Initialize(*this, SourceSkeleton, TargetSkeleton, Call.GetParentOp(), Log);
+			break;
+
+		case ERetargetOpPhase::PostInitializeProcessor:
+			PostInitializeProcessor();
+			break;
+
+		case ERetargetOpPhase::PostInitializeOp:
+			Op->PostInitialize(*this, SourceSkeleton, TargetSkeleton, Log);
+			break;
+
+		default:
+			checkNoEntry();
 		}
 	}
 }
@@ -1347,26 +1355,26 @@ void FIKRetargetProcessor::RunRetargetOps(
 	BuildExecutionStack(InLOD, ExecutionStack);
 
 	// run the ops
-	for (const FRetargetOpCall& Call : ExecutionStack)
+	for (FRetargetOpCall& Call : ExecutionStack)
 	{
 		
 #if WITH_EDITOR
 		const double StartTime = FPlatformTime::Seconds();
 #endif
 
-		FIKRetargetOpBase* Op = Call.Op;
+		FIKRetargetOpBase* Op = Call.GetOp();
 
 		switch (Call.Phase)
 		{
-		case ERetargetOpPhase::ParentBeforeChildren:
+		case ERetargetOpPhase::RunBeforeChildren:
 			Op->RunBeforeChildren(*this, InDeltaTime, InSourceGlobalPose, OutTargetGlobalPose);
 			break;
 
-		case ERetargetOpPhase::Main:
+		case ERetargetOpPhase::Run:
 			Op->Run(*this, InDeltaTime, InSourceGlobalPose, OutTargetGlobalPose);
 			break;
 
-		case ERetargetOpPhase::ChildAfterParent:
+		case ERetargetOpPhase::RunAfterParent:
 			Op->RunAfterParent(*this, InDeltaTime, InSourceGlobalPose, OutTargetGlobalPose);
 			break;
 
@@ -1400,6 +1408,31 @@ bool FIKRetargetProcessor::ShouldExecuteOp(const FIKRetargetOpBase& InOp, int32
 	return true;
 }
 
+FIKRetargetOpBase* FIKRetargetProcessor::FRetargetOpCall::GetOp()
+{
+	if (!OpEntry)
+	{
+		return nullptr;
+	}
+	
+	return OpEntry->Op;
+}
+
+FIKRetargetOpBase* FIKRetargetProcessor::FRetargetOpCall::GetParentOp()
+{
+	if (!OpEntry)
+	{
+		return nullptr;
+	}
+	
+	if (!OpEntry->Parent)
+	{
+		return nullptr;
+	}
+
+	return OpEntry->Parent->Op;
+}
+
 void FIKRetargetProcessor::BuildExecutionStack(int32 InLOD, TArray<FRetargetOpCall>& OutStack)
 {
 	OutStack.Reset();
@@ -1420,18 +1453,18 @@ void FIKRetargetProcessor::BuildExecutionStack(int32 InLOD, TArray<FRetargetOpCa
 			continue;
 		}
 
-		const bool bHasParent = OpEntry.Parent != nullptr;
+		const bool bIsChild = OpEntry.Parent != nullptr;
 		const bool bIsParent = OpEntry.Op->CanHaveChildOps();
 
-		if (bHasParent)
+		if (bIsChild)
 		{
 			continue; // children ops managed by parent
 		}
 
 		// neither a child or a parent, ie a simple op, just run Main once.
-		if (!bIsParent && !bHasParent)
+		if (!bIsParent && !bIsChild)
 		{
-			OutStack.Add({ ERetargetOpPhase::Main, OpEntry.Op });
+			OutStack.Add({ ERetargetOpPhase::Run, &OpEntry });
 			continue;
 		}
 
@@ -1444,31 +1477,80 @@ void FIKRetargetProcessor::BuildExecutionStack(int32 InLOD, TArray<FRetargetOpCa
 		// 4. Children->RunAfterParent
 
 		// 1. Parent->RunBeforeChildren()
-		OutStack.Emplace(ERetargetOpPhase::ParentBeforeChildren, OpEntry.Op);
+		OutStack.Emplace(ERetargetOpPhase::RunBeforeChildren, &OpEntry);
 
 		// 2. Children->Run()
-		for (FOpEntry* Child : OpEntry.Children)
+		for (FOpEntry* ChildEntry : OpEntry.Children)
 		{
-			if (Child->bExecute)
+			if (ChildEntry->bExecute)
 			{
-				OutStack.Emplace(ERetargetOpPhase::Main, Child->Op);
+				OutStack.Emplace(ERetargetOpPhase::Run, ChildEntry);
 			}
 		}
 
 		// 3. Parent->Run()
-		OutStack.Emplace(ERetargetOpPhase::Main, OpEntry.Op);
+		OutStack.Emplace(ERetargetOpPhase::Run, &OpEntry);
 
 		// 4. Children->RunAfterParent()
-		for (FOpEntry* Child : OpEntry.Children)
+		for (FOpEntry* ChildEntry : OpEntry.Children)
 		{
-			if (Child->bExecute)
+			if (ChildEntry->bExecute)
 			{
-				OutStack.Emplace(ERetargetOpPhase::ChildAfterParent, Child->Op);
+				OutStack.Emplace(ERetargetOpPhase::RunAfterParent, ChildEntry);
 			}
 		}
 	}
 }
 
+void FIKRetargetProcessor::BuildInitializationStack(TArray<FRetargetOpCall>& OutStack)
+{
+	OutStack.Reset();
+
+	// NOTE: we initialize all ops regardless of if they are disabled
+	// this is because ops can be toggled on/off at runtime
+
+	// 1st pass add Initialize and InitializeBeforeChildren calls
+	for (FOpEntry& OpEntry : OpEntries)
+	{
+		const bool bIsChild = OpEntry.Parent != nullptr;
+		const bool bIsParent = OpEntry.Op->CanHaveChildOps();
+
+		if (bIsChild)
+		{
+			continue; // children ops managed by parent
+		}
+
+		// neither a child or a parent, ie a simple op, just run Initialize.
+		if (!bIsParent && !bIsChild)
+		{
+			OutStack.Add({ ERetargetOpPhase::Initialize, &OpEntry });
+			continue;
+		}
+
+		if (bIsParent)
+		{
+			// Parent->InitializeBeforeChildren()
+			// Child0...N->Initialize()
+			// Parent->Initialize()
+			OutStack.Emplace(ERetargetOpPhase::InitializeBeforeChildren, &OpEntry);
+			for (FOpEntry* Child : OpEntry.Children)
+			{
+				OutStack.Add({ ERetargetOpPhase::Initialize, Child });
+			}
+			OutStack.Add({ ERetargetOpPhase::Initialize, &OpEntry });
+		}
+	}
+
+	// 2nd pass, give processor a chance to gather data from initialized ops
+	OutStack.Add({ERetargetOpPhase::PostInitializeProcessor, nullptr});
+
+	// 3nd pass, add PostInitialize() for all ops
+	for (FOpEntry& OpEntry : OpEntries)
+	{
+		OutStack.Add({ ERetargetOpPhase::PostInitializeOp, &OpEntry });
+	}
+}
+
 void FIKRetargetProcessor::InitializeOpEntries()
 {
 	auto FindOpEntry = [this](const FIKRetargetOpBase* InOpToFind) -> FOpEntry*
diff --git a/Engine/Plugins/Animation/IKRig/Source/IKRig/Public/Retargeter/IKRetargetProcessor.h b/Engine/Plugins/Animation/IKRig/Source/IKRig/Public/Retargeter/IKRetargetProcessor.h
--- a/Engine/Plugins/Animation/IKRig/Source/IKRig/Public/Retargeter/IKRetargetProcessor.h
+++ b/Engine/Plugins/Animation/IKRig/Source/IKRig/Public/Retargeter/IKRetargetProcessor.h
@@ -599,20 +599,33 @@ private:
 	
 	enum class ERetargetOpPhase : uint8
 	{
-		ParentBeforeChildren,
-		Main,
-		ChildAfterParent
+		// initialization
+		InitializeBeforeChildren,
+		Initialize,
+		PostInitializeProcessor,
+		PostInitializeOp,
+		
+		// runtime
+		RunBeforeChildren,
+		Run,
+		RunAfterParent
 	};
 
 	struct FRetargetOpCall
 	{
-		FRetargetOpCall(const ERetargetOpPhase InPhase,	FIKRetargetOpBase* InOp) : Phase(InPhase), Op(InOp) {}
+		FRetargetOpCall(const ERetargetOpPhase InPhase, FOpEntry* InOpEntry) : Phase(InPhase), OpEntry(InOpEntry) {}
+		FIKRetargetOpBase* GetOp();
+		FIKRetargetOpBase* GetParentOp();
+		
 		ERetargetOpPhase Phase;
-		FIKRetargetOpBase* Op;
+	private:
+		FOpEntry* OpEntry;
 	};
 
 	void BuildExecutionStack(int32 InLOD, TArray<FRetargetOpCall>& OutStack);
 
+	void BuildInitializationStack(TArray<FRetargetOpCall>& OutStack);
+
 	void InitializeOpEntries();
 	TArray<FOpEntry> OpEntries;
 };
diff --git a/Engine/Plugins/Animation/IKRig/Source/IKRigEditor/Private/RetargetEditor/IKRetargeterController.cpp b/Engine/Plugins/Animation/IKRig/Source/IKRigEditor/Private/RetargetEditor/IKRetargeterController.cpp
--- a/Engine/Plugins/Animation/IKRig/Source/IKRigEditor/Private/RetargetEditor/IKRetargeterController.cpp
+++ b/Engine/Plugins/Animation/IKRig/Source/IKRigEditor/Private/RetargetEditor/IKRetargeterController.cpp
@@ -344,7 +344,6 @@ void UIKRetargeterController::AddDefaultOps() const
 	AddRetargetOp(FIKRetargetFKChainsOp::StaticStruct());
 	const int32 RunIKIndex = AddRetargetOp(FIKRetargetRunIKRigOp::StaticStruct());
 	const FName RunIKOpName = GetOpName(RunIKIndex);
-	AddRetargetOp(FIKRetargetIKChainsOp::StaticStruct(), RunIKOpName);
 	AddRetargetOp(FIKRetargetRootMotionOp::StaticStruct());
 	AddRetargetOp(FIKRetargetCurveRemapOp::StaticStruct());
 }
