[IK Retarget] Fixing goal blending behavior on Run IK Rig Op

#Synced-CL 48786069

diff --git a/Engine/Plugins/Animation/IKRig/Source/IKRig/Private/Retargeter/RetargetOps/RunIKRigOp.cpp b/Engine/Plugins/Animation/IKRig/Source/IKRig/Private/Retargeter/RetargetOps/RunIKRigOp.cpp
--- a/Engine/Plugins/Animation/IKRig/Source/IKRig/Private/Retargeter/RetargetOps/RunIKRigOp.cpp
+++ b/Engine/Plugins/Animation/IKRig/Source/IKRig/Private/Retargeter/RetargetOps/RunIKRigOp.cpp
@@ -18,8 +18,9 @@ struct FIKRigGoal;
 bool FIKRetargetRunIKRigChainSettings::operator==(const FIKRetargetRunIKRigChainSettings& Other) const
 {
 	return bEnableIK == Other.bEnableIK
-		&& FMath::IsNearlyEqualByULP(PositionAlpha, Other.PositionAlpha)
-		&& FMath::IsNearlyEqualByULP(RotationAlpha, Other.RotationAlpha);
+		&& bOverrideGoalAlpha == Other.bOverrideGoalAlpha
+		&& FMath::IsNearlyEqualByULP(GoalPositionAlpha, Other.GoalPositionAlpha)
+		&& FMath::IsNearlyEqualByULP(GoalRotationAlpha, Other.GoalRotationAlpha);
 }
 
 const UClass* FIKRetargetRunIKRigOpSettings::GetControllerType() const
@@ -105,7 +106,7 @@ void FRunIKRigChain::Run(
 		return;
 	}
 
-	if (Settings->PositionAlpha > UE_KINDA_SMALL_NUMBER)
+	if (Settings->ChainAlpha > UE_KINDA_SMALL_NUMBER)
 	{
 		// get relevant bone transforms
 		const int32 SourceStartBoneIndex = SourceBoneChain->BoneIndices[0];
@@ -131,10 +132,10 @@ void FRunIKRigChain::Run(
 		FVector GoalPosition = CurrentTargetStartLocation + (CurrentSourceEndDirectionNormalized * InitialTargetChainLength);
 		
 		// alpha and apply to the goal
-		Goal->Position = FMath::Lerp(Goal->Position, GoalPosition, Settings->PositionAlpha);
+		Goal->Position = FMath::Lerp(Goal->Position, GoalPosition, Settings->ChainAlpha);
 	}
 
-	if (Settings->RotationAlpha > UE_KINDA_SMALL_NUMBER)
+	if (Settings->ChainAlpha > UE_KINDA_SMALL_NUMBER)
 	{
 		// calc delta rotation of source relative to source retarget pose
 		const int32 SourceEndBoneIndex = SourceBoneChain->BoneIndices.Last();
@@ -148,10 +149,33 @@ void FRunIKRigChain::Run(
 		const FQuat NewGoalRotation = SourceDeltaRotation * InitialTargetEndGlobal.GetRotation();
 		
 		// alpha blend rotation and apply to goal
-		Goal->Rotation = FQuat::FastLerp(Goal->Rotation.Quaternion(), NewGoalRotation, Settings->RotationAlpha).GetNormalized().Rotator();
+		Goal->Rotation = FQuat::FastLerp(Goal->Rotation.Quaternion(), NewGoalRotation, Settings->ChainAlpha).GetNormalized().Rotator();
 	}
 }
 
+void FRunIKRigChain::SetGoalAlpha(FIKRigGoalContainer& InGoalContainer)
+{
+	if (!bIsReady)
+	{
+		return;
+	}
+
+	if (!Settings->bOverrideGoalAlpha)
+	{
+		return;
+	}
+	
+	FIKRigGoal* Goal = InGoalContainer.FindGoalByName(TargetBoneChain->IKGoalName);
+	if (!ensure(Goal))
+	{
+		return;
+	}
+
+	// NOTE: actual blending happens in the IK Rig processor
+	Goal->PositionAlpha = Settings->GoalPositionAlpha;
+	Goal->RotationAlpha = Settings->GoalRotationAlpha;
+}
+
 bool FIKRetargetRunIKRigOp::Initialize(
 	const FIKRetargetProcessor& InProcessor,
 	const FRetargetSkeleton& InSourceSkeleton,
@@ -271,8 +295,14 @@ void FIKRetargetRunIKRigOp::Run(
 #endif
 
 	FIKRigGoalContainer& GoalContainer = InProcessor.GetIKRigGoalContainer();
+	
+	// set main alpha on the goals
+	for (FRunIKRigChain& Chain : Chains)
+	{
+		Chain.SetGoalAlpha(GoalContainer);
+	}
 
-	// set the goals
+	// copy the goal container from the retargeter to the IK Rig processor
 	IKRigProcessor.ApplyGoalsFromOtherContainer(GoalContainer);
 	// trigger reinitialization if the goal container was modified such that it needs it
 	if (GoalContainer.NeedsInitialized())
diff --git a/Engine/Plugins/Animation/IKRig/Source/IKRig/Public/Retargeter/RetargetOps/RunIKRigOp.h b/Engine/Plugins/Animation/IKRig/Source/IKRig/Public/Retargeter/RetargetOps/RunIKRigOp.h
--- a/Engine/Plugins/Animation/IKRig/Source/IKRig/Public/Retargeter/RetargetOps/RunIKRigOp.h
+++ b/Engine/Plugins/Animation/IKRig/Source/IKRig/Public/Retargeter/RetargetOps/RunIKRigOp.h
@@ -34,18 +34,30 @@ struct FIKRetargetRunIKRigChainSettings
 	 * NOTE: This actually removes the goal from the internal IK Rig solve and triggers a reinitialization of the solver stack when toggled.*/
 	UPROPERTY(EditAnywhere, BlueprintReadWrite, Category = "Chain Settings", meta=(ReinitializeOnEdit))
 	bool bEnableIK = true;
-	
+
 	/** Range 0 to 1. Default 1. Retargets the goal position based on the source chain length and direction.
-	 * At 0, the goal is left at the input pose bone location.
-	 * At 1, the goal is moved in the direction of the source chain, at a distance from the base of the chain proportional to how much the chain is compressed or stretched. */
+	* At 0, the goal is left at the input pose bone location.
+	* At 1, the goal is moved in the direction of the source chain, at a distance from the base of the chain proportional to how much the chain is compressed or stretched. */
 	UPROPERTY(EditAnywhere, BlueprintReadWrite, Category = "Chain Settings", meta = (ClampMin = "0.0", ClampMax = "1.0", UIMin = "0.0", UIMax = "1.0"))
-	double PositionAlpha = 1.0f;
+	double ChainAlpha = 1.0f;
 
-	/** Range 0 to 1. Default 1. Retargets the goal rotation based on the source chain end bone rotation.
-	 * At 0, the goal is left at the target input pose bone orientation.
-	 * At 1, the goal is rotated by the source chain end bone, while maintaining the offset from the retarget pose.. */
-	UPROPERTY(EditAnywhere, BlueprintReadWrite, Category = "Chain Settings", meta = (ClampMin = "0.0", ClampMax = "1.0", UIMin = "0.0", UIMax = "1.0"))
-	double RotationAlpha = 1.0f;
+	/** Default false. If true, the alpha values here will override the goal alpha settings from the IK Rig asset. */
+	UPROPERTY(EditAnywhere, BlueprintReadWrite, Category = "Chain Settings", meta=(ReinitializeOnEdit))
+	bool bOverrideGoalAlpha = false;
+	
+	/** Range 0 to 1. Default 1. Blends all goal translation offsets on/off.
+	 * At 0, the goal is left at the input pose bone location.
+	 * At 1, the goal is moved to the location generated by the IK goal ops. */
+	UPROPERTY(EditAnywhere, BlueprintReadWrite, Category = "Chain Settings",
+		meta = (ClampMin = "0.0", ClampMax = "1.0", UIMin = "0.0", UIMax = "1.0", EditCondition = "bOverrideGoalAlpha"))
+	double GoalPositionAlpha = 1.0f;
+
+	/** Range 0 to 1. Default 1. Blends all goal rotation offsets on/off.
+	 * At 0, the goal is left at the input pose bone rotation.
+	 * At 1, the goal is rotated to the orientation generated by the IK goal ops. */
+	UPROPERTY(EditAnywhere, BlueprintReadWrite, Category = "Chain Settings",
+		meta = (ClampMin = "0.0", ClampMax = "1.0", UIMin = "0.0", UIMax = "1.0", EditCondition = "bOverrideGoalAlpha"))
+	double GoalRotationAlpha = 1.0f;
 };
 
 USTRUCT(BlueprintType, meta = (DisplayName = "Solve IK Goal Settings"))
@@ -95,6 +107,8 @@ struct FRunIKRigChain
 		const TArray<FTransform>& InTargetGlobalPose,
 		const FIKRetargetRunIKRigOpSettings& InOpSettings,
 		const FIKRetargetPelvisMotionOp* InPelvisMotionOp) const;
+
+	void SetGoalAlpha(FIKRigGoalContainer& InGoalContainer);
 	
 	// cached data
 	bool bIsReady = false;
