[IK Retarget] Adding separate rotation blending to Run IK Rig Goal retarget

#Synced-CL 48786069

diff --git a/Engine/Plugins/Animation/IKRig/Source/IKRig/Private/Retargeter/RetargetOps/RunIKRigOp.cpp b/Engine/Plugins/Animation/IKRig/Source/IKRig/Private/Retargeter/RetargetOps/RunIKRigOp.cpp
--- a/Engine/Plugins/Animation/IKRig/Source/IKRig/Private/Retargeter/RetargetOps/RunIKRigOp.cpp
+++ b/Engine/Plugins/Animation/IKRig/Source/IKRig/Private/Retargeter/RetargetOps/RunIKRigOp.cpp
@@ -18,9 +18,10 @@ struct FIKRigGoal;
 bool FIKRetargetRunIKRigChainSettings::operator==(const FIKRetargetRunIKRigChainSettings& Other) const
 {
 	return bEnableIK == Other.bEnableIK
+		&& FMath::IsNearlyEqualByULP(ChainPositionAlpha, Other.ChainPositionAlpha)
 		&& bOverrideGoalAlpha == Other.bOverrideGoalAlpha
-		&& FMath::IsNearlyEqualByULP(GoalPositionAlpha, Other.GoalPositionAlpha)
-		&& FMath::IsNearlyEqualByULP(GoalRotationAlpha, Other.GoalRotationAlpha);
+		&& FMath::IsNearlyEqualByULP(FinalPositionAlpha, Other.FinalPositionAlpha)
+		&& FMath::IsNearlyEqualByULP(FinalRotationAlpha, Other.FinalRotationAlpha);
 }
 
 const UClass* FIKRetargetRunIKRigOpSettings::GetControllerType() const
@@ -106,7 +107,7 @@ void FRunIKRigChain::Run(
 		return;
 	}
 
-	if (Settings->ChainAlpha > UE_KINDA_SMALL_NUMBER)
+	if (Settings->ChainPositionAlpha > UE_KINDA_SMALL_NUMBER)
 	{
 		// get relevant bone transforms
 		const int32 SourceStartBoneIndex = SourceBoneChain->BoneIndices[0];
@@ -132,10 +133,10 @@ void FRunIKRigChain::Run(
 		FVector GoalPosition = CurrentTargetStartLocation + (CurrentSourceEndDirectionNormalized * InitialTargetChainLength);
 		
 		// alpha and apply to the goal
-		Goal->Position = FMath::Lerp(Goal->Position, GoalPosition, Settings->ChainAlpha);
+		Goal->Position = FMath::Lerp(Goal->Position, GoalPosition, Settings->ChainPositionAlpha);
 	}
 
-	if (Settings->ChainAlpha > UE_KINDA_SMALL_NUMBER)
+	if (Settings->ChainRotationAlpha > UE_KINDA_SMALL_NUMBER)
 	{
 		// calc delta rotation of source relative to source retarget pose
 		const int32 SourceEndBoneIndex = SourceBoneChain->BoneIndices.Last();
@@ -149,7 +150,7 @@ void FRunIKRigChain::Run(
 		const FQuat NewGoalRotation = SourceDeltaRotation * InitialTargetEndGlobal.GetRotation();
 		
 		// alpha blend rotation and apply to goal
-		Goal->Rotation = FQuat::FastLerp(Goal->Rotation.Quaternion(), NewGoalRotation, Settings->ChainAlpha).GetNormalized().Rotator();
+		Goal->Rotation = FQuat::FastLerp(Goal->Rotation.Quaternion(), NewGoalRotation, Settings->ChainRotationAlpha).GetNormalized().Rotator();
 	}
 }
 
@@ -172,8 +173,8 @@ void FRunIKRigChain::SetGoalAlpha(FIKRigGoalContainer& InGoalContainer)
 	}
 
 	// NOTE: actual blending happens in the IK Rig processor
-	Goal->PositionAlpha = Settings->GoalPositionAlpha;
-	Goal->RotationAlpha = Settings->GoalRotationAlpha;
+	Goal->PositionAlpha = Settings->FinalPositionAlpha;
+	Goal->RotationAlpha = Settings->FinalRotationAlpha;
 }
 
 bool FIKRetargetRunIKRigOp::Initialize(
diff --git a/Engine/Plugins/Animation/IKRig/Source/IKRig/Public/Retargeter/RetargetOps/RunIKRigOp.h b/Engine/Plugins/Animation/IKRig/Source/IKRig/Public/Retargeter/RetargetOps/RunIKRigOp.h
--- a/Engine/Plugins/Animation/IKRig/Source/IKRig/Public/Retargeter/RetargetOps/RunIKRigOp.h
+++ b/Engine/Plugins/Animation/IKRig/Source/IKRig/Public/Retargeter/RetargetOps/RunIKRigOp.h
@@ -37,27 +37,37 @@ struct FIKRetargetRunIKRigChainSettings
 
 	/** Range 0 to 1. Default 1. Retargets the goal position based on the source chain length and direction.
 	* At 0, the goal is left at the input pose bone location.
-	* At 1, the goal is moved in the direction of the source chain, at a distance from the base of the chain proportional to how much the chain is compressed or stretched. */
-	UPROPERTY(EditAnywhere, BlueprintReadWrite, Category = "Chain Settings", meta = (ClampMin = "0.0", ClampMax = "1.0", UIMin = "0.0", UIMax = "1.0"))
-	double ChainAlpha = 1.0f;
+	* At 1, the goal is translated in the direction of the source chain, at a distance from the base of the chain proportional to how much it is compressed or stretched. */
+	UPROPERTY(EditAnywhere, BlueprintReadWrite, Category = "Chain Settings",
+		meta = (ClampMin = "0.0", ClampMax = "1.0", UIMin = "0.0", UIMax = "1.0", EditCondition = "bEnableIK"))
+	double ChainPositionAlpha = 1.0f;
+
+	/** Range 0 to 1. Default 1. Retargets the goal rotation based on the source chain end bone rotation.
+	* At 0, the goal is left at the input pose bone rotation.
+	* At 1, the goal is rotated by the same world space rotation applied to the source chain end bone.
+	* NOTE: if the target end bone rotation was already retargeted before this op, this setting may have no effect.
+	* It will not double-rotate the goal. This setting is useful if no FK retargeting is being performed on this chain before the IK op. */
+	UPROPERTY(EditAnywhere, BlueprintReadWrite, Category = "Chain Settings",
+		meta = (ClampMin = "0.0", ClampMax = "1.0", UIMin = "0.0", UIMax = "1.0", EditCondition = "bEnableIK"))
+	double ChainRotationAlpha = 1.0f;
 
 	/** Default false. If true, the alpha values here will override the goal alpha settings from the IK Rig asset. */
-	UPROPERTY(EditAnywhere, BlueprintReadWrite, Category = "Chain Settings", meta=(ReinitializeOnEdit))
+	UPROPERTY(EditAnywhere, BlueprintReadWrite, Category = "Chain Settings", meta=(ReinitializeOnEdit, EditCondition = "bEnableIK"))
 	bool bOverrideGoalAlpha = false;
 	
 	/** Range 0 to 1. Default 1. Blends all goal translation offsets on/off.
 	 * At 0, the goal is left at the input pose bone location.
 	 * At 1, the goal is moved to the location generated by the IK goal ops. */
 	UPROPERTY(EditAnywhere, BlueprintReadWrite, Category = "Chain Settings",
-		meta = (ClampMin = "0.0", ClampMax = "1.0", UIMin = "0.0", UIMax = "1.0", EditCondition = "bOverrideGoalAlpha"))
-	double GoalPositionAlpha = 1.0f;
+		meta = (ClampMin = "0.0", ClampMax = "1.0", UIMin = "0.0", UIMax = "1.0", EditCondition = "bEnableIK && bOverrideGoalAlpha"))
+	double FinalPositionAlpha = 1.0f;
 
 	/** Range 0 to 1. Default 1. Blends all goal rotation offsets on/off.
 	 * At 0, the goal is left at the input pose bone rotation.
 	 * At 1, the goal is rotated to the orientation generated by the IK goal ops. */
 	UPROPERTY(EditAnywhere, BlueprintReadWrite, Category = "Chain Settings",
-		meta = (ClampMin = "0.0", ClampMax = "1.0", UIMin = "0.0", UIMax = "1.0", EditCondition = "bOverrideGoalAlpha"))
-	double GoalRotationAlpha = 1.0f;
+		meta = (ClampMin = "0.0", ClampMax = "1.0", UIMin = "0.0", UIMax = "1.0", EditCondition = "bEnableIK && bOverrideGoalAlpha"))
+	double FinalRotationAlpha = 1.0f;
 };
 
 USTRUCT(BlueprintType, meta = (DisplayName = "Solve IK Goal Settings"))
