[IK Rig] Removing root chain from auto characterizer

#Synced-CL 48918725

diff --git a/Engine/Plugins/Animation/IKRig/Source/IKRig/Private/Retargeter/RetargetOps/RootMotionGeneratorOp.cpp b/Engine/Plugins/Animation/IKRig/Source/IKRig/Private/Retargeter/RetargetOps/RootMotionGeneratorOp.cpp
--- a/Engine/Plugins/Animation/IKRig/Source/IKRig/Private/Retargeter/RetargetOps/RootMotionGeneratorOp.cpp
+++ b/Engine/Plugins/Animation/IKRig/Source/IKRig/Private/Retargeter/RetargetOps/RootMotionGeneratorOp.cpp
@@ -197,23 +197,37 @@ void FIKRetargetRootMotionOp::PostInitialize(
 void FIKRetargetRootMotionOp::OnAddedToStack(const UIKRetargeter* InRetargetAsset, const FIKRetargetOpBase* InParentOp)
 {
 #if WITH_EDITOR
-	// get the root of the SOURCE skeleton
-	if (const USkeletalMesh* Mesh = InRetargetAsset->GetPreviewMesh(ERetargetSourceOrTarget::Source))
+
+	// first try and get the root bone from the characterization data on the IK Rig assets
+	if (const UIKRigDefinition* SourceIKRig = InRetargetAsset->GetIKRig(ERetargetSourceOrTarget::Source))
 	{
-		Settings.SourceRoot = Mesh->GetRefSkeleton().GetBoneName(0);
+		Settings.SourceRoot = SourceIKRig->GetRoot();
 	}
-	
-	// get the root of the TARGET skeleton
-	if (const USkeletalMesh* Mesh = InRetargetAsset->GetPreviewMesh(ERetargetSourceOrTarget::Target))
+	if (const UIKRigDefinition* TargetIKRig = InRetargetAsset->GetIKRig(ERetargetSourceOrTarget::Target))
+	{
+		Settings.TargetRoot = TargetIKRig->GetRoot();
+		// get the target pelvis bone too
+		Settings.TargetPelvis = TargetIKRig->GetPelvis();
+	}
+
+	// if source root bone was not registered, fallback to the root of the SOURCE skeleton
+	if (Settings.SourceRoot.BoneName == NAME_None)
 	{
-		Settings.TargetRoot = Mesh->GetRefSkeleton().GetBoneName(0);
+		if (const USkeletalMesh* Mesh = InRetargetAsset->GetPreviewMesh(ERetargetSourceOrTarget::Source))
+		{
+			Settings.SourceRoot = Mesh->GetRefSkeleton().GetBoneName(0);
+		}
 	}
 
-	// get the pelvis bone
-	if (const UIKRigDefinition* IKRig = InRetargetAsset->GetIKRig(ERetargetSourceOrTarget::Target))
+	// if target root bone was not registered, fallback to the root of the TARGET skeleton
+	if (Settings.TargetRoot.BoneName == NAME_None)
 	{
-		Settings.TargetPelvis = IKRig->GetPelvis();
+		if (const USkeletalMesh* Mesh = InRetargetAsset->GetPreviewMesh(ERetargetSourceOrTarget::Target))
+		{
+			Settings.TargetRoot = Mesh->GetRefSkeleton().GetBoneName(0);
+		}
 	}
+	
 #endif
 }
 
diff --git a/Engine/Plugins/Animation/IKRig/Source/IKRig/Private/Rig/IKRigDefinition.cpp b/Engine/Plugins/Animation/IKRig/Source/IKRig/Private/Rig/IKRigDefinition.cpp
--- a/Engine/Plugins/Animation/IKRig/Source/IKRig/Private/Rig/IKRigDefinition.cpp
+++ b/Engine/Plugins/Animation/IKRig/Source/IKRig/Private/Rig/IKRigDefinition.cpp
@@ -287,6 +287,26 @@ FBoneChain* FRetargetDefinition::GetEditableBoneChainByName(FName ChainName)
 	return nullptr;
 }
 
+bool FRetargetDefinition::Serialize(FArchive& Ar)
+{
+	Ar.UsingCustomVersion(FIKRigObjectVersion::GUID);
+	// doesn't actually serialize, just write the custom version for PostSerialize
+	return false;
+}
+
+void FRetargetDefinition::PostSerialize(const FArchive& Ar)
+{
+	if (Ar.IsLoading())
+	{
+		if (Ar.CustomVer(FIKRigObjectVersion::GUID) < FIKRigObjectVersion::AddedRootBoneToRetargetDefinition)
+		{
+			// we used to call "Pelvis" the "RetargetRoot"
+			// but now we also want to characterize the actual skeleton "root" bone itself
+			PelvisBone = RootBone;
+		}
+	}
+}
+
 void UIKRigDefinition::PostDuplicate(bool bDuplicateForPIE)
 {
 	Super::PostDuplicate(bDuplicateForPIE);
diff --git a/Engine/Plugins/Animation/IKRig/Source/IKRig/Public/IKRigObjectVersion.h b/Engine/Plugins/Animation/IKRig/Source/IKRig/Public/IKRigObjectVersion.h
--- a/Engine/Plugins/Animation/IKRig/Source/IKRig/Public/IKRigObjectVersion.h
+++ b/Engine/Plugins/Animation/IKRig/Source/IKRig/Public/IKRigObjectVersion.h
@@ -53,6 +53,9 @@ struct FIKRigObjectVersion
 		// the "Retarget IK Goals" op is deprecated and features are split into separate ops
 		IKChainsOpSplitIntoSeparateOps,
 
+		// IK Rig retarget definition now stores separate root bone and pelvis bone
+		AddedRootBoneToRetargetDefinition,
+
 		// -----<new versions can be added above this line>-------------------------------------------------
 		VersionPlusOne,
 		LatestVersion = VersionPlusOne - 1
diff --git a/Engine/Plugins/Animation/IKRig/Source/IKRig/Public/Rig/IKRigDefinition.h b/Engine/Plugins/Animation/IKRig/Source/IKRig/Public/Rig/IKRigDefinition.h
--- a/Engine/Plugins/Animation/IKRig/Source/IKRig/Public/Rig/IKRigDefinition.h
+++ b/Engine/Plugins/Animation/IKRig/Source/IKRig/Public/Rig/IKRigDefinition.h
@@ -166,9 +166,14 @@ struct FRetargetDefinition
 {
 	GENERATED_BODY()
 
+	// The bone that contains root motion, located at the origin
 	UPROPERTY(BlueprintReadWrite, EditAnywhere, Category=RetargetRoot)
 	FName RootBone;
 
+	// The bone located at the hips or pelvis area, that translates the whole character
+	UPROPERTY(BlueprintReadWrite, EditAnywhere, Category=RetargetPelvis)
+	FName PelvisBone;
+
 	UPROPERTY(BlueprintReadWrite, EditAnywhere, Category=Chains)
 	TArray<FBoneChain> BoneChains;
 
@@ -180,6 +185,21 @@ struct FRetargetDefinition
 		const FName& GoalName = NAME_None);
 	
 	IKRIG_API FBoneChain* GetEditableBoneChainByName(FName ChainName);
+
+	// post load patch old data
+	IKRIG_API bool Serialize(FArchive& Ar);
+	IKRIG_API void PostSerialize(const FArchive& Ar);
+};
+
+// load-patch FRetargetDefinition
+template<>
+struct TStructOpsTypeTraits<FRetargetDefinition> : public TStructOpsTypeTraitsBase2<FRetargetDefinition>
+{
+	enum
+	{
+		WithPostSerialize = true,
+		WithSerializer = true,
+	};
 };
 
 UCLASS(MinimalAPI, BlueprintType)
@@ -208,7 +228,10 @@ public:
 	const TArray<FBoneChain>& GetRetargetChains() const { return RetargetDefinition.BoneChains; };
 
 	/** runtime, read-only access to Pelvis, all modifications must go through UIKRigController */
-	const FName& GetPelvis() const { return RetargetDefinition.RootBone; };
+	const FName& GetPelvis() const { return RetargetDefinition.PelvisBone; };
+
+	/** runtime, read-only access to Root bone, all modifications must go through UIKRigController */
+	const FName& GetRoot() const { return RetargetDefinition.RootBone; };
 
 	/** runtime, read-only access to read a bone chain, all modifications must go through UIKRigController */
 	IKRIG_API const FBoneChain* GetRetargetChainByName(FName ChainName) const;
diff --git a/Engine/Plugins/Animation/IKRig/Source/IKRigEditor/Private/RetargetEditor/SRetargetAnimAssetsWindow.cpp b/Engine/Plugins/Animation/IKRig/Source/IKRigEditor/Private/RetargetEditor/SRetargetAnimAssetsWindow.cpp
--- a/Engine/Plugins/Animation/IKRig/Source/IKRigEditor/Private/RetargetEditor/SRetargetAnimAssetsWindow.cpp
+++ b/Engine/Plugins/Animation/IKRig/Source/IKRigEditor/Private/RetargetEditor/SRetargetAnimAssetsWindow.cpp
@@ -707,12 +707,12 @@ void FProceduralRetargetAssets::AutoGenerateIKRetargetAsset()
 	// setup the root motion generator (but only if the skeleton has a root bone separate from the pelvis)
 	FIKRetargetRootMotionOp* RootOp = RetargetController->GetFirstRetargetOpOfType<FIKRetargetRootMotionOp>();
 	const FName TargetSkeletonRoot = TargetIKRig->GetSkeleton().BoneNames[0];
-	const FName TargetPelvis = TargetCharacterizationResults.AutoRetargetDefinition.RetargetDefinition.RootBone;
+	const FName TargetPelvis = TargetCharacterizationResults.AutoRetargetDefinition.RetargetDefinition.PelvisBone;
 	if (TargetSkeletonRoot != TargetPelvis)
 	{
 		// if the source does not have a dedicated root bone, configure the root motion to be generated from the pelvis motion on target
 		const FName SourceSkeletonRoot = SourceIKRig->GetSkeleton().BoneNames[0];
-		const FName SourcePelvis = SourceCharacterizationResults.AutoRetargetDefinition.RetargetDefinition.RootBone;
+		const FName SourcePelvis = SourceCharacterizationResults.AutoRetargetDefinition.RetargetDefinition.PelvisBone;
 		if (SourceSkeletonRoot == SourcePelvis)
 		{
 			RootOp->Settings.RootMotionSource = ERootMotionSource::GenerateFromTargetPelvis;
diff --git a/Engine/Plugins/Animation/IKRig/Source/IKRigEditor/Private/RigEditor/IKRigAutoCharacterizer.cpp b/Engine/Plugins/Animation/IKRig/Source/IKRigEditor/Private/RigEditor/IKRigAutoCharacterizer.cpp
--- a/Engine/Plugins/Animation/IKRig/Source/IKRigEditor/Private/RigEditor/IKRigAutoCharacterizer.cpp
+++ b/Engine/Plugins/Animation/IKRig/Source/IKRigEditor/Private/RigEditor/IKRigAutoCharacterizer.cpp
@@ -9,7 +9,6 @@
 #define LOCTEXT_NAMESPACE "AutoCharacterizer"
 
 // core
-const FName FCharacterizationStandard::Root("Root");
 const FName FCharacterizationStandard::Spine("Spine");
 const FName FCharacterizationStandard::Neck("Neck");
 const FName FCharacterizationStandard::Head("Head");
@@ -470,8 +469,8 @@ FKnownTemplateHierarchies::FKnownTemplateHierarchies()
 		FTemplateHierarchy& UE4 = AddTemplateHierarchy(UE4Name, UE4Bones, UE4ParentIndices);
 		FRetargetDefinition& UE4Retarget = UE4.AutoRetargetDefinition.RetargetDefinition;
 		// core
-		UE4Retarget.RootBone = FName("pelvis");
-		UE4Retarget.AddBoneChain(FCharacterizationStandard::Root, FName("root"), FName("root"));
+		UE4Retarget.PelvisBone = FName("pelvis");
+		UE4Retarget.RootBone = FName("root");
 		UE4Retarget.AddBoneChain(FCharacterizationStandard::Spine, FName("spine_01"), FName("spine_02"));
 		UE4Retarget.AddBoneChain(FCharacterizationStandard::Neck, FName("neck_01"), FName("neck_01"));
 		UE4Retarget.AddBoneChain(FCharacterizationStandard::Head, FName("head"), FName("head"));
@@ -526,8 +525,8 @@ FKnownTemplateHierarchies::FKnownTemplateHierarchies()
 		FTemplateHierarchy& UE5 = AddTemplateHierarchy(UE5Name, UE5Bones, UE5ParentIndices); 
 		FRetargetDefinition& UE5Retarget = UE5.AutoRetargetDefinition.RetargetDefinition;
 		// core
-		UE5Retarget.RootBone = FName("pelvis");
-		UE5Retarget.AddBoneChain(FCharacterizationStandard::Root, FName("root"), FName("root"));
+		UE5Retarget.PelvisBone = FName("pelvis");
+		UE5Retarget.RootBone = FName("root");
 		UE5Retarget.AddBoneChain(FCharacterizationStandard::Spine, FName("spine_01"), FName("spine_02"));
 		UE5Retarget.AddBoneChain(FCharacterizationStandard::Neck, FName("neck_01"), FName("neck_02"));
 		UE5Retarget.AddBoneChain(FCharacterizationStandard::Head, FName("head"), FName("head"));
@@ -590,7 +589,7 @@ FKnownTemplateHierarchies::FKnownTemplateHierarchies()
 		FTemplateHierarchy& Daz3d = AddTemplateHierarchy(Daz3dName, Daz3dBones, Daz3dParentIndices);
 		FRetargetDefinition& DazRetarget = Daz3d.AutoRetargetDefinition.RetargetDefinition;
 		// core
-		DazRetarget.RootBone = FName("hip");
+		DazRetarget.PelvisBone = FName("hip");
 		DazRetarget.AddBoneChain(FCharacterizationStandard::Spine, FName("abdomenLower"), FName("chestUpper"));
 		DazRetarget.AddBoneChain(FCharacterizationStandard::Neck, FName("neckLower"), FName("neckUpper"));
 		DazRetarget.AddBoneChain(FCharacterizationStandard::Head, FName("head"), FName("head"));
@@ -663,7 +662,7 @@ FKnownTemplateHierarchies::FKnownTemplateHierarchies()
 		FTemplateHierarchy& Mixamo = AddTemplateHierarchy(Mixamo_Name, Mixamo_Bones, Mixamo_ParentIndices);
 		FRetargetDefinition& MixamoRetarget = Mixamo.AutoRetargetDefinition.RetargetDefinition;
 		// core
-		MixamoRetarget.RootBone = FName("Hips");
+		MixamoRetarget.PelvisBone = FName("Hips");
 		MixamoRetarget.AddBoneChain(FCharacterizationStandard::Spine, FName("Spine"), FName("Spine2"));
 		MixamoRetarget.AddBoneChain(FCharacterizationStandard::Neck, FName("Neck"), FName("Neck"));
 		MixamoRetarget.AddBoneChain(FCharacterizationStandard::Head, FName("Head"), FName("Head"));
@@ -708,8 +707,8 @@ FKnownTemplateHierarchies::FKnownTemplateHierarchies()
 		FTemplateHierarchy& CC4 = AddTemplateHierarchy(CC4Name, CC4Bones, CC4ParentIndices);
 		FRetargetDefinition& CC4Retarget = CC4.AutoRetargetDefinition.RetargetDefinition;
 		// core
-		CC4Retarget.RootBone = FName("Hip");
-		CC4Retarget.AddBoneChain(FCharacterizationStandard::Root, FName("BoneRoot"), FName("BoneRoot"));
+		CC4Retarget.PelvisBone = FName("Hip");
+		CC4Retarget.RootBone = FName("BoneRoot");
 		CC4Retarget.AddBoneChain(FCharacterizationStandard::Spine, FName("Waist"), FName("Spine02"));
 		CC4Retarget.AddBoneChain(FCharacterizationStandard::Neck, FName("NeckTwist01"), FName("NeckTwist02"));
 		CC4Retarget.AddBoneChain(FCharacterizationStandard::Head, FName("Head"), FName("Head"));
@@ -765,8 +764,8 @@ FKnownTemplateHierarchies::FKnownTemplateHierarchies()
 		FTemplateHierarchy& Xsens = AddTemplateHierarchy(XsensName, XsenseBones, XsensParentIndices);
 		FRetargetDefinition& XsensRetarget = Xsens.AutoRetargetDefinition.RetargetDefinition;
 		// core
-		XsensRetarget.RootBone = FName("Hips");
-		XsensRetarget.AddBoneChain(FCharacterizationStandard::Root, FName("Reference"), FName("Reference"));
+		XsensRetarget.PelvisBone = FName("Hips");
+		XsensRetarget.RootBone = FName("Reference");
 		XsensRetarget.AddBoneChain(FCharacterizationStandard::Spine, FName("Spine"), FName("Spine3"));
 		XsensRetarget.AddBoneChain(FCharacterizationStandard::Neck, FName("Neck"), FName("Neck"));
 		XsensRetarget.AddBoneChain(FCharacterizationStandard::Head, FName("Head"), FName("Head"));
@@ -811,8 +810,8 @@ FKnownTemplateHierarchies::FKnownTemplateHierarchies()
 		FTemplateHierarchy& mGear = AddTemplateHierarchy(mGearName, mGearBones, mGearParentIndices);
 		FRetargetDefinition& mGearRetarget = mGear.AutoRetargetDefinition.RetargetDefinition;
 		// core
-		mGearRetarget.RootBone = FName("spine_C0_pelvis_jnt");
-		mGearRetarget.AddBoneChain(FCharacterizationStandard::Root, FName("Reference"), FName("Reference"));
+		mGearRetarget.PelvisBone = FName("spine_C0_pelvis_jnt");
+		mGearRetarget.RootBone = FName("Reference");
 		mGearRetarget.AddBoneChain(FCharacterizationStandard::Spine, FName("spine_C0_01_jnt"), FName("spine_C0_05_jnt"));
 		mGearRetarget.AddBoneChain(FCharacterizationStandard::Neck, FName("neck_C0_01_jnt"), FName("neck_C0_02_jnt"));
 		mGearRetarget.AddBoneChain(FCharacterizationStandard::Head, FName("neck_C0_head_jnt"), FName("neck_C0_head_jnt"));
@@ -874,7 +873,7 @@ FKnownTemplateHierarchies::FKnownTemplateHierarchies()
 		FTemplateHierarchy& HumanIK = AddTemplateHierarchy(HumanIKName, HumanIKBones, HumanIKParentIndices);
 		FRetargetDefinition& HumanIKRetarget = HumanIK.AutoRetargetDefinition.RetargetDefinition;
 		// core
-		HumanIKRetarget.RootBone = FName("Hips");
+		HumanIKRetarget.PelvisBone = FName("Hips");
 		HumanIKRetarget.AddBoneChain(FCharacterizationStandard::Spine, FName("Spine"), FName("Spine3"));
 		HumanIKRetarget.AddBoneChain(FCharacterizationStandard::Neck, FName("Neck"), FName("Neck"));
 		HumanIKRetarget.AddBoneChain(FCharacterizationStandard::Head, FName("Head"), FName("Head"));
@@ -920,7 +919,7 @@ FKnownTemplateHierarchies::FKnownTemplateHierarchies()
 		FTemplateHierarchy& Vicon = AddTemplateHierarchy(ViconName, ViconBones, ViconParentIndices);
 		FRetargetDefinition& ViconRetarget = Vicon.AutoRetargetDefinition.RetargetDefinition;
 		// core
-		ViconRetarget.RootBone = FName("Hips");
+		ViconRetarget.PelvisBone = FName("Hips");
 		ViconRetarget.AddBoneChain(FCharacterizationStandard::Spine, FName("Spine"), FName("Spine3"));
 		ViconRetarget.AddBoneChain(FCharacterizationStandard::Neck, FName("Neck"), FName("Neck1"));
 		ViconRetarget.AddBoneChain(FCharacterizationStandard::Head, FName("Head"), FName("Head"));
@@ -970,8 +969,8 @@ FKnownTemplateHierarchies::FKnownTemplateHierarchies()
 		FTemplateHierarchy& Optitrak = AddTemplateHierarchy(OptitrakName, OptitrakBones, OptitrakParentIndices);
 		FRetargetDefinition& OptitrakRetarget = Optitrak.AutoRetargetDefinition.RetargetDefinition;
 		// core
-		OptitrakRetarget.RootBone = FName("pelvis");
-		OptitrakRetarget.AddBoneChain(FCharacterizationStandard::Root, FName("root"), FName("root"));
+		OptitrakRetarget.PelvisBone = FName("pelvis");
+		OptitrakRetarget.RootBone = FName("root");
 		OptitrakRetarget.AddBoneChain(FCharacterizationStandard::Spine, FName("spine_01"), FName("spine_02"));
 		OptitrakRetarget.AddBoneChain(FCharacterizationStandard::Neck, FName("neck_01"), FName("neck_01"));
 		OptitrakRetarget.AddBoneChain(FCharacterizationStandard::Head, FName("head"), FName("head"));
@@ -1005,7 +1004,7 @@ FKnownTemplateHierarchies::FKnownTemplateHierarchies()
 		static TArray<int32> MocopiParentIndices = {-1, 0, 1, 2, 3, 4, 5, 6, 7, 8, 9, 10, 7, 12, 13, 14, 15, 7, 17, 18, 19, 20, 0, 22, 23, 24, 25, 0, 27, 28, 29, 30};
 		FTemplateHierarchy& Mocopi = AddTemplateHierarchy(MocopiName, MocopiBones, MocopiParentIndices);
 		FRetargetDefinition& MocopiRetarget = Mocopi.AutoRetargetDefinition.RetargetDefinition;
-		MocopiRetarget.RootBone = FName("root");
+		MocopiRetarget.PelvisBone = FName("root");
 		MocopiRetarget.AddBoneChain(FCharacterizationStandard::Spine, FName("torso_1"), FName("torso_7"));
 		MocopiRetarget.AddBoneChain(FCharacterizationStandard::Neck, FName("neck_1"), FName("neck_2"));
 		MocopiRetarget.AddBoneChain(FCharacterizationStandard::Head, FName("head"), FName("head"));
@@ -1042,7 +1041,7 @@ FKnownTemplateHierarchies::FKnownTemplateHierarchies()
 		FTemplateHierarchy& AdvancedSkeleton = AddTemplateHierarchy(AdvancedSkeletonName, AdvancedSkeletonBones, AdvancedSkeletonParentIndices);
 		FRetargetDefinition& AdvSkelRetarget = AdvancedSkeleton.AutoRetargetDefinition.RetargetDefinition;
 		// core
-		AdvSkelRetarget.RootBone = FName("Root_M");
+		AdvSkelRetarget.PelvisBone = FName("Root_M");
 		AdvSkelRetarget.AddBoneChain(FCharacterizationStandard::Spine, FName("RootPart1_M"), FName("Chest_M"));
 		AdvSkelRetarget.AddBoneChain(FCharacterizationStandard::Neck, FName("Neck_M"), FName("NeckPart2_M"));
 		AdvSkelRetarget.AddBoneChain(FCharacterizationStandard::Head, FName("Head_M"), FName("Head_M"));
@@ -1101,7 +1100,7 @@ FKnownTemplateHierarchies::FKnownTemplateHierarchies()
 		FTemplateHierarchy& MoveOne = AddTemplateHierarchy(MoveOneName, MoveOneBones, MoveOneParentIndices);
 		FRetargetDefinition& MoveOneRetarget = MoveOne.AutoRetargetDefinition.RetargetDefinition;
 		// core
-		MoveOneRetarget.RootBone = FName("Hips");
+		MoveOneRetarget.PelvisBone = FName("Hips");
 		MoveOneRetarget.AddBoneChain(FCharacterizationStandard::Spine, FName("Spine"), FName("Spine1"));
 		MoveOneRetarget.AddBoneChain(FCharacterizationStandard::Neck, FName("Neck"), FName("Neck"));
 		MoveOneRetarget.AddBoneChain(FCharacterizationStandard::Head, FName("Head"), FName("Head"));
@@ -1155,7 +1154,7 @@ FKnownTemplateHierarchies::FKnownTemplateHierarchies()
 		FTemplateHierarchy& Qualisys = AddTemplateHierarchy(QualisysName, QualisysBones, QualisysParentIndices);
 		FRetargetDefinition& QualisysRetarget = Qualisys.AutoRetargetDefinition.RetargetDefinition;
 		// core
-		QualisysRetarget.RootBone = FName("Hips");
+		QualisysRetarget.PelvisBone = FName("Hips");
 		QualisysRetarget.AddBoneChain(FCharacterizationStandard::Spine, FName("Spine"), FName("Spine2"));
 		QualisysRetarget.AddBoneChain(FCharacterizationStandard::Neck, FName("Neck"), FName("Neck"));
 		QualisysRetarget.AddBoneChain(FCharacterizationStandard::Head, FName("Head"), FName("Head"));
@@ -1211,8 +1210,8 @@ FKnownTemplateHierarchies::FKnownTemplateHierarchies()
 		FTemplateHierarchy& FNHuman = AddTemplateHierarchy(FNHumanName, FNHumanBones, FNHumanParentIndices);
 		FRetargetDefinition& FNHumanRetarget = FNHuman.AutoRetargetDefinition.RetargetDefinition;
 		// core
-		FNHumanRetarget.RootBone = FName("pelvis");
-		FNHumanRetarget.AddBoneChain(FCharacterizationStandard::Root, FName("root"), FName("root"));
+		FNHumanRetarget.PelvisBone = FName("pelvis");
+		FNHumanRetarget.RootBone = FName("root");
 		FNHumanRetarget.AddBoneChain(FCharacterizationStandard::Spine, FName("spine_01"), FName("spine_05"));
 		FNHumanRetarget.AddBoneChain(FCharacterizationStandard::Neck, FName("neck_01"), FName("neck_02"));
 		FNHumanRetarget.AddBoneChain(FCharacterizationStandard::Head, FName("head"), FName("head"));
@@ -1274,7 +1273,7 @@ FKnownTemplateHierarchies::FKnownTemplateHierarchies()
 		FTemplateHierarchy& FNQuadruped = AddTemplateHierarchy(FNQuadrupedName, FNQuadrupedBones, FNQuadrupedParentIndices);
 		FRetargetDefinition& FNQuadrupedRetarget = FNQuadruped.AutoRetargetDefinition.RetargetDefinition;
 		// core
-		FNQuadrupedRetarget.RootBone = FName("QuadSpine_A_Pelvis_C");
+		FNQuadrupedRetarget.PelvisBone = FName("QuadSpine_A_Pelvis_C");
 		FNQuadrupedRetarget.AddBoneChain(FCharacterizationStandard::Spine, FName("QuadSpine_A_Spine1_C"), FName("QuadSpine_A_Spine3_C"));
 		FNQuadrupedRetarget.AddBoneChain(FCharacterizationStandard::Neck, FName("BipedHeadNeck_A_NeckBase_C"), FName("BipedHeadNeck_A_NeckMid_C"));
 		FNQuadrupedRetarget.AddBoneChain(FCharacterizationStandard::Head, FName("BipedHeadNeck_A_Head_C"), FName("BipedHeadNeck_A_Head_C"));
@@ -1303,7 +1302,8 @@ FKnownTemplateHierarchies::FKnownTemplateHierarchies()
 		FTemplateHierarchy& FNRaptor = AddTemplateHierarchy(FNRaptorName, FNRaptorBones, FNRaptorParentIndices);
 		FRetargetDefinition& FNRaptorRetarget = FNRaptor.AutoRetargetDefinition.RetargetDefinition;
 		// core
-		FNRaptorRetarget.RootBone = FName("QuadSpine_A_Pelvis_C");
+		FNRaptorRetarget.PelvisBone = FName("QuadSpine_A_Pelvis_C");
+		FNRaptorRetarget.RootBone = "root";
 		FNRaptorRetarget.AddBoneChain(FCharacterizationStandard::Spine, FName("QuadSpine_A_Spine1_C"), FName("QuadSpine_A_Spine3_C"));
 		FNRaptorRetarget.AddBoneChain(FCharacterizationStandard::Neck, FName("BipedHeadNeck_A_NeckBase_C"), FName("BipedHeadNeck_A_NeckMid_C"));
 		FNRaptorRetarget.AddBoneChain(FCharacterizationStandard::Head, FName("BipedHeadNeck_A_Head_C"), FName("BipedHeadNeck_A_Head_C"));
@@ -1343,7 +1343,8 @@ FKnownTemplateHierarchies::FKnownTemplateHierarchies()
 		FTemplateHierarchy& FNAmphibian = AddTemplateHierarchy(FNAmphibianName, FNAmphibianBones, FNAmphibianParentIndices);
 		FRetargetDefinition& FNAmphibianRetarget = FNAmphibian.AutoRetargetDefinition.RetargetDefinition;
 		// core
-		FNAmphibianRetarget.RootBone = FName("QuadSpine_A_Pelvis_C");
+		FNAmphibianRetarget.PelvisBone = FName("QuadSpine_A_Pelvis_C");
+		FNAmphibianRetarget.RootBone = FName("root");
 		FNAmphibianRetarget.AddBoneChain(FCharacterizationStandard::Spine, FName("QuadSpine_A_Spine1_C"), FName("QuadSpine_A_Spine3_C"));
 		FNAmphibianRetarget.AddBoneChain(FCharacterizationStandard::Neck, FName("BipedHeadNeck_A_NeckBase_C"), FName("BipedHeadNeck_A_NeckMid_C"));
 		FNAmphibianRetarget.AddBoneChain(FCharacterizationStandard::Head, FName("BipedHeadNeck_A_Head_C"), FName("BipedHeadNeck_A_Head_C"));
@@ -1384,7 +1385,8 @@ FKnownTemplateHierarchies::FKnownTemplateHierarchies()
 		FTemplateHierarchy& FNButtercake = AddTemplateHierarchy(FNButtercakeName, FNButtercakeBones, FNButtercakeParentIndices);
 		FRetargetDefinition& FNButtercakeRetarget = FNButtercake.AutoRetargetDefinition.RetargetDefinition;
 		// core
-		FNButtercakeRetarget.RootBone = FName("C_QuadSpine_A_Pelvis_Jnt");
+		FNButtercakeRetarget.PelvisBone = FName("C_QuadSpine_A_Pelvis_Jnt");
+		FNButtercakeRetarget.RootBone = FName("C_Root_Main_Root_Jnt");
 		FNButtercakeRetarget.AddBoneChain(FCharacterizationStandard::Spine, FName("C_QuadSpine_A_Spine1_Jnt"), FName("C_QuadSpine_A_Chest_Jnt"));
 		FNButtercakeRetarget.AddBoneChain(FCharacterizationStandard::Neck, FName("C_BipedHeadNeck_A_NeckBase_Jnt"), FName("C_BipedHeadNeck_A_NeckMid_Jnt"));
 		FNButtercakeRetarget.AddBoneChain(FCharacterizationStandard::Head, FName("C_BipedHeadNeck_A_Head_Jnt"), FName("C_BipedHeadNeck_A_Head_Jnt"));
@@ -1408,8 +1410,8 @@ FKnownTemplateHierarchies::FKnownTemplateHierarchies()
 		FTemplateHierarchy& FNBiped = AddTemplateHierarchy(FNBipedName, FNBipedBones, FNBipedParentIndices);
 		FRetargetDefinition& FNBipedRetarget = FNBiped.AutoRetargetDefinition.RetargetDefinition;
 		// core
-		FNBipedRetarget.RootBone = FName("C_BipedSpine_Main_Pelvis_Jnt");
-		FNBipedRetarget.AddBoneChain(FCharacterizationStandard::Root, FName("C_Root_Main_Root_Jnt"), FName("C_Root_Main_Root_Jnt"));
+		FNBipedRetarget.PelvisBone = FName("C_BipedSpine_Main_Pelvis_Jnt");
+		FNBipedRetarget.RootBone = FName("C_Root_Main_Root_Jnt");
 		FNBipedRetarget.AddBoneChain(FCharacterizationStandard::Spine, FName("C_BipedSpine_Main_Spine1_Jnt"), FName("C_BipedSpine_Main_Spine3_Jnt"));
 		FNBipedRetarget.AddBoneChain(FCharacterizationStandard::Neck, FName("C_BipedHeadNeck_A_NeckBase_Jnt"), FName("C_BipedHeadNeck_A_NeckMid_Jnt"));
 		FNBipedRetarget.AddBoneChain(FCharacterizationStandard::Head, FName("C_BipedHeadNeck_A_Head_Jnt"), FName("C_BipedHeadNeck_A_Head_Jnt"));
@@ -1474,7 +1476,8 @@ FKnownTemplateHierarchies::FKnownTemplateHierarchies()
 		FTemplateHierarchy& Rokoko = AddTemplateHierarchy(RokokoName, RokokoBones, RokokoParentIndices);
 		FRetargetDefinition& RokokoRetarget = Rokoko.AutoRetargetDefinition.RetargetDefinition;
 		// core
-		RokokoRetarget.RootBone = FName("Hips");
+		RokokoRetarget.PelvisBone = FName("Hips");
+		RokokoRetarget.RootBone = FName("Root");
 		RokokoRetarget.AddBoneChain(FCharacterizationStandard::Spine, FName("Spine1"), FName("Spine4"));
 		RokokoRetarget.AddBoneChain(FCharacterizationStandard::Neck, FName("Neck"), FName("Neck"));
 		RokokoRetarget.AddBoneChain(FCharacterizationStandard::Head, FName("Head"), FName("Head"));
@@ -1614,13 +1617,27 @@ void FAutoCharacterizer::AdaptTemplateToHierarchy(
 	// reset the retarget definition, we're about to adapt it to the given hierarchy
 	Results.AutoRetargetDefinition.RetargetDefinition = FRetargetDefinition();
 
+	// adapt the root bone
+	const int32 RootBoneIndex = TargetHierarchy.GetBoneIndex(Template.AutoRetargetDefinition.RetargetDefinition.RootBone, ECleanOrFullName::Clean);
+	if (RootBoneIndex != INDEX_NONE)
+	{
+		// convert to the full, non-clean name to use for the actual bone chain
+		const FName RootBoneName =  TargetHierarchy.GetBoneName(RootBoneIndex, ECleanOrFullName::Full);
+		Results.AutoRetargetDefinition.RetargetDefinition.RootBone = RootBoneName;
+	}
+	else
+	{
+		// default to the first bone if none was explictly specified
+		Results.AutoRetargetDefinition.RetargetDefinition.RootBone = TargetHierarchy.GetBoneName(0, ECleanOrFullName::Full);
+	}
+
 	// adapt the pelvis bone
-	const int32 PelvisBoneIndex = TargetHierarchy.GetBoneIndex(Template.AutoRetargetDefinition.RetargetDefinition.RootBone, ECleanOrFullName::Clean);
+	const int32 PelvisBoneIndex = TargetHierarchy.GetBoneIndex(Template.AutoRetargetDefinition.RetargetDefinition.PelvisBone, ECleanOrFullName::Clean);
 	if (PelvisBoneIndex != INDEX_NONE)
 	{
 		// convert to the full, non-clean name to use for the actual bone chain
 		const FName PelvisBoneName =  TargetHierarchy.GetBoneName(PelvisBoneIndex, ECleanOrFullName::Full);
-		Results.AutoRetargetDefinition.RetargetDefinition.RootBone = PelvisBoneName;
+		Results.AutoRetargetDefinition.RetargetDefinition.PelvisBone = PelvisBoneName;
 	}
 
 	// adapt the bone chains by shrinking from start and end until a valid chain is found on the target hierarchy
diff --git a/Engine/Plugins/Animation/IKRig/Source/IKRigEditor/Private/RigEditor/IKRigAutoFBIK.cpp b/Engine/Plugins/Animation/IKRig/Source/IKRigEditor/Private/RigEditor/IKRigAutoFBIK.cpp
--- a/Engine/Plugins/Animation/IKRig/Source/IKRigEditor/Private/RigEditor/IKRigAutoFBIK.cpp
+++ b/Engine/Plugins/Animation/IKRig/Source/IKRigEditor/Private/RigEditor/IKRigAutoFBIK.cpp
@@ -67,10 +67,10 @@ void FAutoFBIKCreator::CreateFBIKSetup(const UIKRigController& IKRigController,
 	}
 
 	// set the root of the solver
-	const bool bSetRoot = IKRigController.SetStartBone(CharacterizeResults.AutoRetargetDefinition.RetargetDefinition.RootBone, SolverIndex);
+	const bool bSetRoot = IKRigController.SetStartBone(CharacterizeResults.AutoRetargetDefinition.RetargetDefinition.PelvisBone, SolverIndex);
 	if (!bSetRoot)
 	{
-		Results.Outcome = EAutoFBIKResult::MissingRootBone;
+		Results.Outcome = EAutoFBIKResult::MissingPelvisBone;
 		return;
 	}
 
diff --git a/Engine/Plugins/Animation/IKRig/Source/IKRigEditor/Private/RigEditor/IKRigController.cpp b/Engine/Plugins/Animation/IKRig/Source/IKRigEditor/Private/RigEditor/IKRigController.cpp
--- a/Engine/Plugins/Animation/IKRig/Source/IKRigEditor/Private/RigEditor/IKRigController.cpp
+++ b/Engine/Plugins/Animation/IKRig/Source/IKRigEditor/Private/RigEditor/IKRigController.cpp
@@ -409,20 +409,20 @@ const TArray<FBoneChain>& UIKRigController::GetRetargetChains() const
 	return Asset->GetRetargetChains();
 }
 
-bool UIKRigController::SetRetargetRoot(const FName RootBoneName) const
+bool UIKRigController::SetRetargetRoot(const FName InPelvisBoneName) const
 {
 	check(Asset)
 
-	FName NewRootBone = RootBoneName;
-	if (RootBoneName != NAME_None && Asset->Skeleton.GetBoneIndexFromName(RootBoneName) == INDEX_NONE)
+	FName NewPelvisBone = InPelvisBoneName;
+	if (InPelvisBoneName != NAME_None && Asset->Skeleton.GetBoneIndexFromName(InPelvisBoneName) == INDEX_NONE)
 	{
-		NewRootBone = NAME_None;
+		NewPelvisBone = NAME_None;
 	}
 
 	FScopedTransaction Transaction(LOCTEXT("SetPelvisBone_Label", "Set Pelvis Bone"));
 	FScopedReinitializeIKRig Reinitialize(this);
 	Asset->Modify();
-	Asset->RetargetDefinition.RootBone = NewRootBone;
+	Asset->RetargetDefinition.PelvisBone = NewPelvisBone;
 	return true;
 }
 
@@ -430,6 +430,29 @@ FName UIKRigController::GetRetargetRoot() const
 {
 	check(Asset)
 
+	return Asset->RetargetDefinition.PelvisBone;
+}
+
+bool UIKRigController::SetRootBone(const FName InRootBoneName) const
+{
+	check(Asset)
+
+	FName NewRootBone = InRootBoneName;
+	if (InRootBoneName != NAME_None && Asset->Skeleton.GetBoneIndexFromName(InRootBoneName) == INDEX_NONE)
+	{
+		NewRootBone = NAME_None;
+	}
+
+	FScopedTransaction Transaction(LOCTEXT("SetRootBone_Label", "Set Root Bone"));
+	FScopedReinitializeIKRig Reinitialize(this);
+	Asset->Modify();
+	Asset->RetargetDefinition.RootBone = NewRootBone;
+	return true;
+}
+
+FName UIKRigController::GetRootBone() const
+{
+	check(Asset);
 	return Asset->RetargetDefinition.RootBone;
 }
 
diff --git a/Engine/Plugins/Animation/IKRig/Source/IKRigEditor/Private/RigEditor/IKRigEditorController.cpp b/Engine/Plugins/Animation/IKRig/Source/IKRigEditor/Private/RigEditor/IKRigEditorController.cpp
--- a/Engine/Plugins/Animation/IKRig/Source/IKRigEditor/Private/RigEditor/IKRigEditorController.cpp
+++ b/Engine/Plugins/Animation/IKRig/Source/IKRigEditor/Private/RigEditor/IKRigEditorController.cpp
@@ -614,7 +614,8 @@ void FIKRigEditorController::AutoGenerateFBIK() const
 			Info.ExpireDuration = 3.0f;
 			FSlateNotificationManager::Get().AddNotification(Info);
 		}
-	case EAutoFBIKResult::MissingRootBone:
+		break;
+	case EAutoFBIKResult::MissingPelvisBone:
 		{
 			FNotificationInfo Info(LOCTEXT("AutoIKMissingRootBone", "Auto FBIK setup, but root bone was missing. Please assign a root bone manually."));
 			Info.ExpireDuration = 3.0f;
diff --git a/Engine/Plugins/Animation/IKRig/Source/IKRigEditor/Private/RigEditor/SIKRigRetargetChainList.cpp b/Engine/Plugins/Animation/IKRig/Source/IKRigEditor/Private/RigEditor/SIKRigRetargetChainList.cpp
--- a/Engine/Plugins/Animation/IKRig/Source/IKRigEditor/Private/RigEditor/SIKRigRetargetChainList.cpp
+++ b/Engine/Plugins/Animation/IKRig/Source/IKRigEditor/Private/RigEditor/SIKRigRetargetChainList.cpp
@@ -281,21 +281,13 @@ void SIKRigRetargetChainRow::OnRenameChain(const FText& InText, ETextCommit::Typ
 
 const FReferenceSkeleton& SIKRigRetargetChainRow::GetReferenceSkeleton() const
 {
-	static const FReferenceSkeleton DummySkeleton;
-	
-	const TSharedPtr<FIKRigEditorController> Controller = ChainList.Pin()->EditorController.Pin();
-	if (!Controller.IsValid())
+	if (ChainList.IsValid())
 	{
-		return DummySkeleton; 
-	}
-
-	USkeletalMesh* SkeletalMesh = Controller->AssetController->GetSkeletalMesh();
-	if (SkeletalMesh == nullptr)
-	{
-		return DummySkeleton;
+		return ChainList.Pin()->GetReferenceSkeleton();	
 	}
-
-	return SkeletalMesh->GetRefSkeleton();
+	
+	static const FReferenceSkeleton DummySkeleton;
+	return DummySkeleton;
 }
 
 void SIKRigRetargetChainList::Construct(const FArguments& InArgs, TSharedRef<FIKRigEditorController> InEditorController)
@@ -325,17 +317,40 @@ void SIKRigRetargetChainList::Construct(const FArguments& InArgs, TSharedRef<FIK
 				.Text(LOCTEXT("PelvisLabel", "Pelvis:"))
 				.TextStyle(FAppStyle::Get(), "NormalText")
 			]
-				
-			+ SHorizontalBox::Slot()
+        	
+        	+ SHorizontalBox::Slot()
+			.AutoWidth()
+			.HAlign(HAlign_Left)
+			.VAlign(VAlign_Center)
+			.Padding(3.0f, 1.0f)
+			[
+				SNew(SBoneSelectionWidget)
+				.OnBoneSelectionChanged(this, &SIKRigRetargetChainList::OnPelvisBoneComboSelectionChanged)
+				.OnGetSelectedBone(this, &SIKRigRetargetChainList::GetPelvisBoneName)
+				.OnGetReferenceSkeleton(this, &SIKRigRetargetChainList::GetReferenceSkeleton)
+			]
+
+        	+ SHorizontalBox::Slot()
 			.AutoWidth()
 			.HAlign(HAlign_Left)
 			.VAlign(VAlign_Center)
 			.Padding(3.0f, 3.0f)
 			[
-				SAssignNew(PelvisTextBox, SEditableTextBox)
-				.Text(FText::FromName(InEditorController->AssetController->GetRetargetRoot()))
-				.Font(FAppStyle::GetFontStyle(TEXT("BoldFont")))
-				.IsReadOnly(true)
+				SNew(STextBlock)
+				.Text(LOCTEXT("RootLabel", "Root Motion:"))
+				.TextStyle(FAppStyle::Get(), "NormalText")
+			]
+        	
+			+ SHorizontalBox::Slot()
+			.AutoWidth()
+			.HAlign(HAlign_Left)
+			.VAlign(VAlign_Center)
+			.Padding(3.0f, 1.0f)
+			[
+				SNew(SBoneSelectionWidget)
+				.OnBoneSelectionChanged(this, &SIKRigRetargetChainList::OnRootBoneComboSelectionChanged)
+				.OnGetSelectedBone(this, &SIKRigRetargetChainList::GetRootBoneName)
+				.OnGetReferenceSkeleton(this, &SIKRigRetargetChainList::GetReferenceSkeleton)
 			]
         ]
 
@@ -439,6 +454,73 @@ TArray<FName> SIKRigRetargetChainList::GetSelectedChains() const
 	return SelectedChainNames;
 }
 
+void SIKRigRetargetChainList::OnPelvisBoneComboSelectionChanged(FName InName)
+{
+	const TSharedPtr<FIKRigEditorController> Controller = EditorController.Pin();
+	if (!Controller.IsValid())
+	{
+		return; 
+	}
+	
+	Controller->AssetController->SetRetargetRoot(InName);
+	RefreshView();
+}
+
+void SIKRigRetargetChainList::OnRootBoneComboSelectionChanged(FName InName)
+{
+	const TSharedPtr<FIKRigEditorController> Controller = EditorController.Pin();
+	if (!Controller.IsValid())
+	{
+		return; 
+	}
+	
+	Controller->AssetController->SetRootBone(InName);
+	RefreshView();
+}
+
+FName SIKRigRetargetChainList::GetPelvisBoneName(bool& bMultipleValues) const
+{
+	const TSharedPtr<FIKRigEditorController> Controller = EditorController.Pin();
+	if (!Controller.IsValid())
+	{
+		return NAME_None; 
+	}
+
+	bMultipleValues = false;
+	return Controller->AssetController->GetRetargetRoot();
+}
+
+FName SIKRigRetargetChainList::GetRootBoneName(bool& bMultipleValues) const
+{
+	const TSharedPtr<FIKRigEditorController> Controller = EditorController.Pin();
+	if (!Controller.IsValid())
+	{
+		return NAME_None; 
+	}
+
+	bMultipleValues = false;
+	return Controller->AssetController->GetRootBone();
+}
+
+const FReferenceSkeleton& SIKRigRetargetChainList::GetReferenceSkeleton() const
+{
+	static const FReferenceSkeleton DummySkeleton;
+	
+	const TSharedPtr<FIKRigEditorController> Controller = EditorController.Pin();
+	if (!Controller.IsValid())
+	{
+		return DummySkeleton; 
+	}
+
+	USkeletalMesh* SkeletalMesh = Controller->AssetController->GetSkeletalMesh();
+	if (SkeletalMesh == nullptr)
+	{
+		return DummySkeleton;
+	}
+
+	return SkeletalMesh->GetRefSkeleton();
+}
+
 void SIKRigRetargetChainList::OnFilterTextChanged(const FText& SearchText)
 {
 	TextFilter->SetFilterText(SearchText);
@@ -563,9 +645,6 @@ void SIKRigRetargetChainList::RefreshView()
 		return TextFilter->TestTextFilter(FBasicStringFilterExpressionContext(StringToTest));
 	};
 	
-	// refresh pelvis
-	PelvisTextBox.Get()->SetText(FText::FromName(Controller->AssetController->GetRetargetRoot()));
-
 	// refresh list of chains
 	ListViewItems.Reset();
 	const TArray<FBoneChain>& Chains = Controller->AssetController->GetRetargetChains();
diff --git a/Engine/Plugins/Animation/IKRig/Source/IKRigEditor/Public/RigEditor/IKRigAutoCharacterizer.h b/Engine/Plugins/Animation/IKRig/Source/IKRigEditor/Public/RigEditor/IKRigAutoCharacterizer.h
--- a/Engine/Plugins/Animation/IKRig/Source/IKRigEditor/Public/RigEditor/IKRigAutoCharacterizer.h
+++ b/Engine/Plugins/Animation/IKRig/Source/IKRigEditor/Public/RigEditor/IKRigAutoCharacterizer.h
@@ -29,7 +29,6 @@ struct FCharacterizationStandard
 	//
 	
 	// core
-	static const FName Root;
 	static const FName Spine;
 	static const FName Neck;
 	static const FName Head;
diff --git a/Engine/Plugins/Animation/IKRig/Source/IKRigEditor/Public/RigEditor/IKRigAutoFBIK.h b/Engine/Plugins/Animation/IKRig/Source/IKRigEditor/Public/RigEditor/IKRigAutoFBIK.h
--- a/Engine/Plugins/Animation/IKRig/Source/IKRigEditor/Public/RigEditor/IKRigAutoFBIK.h
+++ b/Engine/Plugins/Animation/IKRig/Source/IKRigEditor/Public/RigEditor/IKRigAutoFBIK.h
@@ -13,7 +13,7 @@ enum class EAutoFBIKResult
 	MissingMesh,
 	UnknownSkeletonType,
 	MissingChains,
-	MissingRootBone
+	MissingPelvisBone
 };
 
 
diff --git a/Engine/Plugins/Animation/IKRig/Source/IKRigEditor/Public/RigEditor/IKRigController.h b/Engine/Plugins/Animation/IKRig/Source/IKRigEditor/Public/RigEditor/IKRigController.h
--- a/Engine/Plugins/Animation/IKRig/Source/IKRigEditor/Public/RigEditor/IKRigController.h
+++ b/Engine/Plugins/Animation/IKRig/Source/IKRigEditor/Public/RigEditor/IKRigController.h
@@ -345,14 +345,22 @@ public:
 	UFUNCTION(BlueprintCallable, Category="IK Rig")
 	UE_API const TArray<FBoneChain>& GetRetargetChains() const;
 
-	// Set the Root Bone of the retargeting (can only be one). 
+	// Set the Pelvis Bone of the retargeting (can only be one). 
 	UFUNCTION(BlueprintCallable, BlueprintPure=false, Category="IK Rig", DisplayName="SetPelvisBone")
 	UE_API bool SetRetargetRoot(const FName RootBoneName) const;
 
-	// Get the name of the Root Bone of the retargeting (can only be one). 
+	// Get the name of the Pelvis Bone of the retargeting (can only be one). 
 	UFUNCTION(BlueprintCallable, Category="IK Rig", DisplayName="GetPelvisBone")
 	UE_API FName GetRetargetRoot() const;
 
+	// Set the root bone of the skeleton (used for root motion). 
+	UFUNCTION(BlueprintCallable, BlueprintPure=false, Category="IK Rig")
+	UE_API bool SetRootBone(const FName RootBoneName) const;
+
+	// Get the root bone of the skeleton (used for root motion).
+	UFUNCTION(BlueprintCallable, Category="IK Rig")
+	UE_API FName GetRootBone() const;
+
 	// Analyse the skeleton to see if it matches a known template and automatically generates all retarget chains and sets the pelvis
 	// Returns true if a matching skeletal template was found and the retarget definition for it was applied.
 	UFUNCTION(BlueprintCallable, Category="IK Rig")
diff --git a/Engine/Plugins/Animation/IKRig/Source/IKRigEditor/Public/RigEditor/SIKRigRetargetChainList.h b/Engine/Plugins/Animation/IKRig/Source/IKRigEditor/Public/RigEditor/SIKRigRetargetChainList.h
--- a/Engine/Plugins/Animation/IKRig/Source/IKRigEditor/Public/RigEditor/SIKRigRetargetChainList.h
+++ b/Engine/Plugins/Animation/IKRig/Source/IKRigEditor/Public/RigEditor/SIKRigRetargetChainList.h
@@ -135,8 +135,12 @@ private:
 	/** editor controller */
 	TWeakPtr<FIKRigEditorController> EditorController;
 
-	/** pelvis widget */
-	TSharedPtr<SEditableTextBox> PelvisTextBox;
+	/** pelvis / root bone widgets */
+	void OnPelvisBoneComboSelectionChanged(FName InName);
+	void OnRootBoneComboSelectionChanged(FName InName);
+	FName GetPelvisBoneName(bool& bMultipleValues) const;
+	FName GetRootBoneName(bool& bMultipleValues) const;
+	const FReferenceSkeleton& GetReferenceSkeleton() const;
 	
 	/** list view */
 	TSharedPtr<SRetargetChainListViewType> ListView;
