[IK Retarget] Stretch Chain op now lazy initializes chains.

#Synced-CL 48918725

diff --git a/Engine/Plugins/Animation/IKRig/Source/IKRig/Private/Retargeter/RetargetOps/StretchChainOp.cpp b/Engine/Plugins/Animation/IKRig/Source/IKRig/Private/Retargeter/RetargetOps/StretchChainOp.cpp
--- a/Engine/Plugins/Animation/IKRig/Source/IKRig/Private/Retargeter/RetargetOps/StretchChainOp.cpp
+++ b/Engine/Plugins/Animation/IKRig/Source/IKRig/Private/Retargeter/RetargetOps/StretchChainOp.cpp
@@ -10,17 +10,21 @@
 
 #define LOCTEXT_NAMESPACE "StretchChainOp"
 
-bool FChainStretcher::Initialize(
+FChainStretcher::FChainStretcher(
 	const FRetargetStretchChainSettings* InSettings,
 	const FResolvedBoneChain& InSourceBoneChain,
-	const FResolvedBoneChain& InTargetBoneChain,
-	const FRetargetSkeleton& InSourceSkeleton,
-	const FRetargetSkeleton& InTargetSkeleton)
+	const FResolvedBoneChain& InTargetBoneChain)
 {
 	Settings = InSettings;
 	SourceBoneChain = &InSourceBoneChain;
 	TargetBoneChain = &InTargetBoneChain;
 
+	// lazy initialized to avoid wasting memory
+	bIsInitialized = false;
+}
+
+void FChainStretcher::Initialize(const FRetargetSkeleton& InTargetSkeleton)
+{
 	// cache indices of all children of all bones in chain (recursively)
 	AllChildrenOfChain.Reset();
 	InTargetSkeleton.GetChildrenIndicesRecursive(TargetBoneChain->BoneIndices[0],AllChildrenOfChain);
@@ -30,15 +34,25 @@ bool FChainStretcher::Initialize(
 	{
 		return IndicesToRemove.Contains(Element);
 	});
-	
-	return true;
+
+	bIsInitialized = true;
 }
 
 void FChainStretcher::StretchChain(
 	const TArray<FTransform> &InSourceGlobalPose,
 	const FRetargetSkeleton& InTargetSkeleton,
-	TArray<FTransform> &OutTargetGlobalPose) const
+	TArray<FTransform> &OutTargetGlobalPose)
 {
+	if (!(Settings && Settings->bEnabled))
+	{
+		return;
+	}
+
+	if (!bIsInitialized)
+	{
+		Initialize(InTargetSkeleton);
+	}
+	
 	const bool bMatchingSourceLength = !FMath::IsNearlyZero(Settings->MatchSourceLength);
 	const bool bScalingChainLength = !FMath::IsNearlyEqual(Settings->ScaleChainLength, 1.0f);
 	if (!bMatchingSourceLength && !bScalingChainLength)
@@ -113,11 +127,6 @@ bool FIKRetargetStretchChainOp::Initialize(
 	const FRetargeterBoneChains& BoneChains = InProcessor.GetBoneChains();
 	for (const FRetargetStretchChainSettings& ChainSettings : Settings.ChainsToStretch)
 	{
-		if (!ChainSettings.bEnabled)
-		{
-			continue; // if the enabled flag is toggled it will trigger reinit to make a Chain Stretcher
-		}
-		
 		const FResolvedBoneChain* TargetBoneChain = BoneChains.GetResolvedBoneChainByName(
 			ChainSettings.TargetChainName,
 			ERetargetSourceOrTarget::Target,
@@ -141,29 +150,12 @@ bool FIKRetargetStretchChainOp::Initialize(
 			continue;
 		}
 		
-		// initialize the mapped pair of source/target bone chains
-		FChainStretcher ChainStretcher;
-		const bool bChainStretcherInitialized = ChainStretcher.Initialize(
-			&ChainSettings,
-			*SourceBoneChain,
-			*TargetBoneChain,
-			InSourceSkeleton,
-			InTargetSkeleton);
-		if (!bChainStretcherInitialized)
-		{
-			InLog.LogWarning( FText::Format(
-				LOCTEXT("ChainStretchOpFailedInit", "Chain Stretch Op: failed to initialize stretcher for chain, '{0}', on Skeletal Mesh: '{1}'"),
-				FText::FromName(ChainSettings.TargetChainName), FText::FromString(InTargetSkeleton.SkeletalMesh->GetName())));
-			return false;
-		}
-
-		// store valid chain pair to be retargeted
-		ChainStretchers.Add(ChainStretcher);
+		// store a chain stretcher
+		ChainStretchers.Emplace(&ChainSettings, *SourceBoneChain, *TargetBoneChain);
 	}
 
-	bIsInitialized = !ChainStretchers.IsEmpty();
-
-	return bIsInitialized;
+	bIsInitialized = true;
+	return true;
 }
 
 void FIKRetargetStretchChainOp::Run(
diff --git a/Engine/Plugins/Animation/IKRig/Source/IKRig/Public/Retargeter/RetargetOps/StretchChainOp.h b/Engine/Plugins/Animation/IKRig/Source/IKRig/Public/Retargeter/RetargetOps/StretchChainOp.h
--- a/Engine/Plugins/Animation/IKRig/Source/IKRig/Public/Retargeter/RetargetOps/StretchChainOp.h
+++ b/Engine/Plugins/Animation/IKRig/Source/IKRig/Public/Retargeter/RetargetOps/StretchChainOp.h
@@ -15,20 +15,21 @@ struct FRetargetStretchChainSettings;
 
 struct FChainStretcher
 {
-	bool Initialize(
+	FChainStretcher(
 		const FRetargetStretchChainSettings* InSettings,
 		const FResolvedBoneChain& InSourceBoneChain,
-		const FResolvedBoneChain& InTargetBoneChain,
-		const FRetargetSkeleton& InSourceSkeleton,
-		const FRetargetSkeleton& InTargetSkeleton);
+		const FResolvedBoneChain& InTargetBoneChain);
+	
+	void Initialize(const FRetargetSkeleton& InTargetSkeleton);
 
 	void StretchChain(
 		const TArray<FTransform> &InSourceGlobalPose,
 		const FRetargetSkeleton& InTargetSkeleton,
-		TArray<FTransform> &OutTargetGlobalPose) const;
+		TArray<FTransform> &OutTargetGlobalPose);
 	
 private:
-	
+
+	bool bIsInitialized = false;
 	TArray<int32> AllChildrenOfChain;
 	const FRetargetStretchChainSettings* Settings;
 	const FResolvedBoneChain* SourceBoneChain;
@@ -48,7 +49,7 @@ struct FRetargetStretchChainSettings
 	FName TargetChainName;
 
 	/** Enable/disable stretching. */
-	UPROPERTY(EditAnywhere, BlueprintReadWrite, Category = "Stretch Chain", meta=(ReinitializeOnEdit))
+	UPROPERTY(EditAnywhere, BlueprintReadWrite, Category = "Stretch Chain")
 	bool bEnabled = false;
 
 	/** Range 0 to 1. Default 0. Matches the total length of this chain with the mapped source chain.
