[IK Retarget] Renamed SetRootBone to SetRootMotionBone to avoid function redirect name clash.

#Synced-CL 49046516

diff --git a/Engine/Plugins/Animation/IKRig/Source/IKRigEditor/Private/RigEditor/IKRigController.cpp b/Engine/Plugins/Animation/IKRig/Source/IKRigEditor/Private/RigEditor/IKRigController.cpp
--- a/Engine/Plugins/Animation/IKRig/Source/IKRigEditor/Private/RigEditor/IKRigController.cpp
+++ b/Engine/Plugins/Animation/IKRig/Source/IKRigEditor/Private/RigEditor/IKRigController.cpp
@@ -433,7 +433,7 @@ FName UIKRigController::GetRetargetRoot() const
 	return Asset->RetargetDefinition.PelvisBone;
 }
 
-bool UIKRigController::SetRootBone(const FName InRootBoneName) const
+bool UIKRigController::SetRootMotionBone(const FName InRootBoneName) const
 {
 	check(Asset)
 
@@ -450,7 +450,7 @@ bool UIKRigController::SetRootBone(const FName InRootBoneName) const
 	return true;
 }
 
-FName UIKRigController::GetRootBone() const
+FName UIKRigController::GetRootMotionBone() const
 {
 	check(Asset);
 	return Asset->RetargetDefinition.RootBone;
diff --git a/Engine/Plugins/Animation/IKRig/Source/IKRigEditor/Private/RigEditor/SIKRigRetargetChainList.cpp b/Engine/Plugins/Animation/IKRig/Source/IKRigEditor/Private/RigEditor/SIKRigRetargetChainList.cpp
--- a/Engine/Plugins/Animation/IKRig/Source/IKRigEditor/Private/RigEditor/SIKRigRetargetChainList.cpp
+++ b/Engine/Plugins/Animation/IKRig/Source/IKRigEditor/Private/RigEditor/SIKRigRetargetChainList.cpp
@@ -474,7 +474,7 @@ void SIKRigRetargetChainList::OnRootBoneComboSelectionChanged(FName InName)
 		return; 
 	}
 	
-	Controller->AssetController->SetRootBone(InName);
+	Controller->AssetController->SetRootMotionBone(InName);
 	RefreshView();
 }
 
@@ -499,7 +499,7 @@ FName SIKRigRetargetChainList::GetRootBoneName(bool& bMultipleValues) const
 	}
 
 	bMultipleValues = false;
-	return Controller->AssetController->GetRootBone();
+	return Controller->AssetController->GetRootMotionBone();
 }
 
 const FReferenceSkeleton& SIKRigRetargetChainList::GetReferenceSkeleton() const
diff --git a/Engine/Plugins/Animation/IKRig/Source/IKRigEditor/Public/RigEditor/IKRigController.h b/Engine/Plugins/Animation/IKRig/Source/IKRigEditor/Public/RigEditor/IKRigController.h
--- a/Engine/Plugins/Animation/IKRig/Source/IKRigEditor/Public/RigEditor/IKRigController.h
+++ b/Engine/Plugins/Animation/IKRig/Source/IKRigEditor/Public/RigEditor/IKRigController.h
@@ -353,13 +353,13 @@ public:
 	UFUNCTION(BlueprintCallable, Category="IK Rig", DisplayName="GetPelvisBone")
 	UE_API FName GetRetargetRoot() const;
 
-	// Set the root bone of the skeleton (used for root motion). 
+	// Set the bone used to store root motion for the skeleton (could be different from the root of the hierarchy).
 	UFUNCTION(BlueprintCallable, BlueprintPure=false, Category="IK Rig")
-	UE_API bool SetRootBone(const FName RootBoneName) const;
+	UE_API bool SetRootMotionBone(const FName RootBoneName) const;
 
-	// Get the root bone of the skeleton (used for root motion).
+	// Get the bone used to store root motion for the skeleton (could be different from the root of the hierarchy).
 	UFUNCTION(BlueprintCallable, Category="IK Rig")
-	UE_API FName GetRootBone() const;
+	UE_API FName GetRootMotionBone() const;
 
 	// Analyse the skeleton to see if it matches a known template and automatically generates all retarget chains and sets the pelvis
 	// Returns true if a matching skeletal template was found and the retarget definition for it was applied.
