[IK Retarget] Fixed controller settings type and name for floor constraint op.

#Synced-CL 49060511

diff --git a/Engine/Plugins/Animation/IKRig/Source/IKRig/Private/Retargeter/RetargetOps/FloorConstraintOp.cpp b/Engine/Plugins/Animation/IKRig/Source/IKRig/Private/Retargeter/RetargetOps/FloorConstraintOp.cpp
--- a/Engine/Plugins/Animation/IKRig/Source/IKRig/Private/Retargeter/RetargetOps/FloorConstraintOp.cpp
+++ b/Engine/Plugins/Animation/IKRig/Source/IKRig/Private/Retargeter/RetargetOps/FloorConstraintOp.cpp
@@ -26,7 +26,7 @@ bool FFloorConstraintChainSettings::operator==(const FFloorConstraintChainSettin
 
 const UClass* FIKRetargetFloorConstraintOpSettings::GetControllerType() const
 {
-	return UIKRetargetFloorGoalsController::StaticClass();
+	return UIKRetargetFloorConstraintController::StaticClass();
 }
 
 void FIKRetargetFloorConstraintOpSettings::CopySettingsAtRuntime(const FIKRetargetOpSettingsBase* InSettingsToCopyFrom)
@@ -807,12 +807,12 @@ void FIKRetargetFloorConstraintOp::DebugDraw(
 
 #endif // WITH_EDITOR
 
-FIKRetargetIKChainsOpSettings UIKRetargetFloorGoalsController::GetSettings()
+FIKRetargetFloorConstraintOpSettings UIKRetargetFloorConstraintController::GetSettings()
 {
-	return *reinterpret_cast<FIKRetargetIKChainsOpSettings*>(OpSettingsToControl);
+	return *reinterpret_cast<FIKRetargetFloorConstraintOpSettings*>(OpSettingsToControl);
 }
 
-void UIKRetargetFloorGoalsController::SetSettings(FIKRetargetIKChainsOpSettings InSettings)
+void UIKRetargetFloorConstraintController::SetSettings(FIKRetargetFloorConstraintOpSettings InSettings)
 {
 	OpSettingsToControl->CopySettingsAtRuntime(&InSettings);
 }
diff --git a/Engine/Plugins/Animation/IKRig/Source/IKRig/Public/Retargeter/RetargetOps/FloorConstraintOp.h b/Engine/Plugins/Animation/IKRig/Source/IKRig/Public/Retargeter/RetargetOps/FloorConstraintOp.h
--- a/Engine/Plugins/Animation/IKRig/Source/IKRig/Public/Retargeter/RetargetOps/FloorConstraintOp.h
+++ b/Engine/Plugins/Animation/IKRig/Source/IKRig/Public/Retargeter/RetargetOps/FloorConstraintOp.h
@@ -361,22 +361,22 @@ private:
 	TArray<FFloorConstraint> FloorConstraints;
 };
 
-/* The blueprint/python API for editing a Floor Goal Op */
+/* The blueprint/python API for editing a Floor Constraint Op */
 UCLASS(MinimalAPI, BlueprintType)
-class UIKRetargetFloorGoalsController : public UIKRetargetOpControllerBase
+class UIKRetargetFloorConstraintController : public UIKRetargetOpControllerBase
 {
 	GENERATED_BODY()
 	
 public:
 	/* Get the current op settings as a struct.
-	 * @return FIKRetargetIKChainsOpSettings struct with the current settings used by the op. */
+	 * @return FIKRetargetFloorConstraintOpSettings struct with the current settings used by the op. */
 	UFUNCTION(BlueprintCallable, Category = Settings)
-	UE_API FIKRetargetIKChainsOpSettings GetSettings();
+	UE_API FIKRetargetFloorConstraintOpSettings GetSettings();
 
 	/* Set the op settings. Input is a custom struct type for this op.
-	 * @param InSettings a FIKRetargetIKChainsOpSettings struct containing all the settings to apply to this op */
+	 * @param InSettings a FIKRetargetFloorConstraintOpSettings struct containing all the settings to apply to this op */
 	UFUNCTION(BlueprintCallable, Category = Settings)
-	UE_API void SetSettings(FIKRetargetIKChainsOpSettings InSettings);
+	UE_API void SetSettings(FIKRetargetFloorConstraintOpSettings InSettings);
 };
 
 #undef LOCTEXT_NAMESPACE
