[IK Retarget] Fixes performance regression causing constant IK Rig reinitialization.

#Synced-CL 49078080

diff --git a/Engine/Plugins/Animation/IKRig/Source/IKRig/Private/AnimNodes/RigUnit_IKRig.cpp b/Engine/Plugins/Animation/IKRig/Source/IKRig/Private/AnimNodes/RigUnit_IKRig.cpp
--- a/Engine/Plugins/Animation/IKRig/Source/IKRig/Private/AnimNodes/RigUnit_IKRig.cpp
+++ b/Engine/Plugins/Animation/IKRig/Source/IKRig/Private/AnimNodes/RigUnit_IKRig.cpp
@@ -77,12 +77,6 @@ FRigUnit_IKRig_Execute()
 		}
 	}
 	
-	// trigger reinitialization if the goal container was modified such that it needs it
-	if (Processor.GetGoalContainer().NeedsInitialized())
-	{
-		Processor.Initialize(IKRigAsset, GetSkeletalMesh(), FIKRigGoalContainer());
-	}
-	
 	// copy input pose to start IK solve from
 	for (int32 BoneIndex=0; BoneIndex<WorkData.CachedRigElements.Num(); BoneIndex++)
 	{
diff --git a/Engine/Plugins/Animation/IKRig/Source/IKRig/Private/Retargeter/RetargetOps/RunIKRigOp.cpp b/Engine/Plugins/Animation/IKRig/Source/IKRig/Private/Retargeter/RetargetOps/RunIKRigOp.cpp
--- a/Engine/Plugins/Animation/IKRig/Source/IKRig/Private/Retargeter/RetargetOps/RunIKRigOp.cpp
+++ b/Engine/Plugins/Animation/IKRig/Source/IKRig/Private/Retargeter/RetargetOps/RunIKRigOp.cpp
@@ -295,22 +295,16 @@ void FIKRetargetRunIKRigOp::Run(
 	}
 #endif
 
-	FIKRigGoalContainer& GoalContainer = InProcessor.GetIKRigGoalContainer();
+	FIKRigGoalContainer& GoalsFromOpStack = InProcessor.GetIKRigGoalContainer();
 	
 	// set main alpha on the goals
 	for (FRunIKRigChain& Chain : Chains)
 	{
-		Chain.SetGoalAlpha(GoalContainer);
+		Chain.SetGoalAlpha(GoalsFromOpStack);
 	}
 
 	// copy the goal container from the retargeter to the IK Rig processor
-	IKRigProcessor.ApplyGoalsFromOtherContainer(GoalContainer);
-	// trigger reinitialization if the goal container was modified such that it needs it
-	if (GoalContainer.NeedsInitialized())
-	{
-		const USkeletalMesh* TargetSkeletalMesh = InProcessor.GetSkeleton(ERetargetSourceOrTarget::Target).SkeletalMesh;
-		IKRigProcessor.Initialize(Settings.IKRigAsset, TargetSkeletalMesh, GoalContainer);
-	}
+	IKRigProcessor.ApplyGoalsFromOtherContainer(GoalsFromOpStack);
 	
 #if WITH_EDITOR
 	// store initial goals transforms
diff --git a/Engine/Plugins/Animation/IKRig/Source/IKRig/Private/Rig/IKRigProcessor.cpp b/Engine/Plugins/Animation/IKRig/Source/IKRig/Private/Rig/IKRigProcessor.cpp
--- a/Engine/Plugins/Animation/IKRig/Source/IKRig/Private/Rig/IKRigProcessor.cpp
+++ b/Engine/Plugins/Animation/IKRig/Source/IKRig/Private/Rig/IKRigProcessor.cpp
@@ -204,6 +204,7 @@ void FIKRigProcessor::Initialize(
 				LOCTEXT("SuccessfulInit", "IK Rig, '{0}' ready to run on {1}."),
 				FText::FromString(InRigAsset->GetName()), FText::FromString(InSkeletalMesh->GetName())));
 
+	AssetInitializedWith = InRigAsset;
 	GoalContainer.bRigNeedsInitialized = false;
 	bTriedToInitialize = false;
 	bInitialized = true;
@@ -367,11 +368,25 @@ void FIKRigProcessor::SetIKGoal(const UIKRigEffectorGoal* InGoal)
 
 void FIKRigProcessor::Solve(const FTransform& ComponentToWorld)
 {
+	// trigger reinitialization if the goal container was modified such that it needs it
+	if (GoalContainer.NeedsInitialized())
+	{
+		// save input pose
+		TArray<FTransform> InputPoseGlobal = Skeleton.CurrentPoseGlobal;
+		// make a copy of the goal container in it's current state to reinitialize with
+		FIKRigGoalContainer CopiedGoalContainer = GoalContainer;
+		// reinitialize
+		Initialize(AssetInitializedWith, Skeleton.SkeletalMesh, CopiedGoalContainer);
+		// restore input pose after reinit
+		SetInputPoseGlobal(InputPoseGlobal);
+	}
+
 	if (!bInitialized)
 	{
 		return;
 	}
 
+	// generate local space pose
 	Skeleton.UpdateAllLocalTransformFromGlobal();
 	
 	// convert goals into component space and blend towards input pose by alpha
diff --git a/Engine/Plugins/Animation/IKRig/Source/IKRig/Private/Rig/IKRigSkeleton.cpp b/Engine/Plugins/Animation/IKRig/Source/IKRig/Private/Rig/IKRigSkeleton.cpp
--- a/Engine/Plugins/Animation/IKRig/Source/IKRig/Private/Rig/IKRigSkeleton.cpp
+++ b/Engine/Plugins/Animation/IKRig/Source/IKRig/Private/Rig/IKRigSkeleton.cpp
@@ -17,16 +17,16 @@ FIKRigInputSkeleton::FIKRigInputSkeleton(const USkeletalMesh* SkeletalMesh)
 	Initialize(SkeletalMesh);
 }
 
-void FIKRigInputSkeleton::Initialize(const USkeletalMesh* SkeletalMesh)
+void FIKRigInputSkeleton::Initialize(const USkeletalMesh* InSkeletalMesh)
 {
 	Reset();
 
-	if (!SkeletalMesh)
+	if (!InSkeletalMesh)
 	{
 		return;
 	}
 	
-	SkeletalMeshName = FName(SkeletalMesh->GetName());
+	SkeletalMesh = InSkeletalMesh;
 	const FReferenceSkeleton& RefSkeleton = SkeletalMesh->GetRefSkeleton();
 	const TArray<FMeshBoneInfo>& BoneInfo = RefSkeleton.GetRefBoneInfo();
 	for (int32 BoneIndex=0; BoneIndex<BoneInfo.Num(); ++BoneIndex)
@@ -42,9 +42,9 @@ void FIKRigInputSkeleton::Reset()
 	*this = FIKRigInputSkeleton();
 }
 
-void FIKRigSkeleton::SetInputSkeleton(const USkeletalMesh* SkeletalMesh, const TArray<FName>& InExcludedBones)
+void FIKRigSkeleton::SetInputSkeleton(const USkeletalMesh* InSkeletalMesh, const TArray<FName>& InExcludedBones)
 {
-	const FIKRigInputSkeleton InputSkeleton = FIKRigInputSkeleton(SkeletalMesh);
+	const FIKRigInputSkeleton InputSkeleton = FIKRigInputSkeleton(InSkeletalMesh);
 	SetInputSkeleton(InputSkeleton, InExcludedBones);
 }
 
@@ -55,6 +55,9 @@ void FIKRigSkeleton::SetInputSkeleton(const FIKRigInputSkeleton& InputSkeleton,
 	
 	// reset all containers
 	Reset();
+
+	// store skeletal mesh for reinit
+	SkeletalMesh = InputSkeleton.SkeletalMesh;
 	
 	// we use the bone names and parent indices from the input skeleton
 	BoneNames = InputSkeleton.BoneNames;
diff --git a/Engine/Plugins/Animation/IKRig/Source/IKRig/Public/Rig/IKRigProcessor.h b/Engine/Plugins/Animation/IKRig/Source/IKRig/Public/Rig/IKRigProcessor.h
--- a/Engine/Plugins/Animation/IKRig/Source/IKRig/Public/Rig/IKRigProcessor.h
+++ b/Engine/Plugins/Animation/IKRig/Source/IKRig/Public/Rig/IKRigProcessor.h
@@ -138,6 +138,9 @@ private:
 	/** storage for hierarchy and bone transforms */
 	FIKRigSkeleton Skeleton;
 
+	/** store asset for dynamic reinitialization */
+	TObjectPtr<const UIKRigDefinition> AssetInitializedWith;
+
 	/** true if processor is currently initialized and ready to run */
 	bool bInitialized = false;
 	/** true if processor attempted an initialization but failed (prevents infinite reinit) */
diff --git a/Engine/Plugins/Animation/IKRig/Source/IKRig/Public/Rig/IKRigSkeleton.h b/Engine/Plugins/Animation/IKRig/Source/IKRig/Public/Rig/IKRigSkeleton.h
--- a/Engine/Plugins/Animation/IKRig/Source/IKRig/Public/Rig/IKRigSkeleton.h
+++ b/Engine/Plugins/Animation/IKRig/Source/IKRig/Public/Rig/IKRigSkeleton.h
@@ -31,13 +31,13 @@ struct FIKRigInputSkeleton
 	TArray<FName> BoneNames;
 	TArray<int32> ParentIndices;
 	TArray<FTransform> LocalRefPose;
-	FName SkeletalMeshName;
+	const USkeletalMesh* SkeletalMesh;
 
 	FIKRigInputSkeleton() = default;
 	
-	UE_API FIKRigInputSkeleton(const USkeletalMesh* SkeletalMesh);
+	UE_API FIKRigInputSkeleton(const USkeletalMesh* InSkeletalMesh);
 
-	UE_API void Initialize(const USkeletalMesh* SkeletalMesh);
+	UE_API void Initialize(const USkeletalMesh* InSkeletalMesh);
 
 	UE_API void Reset();
 };
@@ -47,6 +47,10 @@ struct FIKRigSkeleton
 {
 	GENERATED_BODY()
 
+	/** The skeletal mesh this was initialized with. */
+	UPROPERTY(VisibleAnywhere, Category = Skeleton)
+	TObjectPtr<const USkeletalMesh> SkeletalMesh;
+	
 	/** Names of bones. Used to match hierarchy with runtime skeleton. */
 	UPROPERTY(VisibleAnywhere, Category = Skeleton)
 	TArray<FName> BoneNames;
@@ -71,7 +75,7 @@ struct FIKRigSkeleton
 	UPROPERTY(VisibleAnywhere, Category = Skeleton)
 	TArray<FTransform> RefPoseGlobal;
 
-	UE_API void SetInputSkeleton(const USkeletalMesh* SkeletalMesh, const TArray<FName>& InExcludedBones);
+	UE_API void SetInputSkeleton(const USkeletalMesh* InSkeletalMesh, const TArray<FName>& InExcludedBones);
 
 	UE_API void SetInputSkeleton(const FIKRigInputSkeleton& InputSkeleton, const TArray<FName>& InExcludedBones);
 
