[IK Retarget] Optimization in copying chain settings.

#Synced-CL 49078080

diff --git a/Engine/Plugins/Animation/IKRig/Source/IKRig/Public/Retargeter/IKRetargetOpUtils.h b/Engine/Plugins/Animation/IKRig/Source/IKRig/Public/Retargeter/IKRetargetOpUtils.h
--- a/Engine/Plugins/Animation/IKRig/Source/IKRig/Public/Retargeter/IKRetargetOpUtils.h
+++ b/Engine/Plugins/Animation/IKRig/Source/IKRig/Public/Retargeter/IKRetargetOpUtils.h
@@ -171,6 +171,7 @@ namespace IKRetargetOpUtils
 				if (SettingsToCopyFrom.GetName() == SettingsToCopyTo.GetName())
 				{
 					SettingsToCopyTo = SettingsToCopyFrom;
+					break;
 				}
 			}
 		}
