[IKRetarget] Adding profile scope to retarget processor.

#Synced-CL 49078080

diff --git a/Engine/Plugins/Animation/IKRig/Source/IKRig/Private/Retargeter/IKRetargetProcessor.cpp b/Engine/Plugins/Animation/IKRig/Source/IKRig/Private/Retargeter/IKRetargetProcessor.cpp
--- a/Engine/Plugins/Animation/IKRig/Source/IKRig/Private/Retargeter/IKRetargetProcessor.cpp
+++ b/Engine/Plugins/Animation/IKRig/Source/IKRig/Private/Retargeter/IKRetargetProcessor.cpp
@@ -16,6 +16,8 @@
 
 #include UE_INLINE_GENERATED_CPP_BY_NAME(IKRetargetProcessor)
 
+DEFINE_STAT(STAT_Anim_RetargetProcessorRun);
+
 #define LOCTEXT_NAMESPACE "IKRetargetProcessor"
 
 // This is the default end of branch index value, meaning we haven't cached it yet
@@ -1283,6 +1285,8 @@ TArray<FTransform>&  FIKRetargetProcessor::RunRetargeter(
 	const float InDeltaTime,
 	const int32 InLOD)
 {
+	SCOPE_CYCLE_COUNTER(STAT_Anim_RetargetProcessorRun);
+	
 	if (!ensure(bIsInitialized))
 	{
 		return InSourceGlobalPose;
diff --git a/Engine/Plugins/Animation/IKRig/Source/IKRig/Public/Retargeter/IKRetargetProcessor.h b/Engine/Plugins/Animation/IKRig/Source/IKRig/Public/Retargeter/IKRetargetProcessor.h
--- a/Engine/Plugins/Animation/IKRig/Source/IKRig/Public/Retargeter/IKRetargetProcessor.h
+++ b/Engine/Plugins/Animation/IKRig/Source/IKRig/Public/Retargeter/IKRetargetProcessor.h
@@ -14,6 +14,8 @@
 
 #define UE_API IKRIG_API
 
+DECLARE_CYCLE_STAT_EXTERN(TEXT("IK Retarget Processor Run"), STAT_Anim_RetargetProcessorRun, STATGROUP_Anim, UE_API);
+
 struct FIKRetargetOpBase;
 enum class ERetargetTranslationMode : uint8;
 enum class ERetargetRotationMode : uint8;
