[IKRig] Auto remove null goals from goal container on load.

#Synced-CL 49078080

diff --git a/Engine/Plugins/Animation/IKRig/Source/IKRig/Private/Rig/IKRigDefinition.cpp b/Engine/Plugins/Animation/IKRig/Source/IKRig/Private/Rig/IKRigDefinition.cpp
--- a/Engine/Plugins/Animation/IKRig/Source/IKRig/Private/Rig/IKRigDefinition.cpp
+++ b/Engine/Plugins/Animation/IKRig/Source/IKRig/Private/Rig/IKRigDefinition.cpp
@@ -378,6 +378,12 @@ void UIKRigDefinition::PostLoad()
 	{
 		return !InSolver.IsValid();
 	});
+
+	// clean out null goals
+	Goals.RemoveAll([](const TObjectPtr<UIKRigEffectorGoal>& InGoal)
+	{
+		return !InGoal;
+	});
 }
 
 const TArray<FInstancedStruct>& UIKRigDefinition::GetSolverStructs() const
