[IKRetarget] Fixed issue preventing export when source or target IK rig are not present.

#Synced-CL 49217940

diff --git a/Engine/Plugins/Animation/IKRig/Source/IKRigEditor/Private/RetargetEditor/IKRetargetBatchOperation.cpp b/Engine/Plugins/Animation/IKRig/Source/IKRigEditor/Private/RetargetEditor/IKRetargetBatchOperation.cpp
--- a/Engine/Plugins/Animation/IKRig/Source/IKRigEditor/Private/RetargetEditor/IKRetargetBatchOperation.cpp
+++ b/Engine/Plugins/Animation/IKRig/Source/IKRigEditor/Private/RetargetEditor/IKRetargetBatchOperation.cpp
@@ -898,27 +898,17 @@ void UIKRetargetBatchOperation::RunRetarget(FIKRetargetBatchOperationContext& Co
 		return;
 	}
 
-	// validate a source IK rig was provided
-	const UIKRigDefinition* SrcIKRig = Context.IKRetargetAsset->GetIKRig(ERetargetSourceOrTarget::Source);
-	if (!SrcIKRig)
-	{
-		UE_LOG(LogTemp, Warning, TEXT("Batch retarget aborted. Specified IK Retargeter does not reference a source IK Rig."));
-		return;
-	}
-
-	// validate a target IK rig was provided
-	const UIKRigDefinition* TgtIKRig = Context.IKRetargetAsset->GetIKRig(ERetargetSourceOrTarget::Target);
-	if (!TgtIKRig)
-	{
-		UE_LOG(LogTemp, Warning, TEXT("Batch retarget aborted. Specified IK Retargeter does not reference a target IK Rig."));
-		return;
-	}
-
+	
 	// validate a source mesh was provided
 	if (!Context.SourceMesh)
 	{
 		// fallback to mesh provided by IK Rig
-		Context.SourceMesh = SrcIKRig->GetPreviewMesh();
+		const UIKRigDefinition* SrcIKRig = Context.IKRetargetAsset->GetIKRig(ERetargetSourceOrTarget::Source);
+		if (SrcIKRig)
+		{
+			Context.SourceMesh = SrcIKRig->GetPreviewMesh();
+		}
+		
 		if (!Context.SourceMesh)
 		{
 			UE_LOG(LogTemp, Warning, TEXT("Batch retarget aborted. No source mesh was specified and the source IK Rig did not have one. "));
@@ -930,7 +920,12 @@ void UIKRetargetBatchOperation::RunRetarget(FIKRetargetBatchOperationContext& Co
 	if (!Context.TargetMesh)
 	{
 		// fallback to mesh provided by IK Rig
-		Context.TargetMesh = TgtIKRig->GetPreviewMesh();
+		const UIKRigDefinition* TgtIKRig = Context.IKRetargetAsset->GetIKRig(ERetargetSourceOrTarget::Target);
+		if (TgtIKRig)
+		{
+			Context.TargetMesh = TgtIKRig->GetPreviewMesh();
+		}
+		
 		if (!Context.TargetMesh)
 		{
 			UE_LOG(LogTemp, Warning, TEXT("Batch retarget aborted. No target mesh was specified and the target IK Rig did not have one. "));
