[IK Retarget] Removed default ops from new retargeter

#Synced-CL 49217940

diff --git a/Engine/Plugins/Animation/IKRig/Source/IKRigEditor/Private/RetargetEditor/IKRetargetFactory.cpp b/Engine/Plugins/Animation/IKRig/Source/IKRigEditor/Private/RetargetEditor/IKRetargetFactory.cpp
--- a/Engine/Plugins/Animation/IKRig/Source/IKRigEditor/Private/RetargetEditor/IKRetargetFactory.cpp
+++ b/Engine/Plugins/Animation/IKRig/Source/IKRigEditor/Private/RetargetEditor/IKRetargetFactory.cpp
@@ -27,10 +27,7 @@ UObject* UIKRetargetFactory::FactoryCreateNew(
 	UObject* Context,
 	FFeedbackContext* Warn)
 {
-	UIKRetargeter* NewRetargeter = NewObject<UIKRetargeter>(InParent, Class, Name, Flags);
-	UIKRetargeterController* Controller = UIKRetargeterController::GetController(NewRetargeter);
-	Controller->AddDefaultOps();
-	return NewRetargeter;
+	return NewObject<UIKRetargeter>(InParent, Class, Name, Flags);
 }
 
 bool UIKRetargetFactory::ShouldShowInNewMenu() const
