[IK Retarget] Integrating IK Retargeter into UAF

#Synced-CL 49355612

diff --git a/Engine/Plugins/Animation/IKRig/IKRig.uplugin b/Engine/Plugins/Animation/IKRig/IKRig.uplugin
--- a/Engine/Plugins/Animation/IKRig/IKRig.uplugin
+++ b/Engine/Plugins/Animation/IKRig/IKRig.uplugin
@@ -23,23 +23,38 @@
 		{
 			"Name": "FullBodyIK",
 			"Enabled": true
+		},
+		{
+			"Name": "UAF",
+			"Enabled": false,
+			"Optional": true
+		},
+		{
+			"Name": "UAFAnimGraph",
+			"Enabled": false,
+			"Optional": true
 		}
 	],
 	"Modules": [
 		{
 			"Name": "IKRig",
 			"Type": "Runtime",
-			"LoadingPhase": "PreDefault"
+			"LoadingPhase": "Default"
 		},
 		{
 			"Name": "IKRigDeveloper",
 			"Type": "UncookedOnly",
-			"LoadingPhase": "PreDefault"
+			"LoadingPhase": "Default"
 		},
 		{
 			"Name": "IKRigEditor",
 			"Type": "Editor",
 			"LoadingPhase": "PostEngineInit"
+		},
+		{
+			"Name": "IKRigUAF",
+			"Type": "Runtime",
+			"LoadingPhase": "None"
 		}
 	]
 }
\ No newline at end of file
diff --git a/Engine/Plugins/Animation/IKRig/Source/IKRigUAF/IKRigUAF.build.cs b/Engine/Plugins/Animation/IKRig/Source/IKRigUAF/IKRigUAF.build.cs
new file mode 100644
--- /dev/null
+++ b/Engine/Plugins/Animation/IKRig/Source/IKRigUAF/IKRigUAF.build.cs
@@ -0,0 +1,27 @@
+﻿// Copyright Epic Games, Inc. All Rights Reserved.
+
+namespace UnrealBuildTool.Rules
+{
+	public class IKRigUAF : ModuleRules
+	{
+		public IKRigUAF(ReadOnlyTargetRules Target) : base(Target)
+		{
+			PCHUsage = PCHUsageMode.UseExplicitOrSharedPCHs;
+
+			PublicDependencyModuleNames.AddRange(
+				new string[]
+				{
+					"Core",
+					"CoreUObject",
+					"Engine",
+					
+					"RigVM",
+					"UAF",
+					"UAFAnimGraph",
+					
+					"IKRig"
+				}
+			);
+		}
+	}
+}
\ No newline at end of file
diff --git a/Engine/Plugins/Animation/IKRig/Source/IKRigUAF/Private/IKRigUAFModule.cpp b/Engine/Plugins/Animation/IKRig/Source/IKRigUAF/Private/IKRigUAFModule.cpp
new file mode 100644
--- /dev/null
+++ b/Engine/Plugins/Animation/IKRig/Source/IKRigUAF/Private/IKRigUAFModule.cpp
@@ -0,0 +1,18 @@
+// Copyright Epic Games, Inc. All Rights Reserved.
+
+#include "IKRigUAFModule.h"
+#include "Modules/ModuleInterface.h"
+#include "Modules/ModuleManager.h"
+#include "Retargeter/IKRetargeter.h"
+#include "RigVMCore/RigVMRegistry.h"
+
+void FIKRigUAFModule::StartupModule()
+{
+}
+
+void FIKRigUAFModule::ShutdownModule()
+{
+	IModuleInterface::ShutdownModule();
+}
+
+IMPLEMENT_MODULE(FIKRigUAFModule, IKRigUAF)
\ No newline at end of file
diff --git a/Engine/Plugins/Animation/IKRig/Source/IKRigUAF/Private/RigUnit_UAFIKRetargeter.cpp b/Engine/Plugins/Animation/IKRig/Source/IKRigUAF/Private/RigUnit_UAFIKRetargeter.cpp
new file mode 100644
--- /dev/null
+++ b/Engine/Plugins/Animation/IKRig/Source/IKRigUAF/Private/RigUnit_UAFIKRetargeter.cpp
@@ -0,0 +1,88 @@
+﻿// Copyright Epic Games, Inc. All Rights Reserved.
+
+#include "RigUnit_UAFIKRetargeter.h"
+
+#include "Engine/SkeletalMesh.h"
+#include "AnimationRuntime.h"
+
+#include UE_INLINE_GENERATED_CPP_BY_NAME(RigUnit_UAFIKRetargeter)
+
+DEFINE_STAT(STAT_AnimNext_RigUnit_IKRetargeter);
+
+FRigUnit_UAFIKRetargeter_Execute()
+{
+	SCOPE_CYCLE_COUNTER(STAT_AnimNext_RigUnit_IKRetargeter);
+
+	using namespace UE::UAF;
+
+	// validate prerequisites
+	if( !IsValid(IKRetargetAsset) ||
+		!SourcePose.LODPose.IsValid() ||
+		!TargetAnimGraphRefPose.ReferencePose.IsValid())
+	{
+		return;
+	}
+
+	// get source skeletal mesh
+	const UE::UAF::FReferencePose& SourceRefPose = SourcePose.LODPose.GetRefPose();
+	const USkeletalMesh* SourceMesh = SourceRefPose.SkeletalMesh.Get();
+	
+	// get target skeletal mesh
+	const UE::UAF::FReferencePose& TargetRefPose = TargetAnimGraphRefPose.ReferencePose.GetRef<UE::UAF::FReferencePose>();
+	const USkeletalMesh* TargetMesh = TargetRefPose.SkeletalMesh.Get();
+
+	// collect settings to retarget with starting with asset settings and override/merge with custom profile
+	WorkData.Profile.FillProfileWithAssetSettings(IKRetargetAsset);
+	WorkData.Profile.MergeWithOtherProfile(CustomRetargetProfile);
+		
+	// ensure the processor is initialized with the current inputs
+	FIKRetargetProcessor& Processor = WorkData.Processor;
+	if (!Processor.WasInitializedWithTheseAssets(SourceMesh, TargetMesh, IKRetargetAsset))
+	{
+		// initialize retarget processor with source and target skeletal meshes
+		Processor.Initialize(SourceMesh, TargetMesh, IKRetargetAsset, WorkData.Profile);
+	}
+
+	// bail out if not initialized
+	if (!Processor.IsInitialized())
+	{
+		return;
+	}
+
+	// copy the full local ref pose of the source
+	const FReferenceSkeleton& SourceRefSkeleton = SourceMesh->GetRefSkeleton();
+	WorkData.SourceMeshPoseLocal = SourceRefSkeleton.GetRefBonePose();
+	// override all bones in the current LOD with the current local input pose
+	const int32 NumBonesInSourceLOD = SourcePose.LODPose.GetNumBones();
+	for (int32 SourceLODBoneIndex = 0; SourceLODBoneIndex < NumBonesInSourceLOD; ++SourceLODBoneIndex)
+	{
+		const int32 MeshBoneIndex = SourceRefPose.GetMeshBoneIndexFromLODBoneIndex(SourceLODBoneIndex);
+		WorkData.SourceMeshPoseLocal[MeshBoneIndex] = SourcePose.LODPose.LocalTransforms[SourceLODBoneIndex];
+	}
+
+	// convert to component space
+	FAnimationRuntime::FillUpComponentSpaceTransforms(SourceRefSkeleton, WorkData.SourceMeshPoseLocal, WorkData.SourceMeshPoseGlobal);
+	
+	// run the retargeter
+	const TArray<FTransform>& TargetGlobalPose = Processor.RunRetargeter(
+	WorkData.SourceMeshPoseGlobal,
+	WorkData.Profile,
+	ExecuteContext.GetDeltaTime(),
+	TargetRefPose.GetSourceLODLevel());
+
+	// fill the target pose with a ref pose at the appropriate LOD level
+	TargetPose.LODPose.PrepareForLOD(TargetRefPose, TargetRefPose.GetSourceLODLevel(), true/*bSetRefPose*/, false/*bAdditive*/);
+
+	// copy output pose back to target lod after converting back to local space
+	const TArrayView<const FBoneIndexType>& LODToMeshBoneMap = TargetPose.LODPose.GetLODBoneIndexToMeshBoneIndexMap();
+	const TArrayView<const FBoneIndexType>& MeshToParentMeshBoneMap = TargetPose.LODPose.GetMeshBoneIndexToParentMeshBoneIndexMap();
+	const int32 NumBonesInTargetLOD = TargetPose.LODPose.GetNumBones();
+	for (int32 TargetLODBoneIndex = 0; TargetLODBoneIndex < NumBonesInTargetLOD; ++TargetLODBoneIndex)
+	{
+		const FBoneIndexType ChildMeshBoneIndex = LODToMeshBoneMap[TargetLODBoneIndex];
+		const FBoneIndexType ParentMeshBoneIndex = MeshToParentMeshBoneMap[ChildMeshBoneIndex];
+		const FTransform& ChildGlobal = TargetGlobalPose[ChildMeshBoneIndex];
+		const FTransform& ParentGlobal = TargetGlobalPose.IsValidIndex(ParentMeshBoneIndex) ? TargetGlobalPose[ParentMeshBoneIndex] : FTransform::Identity;
+		TargetPose.LODPose.LocalTransforms[TargetLODBoneIndex] = (ChildGlobal * ParentGlobal.Inverse());
+	}
+}
\ No newline at end of file
diff --git a/Engine/Plugins/Animation/IKRig/Source/IKRigUAF/Public/IKRigUAFModule.h b/Engine/Plugins/Animation/IKRig/Source/IKRigUAF/Public/IKRigUAFModule.h
new file mode 100644
--- /dev/null
+++ b/Engine/Plugins/Animation/IKRig/Source/IKRigUAF/Public/IKRigUAFModule.h
@@ -0,0 +1,13 @@
+// Copyright Epic Games, Inc. All Rights Reserved.
+#pragma once
+
+#include "CoreMinimal.h"
+#include "Modules/ModuleInterface.h"
+
+class FIKRigUAFModule : public IModuleInterface
+{
+public:
+	virtual void StartupModule() override;
+	virtual void ShutdownModule() override;
+};
+
diff --git a/Engine/Plugins/Animation/IKRig/Source/IKRigUAF/Public/RigUnit_UAFIKRetargeter.h b/Engine/Plugins/Animation/IKRig/Source/IKRigUAF/Public/RigUnit_UAFIKRetargeter.h
new file mode 100644
--- /dev/null
+++ b/Engine/Plugins/Animation/IKRig/Source/IKRigUAF/Public/RigUnit_UAFIKRetargeter.h
@@ -0,0 +1,67 @@
+﻿// Copyright Epic Games, Inc. All Rights Reserved.
+
+#pragma once
+
+#include "Retargeter/IKRetargeter.h"
+#include "AnimNextExecuteContext.h"
+#include "Graph/AnimNext_LODPose.h"
+#include "Graph/RigUnit_AnimNextBase.h"
+#include "RemapPoseData.h"
+#include "AnimNextStats.h"
+#include "Retargeter/IKRetargetProcessor.h"
+
+#include "RigUnit_UAFIKRetargeter.generated.h"
+
+DECLARE_CYCLE_STAT_EXTERN(TEXT("UAF RigUnit: IK Retargeter Total"), STAT_AnimNext_RigUnit_IKRetargeter, STATGROUP_AnimNext, IKRIGUAF_API);
+
+USTRUCT(BlueprintType)
+struct FIKRetargetWorkData
+{
+	GENERATED_BODY()
+ 
+	// the runtime processor used to run the retarget and generate new poses
+	FIKRetargetProcessor Processor;
+
+	// scratch space for adapting the LOD'd pose to the full input pose
+	TArray<FTransform> SourceMeshPoseLocal;
+	TArray<FTransform> SourceMeshPoseGlobal;
+
+	// reusable profile to pass to the retargeter (merged with input profile)
+	FRetargetProfile Profile;
+};
+
+// IK Retargeting UAF Graph node.
+USTRUCT(MinimalAPI, meta=(DisplayName="IK Retargeter", Category="Animation Graph", NodeColor="1, 0, 1", Keywords="Retarget, IK"))
+struct FRigUnit_UAFIKRetargeter : public FRigUnit_AnimNextBase
+{
+	GENERATED_BODY()
+
+	RIGVM_METHOD()
+	void Execute();
+
+	// The source input pose to retarget.
+	UPROPERTY(EditAnywhere, DisplayName = "Source Pose", Category = "Graph", meta = (Input))
+	FAnimNextGraphLODPose SourcePose;
+
+	// The ref pose of the target skeletal mesh.
+	UPROPERTY(EditAnywhere, DisplayName = "Target Ref Pose", Category = "Graph", meta = (Input))
+	FAnimNextGraphReferencePose TargetAnimGraphRefPose;
+
+	// An IK Retarget asset defining the retargeting operations to run.
+	UPROPERTY(EditAnywhere, DisplayName = "IK Retarget Asset", Category = "Graph", meta = (Input))
+	TObjectPtr<UIKRetargeter> IKRetargetAsset;
+
+	// Connect a custom retarget profile to modify the retargeter's settings at runtime.
+	UPROPERTY(EditAnywhere, DisplayName = "Custom Profile", Category = "Graph", meta = (Input))
+	FRetargetProfile CustomRetargetProfile;
+
+	// The output pose for the target skeletal mesh.
+	UPROPERTY(EditAnywhere, DisplayName = "Retargeted Pose", Category = "Graph", meta = (Output))
+	FAnimNextGraphLODPose TargetPose;
+	
+	UPROPERTY(EditAnywhere, DisplayName = "Execute", Category = "BeginExecution", meta = (Input, Output))
+	FAnimNextExecuteContext ExecuteContext;
+
+	UPROPERTY(meta = (Hidden))
+	FIKRetargetWorkData WorkData;
+};
diff --git a/Engine/Plugins/Experimental/UAF/UAF/Source/UAF/Public/LODPose.h b/Engine/Plugins/Experimental/UAF/UAF/Source/UAF/Public/LODPose.h
--- a/Engine/Plugins/Experimental/UAF/UAF/Source/UAF/Public/LODPose.h
+++ b/Engine/Plugins/Experimental/UAF/UAF/Source/UAF/Public/LODPose.h
@@ -199,6 +199,18 @@ struct FLODPose
 
 	}
 
+	const TArrayView<const FBoneIndexType> GetMeshBoneIndexToParentMeshBoneIndexMap() const
+	{
+		if (RefPose != nullptr)
+		{
+			return RefPose->GetMeshBoneIndexToParentMeshBoneIndexMap();
+		}
+		else
+		{
+			return TArrayView<const FBoneIndexType>();
+		}
+	}
+
 	const TArrayView<const FBoneIndexType> GetSkeletonBoneIndexToLODBoneIndexMap() const
 	{
 		if (RefPose != nullptr)
