[IKRetarget] Pulling out One Euro filter into resueable utils.

Fixed some minor issues in the filters themselves.

#Synced-CL 49355612

diff --git a/Engine/Plugins/Animation/IKRig/Source/IKRig/Private/Retargeter/RetargetOps/FilterBoneOp.cpp b/Engine/Plugins/Animation/IKRig/Source/IKRig/Private/Retargeter/RetargetOps/FilterBoneOp.cpp
--- a/Engine/Plugins/Animation/IKRig/Source/IKRig/Private/Retargeter/RetargetOps/FilterBoneOp.cpp
+++ b/Engine/Plugins/Animation/IKRig/Source/IKRig/Private/Retargeter/RetargetOps/FilterBoneOp.cpp
@@ -7,66 +7,6 @@
 #define LOCTEXT_NAMESPACE "FilterBoneOp"
 
 
-FQuat FOneEuroRotationFilter::Update(
-	const FQuat& InTargetRotation,
-	double InDeltaTime,
-	const FIKRetargetFilterBoneOpSettings& InSettings)
-{
-	if (bFirstRun)
-	{
-		PrevFilteredRot = InTargetRotation;
-		PrevHalfAngVel = FVector::ZeroVector;
-		bFirstRun = false;
-		return InTargetRotation;
-	}
-
-	// sanity check delta time
-	constexpr double MAX_DELTATIME = 1.0/15.0;
-	constexpr double MIN_DELTATIME = 1.0/240.0;
-	InDeltaTime = FMath::Clamp(InDeltaTime, MIN_DELTATIME, MAX_DELTATIME);
-
-	auto AlphaFromHz = [&InDeltaTime](double FreqCutoff) -> double
-	{
-		const double Tau = (FreqCutoff > 0.0) ? (1.0 / (2.0 * PI * FreqCutoff)) : UE_SMALL_NUMBER;
-		return 1.0 / (1.0 + Tau / FMath::Max(InDeltaTime, UE_KINDA_SMALL_NUMBER));
-	};
-
-	// This an adaptive first-order low-pass filter based on the "One-Euro" 2012 paper: https://gery.casiez.net/1euro/
-	//
-	// NOTE: because we are filtering rotations, and desire to use linear math,
-	// all calculations are done with the exponential map of the delta quaternion
-
-	// calc error from current filtered to target
-	FQuat ErrorQ = InTargetRotation * PrevFilteredRot.Inverse();
-	ErrorQ = ErrorQ.W < 0 ? -ErrorQ : ErrorQ; // negate if W negative
-
-	// convert into tangent space
-	// NOTE: this produces a 3d vector amenable to filtering, with |Phi| = angle/2 
-	const FQuat LogQ = ErrorQ.Log();
-	const FVector Phi(LogQ.X, LogQ.Y, LogQ.Z);
-	
-	// compute angular velocity that would reduce error to zero if fully applied
-	const FVector HalfAngVel = (InDeltaTime > 0.0) ? (Phi / InDeltaTime) : FVector::ZeroVector; // rad/s
-	// low-pass filter the velocity with fixed cutoff frequency
-	const double VelCutoffFreq = (InSettings.VelocityCutoffHz > 0.0) ? InSettings.VelocityCutoffHz : 20.0; // Hz
-	PrevHalfAngVel = FMath::Lerp(PrevHalfAngVel, HalfAngVel, AlphaFromHz(VelCutoffFreq));
-
-	// adaptive cutoff for the signal based on filtered speed
-	const double Speed = 2.0 * PrevHalfAngVel.Size(); // rad/s
-	const double FreqCutoff = InSettings.MinFrequency + InSettings.Responsiveness * Speed; // Hz
-	// take only a fraction of the filtered velocity this frame
-	const FVector PhiStep = PrevHalfAngVel * (AlphaFromHz(FreqCutoff) * InDeltaTime);
-	
-	// exponentiate back into SO(3)
-	const FQuat DeltaQ = FQuat(PhiStep.X, PhiStep.Y, PhiStep.Z, 0).Exp();
-	// integrate the angular velocity to get new rotation
-	const FQuat NewRotation = (DeltaQ * PrevFilteredRot).GetNormalized();
-	// store prev filtered rotation for next frame
-	PrevFilteredRot = NewRotation;
-	
-	return NewRotation;
-}
-
 const UClass* FIKRetargetFilterBoneOpSettings::GetControllerType() const
 {
 	return UIKRetargetFilterBoneController::StaticClass();
@@ -134,7 +74,7 @@ void FIKRetargetFilterBoneOp::Run(
 		}
 
 		FTransform TargetGlobalTransform = OutTargetGlobalPose[FilterBoneData.TargetBone.BoneIndex];
-		FQuat FilteredRotation = FilterBoneData.Filter.Update(TargetGlobalTransform.GetRotation(), InDeltaTime, Settings);
+		FQuat FilteredRotation = FilterBoneData.Filter.Update(TargetGlobalTransform.GetRotation(), InDeltaTime, Settings.FilterSettings);
 		TargetGlobalTransform.SetRotation(FilteredRotation);
 
 		// assign result and update children
diff --git a/Engine/Plugins/Animation/IKRig/Source/IKRig/Private/Retargeter/RetargetOps/FloorConstraintOp.cpp b/Engine/Plugins/Animation/IKRig/Source/IKRig/Private/Retargeter/RetargetOps/FloorConstraintOp.cpp
--- a/Engine/Plugins/Animation/IKRig/Source/IKRig/Private/Retargeter/RetargetOps/FloorConstraintOp.cpp
+++ b/Engine/Plugins/Animation/IKRig/Source/IKRig/Private/Retargeter/RetargetOps/FloorConstraintOp.cpp
@@ -104,54 +104,6 @@ void FIKRetargetFloorConstraintOpSettings::AssignSkeletonAssets(const USkeleton*
 }
 #endif
 
-void FScalarOneEuroFilter::Tune(
-	const double InFrequency,
-	const double InMinCutoff,
-	const double InBeta,
-	const double InDerivativeCutoff)
-{
-	Frequency = InFrequency;
-	MinCuttoff = InMinCutoff;
-	Beta = InBeta;
-	DerivativeCutoff = InDerivativeCutoff;
-}
-
-double FScalarOneEuroFilter::Run(double InValue, double InDeltaTime)
-{
-	auto LowPass = [](const double InPrev, const double InCurrent, const double InAlpha) -> double
-		{
-			return InAlpha * InCurrent + (1.0 - InAlpha) * InPrev;
-		};
-
-	auto Alpha = [](const double InInvDeltaTime, const double InCutoff) -> double
-		{
-			const double Tau = 1.0 / (2.0 * UE_PI * InCutoff);
-			return 1.0 / (1.0 + Tau * InInvDeltaTime);
-		};
-		
-	if (bIsFirstRun)
-	{
-		PrevValue = InValue;
-		PrevValueDerivative = 0.0;
-		bIsFirstRun = false;
-	}
-
-	// compute the derivative of the signal
-	const double InvDeltaTime = 1.0 / FMath::Max(InDeltaTime, UE_KINDA_SMALL_NUMBER);
-	const double Derivative = (InValue - PrevValue) * InvDeltaTime;
-	const double LowPassDerivative = LowPass(PrevValueDerivative, Derivative, Alpha(InvDeltaTime, DerivativeCutoff));
-	PrevValueDerivative = LowPassDerivative;
-
-	// compute the cutoff frequency based on signal speed
-	const double Cutoff = MinCuttoff + Beta * FMath::Abs(LowPassDerivative);
-
-	// filter the signal
-	const double Result = LowPass(PrevValue, InValue, Alpha(InvDeltaTime, Cutoff));
-	PrevValue = Result;
-
-	return Result;
-}
-
 void FFootConstraint::Initialize(
 	const FTargetSkeleton& InTargetSkeleton,
 	const FFloorConstraintFootDefinition& InFootSettings,
@@ -289,12 +241,11 @@ void FToeConstraint::RunAfterParent(
 	double ThetaRadians = CalcToeCorrectionAngle();
 
 	// temporal filtering of the toe adjustment
-	AngleFilter.Tune(
-		AllToeSettings->FilterFrequency,
-		AllToeSettings->FilterMinCutoff,
-		AllToeSettings->FilterBeta,
-		AllToeSettings->FilterDerivCutoff);
-	ThetaRadians = AngleFilter.Run(ThetaRadians, InDeltaTime);
+	FOneEuroFilterSettings OneEuroFilterSettings;
+	OneEuroFilterSettings.CutoffFrequency = AllToeSettings->CutoffFrequency;
+	OneEuroFilterSettings.Responsiveness = AllToeSettings->Responsiveness;
+	OneEuroFilterSettings.VelocityCutoffFrequency = AllToeSettings->VelocityCutoffFrequency;
+	ThetaRadians = AngleFilter.Update(ThetaRadians, InDeltaTime, OneEuroFilterSettings);
 
 	// alpha the whole effect
 	ThetaRadians = FMath::Lerp(0, ThetaRadians, AllToeSettings->Alpha * InChainSettings->Alpha);
diff --git a/Engine/Plugins/Animation/IKRig/Source/IKRig/Private/Retargeter/RetargetOps/RetargetOpUtils.cpp b/Engine/Plugins/Animation/IKRig/Source/IKRig/Private/Retargeter/RetargetOps/RetargetOpUtils.cpp
new file mode 100644
--- /dev/null
+++ b/Engine/Plugins/Animation/IKRig/Source/IKRig/Private/Retargeter/RetargetOps/RetargetOpUtils.cpp
@@ -0,0 +1,154 @@
+﻿// Copyright Epic Games, Inc. All Rights Reserved.
+
+#include "Retargeter/RetargetOps/RetargetOpUtils.h"
+
+#include "VectorUtil.h"
+
+double FOneEuroScalarFilter::Update(
+	double InValue,
+	double InDeltaTime,
+	const FOneEuroFilterSettings& InSettings)
+{
+	auto LowPass = [](const double InPrev, const double InCurrent, const double InAlpha) -> double
+		{
+			return InAlpha * InCurrent + (1.0 - InAlpha) * InPrev;
+		};
+		
+	if (bIsFirstRun)
+	{
+		PrevValue = PrevRawValue = InValue;
+		PrevValueDerivative = 0.0;
+		bIsFirstRun = false;
+	}
+
+	const double ClampedDeltaTime = ClampDeltaTime(InDeltaTime);
+
+	// compute the derivative of the signal
+	const double Derivative = (InValue - PrevRawValue) / ClampedDeltaTime;
+	PrevRawValue = InValue;
+	const double VelCutoffFreq = ClampFreqToFramerate(InSettings.VelocityCutoffFrequency, ClampedDeltaTime);
+	const double LowPassDerivative = LowPass(PrevValueDerivative, Derivative, AlphaFromHz(ClampedDeltaTime, VelCutoffFreq));
+	PrevValueDerivative = LowPassDerivative;
+
+	// compute the cutoff frequency based on signal speed
+	const double CutoffFreq = ClampFreqToFramerate(InSettings.CutoffFrequency, ClampedDeltaTime);
+	const double Cutoff = CutoffFreq + InSettings.Responsiveness * FMath::Abs(LowPassDerivative);
+
+	// filter the signal
+	const double Result = LowPass(PrevValue, InValue, AlphaFromHz(ClampedDeltaTime, Cutoff));
+	PrevValue = Result;
+
+	return Result;
+}
+
+void FOneEuroScalarFilter::Reset(const double InValue)
+{
+	PrevValue = PrevRawValue = InValue;
+	PrevValueDerivative = 0.0;
+	bIsFirstRun = false;
+}
+
+double FOneEuroScalarFilter::AlphaFromHz(const double InDeltaTime, const double InFreqCutoffHz)
+{
+	const double OmegaDt = 2.0 * PI * InFreqCutoffHz * InDeltaTime;
+	return OmegaDt / (1.0 + OmegaDt);
+}
+
+double FOneEuroScalarFilter::ClampDeltaTime(const double InDeltaTime)
+{
+	constexpr double MAX_DELTATIME = 1.0 / 15.0;
+	constexpr double MIN_DELTATIME = 1.0 / 240.0;
+	return FMath::Clamp(InDeltaTime, MIN_DELTATIME, MAX_DELTATIME);
+}
+
+double FOneEuroScalarFilter::ClampFreqToFramerate(const double InHz, const double InDeltaTime)
+{
+	// calculate the theoretical maximum frequency we can sample (Speed Limit)
+	// Nyquist = SampleRate / 2 = (1 / DeltaTime) / 2
+	const double NyquistFreq = 0.5 / InDeltaTime;
+    
+	// clamp to AT MOST the Nyquist frequency
+	return (InHz > NyquistFreq) ? NyquistFreq : InHz;
+}
+
+FVector FOneEuroVectorFilter::Update(const FVector& InValue, const double InDeltaTime, const FOneEuroFilterSettings& InSettings)
+{
+	const double NewX = X.Update(InValue.X, InDeltaTime, InSettings);
+	const double NewY = Y.Update(InValue.Y, InDeltaTime, InSettings);
+	const double NewZ = Z.Update(InValue.Z, InDeltaTime, InSettings);
+	return FVector(NewX,NewY,NewZ);
+}
+
+void FOneEuroVectorFilter::Reset(const FVector& StartingValue)
+{
+	X.Reset(StartingValue.X);
+	Y.Reset(StartingValue.Y);
+	Z.Reset(StartingValue.Z);
+}
+
+void FOneEuroVectorFilter::Reset()
+{
+	X.Reset();
+	Y.Reset();
+	Z.Reset();
+}
+
+FQuat FOneEuroQuatFilter::Update(
+	const FQuat& InTargetRotation,
+	double InDeltaTime,
+	const FOneEuroFilterSettings& InSettings)
+{
+	if (bIsFirstRun)
+	{
+		PrevFilteredRot = PrevRawTarget = InTargetRotation;
+		PrevHalfAngVel = FVector::ZeroVector;
+		bIsFirstRun = false;
+		return InTargetRotation;
+	}
+	
+	// NOTE: because we are filtering rotations, and desire to use linear math,
+	// all calculations are done with the exponential map of the delta quaternion
+
+	// sanity check delta time
+	const double ClampedDeltaTime = FOneEuroScalarFilter::ClampDeltaTime(InDeltaTime);
+
+	// calc delta rotation from prev input to current input
+	FQuat DeltaRot = InTargetRotation * PrevRawTarget.Inverse();
+	DeltaRot = DeltaRot.W < 0 ? -DeltaRot : DeltaRot; // negate if W negative
+
+	// store prev raw target for next frame
+	PrevRawTarget = InTargetRotation;
+
+	// convert into tangent space
+	// NOTE: this produces a 3d vector amenable to filtering, with |Phi| = angle/2
+	const FQuat DeltaLog = DeltaRot.Log();
+	const FVector Phi(DeltaLog.X, DeltaLog.Y, DeltaLog.Z);
+	
+	// compute angular velocity
+	const FVector HalfAngVel = Phi / ClampedDeltaTime; // rad/s
+	// low-pass filter the velocity with fixed cutoff frequency
+	const double VelCutoffFreq = FOneEuroScalarFilter::ClampFreqToFramerate(InSettings.VelocityCutoffFrequency, ClampedDeltaTime);
+	const double VelAlpha = FOneEuroScalarFilter::AlphaFromHz(ClampedDeltaTime, VelCutoffFreq);
+	PrevHalfAngVel = FMath::Lerp(PrevHalfAngVel, HalfAngVel, VelAlpha);
+
+	// adaptive cutoff for the signal based on filtered speed
+	const double Speed = 2.0 * PrevHalfAngVel.Size(); // rad/s
+	const double ClampedFreqCutoff = FOneEuroScalarFilter::ClampFreqToFramerate(InSettings.CutoffFrequency, ClampedDeltaTime);
+	const double FreqCutoff = ClampedFreqCutoff + (InSettings.Responsiveness * Speed); // Hz
+	// take only a fraction of the filtered velocity this frame
+	const double Alpha = FOneEuroScalarFilter::AlphaFromHz(ClampedDeltaTime, FreqCutoff);
+
+	// Convert to tangent space (this is the total displacement needed)
+	FQuat ErrorQ = InTargetRotation * PrevFilteredRot.Inverse();
+	ErrorQ = ErrorQ.W < 0 ? -ErrorQ : ErrorQ;
+	FVector ErrorPhi(ErrorQ.Log().X, ErrorQ.Log().Y, ErrorQ.Log().Z);
+	const FVector PhiStep = ErrorPhi * Alpha;
+	
+	// exponentiate back into SO(3)
+	const FQuat StepQ = FQuat(PhiStep.X, PhiStep.Y, PhiStep.Z, 0).Exp();
+	// integrate the angular velocity to get new rotation
+	const FQuat NewRotation = (StepQ * PrevFilteredRot).GetNormalized();
+	// store prev filtered rotation for next frame
+	PrevFilteredRot = NewRotation;
+	return NewRotation;
+}
diff --git a/Engine/Plugins/Animation/IKRig/Source/IKRig/Public/Retargeter/RetargetOps/FilterBoneOp.h b/Engine/Plugins/Animation/IKRig/Source/IKRig/Public/Retargeter/RetargetOps/FilterBoneOp.h
--- a/Engine/Plugins/Animation/IKRig/Source/IKRig/Public/Retargeter/RetargetOps/FilterBoneOp.h
+++ b/Engine/Plugins/Animation/IKRig/Source/IKRig/Public/Retargeter/RetargetOps/FilterBoneOp.h
@@ -2,6 +2,7 @@
 
 #pragma once
 
+#include "RetargetOpUtils.h"
 #include "Animation/BoneReference.h"
 #include "Retargeter/IKRetargetOps.h"
 #include "Retargeter/IKRetargetProcessor.h"
@@ -14,23 +15,6 @@
 
 struct FIKRetargetOpSettingsBase;
 
-struct FOneEuroRotationFilter
-{
-	FQuat Update(
-		const FQuat& InTargetRotation,
-		double InDeltaTime,
-		const FIKRetargetFilterBoneOpSettings& InSettings);
-
-	void Reset() { bFirstRun = true; };
-
-private:
-	
-	FQuat PrevFilteredRot = FQuat::Identity;
-	FVector PrevHalfAngVel = FVector::ZeroVector;
-	FVector PrevPhiFilt = FVector::ZeroVector;
-	bool bFirstRun = true;
-};
-
 USTRUCT(BlueprintType)
 struct FFilterBoneData
 {
@@ -41,7 +25,7 @@ struct FFilterBoneData
 	FBoneReference TargetBone;
 
 	bool bIsReady = false;
-	FOneEuroRotationFilter Filter;
+	FOneEuroQuatFilter Filter;
 };
 
 USTRUCT(BlueprintType, meta = (DisplayName = "Filter Bone Settings"))
@@ -53,29 +37,9 @@ struct FIKRetargetFilterBoneOpSettings : public FIKRetargetOpSettingsBase
 	UPROPERTY(EditAnywhere, BlueprintReadWrite, Category=Setup, meta=(ReinitializeOnEdit))
 	TArray<FFilterBoneData> BonesToFilter;
 
-	/*/* When true, applies a One-Euro filter to smooth the rotation of the target bone #1#
-	UPROPERTY(EditAnywhere, BlueprintReadWrite, Category="One Euro Filter")
-	bool bFilterRotation = true;*/
-
-	/* Hz. Sets the low-pass cutoff when motion is near zero.
-	 * Higher = less smoothing at rest (more responsive but more jitter).
-	 * Lower = more smoothing at rest (less jitter but more “stickiness”). */
-	UPROPERTY(EditAnywhere, BlueprintReadWrite, Category="One Euro Filter")
-	double MinFrequency = 1.5;
-
-	/* Raises the cutoff as angular speed grows.
-	 * Larger values are snappier on fast turns.
-	 * Smaller values are smoother but laggier during quick motion.
-	 * Typical sweet spot: 0.3 – 0.8 */
-	UPROPERTY(EditAnywhere, BlueprintReadWrite, Category="One Euro Filter")
-	double Responsiveness = 0.5;
-
-	/* Hz. Low-passes the raw angular velocity before we use it to adapt the derivative cutoff.
-	* If you see breathing/pumping of the smoothing during motion onsets or reversals, lower this value (e.g., 30 → 20 Hz).
-	* If responsiveness to fast changes is sluggish, raise this a bit.
-	* Good starting range: 15–30 Hz. Keep under Nyquist frequency (frame_rate/2). */
+	/** The settings to control the One-Euro filter behavior */
 	UPROPERTY(EditAnywhere, BlueprintReadWrite, Category="One Euro Filter")
-	double VelocityCutoffHz = 20.0;
+	FOneEuroFilterSettings FilterSettings;
 
 	/* If true, filter is reset when playback loops. */
 	UPROPERTY(EditAnywhere, BlueprintReadWrite, Category="One Euro Filter")
diff --git a/Engine/Plugins/Animation/IKRig/Source/IKRig/Public/Retargeter/RetargetOps/FloorConstraintOp.h b/Engine/Plugins/Animation/IKRig/Source/IKRig/Public/Retargeter/RetargetOps/FloorConstraintOp.h
--- a/Engine/Plugins/Animation/IKRig/Source/IKRig/Public/Retargeter/RetargetOps/FloorConstraintOp.h
+++ b/Engine/Plugins/Animation/IKRig/Source/IKRig/Public/Retargeter/RetargetOps/FloorConstraintOp.h
@@ -1,6 +1,7 @@
 ﻿// Copyright Epic Games, Inc. All Rights Reserved.
 #pragma once
 
+#include "RetargetOpUtils.h"
 #include "Retargeter/IKRetargetChainMapping.h"
 #include "Retargeter/IKRetargetOps.h"
 #include "Animation/BoneReference.h"
@@ -47,22 +48,18 @@ struct FFloorConstraintToesDefinition
 	/** Default 1.0. Range 0 to 1. Blend the effect of the toe solver on/off. */
 	UPROPERTY(EditAnywhere, BlueprintReadWrite, Category = "Toe Settings", meta = (ReinitializeOnEdit, ClampMin = "0.0", ClampMax = "1.0", UIMin = "0.0", UIMax = "1.0"))
 	double Alpha = 1.0;
-	
-	/** Sampling frequency (Hz) used by the One Euro Filter. */
-	UPROPERTY(EditAnywhere, BlueprintReadWrite, Category = "Toe Filter", meta = (ClampMin = "1.0", UIMin = "1.0", UIMax = "240.0"))
-	float FilterFrequency = 120.0f;
-
-	/** Minimum cutoff frequency for the One Euro Filter (higher = less smoothing). */
-	UPROPERTY(EditAnywhere, BlueprintReadWrite, Category = "Toe Filter", meta = (ClampMin = "0.0", UIMin = "0.0", UIMax = "10.0"))
-	float FilterMinCutoff = 8.0f;
 
 	/** Speed coefficient; increases cutoff when signal changes quickly (reactivity). */
 	UPROPERTY(EditAnywhere, BlueprintReadWrite, Category = "Toe Filter", meta = (ClampMin = "0.0", UIMin = "0.0", UIMax = "5.0"))
-	float FilterBeta = 0.02f;
+	float Responsiveness = 0.02f;
+	
+	/** Minimum cutoff frequency for the One Euro Filter (higher = less smoothing). */
+	UPROPERTY(EditAnywhere, BlueprintReadWrite, Category = "Toe Filter", meta = (ClampMin = "0.0", UIMin = "0.0", UIMax = "10.0"))
+	float CutoffFrequency = 8.0f;
 
 	/** Cutoff frequency for the derivative smoothing step. */
 	UPROPERTY(EditAnywhere, BlueprintReadWrite, Category = "Toe Filter", meta = (ClampMin = "0.0", UIMin = "0.0", UIMax = "10.0"))
-	float FilterDerivCutoff = 1.0f;
+	float VelocityCutoffFrequency = 1.0f;
 	
 	/** An array of toes to bend when contacting the floor */
 	UPROPERTY(EditAnywhere, Category="Toe Definitions", meta=(ReinitializeOnEdit))
@@ -179,28 +176,6 @@ struct FIKRetargetFloorConstraintOpSettings : public FIKRetargetOpSettingsBase
 #endif
 };
 
-// filter a double over time using the "One Euro" method
-struct FScalarOneEuroFilter
-{
-	void Tune(
-		const double InFrequency,
-		const double InMinCutoff,
-		const double InBeta,
-		const double InDerivativeCutoff);
-	
-	double Run(double InValue, double InDeltaTime);
-
-private:
-	
-	double Frequency = 1.0;
-	double MinCuttoff = 1.0;
-	double Beta = 0.0;
-	double DerivativeCutoff = 1.0;
-	double PrevValue = 0.0;
-	double PrevValueDerivative = 0.0;
-	bool bIsFirstRun = true;
-};
-
 struct FToeConstraint
 {
 	// constant
@@ -218,7 +193,7 @@ struct FToeConstraint
 	FVector Start;
 	FVector End;
 	FVector PitchAxis;
-	FScalarOneEuroFilter AngleFilter;
+	FOneEuroScalarFilter AngleFilter;
 
 	void Initialize(
 		const FTargetSkeleton& InTargetSkeleton,
diff --git a/Engine/Plugins/Animation/IKRig/Source/IKRig/Public/Retargeter/RetargetOps/RetargetOpUtils.h b/Engine/Plugins/Animation/IKRig/Source/IKRig/Public/Retargeter/RetargetOps/RetargetOpUtils.h
new file mode 100644
--- /dev/null
+++ b/Engine/Plugins/Animation/IKRig/Source/IKRig/Public/Retargeter/RetargetOps/RetargetOpUtils.h
@@ -0,0 +1,85 @@
+﻿// Copyright Epic Games, Inc. All Rights Reserved.
+
+#pragma once
+
+#include "RetargetOpUtils.generated.h"
+
+// These are settings for an adaptive first-order low-pass filter based on the "One-Euro" 2012 paper: https://gery.casiez.net/1euro/
+USTRUCT(BlueprintType)
+struct FOneEuroFilterSettings
+{
+	GENERATED_BODY()
+
+	/* Raises the cutoff as speed grows.
+	 * Larger values are snappier on fast turns.
+	 * Smaller values are smoother but laggier during quick motion.
+	 * Typical sweet spot: 0.3 – 0.8 */
+	UPROPERTY(EditAnywhere, BlueprintReadWrite, Category="One Euro Filter")
+	double Responsiveness = 0.5;
+	
+	/* Hz. Sets the low-pass cutoff when motion is near zero.
+	 * Higher = less smoothing at rest (more responsive but more jitter).
+	 * Lower = more smoothing at rest (less jitter but more “stickiness”). */
+	UPROPERTY(EditAnywhere, BlueprintReadWrite, Category="One Euro Filter")
+	double CutoffFrequency = 1.5;
+
+	/* Hz. Low-passes the raw velocity before we use it to adapt the derivative cutoff.
+	* If you see breathing/pumping of the smoothing during motion onsets or reversals, lower this value (e.g., 30 → 20 Hz).
+	* If responsiveness to fast changes is sluggish, raise this a bit.
+	* Good starting range: 15–30 Hz. Keep under Nyquist frequency (frame_rate/2). */
+	UPROPERTY(EditAnywhere, BlueprintReadWrite, Category="One Euro Filter")
+	double VelocityCutoffFrequency = 20.0;
+};
+
+// filter a double over time using the "One Euro" method
+struct FOneEuroScalarFilter
+{
+	double Update(const double InValue, const double InDeltaTime, const FOneEuroFilterSettings& InSettings);
+
+	void Reset(const double InValue);
+	void Reset() { bIsFirstRun = true; }
+
+	static double AlphaFromHz(const double InDeltaTime, const double InFreqCutoffHz);
+
+	static double ClampDeltaTime(const double InDeltaTime);
+
+	static double ClampFreqToFramerate(const double InHz, const double InDeltaTime);
+
+private:
+	
+	double PrevValue = 0.0;
+	double PrevRawValue = 0.0;
+	double PrevValueDerivative = 0.0;
+	bool bIsFirstRun = true;
+};
+
+// filter a 3d vector over time using the "One Euro" method
+struct FOneEuroVectorFilter
+{
+	FVector Update(const FVector& InValue, const double InDeltaTime, const FOneEuroFilterSettings& InSettings);
+
+	void Reset(const FVector& StartingValue);
+	void Reset();
+
+private:
+	
+	FOneEuroScalarFilter X;
+	FOneEuroScalarFilter Y;
+	FOneEuroScalarFilter Z;
+};
+
+// filter a quaternion over time using the "One Euro" method
+struct FOneEuroQuatFilter
+{
+	FQuat Update(const FQuat& InTargetRotation, double InDeltaTime, const FOneEuroFilterSettings& InSettings);
+
+	void Reset() { bIsFirstRun = true; };
+
+private:
+	
+	FQuat PrevFilteredRot = FQuat::Identity;
+	FQuat PrevRawTarget = FQuat::Identity;
+	FVector PrevHalfAngVel = FVector::ZeroVector;
+	FVector PrevPhiFilt = FVector::ZeroVector;
+	bool bIsFirstRun = true;
+};
\ No newline at end of file
