[IK Retarget] Follow up change to load UAF retarget node if UAF module available

#Synced-CL 49392216

diff --git a/Engine/Plugins/Animation/IKRig/Source/IKRig/Private/IKRig.cpp b/Engine/Plugins/Animation/IKRig/Source/IKRig/Private/IKRig.cpp
--- a/Engine/Plugins/Animation/IKRig/Source/IKRig/Private/IKRig.cpp
+++ b/Engine/Plugins/Animation/IKRig/Source/IKRig/Private/IKRig.cpp
@@ -16,7 +16,11 @@ FDevVersionRegistration GRegisterIKRigObjectVersion(FIKRigObjectVersion::GUID, F
 
 void FIKRigModule::StartupModule()
 {
-	// This code will execute after your module is loaded into memory; the exact timing is specified in the .uplugin file per-module
+	// If UAF is loaded, then let's enable our UAF-compatible module.
+	if (FModuleManager::Get().IsModuleLoaded("UAF"))
+	{
+		FModuleManager::Get().LoadModule("IKRigUAF");
+	}
 }
 
 void FIKRigModule::ShutdownModule()
