[IKRetarget] Improved Speed Planting spring behavior. Added debug drawing.

#Synced-CL 49392216

diff --git a/Engine/Plugins/Animation/IKRig/Source/IKRig/Private/Retargeter/IKRetargetProfile.cpp b/Engine/Plugins/Animation/IKRig/Source/IKRig/Private/Retargeter/IKRetargetProfile.cpp
--- a/Engine/Plugins/Animation/IKRig/Source/IKRig/Private/Retargeter/IKRetargetProfile.cpp
+++ b/Engine/Plugins/Animation/IKRig/Source/IKRig/Private/Retargeter/IKRetargetProfile.cpp
@@ -242,7 +242,6 @@ void FRetargetProfile::MergeWithOtherProfile(const FRetargetProfile& OtherProfil
 					{
 						SettingsInProfile->SpeedThreshold = OtherChainSettings.SpeedPlanting.SpeedThreshold;
 						SettingsInProfile->Stiffness = OtherChainSettings.SpeedPlanting.UnplantStiffness;
-						SettingsInProfile->CriticalDamping = OtherChainSettings.SpeedPlanting.UnplantCriticalDamping;
 					}
 				}
 			}
diff --git a/Engine/Plugins/Animation/IKRig/Source/IKRig/Private/Retargeter/IKRetargeter.cpp b/Engine/Plugins/Animation/IKRig/Source/IKRig/Private/Retargeter/IKRetargeter.cpp
--- a/Engine/Plugins/Animation/IKRig/Source/IKRig/Private/Retargeter/IKRetargeter.cpp
+++ b/Engine/Plugins/Animation/IKRig/Source/IKRig/Private/Retargeter/IKRetargeter.cpp
@@ -426,7 +426,6 @@ void UIKRetargeter::PostLoadConvertEverythingToOps()
 
 				Settings->SpeedThreshold = Chain->Settings.SpeedPlanting.SpeedThreshold;
 				Settings->Stiffness = Chain->Settings.SpeedPlanting.UnplantStiffness;
-				Settings->CriticalDamping = Chain->Settings.SpeedPlanting.UnplantCriticalDamping;
 			}
 		}
 	}
diff --git a/Engine/Plugins/Animation/IKRig/Source/IKRig/Private/Retargeter/RetargetOps/SpeedPlantingOp.cpp b/Engine/Plugins/Animation/IKRig/Source/IKRig/Private/Retargeter/RetargetOps/SpeedPlantingOp.cpp
--- a/Engine/Plugins/Animation/IKRig/Source/IKRig/Private/Retargeter/RetargetOps/SpeedPlantingOp.cpp
+++ b/Engine/Plugins/Animation/IKRig/Source/IKRig/Private/Retargeter/RetargetOps/SpeedPlantingOp.cpp
@@ -9,6 +9,11 @@
 #include "Retargeter/IKRetargetOpUtils.h"
 #include "Retargeter/RetargetOps/RunIKRigOp.h"
 
+#if WITH_EDITOR
+#include "PrimitiveDrawingUtils.h"
+#include "Engine/EngineTypes.h"
+#endif
+
 #include UE_INLINE_GENERATED_CPP_BY_NAME(SpeedPlantingOp)
 
 #define LOCTEXT_NAMESPACE "SpeedPlantingOp"
@@ -33,6 +38,77 @@ void FIKRetargetSpeedPlantingOpSettings::CopySettingsAtRuntime(const FIKRetarget
 	IKRetargetOpUtils::CopyChainSettingsAtRuntime(NewSettings->ChainsToSpeedPlant, ChainsToSpeedPlant);
 }
 
+void FIKRetargetSpeedPlantingOpSettings::PostLoad(const FIKRigObjectVersion::Type InVersion)
+{
+	SpringStrength = FMath::Sqrt(Stiffness) / (2.0 * PI);
+}
+
+void FSpeedPlantingGoalState::Update(
+	FIKRigGoal* IKRigGoal,
+	const double InDeltaTime,
+	const FIKRetargetSpeedPlantingOpSettings& InSettings)
+{
+	// negative speed means no curve was found
+	// (we never want to pin a goal that is missing a curve)
+	if (CurrentSpeedValue < 0.0f)
+	{
+		return;
+	}
+
+	// forcibly update current position when op is reset
+	if (bWasReset)
+	{
+		CurrentPosition = TargetPosition = IKRigGoal->Position;
+		bWasReset = false;
+	}
+
+	// update target position
+	{
+		// update if the goal should be planted
+		const bool bWasPlanted = bIsPlanted;
+		bIsPlanted = CurrentSpeedValue < InSettings.SpeedThreshold;
+
+		// if we planted this tick, record the location to stay planted at
+		if (bIsPlanted && !bWasPlanted)
+		{
+			// update target position only once when first planted
+			TargetPosition = IKRigGoal->Position;
+		}
+
+		if (!bIsPlanted)
+		{
+			// if source is unplanted, continuously update target
+			TargetPosition  = IKRigGoal->Position;
+		}
+	}
+	
+	// run the goal position through a spring damper to un-plant it
+	const double AngularFrequency = InSettings.SpringStrength * 2.0f * PI;
+	const double Stiffness = AngularFrequency * AngularFrequency;
+	constexpr double Mass = 1.0;
+	constexpr double TargetVelocityAmount = 0.0;
+	constexpr double CriticalDamping = 1.0;
+	CurrentPosition = UKismetMathLibrary::VectorSpringInterp(
+		CurrentPosition,
+		TargetPosition,
+		PositionSpringState,
+		Stiffness,
+		CriticalDamping,
+		InDeltaTime,
+		Mass,
+		TargetVelocityAmount);
+
+	// set the actual goal to the current position from the spring
+	IKRigGoal->Position = CurrentPosition;
+}
+
+void FSpeedPlantingGoalState::Reset()
+{
+	bIsPlanted = false;
+	bWasReset = true;
+	PositionSpringState.Reset();
+}
+
 bool FIKRetargetSpeedPlantingOp::Initialize(
 	const FIKRetargetProcessor& InProcessor,
 	const FRetargetSkeleton& InSourceSkeleton,
@@ -63,37 +139,23 @@ bool FIKRetargetSpeedPlantingOp::Initialize(
 	const FIKRigGoalContainer& GoalContainer = InProcessor.GetIKRigGoalContainer();
 	for (const FRetargetSpeedPlantingSettings& SettingsForChain : Settings.ChainsToSpeedPlant)
 	{
-		const FName TargetChainName = SettingsForChain.TargetChainName;
-		const FResolvedBoneChain* TargetBoneChain = BoneChains.GetResolvedBoneChainByName(
+		FName TargetChainName = SettingsForChain.TargetChainName;
+		const FResolvedBoneChain* SourceBoneChain = nullptr;
+		const FResolvedBoneChain* TargetBoneChain = nullptr;
+		if (!FIKRetargetRunIKRigOp::ValidateIKChain(
 			TargetChainName,
-			ERetargetSourceOrTarget::Target,
-			ParentOp->Settings.IKRigAsset);
-		if (TargetBoneChain == nullptr)
-		{
-			InLog.LogWarning( FText::Format(
-			LOCTEXT("SpeedPlantingMissingChain", "Speed Planting Op: chain data is out of sync with IK Rig. Missing target chain, '{0}."),
-			FText::FromName(TargetChainName)));
-			continue;
-		}
-
-		if (TargetBoneChain->IKGoalName == NAME_None)
+			GoalContainer,
+			BoneChains,
+			this,
+			InParentOp,
+			InLog,
+			SourceBoneChain /*out*/,
+			TargetBoneChain /*out*/))
 		{
-			InLog.LogWarning( FText::Format(
-			LOCTEXT("SpeedPlantingChainWithoutGoal", "Speed Planting Op: specified chain does not have an IK goal. Cannot speed plant, '{0}."),
-			FText::FromName(TargetChainName)));
-			continue;
-		}
-
-		const FIKRigGoal* Goal = GoalContainer.FindGoalByName(TargetBoneChain->IKGoalName);
-		if (Goal == nullptr)
-		{
-			InLog.LogWarning( FText::Format(
-			LOCTEXT("SpeedPlantingRigWithoutGoal", "Speed Planting Op: target chain references a goal that is not present in the IK Rig, '{0}."),
-			FText::FromName(TargetBoneChain->IKGoalName)));
 			continue;
 		}
 		
-		GoalsToPlant.Emplace(TargetBoneChain->IKGoalName, &SettingsForChain, Goal->Position);
+		GoalsToPlant.Emplace(TargetBoneChain->IKGoalName, &SettingsForChain);
 	}
 
 	bIsInitialized = !GoalsToPlant.IsEmpty();
@@ -120,24 +182,8 @@ void FIKRetargetSpeedPlantingOp::Run(
 		{
 			continue; // goal excluded, just ignore it
 		}
-		
-		if (GoalToPlant.CurrentSpeedValue < 0.0f || GoalToPlant.CurrentSpeedValue > Settings.SpeedThreshold)
-		{
-			GoalToPlant.PrevGoalPosition = IKRigGoal->Position;
-			GoalToPlant.PositionSpring.Reset();
-			continue;
-		}
 
-		// run the goal position through a spring damper to un-plant it
-		IKRigGoal->Position = UKismetMathLibrary::VectorSpringInterp(
-				GoalToPlant.PrevGoalPosition,
-				IKRigGoal->Position,
-				GoalToPlant.PositionSpring,
-				static_cast<float>(Settings.Stiffness),
-				static_cast<float>(Settings.CriticalDamping),
-				static_cast<float>(InDeltaTime),
-				1.0f /*mass*/,
-				0.0f /*target velocity amount*/);
+		GoalToPlant.Update(IKRigGoal, InDeltaTime, Settings);
 	}
 }
 
@@ -196,10 +242,8 @@ void FIKRetargetSpeedPlantingOp::OnPlaybackReset()
 {
 	for (FSpeedPlantingGoalState& GoalToPlant : GoalsToPlant)
 	{
-		GoalToPlant.PositionSpring.Reset();
+		GoalToPlant.Reset();
 	}
-	
-	ResetThisTick = true;
 }
 
 void FIKRetargetSpeedPlantingOp::AnimGraphPreUpdateMainThread(
@@ -259,7 +303,33 @@ TArray<FName> FIKRetargetSpeedPlantingOp::GetRequiredSpeedCurves() const
 	return MoveTemp(OutSpeedCurveNames);
 }
 
+
+
 #if WITH_EDITOR
+
+void FIKRetargetSpeedPlantingOp::DebugDraw(
+	FPrimitiveDrawInterface* InPDI,
+	const FTransform& InSourceTransform,
+	const FTransform& InComponentTransform,
+	const double InComponentScale,
+	const FIKRetargetDebugDrawState& InEditorState) const
+{
+	for (const FSpeedPlantingGoalState& GoalToPlant : GoalsToPlant)
+	{
+		const FVector WorldPosition = InComponentTransform.TransformPosition(GoalToPlant.CurrentPosition);
+		DrawWireSphere(
+				InPDI,
+				WorldPosition,
+				GoalToPlant.bIsPlanted ? FLinearColor::Red : FLinearColor::Green,
+				Settings.DebugDrawSize,
+				12,
+				SDPG_MAX,
+				0.0f,
+				0.001f,
+				false);
+	}
+}
+
 FText FIKRetargetSpeedPlantingOp::GetWarningMessage() const
 {
 	// warn about missing curves
@@ -301,44 +371,9 @@ uint8* FIKRetargetSpeedPlantingOp::GetChainSettingsMemory(const FName InChainNam
 
 void FIKRetargetSpeedPlantingOp::RegenerateChainSettings(const FIKRetargetOpBase* InParentOp)
 {
-	const FIKRetargetRunIKRigOp* ParentOp = reinterpret_cast<const FIKRetargetRunIKRigOp*>(InParentOp);
-	if (!ensure(ParentOp))
-	{
-		return;
-	}
-	
-	// find the target chains that require goal retargeting
-	const TArray<FName> RequiredTargetChains = ParentOp->GetRequiredTargetChains();
-	if (RequiredTargetChains.IsEmpty())
-	{
-		// NOTE: if there's no chains, we don't clear the settings
-		// this allows users to clear and reassign a different rig and potentially retain/restore compatible settings
-		return;
-	}
-
-	// remove chains that are not required
-	Settings.ChainsToSpeedPlant.RemoveAll([&RequiredTargetChains](const FRetargetSpeedPlantingSettings& InChainSettings)
-	{
-		return !RequiredTargetChains.Contains(InChainSettings.TargetChainName);
-	});
-	
-	// add any required chains not already present
-	for (FName RequiredTargetChain : RequiredTargetChains)
-	{
-		bool bFound = false;
-		for (const FRetargetSpeedPlantingSettings& ChainToSpeedPlant : Settings.ChainsToSpeedPlant)
-		{
-			if (ChainToSpeedPlant.TargetChainName == RequiredTargetChain)
-			{
-				bFound = true;
-				break;
-			}
-		}
-		if (!bFound)
-		{
-			Settings.ChainsToSpeedPlant.Emplace(RequiredTargetChain);
-		}
-	}
+	constexpr bool bSkipUnmappedChains = true;
+	constexpr bool bSkipNonIKChains = true;
+	IKRetargetOpUtils::SynchronizeChainSettingsWithIKRig(Settings.ChainsToSpeedPlant, InParentOp, bSkipUnmappedChains, bSkipNonIKChains);
 }
 
 FIKRetargetSpeedPlantingOpSettings UIKRetargetSpeedPlantingController::GetSettings()
diff --git a/Engine/Plugins/Animation/IKRig/Source/IKRig/Public/Retargeter/RetargetOps/SpeedPlantingOp.h b/Engine/Plugins/Animation/IKRig/Source/IKRig/Public/Retargeter/RetargetOps/SpeedPlantingOp.h
--- a/Engine/Plugins/Animation/IKRig/Source/IKRig/Public/Retargeter/RetargetOps/SpeedPlantingOp.h
+++ b/Engine/Plugins/Animation/IKRig/Source/IKRig/Public/Retargeter/RetargetOps/SpeedPlantingOp.h
@@ -1,8 +1,8 @@
 // Copyright Epic Games, Inc. All Rights Reserved.
 #pragma once
 
-#include "Retargeter/IKRetargetOps.h"
 #include "Kismet/KismetMathLibrary.h"
+#include "Retargeter/IKRetargetOps.h"
 
 #include "SpeedPlantingOp.generated.h"
 
@@ -10,6 +10,8 @@
 
 #define LOCTEXT_NAMESPACE "SpeedPlantingOp"
 
+struct FIKRigGoal;
+
 USTRUCT(BlueprintType)
 struct FRetargetSpeedPlantingSettings
 {
@@ -42,37 +44,53 @@ struct FIKRetargetSpeedPlantingOpSettings : public FIKRetargetOpSettingsBase
 
 	/** Range 0 to 100. Default 15. The maximum speed a source bone can be moving while being considered 'planted'.
 	*  The target IK goal will not be allowed to move whenever the source bone speed drops below this threshold speed. */
-	UPROPERTY(EditAnywhere, BlueprintReadWrite, Category = "Speed Planting", meta = (ClampMin = "0.0", ClampMax = "100.0", UIMin = "0.0", UIMax = "100.0"))
+	UPROPERTY(EditAnywhere, BlueprintReadWrite, Category = "Speed Planting", meta = (ClampMin = "0.0", UIMin = "0.0", UIMax = "100.0"))
 	double SpeedThreshold = 15.0f;
-
-	// How stiff the spring model is that smoothly pulls the IK position after unplanting (more stiffness means more oscillation around the target value)
-	UPROPERTY(EditAnywhere, BlueprintReadWrite, Category = "Speed Planting", meta = (ClampMin = "0.0", UIMin = "0.0"))
-	double Stiffness = 250.0f;
-
-	// How much damping to apply to the spring (0 means no damping, 1 means critically damped which means no oscillation)
-	UPROPERTY(EditAnywhere, BlueprintReadWrite, Category = "Speed Planting", meta = (ClampMin = "0.0", ClampMax = "1.0", UIMin = "0.0", UIMax = "1.0"))
-	double CriticalDamping = 1.0f;
+	
+	/**
+	* The spring strength determines how hard it will pull towards the target when planting and unplanting.
+	* The value is the frequency (hz) at which it will oscillate when there is no damping. */
+	UPROPERTY(EditAnywhere, BlueprintReadWrite, Category = "Speed Planting", meta=(Input, ClampMin = "0.0", UIMin = "0.0", UIMax = "10.0"))
+	double SpringStrength = 3.0f;
+	// Deprecated in 5.8, replaced with higher-level "Strength" param
+	UPROPERTY(meta=(DeprecatedProperty))
+	double Stiffness;
+
+	/* Adjust size of goal debug drawing in viewport. Red is planted, green is unplanted. */
+	UPROPERTY(EditAnywhere, BlueprintReadWrite, Category = "Speed Planting", meta = (ClampMin = "0.0", UIMin = "0.0", UIMax = "10.0"))
+	double DebugDrawSize = 5.0;
 
 	UE_API virtual const UClass* GetControllerType() const override;
 
 	UE_API virtual void CopySettingsAtRuntime(const FIKRetargetOpSettingsBase* InSettingsToCopyFrom) override;
+
+	/** Custom override to load old "Stiffness" as new "Strength" */
+	virtual void PostLoad(const FIKRigObjectVersion::Type InVersion) override;
 };
 
 struct FSpeedPlantingGoalState
 {
 	FSpeedPlantingGoalState(
 		const FName InGoalName,
-		const FRetargetSpeedPlantingSettings* InSettings,
-		const FVector& InInitialGoalPosition) :
+		const FRetargetSpeedPlantingSettings* InSettings) :
 		GoalName(InGoalName),
-		Settings(InSettings),
-		PrevGoalPosition(InInitialGoalPosition) {};
+		Settings(InSettings) {};
+
+	void Update(
+		FIKRigGoal* IKRigGoal,
+		const double InDeltaTime,
+		const FIKRetargetSpeedPlantingOpSettings& InSettings);
+
+	void Reset();
 	
 	FName GoalName;
-	const FRetargetSpeedPlantingSettings *Settings;
-	FVectorSpringState PositionSpring;
-	FVector PrevGoalPosition;
+	const FRetargetSpeedPlantingSettings* Settings;
+	FVectorSpringState PositionSpringState;
+	FVector CurrentPosition;
+	FVector TargetPosition;
 	double CurrentSpeedValue = -1.0f;
+	bool bIsPlanted = false;
+	bool bWasReset = true;
 
 #if WITH_EDITOR
 	bool bFoundCurveInSourceComponent = false;
@@ -128,14 +146,19 @@ struct FIKRetargetSpeedPlantingOp : public FIKRetargetOpBase
 
 #if WITH_EDITOR
 public:
+	UE_API virtual void DebugDraw(
+		FPrimitiveDrawInterface* InPDI,
+		const FTransform& InSourceTransform,
+		const FTransform& InComponentTransform,
+		const double InComponentScale,
+		const FIKRetargetDebugDrawState& InEditorState) const override;
+	UE_API virtual bool HasDebugDrawing() override { return true; };
 	UE_API virtual FText GetWarningMessage() const override;
 	UE_API virtual void ResetChainSettingsToDefault(const FName& InChainName) override;
 	UE_API virtual bool AreChainSettingsAtDefault(const FName& InChainName) override;
 	UE_API virtual uint8* GetChainSettingsMemory(const FName InChainName) override;
 #endif
 
-	bool ResetThisTick;
-
 private:
 	
 	void RegenerateChainSettings(const FIKRetargetOpBase* InParentOp);
