[IKRetarget] Fixed non-initialized speed planting parameter.

#Synced-CL 49392216

diff --git a/Engine/Plugins/Animation/IKRig/Source/IKRig/Public/Retargeter/RetargetOps/SpeedPlantingOp.h b/Engine/Plugins/Animation/IKRig/Source/IKRig/Public/Retargeter/RetargetOps/SpeedPlantingOp.h
--- a/Engine/Plugins/Animation/IKRig/Source/IKRig/Public/Retargeter/RetargetOps/SpeedPlantingOp.h
+++ b/Engine/Plugins/Animation/IKRig/Source/IKRig/Public/Retargeter/RetargetOps/SpeedPlantingOp.h
@@ -54,7 +54,7 @@ struct FIKRetargetSpeedPlantingOpSettings : public FIKRetargetOpSettingsBase
 	double SpringStrength = 3.0f;
 	// Deprecated in 5.8, replaced with higher-level "Strength" param
 	UPROPERTY(meta=(DeprecatedProperty))
-	double Stiffness;
+	double Stiffness = 250.0f;
 
 	/* Adjust size of goal debug drawing in viewport. Red is planted, green is unplanted. */
 	UPROPERTY(EditAnywhere, BlueprintReadWrite, Category = "Speed Planting", meta = (ClampMin = "0.0", UIMin = "0.0", UIMax = "10.0"))
