[IK Retarget] Fixed crash when patching missing ops.

#Synced-CL 49601163

diff --git a/Engine/Plugins/Animation/IKRig/Source/IKRig/Private/Retargeter/IKRetargeter.cpp b/Engine/Plugins/Animation/IKRig/Source/IKRig/Private/Retargeter/IKRetargeter.cpp
--- a/Engine/Plugins/Animation/IKRig/Source/IKRig/Private/Retargeter/IKRetargeter.cpp
+++ b/Engine/Plugins/Animation/IKRig/Source/IKRig/Private/Retargeter/IKRetargeter.cpp
@@ -132,6 +132,9 @@ void UIKRetargeter::PostLoad()
 		TargetIKRigAsset->ConditionalPostLoad();
 	}
 
+	// clean must run before patching to avoid potential null ops making it into the patch logic
+	CleanOpStack();
+
 #if WITH_EDITOR
 	// post-load patches in order of oldest to newest asset upgrades
 	PostLoadOldSettingsToNew();
@@ -151,8 +154,6 @@ void UIKRetargeter::PostLoad()
 #endif
 
 	CleanRetargetPoses();
-	
-	CleanOpStack();
 }
 
 #if WITH_EDITOR
