[IK Retarget] Fixes regression causing editor test breakage.

#Synced-CL 49658913

diff --git a/Engine/Plugins/Animation/IKRig/Source/IKRig/Private/Retargeter/IKRetargeter.cpp b/Engine/Plugins/Animation/IKRig/Source/IKRig/Private/Retargeter/IKRetargeter.cpp
--- a/Engine/Plugins/Animation/IKRig/Source/IKRig/Private/Retargeter/IKRetargeter.cpp
+++ b/Engine/Plugins/Animation/IKRig/Source/IKRig/Private/Retargeter/IKRetargeter.cpp
@@ -132,8 +132,11 @@ void UIKRetargeter::PostLoad()
 		TargetIKRigAsset->ConditionalPostLoad();
 	}
 
-	// clean must run before patching to avoid potential null ops making it into the patch logic
-	CleanOpStack();
+	// remove null ops (could happen if op is in plugin that is not loaded)
+	RetargetOps.RemoveAll([](const FInstancedStruct& InOp)
+	{
+		return !InOp.IsValid();
+	});
 
 #if WITH_EDITOR
 	// post-load patches in order of oldest to newest asset upgrades
@@ -152,7 +155,8 @@ void UIKRetargeter::PostLoad()
 	PostLoadSplitIKChainsOp();	
 	PostLoadSplitOps(); // Note: Generic ops splitting, i.e. when turning singleton op into mutliple ones
 #endif
-
+	
+	CleanOpStack();
 	CleanRetargetPoses();
 }
 
@@ -912,12 +916,6 @@ void UIKRetargeter::CleanOpStack()
 		}
 		return INDEX_NONE;
 	};
-
-	// remove null ops (could happen if op is in plugin that is not loaded)
-	RetargetOps.RemoveAll([](const FInstancedStruct& InOp)
-	{
-		return !InOp.IsValid();
-	});
 	
 	// enforce unique non-None names on all ops
 	for (int32 OpIndex=0; OpIndex<RetargetOps.Num(); ++OpIndex)
