[IKRetarget] New op for weapon retargeting.

Adds support for retarget ops applying scale to bones.

#Synced-CL 49696721

diff --git a/Engine/Plugins/Animation/IKRig/Source/IKRig/Private/AnimNodes/AnimNode_RetargetPoseFromMesh.cpp b/Engine/Plugins/Animation/IKRig/Source/IKRig/Private/AnimNodes/AnimNode_RetargetPoseFromMesh.cpp
--- a/Engine/Plugins/Animation/IKRig/Source/IKRig/Private/AnimNodes/AnimNode_RetargetPoseFromMesh.cpp
+++ b/Engine/Plugins/Animation/IKRig/Source/IKRig/Private/AnimNodes/AnimNode_RetargetPoseFromMesh.cpp
@@ -166,7 +166,7 @@ void FAnimNode_RetargetPoseFromMesh::Evaluate_AnyThread(FPoseContext& Output)
 		// convert to component space
 		FAnimationRuntime::FillUpComponentSpaceTransforms(TargetRefSkeleton, InputLocalTransforms, PoseToRetargetFromComponentSpace);
 
-		// scale it accorrding to the scale source op (if used)
+		// scale it according to the scale source op (if used)
 		Processor.ApplySourceScaleToPose(PoseToRetargetFromComponentSpace);
 	};
 	
@@ -198,7 +198,10 @@ void FAnimNode_RetargetPoseFromMesh::Evaluate_AnyThread(FPoseContext& Output)
 			const FCompactPoseBoneIndex CompactBoneIndex(Pair.Key);
 			if (Output.Pose.IsValidIndex(CompactBoneIndex))
 			{
-				Output.Pose[CompactBoneIndex].SetScale3D(RefPose[Pair.Value].GetScale3D());
+				const FVector ScaleFromRetarget = Output.Pose[CompactBoneIndex].GetScale3D();
+				const FVector ScaleFromSkeletalMesh = RefPose[Pair.Value].GetScale3D() - FVector::OneVector;
+				const FVector NewScale = ScaleFromRetarget + ScaleFromSkeletalMesh;
+				Output.Pose[CompactBoneIndex].SetScale3D(NewScale);
 			}
 		}
 	};
diff --git a/Engine/Plugins/Animation/IKRig/Source/IKRig/Private/Retargeter/RetargetOps/WeaponGoalsOp.cpp b/Engine/Plugins/Animation/IKRig/Source/IKRig/Private/Retargeter/RetargetOps/WeaponGoalsOp.cpp
new file mode 100644
--- /dev/null
+++ b/Engine/Plugins/Animation/IKRig/Source/IKRig/Private/Retargeter/RetargetOps/WeaponGoalsOp.cpp
@@ -0,0 +1,223 @@
+﻿// Copyright Epic Games, Inc. All Rights Reserved.
+
+#include "Retargeter/RetargetOps/WeaponGoalsOp.h"
+
+#include "GeometryCollection/Facades/CollectionConstraintOverrideFacade.h"
+#include "Retargeter/IKRetargetProcessor.h"
+#include "Retargeter/RetargetOps/RunIKRigOp.h"
+
+#define LOCTEXT_NAMESPACE "WeaponGoalsOp"
+
+#include UE_INLINE_GENERATED_CPP_BY_NAME(WeaponGoalsOp)
+
+struct FIKRetargetWeaponGoalsOpSettings;
+
+
+FIKRetargetWeaponGoalsOpSettings::FIKRetargetWeaponGoalsOpSettings()
+{
+	// construct with expected default names (based on UE skeleton convention)
+	LeftWeaponBone.BoneName = FName("weapon_l");
+	RightWeaponBone.BoneName = FName("weapon_r");
+	LeftHandAttachBone.BoneName = FName("hand_attach_l");
+	RightHandAttachBone.BoneName = FName("hand_attach_r");
+	LeftHandBone.BoneName = FName("hand_l");
+	RightHandBone.BoneName = FName("hand_r");
+	IKHandLeftBone.BoneName = FName("ik_hand_l");
+	IKHandRightBone.BoneName = FName("ik_hand_r");
+	IKHandGun.BoneName = FName("ik_hand_gun");
+}
+
+const UClass* FIKRetargetWeaponGoalsOpSettings::GetControllerType() const
+{
+	return UIKRetargetWeaponGoalsOpController::StaticClass();
+}
+
+void FIKRetargetWeaponGoalsOpSettings::CopySettingsAtRuntime(const FIKRetargetOpSettingsBase* InSettingsToCopyFrom)
+{
+	// copy everything except the bones (must be constant at runtime)
+	const TArray<FName> PropertiesToIgnore = {
+		GET_MEMBER_NAME_CHECKED(FIKRetargetWeaponGoalsOpSettings, LeftWeaponBone),
+		GET_MEMBER_NAME_CHECKED(FIKRetargetWeaponGoalsOpSettings, RightWeaponBone),
+		GET_MEMBER_NAME_CHECKED(FIKRetargetWeaponGoalsOpSettings, LeftHandAttachBone),
+		GET_MEMBER_NAME_CHECKED(FIKRetargetWeaponGoalsOpSettings, RightHandAttachBone),
+		GET_MEMBER_NAME_CHECKED(FIKRetargetWeaponGoalsOpSettings, LeftHandBone),
+		GET_MEMBER_NAME_CHECKED(FIKRetargetWeaponGoalsOpSettings, RightHandBone),
+		GET_MEMBER_NAME_CHECKED(FIKRetargetWeaponGoalsOpSettings, IKHandLeftBone),
+		GET_MEMBER_NAME_CHECKED(FIKRetargetWeaponGoalsOpSettings, IKHandRightBone),
+		GET_MEMBER_NAME_CHECKED(FIKRetargetWeaponGoalsOpSettings, IKHandGun)};
+	FIKRetargetOpBase::CopyStructProperties(
+		FIKRetargetWeaponGoalsOpSettings::StaticStruct(),
+		InSettingsToCopyFrom,
+		this,
+		PropertiesToIgnore);
+}
+
+bool FIKRetargetWeaponGoalOp::Initialize(
+	const FIKRetargetProcessor& InProcessor,
+	const FRetargetSkeleton& InSourceSkeleton,
+	const FTargetSkeleton& InTargetSkeleton,
+	const FIKRetargetOpBase* InParentOp,
+	FIKRigLogger& InLog)
+{
+	bIsInitialized = true;
+	
+	// load target bones
+	BoneIndices[(int32)ERetargetWeaponBone::LeftWeapon] = InTargetSkeleton.FindBoneIndexByName(Settings.LeftWeaponBone.BoneName);
+	BoneIndices[(int32)ERetargetWeaponBone::RightWeapon] = InTargetSkeleton.FindBoneIndexByName(Settings.RightWeaponBone.BoneName);
+	BoneIndices[(int32)ERetargetWeaponBone::LeftHandAttach] = InTargetSkeleton.FindBoneIndexByName(Settings.LeftHandAttachBone.BoneName);
+	BoneIndices[(int32)ERetargetWeaponBone::RightHandAttach] = InTargetSkeleton.FindBoneIndexByName(Settings.RightHandAttachBone.BoneName);
+	BoneIndices[(int32)ERetargetWeaponBone::LeftHand] = InTargetSkeleton.FindBoneIndexByName(Settings.LeftHandBone.BoneName);
+	BoneIndices[(int32)ERetargetWeaponBone::RightHand] = InTargetSkeleton.FindBoneIndexByName(Settings.RightHandBone.BoneName);
+	BoneIndices[(int32)ERetargetWeaponBone::IKHandLeft] = InTargetSkeleton.FindBoneIndexByName(Settings.IKHandLeftBone.BoneName);
+	BoneIndices[(int32)ERetargetWeaponBone::IKHandRight] = InTargetSkeleton.FindBoneIndexByName(Settings.IKHandRightBone.BoneName);
+	BoneIndices[(int32)ERetargetWeaponBone::HandGun] = InTargetSkeleton.FindBoneIndexByName(Settings.IKHandGun.BoneName);
+
+	// warn if any bones were missing
+	if (!Settings.bSuppressMissingBoneWarnings)
+	{
+		const UEnum* WeaponBoneEnum = StaticEnum<ERetargetWeaponBone>();
+		check(WeaponBoneEnum);
+		
+		for (int32 Index = 0; Index < (int32)ERetargetWeaponBone::Count; ++Index)
+		{
+			if (BoneIndices[Index] == INDEX_NONE)
+			{
+				FText BoneEnumName = WeaponBoneEnum->GetDisplayNameTextByIndex(Index);
+				FText ErrorMsg = FText::Format(LOCTEXT("WeaponGoalOpMissingBone", "Weapon Goals Op: Missing bone for '{0}'."),BoneEnumName);
+				InLog.LogWarning(ErrorMsg);
+			}
+		}
+	}
+
+	// load and validate the goals associated with each hand
+	const FIKRigGoalContainer& GoalContainer = InProcessor.GetIKRigGoalContainer();
+	const TArray<FIKRigGoal>& AllGoals = GoalContainer.GetGoalArray();
+	LeftHandGoalName = NAME_None;
+	RightHandGoalName = NAME_None;
+	for (const FIKRigGoal& Goal : AllGoals)
+	{
+		if (Goal.BoneName == Settings.LeftHandBone.BoneName)
+		{
+			LeftHandGoalName = Goal.Name;
+		}
+
+		if (Goal.BoneName == Settings.RightHandBone.BoneName)
+		{
+			RightHandGoalName = Goal.Name;
+		}
+	}
+	
+	if (LeftHandGoalName == NAME_None)
+	{
+		FText ErrorMsg = FText::Format(LOCTEXT("WeaponGoalOpMissingLeftGoal", "Weapon Goals Op: No IK Goal found on Left Hand Bone, '{0}'."), FText::FromName(Settings.LeftHandBone.BoneName));
+		InLog.LogWarning(ErrorMsg);
+	}
+	
+	if (RightHandGoalName == NAME_None)
+	{
+		FText ErrorMsg = FText::Format(LOCTEXT("WeaponGoalOpMissingRightGoal", "Weapon Goals Op: No IK Goal found on Right Hand Bone, '{0}'."), FText::FromName(Settings.RightHandBone.BoneName));
+		InLog.LogWarning(ErrorMsg);
+	}
+	
+	return bIsInitialized;
+}
+
+void FIKRetargetWeaponGoalOp::Run(
+	FIKRetargetProcessor& InProcessor,
+	const double InDeltaTime,
+	const TArray<FTransform>& InSourceGlobalPose,
+	TArray<FTransform>& OutTargetGlobalPose)
+{
+	const FTargetSkeleton& TargetSkeleton = InProcessor.GetTargetSkeleton();
+	
+	// scale bones
+	if (Settings.bEnableWeaponScale)
+	{
+		const FVector Scale(Settings.WeaponScale,Settings.WeaponScale,Settings.WeaponScale);
+		TArray<int32, TFixedAllocator<5>> BonesToScale = {
+			(int32)ERetargetWeaponBone::LeftWeapon,
+			(int32)ERetargetWeaponBone::RightWeapon,
+			(int32)ERetargetWeaponBone::LeftHandAttach,
+			(int32)ERetargetWeaponBone::RightHandAttach,
+			(int32)ERetargetWeaponBone::HandGun };
+		for (int32 Index : BonesToScale)
+		{
+			const int32 BoneIndex = BoneIndices[Index];
+			if (BoneIndex != INDEX_NONE)
+			{
+				// scale the bone and propagate to children
+				FTransform BoneTransform = OutTargetGlobalPose[BoneIndex];
+				BoneTransform.SetScale3D(Scale);
+				TargetSkeleton.SetGlobalTransformAndUpdateChildren(BoneIndex, BoneTransform, OutTargetGlobalPose);
+			}
+		}
+	}
+
+	// blend IK hand gun location between Left/Right hand
+	if (Settings.bSnapIKHandGun)
+	{
+		const int32 LeftHandBoneIndex = BoneIndices[(int32)ERetargetWeaponBone::LeftHand];
+		const int32 RightHandBoneIndex = BoneIndices[(int32)ERetargetWeaponBone::RightHand];
+		const int32 HandGunBoneIndex = BoneIndices[(int32)ERetargetWeaponBone::HandGun];
+		const bool bAllBonesFound = LeftHandBoneIndex != INDEX_NONE && RightHandBoneIndex != INDEX_NONE && HandGunBoneIndex != INDEX_NONE;
+		if (bAllBonesFound)
+		{
+			const FTransform& LeftHand = OutTargetGlobalPose[LeftHandBoneIndex];
+			const FTransform& RightHand = OutTargetGlobalPose[RightHandBoneIndex];
+			const FVector BlendedPosition = FMath::Lerp(LeftHand.GetLocation(), RightHand.GetLocation(), Settings.LeftRightHandIKAlpha);
+			FTransform HandGunTransform = OutTargetGlobalPose[HandGunBoneIndex];
+			HandGunTransform.SetTranslation(BlendedPosition);
+			TargetSkeleton.SetGlobalTransformAndUpdateChildren(HandGunBoneIndex, HandGunTransform, OutTargetGlobalPose);
+		}	
+	}
+	
+	// blend the IK goals
+	FIKRigGoalContainer& GoalContainer = InProcessor.GetIKRigGoalContainer();
+	auto BlendGoalToIKHandBone = [&OutTargetGlobalPose](FIKRigGoal* InGoal, int32 InIKBoneIndex, float InAlpha)
+		{
+			const FTransform& IKBoneTransform = OutTargetGlobalPose[InIKBoneIndex];
+			InGoal->Position = FMath::Lerp(InGoal->Position, IKBoneTransform.GetLocation(), InAlpha);
+			const FQuat Rotation = FQuat::FastLerp(InGoal->Rotation.Quaternion(), IKBoneTransform.GetRotation(), InAlpha);
+			InGoal->Rotation = Rotation.GetNormalized().Rotator();
+		};
+	if (FIKRigGoal* LeftHandGoal = GoalContainer.FindGoalByName(LeftHandGoalName))
+	{
+		BlendGoalToIKHandBone(LeftHandGoal, BoneIndices[(int32)ERetargetWeaponBone::IKHandLeft], Settings.LeftHandIKAlpha);
+	}
+	if (FIKRigGoal* RightHandGoal = GoalContainer.FindGoalByName(RightHandGoalName))
+	{
+		BlendGoalToIKHandBone(RightHandGoal, BoneIndices[(int32)ERetargetWeaponBone::IKHandRight], Settings.RightHandIKAlpha);
+	}
+}
+
+FIKRetargetOpSettingsBase* FIKRetargetWeaponGoalOp::GetSettings()
+{
+	return &Settings;
+}
+
+const UScriptStruct* FIKRetargetWeaponGoalOp::GetSettingsType() const
+{
+	return FIKRetargetWeaponGoalsOpSettings::StaticStruct();
+}
+
+const UScriptStruct* FIKRetargetWeaponGoalOp::GetType() const
+{
+	return FIKRetargetWeaponGoalOp::StaticStruct();
+}
+
+const UScriptStruct* FIKRetargetWeaponGoalOp::GetParentOpType() const
+{
+	return FIKRetargetRunIKRigOp::StaticStruct();
+}
+
+FIKRetargetWeaponGoalsOpSettings UIKRetargetWeaponGoalsOpController::GetSettings()
+{
+	return *reinterpret_cast<FIKRetargetWeaponGoalsOpSettings*>(OpSettingsToControl);
+}
+
+void UIKRetargetWeaponGoalsOpController::SetSettings(FIKRetargetWeaponGoalsOpSettings InSettings)
+{
+	OpSettingsToControl->CopySettingsAtRuntime(&InSettings);
+}
+
+#undef LOCTEXT_NAMESPACE
diff --git a/Engine/Plugins/Animation/IKRig/Source/IKRig/Public/Retargeter/RetargetOps/WeaponGoalsOp.h b/Engine/Plugins/Animation/IKRig/Source/IKRig/Public/Retargeter/RetargetOps/WeaponGoalsOp.h
new file mode 100644
--- /dev/null
+++ b/Engine/Plugins/Animation/IKRig/Source/IKRig/Public/Retargeter/RetargetOps/WeaponGoalsOp.h
@@ -0,0 +1,168 @@
+﻿// Copyright Epic Games, Inc. All Rights Reserved.
+#pragma once
+
+#include "Animation/BoneReference.h"
+#include "Retargeter/IKRetargetOps.h"
+#include "WeaponGoalsOp.generated.h"
+
+#define UE_API IKRIG_API
+
+struct FResolvedBoneChain;
+struct FIKRigGoalContainer;
+
+USTRUCT(BlueprintType, meta = (DisplayName = "Weapon Op Settings"))
+struct FIKRetargetWeaponGoalsOpSettings : public FIKRetargetOpSettingsBase
+{
+	GENERATED_BODY()
+
+	FIKRetargetWeaponGoalsOpSettings();
+
+	/* An ancillary weapon bone for the left hand.
+	 * This will be scaled with Weapon Scale. Defaults to "weapon_l". */
+	UPROPERTY(EditAnywhere, Category = "Bones")
+	FBoneReference LeftWeaponBone;
+
+	/* An ancillary weapon bone for the right hand.
+	 * This will be scaled with Weapon Scale. Defaults to "weapon_r". */
+	UPROPERTY(EditAnywhere, Category = "Bones")
+	FBoneReference RightWeaponBone;
+
+	/* An optional bone used to attach props to the left hand.
+	 * This will be scaled with Weapon Scale. Defaults to "hand_attach_l". */
+	UPROPERTY(EditAnywhere, Category = "Bones")
+	FBoneReference LeftHandAttachBone;
+
+	/* An optional bone used to attach props to the right hand.
+	 * This will be scaled with Weapon Scale. Defaults to "hand_attach_r". */
+	UPROPERTY(EditAnywhere, Category = "Bones")
+	FBoneReference RightHandAttachBone;
+
+	/* The bone skinned to the left hand at wrist location.
+	 * This bone must have an associated IK Goal. Defaults to "hand_l". */
+	UPROPERTY(EditAnywhere, Category = "Bones")
+	FBoneReference LeftHandBone;
+
+	/* The bone skinned to the right hand at wrist location.
+	 * This bone must have an associated IK Goal. Defaults to "hand_r". */
+	UPROPERTY(EditAnywhere, Category = "Bones")
+	FBoneReference RightHandBone;
+
+	/* A bone representing the location of the left hand IK position while holding a two-handed prop.
+	 * Typically, a child of the IKHandGun. Defaults to "ik_hand_l". */
+	UPROPERTY(EditAnywhere, Category = "Bones", DisplayName="IK Hand Left Bone")
+	FBoneReference IKHandLeftBone;
+
+	/* A bone representing the location of the right hand IK position while holding a two-handed prop.
+	 * Typically, a child of the IKHandGun. Defaults to "ik_hand_r". */
+	UPROPERTY(EditAnywhere, Category = "Bones", DisplayName="IK Hand Right Bone")
+	FBoneReference IKHandRightBone;
+
+	/* A bone located at the right hand, and parent of IKHandLeftBone and IKHandRightBone.
+	 * This will be scaled with Weapon Scale. Defaults to "ik_hand_gun". */
+	UPROPERTY(EditAnywhere, Category = "Bones")
+	FBoneReference IKHandGun;
+
+	/* When false, all weapon scaling is skipped and WeaponScale has no effect. */
+	UPROPERTY(EditAnywhere, BlueprintReadWrite, Category = "Op Settings")
+	bool bEnableWeaponScale = true;
+	
+	/* Range 0 to inf. Default 1. Scales the IKHandGun bones, Left/Right Weapon bones and hand attach bones. */
+	UPROPERTY(EditAnywhere, BlueprintReadWrite, Category = "Op Settings", meta = (ClampMin = "0.01", UIMin = "0.01", UIMax = "10.0"))
+	double WeaponScale = 1.0f;
+
+	/** When true, enables RetargetHandIKAlpha to blend the IK Gun hierarchy between left/right hands. */
+	UPROPERTY(EditAnywhere, BlueprintReadWrite, Category = "Op Settings", DisplayName="Snap IK Hand Gun")
+	bool bSnapIKHandGun = true;
+
+	/* Blends the entire Hand Gun bone hierarchy location from the left hand (at 0.0) to the right hand (at 1.0). Default is 1.0 (right hand).
+	 * NOTE: this propagates the translational offset to the children bones (IKHandLeftBone/IKHandRightBone). */
+	UPROPERTY(EditAnywhere, BlueprintReadWrite, Category = "Op Settings", DisplayName="Retarget Hand IK Alpha", meta = (ClampMin = "0.0", ClampMax = "1.0", UIMin = "0.0", UIMax = "1.0"))
+	double LeftRightHandIKAlpha = 1.0f;
+
+	/* Blend the effect of the left hand IK Goal from it's input location (0.0) to the IKHandLeftBone location (1.0). Default is 1.0. */
+	UPROPERTY(EditAnywhere, BlueprintReadWrite, Category = "Op Settings", DisplayName="Left Hand IK Alpha", meta = (ClampMin = "0.0", ClampMax = "1.0", UIMin = "0.0", UIMax = "1.0"))
+	double LeftHandIKAlpha = 1.0f;
+
+	/* Blend the effect of the right hand IK Goal from it's input location (0.0) to the IKHandRightBone location (1.0). Default is 1.0. */
+	UPROPERTY(EditAnywhere, BlueprintReadWrite, Category = "Op Settings", DisplayName="Right Hand IK Alpha", meta = (ClampMin = "0.0", ClampMax = "1.0", UIMin = "0.0", UIMax = "1.0"))
+	double RightHandIKAlpha = 1.0f;
+
+	/* When true, will suppress warnings in the retarget editor output log about missing bones from this op. */
+	UPROPERTY(EditAnywhere, BlueprintReadWrite, Category = "Op Settings")
+	bool bSuppressMissingBoneWarnings = false;
+
+	virtual const UClass* GetControllerType() const override;
+
+	virtual void CopySettingsAtRuntime(const FIKRetargetOpSettingsBase* InSettingsToCopyFrom) override;
+};
+
+
+UENUM()
+enum class ERetargetWeaponBone : uint8
+{
+	LeftWeapon      UMETA(DisplayName = "Left Weapon"), // labels here used for warnings only, not exposed to UI
+	RightWeapon     UMETA(DisplayName = "Right Weapon"),
+	LeftHandAttach  UMETA(DisplayName = "Left Hand Attach"),
+	RightHandAttach UMETA(DisplayName = "Right Hand Attach"),
+	LeftHand        UMETA(DisplayName = "Left Hand"),
+	RightHand       UMETA(DisplayName = "Right Hand"),
+	IKHandLeft      UMETA(DisplayName = "IK Hand Left"),
+	IKHandRight     UMETA(DisplayName = "IK Hand Right"),
+	HandGun         UMETA(DisplayName = "Hand Gun"),
+	Count           UMETA(Hidden)
+};
+
+USTRUCT(BlueprintType, meta = (DisplayName = "Weapon Goals"))
+struct FIKRetargetWeaponGoalOp : public FIKRetargetOpBase
+{
+	GENERATED_BODY()
+	
+	UE_API virtual bool Initialize(
+		const FIKRetargetProcessor& InProcessor,
+		const FRetargetSkeleton& InSourceSkeleton,
+		const FTargetSkeleton& InTargetSkeleton,
+		const FIKRetargetOpBase* InParentOp,
+		FIKRigLogger& InLog) override;
+	
+	UE_API virtual void Run(
+		FIKRetargetProcessor& InProcessor,
+		const double InDeltaTime,
+		const TArray<FTransform>& InSourceGlobalPose,
+		TArray<FTransform>& OutTargetGlobalPose) override;
+	
+	UE_API virtual FIKRetargetOpSettingsBase* GetSettings() override;
+	
+	UE_API virtual const UScriptStruct* GetSettingsType() const override;
+
+	UE_API virtual const UScriptStruct* GetType() const override;
+
+	UE_API virtual const UScriptStruct* GetParentOpType() const override;
+	
+	UPROPERTY()
+	FIKRetargetWeaponGoalsOpSettings Settings;
+
+private:
+	int32 BoneIndices[(int32)ERetargetWeaponBone::Count];
+	FName LeftHandGoalName = NAME_None;
+	FName RightHandGoalName = NAME_None;
+};
+
+/* The blueprint/python API for editing a Weapon Goals Op */
+UCLASS(MinimalAPI, BlueprintType)
+class UIKRetargetWeaponGoalsOpController : public UIKRetargetOpControllerBase
+{
+	GENERATED_BODY()
+	
+public:
+	/* Get the current op settings as a struct.
+	 * @return FIKRetargetWeaponGoalsOpSettings struct with the current settings used by the op. */
+	UFUNCTION(BlueprintCallable, Category = Settings)
+	UE_API FIKRetargetWeaponGoalsOpSettings GetSettings();
+
+	/* Set the op settings. Input is a custom struct type for this op.
+	 * @param InSettings a FIKRetargetWeaponGoalsOpSettings struct containing all the settings to apply to this op */
+	UFUNCTION(BlueprintCallable, Category = Settings)
+	UE_API void SetSettings(FIKRetargetWeaponGoalsOpSettings InSettings);
+};
+
+#undef UE_API
\ No newline at end of file
