[IKRetarget] Added absolute mode to Pelvis Motion Op

#Synced-CL 49791386

#Synced-CL 49791386

diff --git a/Engine/Plugins/Animation/IKRig/Source/IKRig/Private/Retargeter/RetargetOps/PelvisMotionOp.cpp b/Engine/Plugins/Animation/IKRig/Source/IKRig/Private/Retargeter/RetargetOps/PelvisMotionOp.cpp
--- a/Engine/Plugins/Animation/IKRig/Source/IKRig/Private/Retargeter/RetargetOps/PelvisMotionOp.cpp
+++ b/Engine/Plugins/Animation/IKRig/Source/IKRig/Private/Retargeter/RetargetOps/PelvisMotionOp.cpp
@@ -122,17 +122,15 @@ void FIKRetargetPelvisMotionOp::Run(
 	const TArray<FTransform>& InSourceGlobalPose,
 	TArray<FTransform>& OutTargetGlobalPose)
 {
-	EncodePose(InSourceGlobalPose);
-
-	FTransform PelvisGlobal;
-	DecodePose(PelvisGlobal);
+	FTransform NewPelvisGlobal;
+	RetargetPelvisTransform(InSourceGlobalPose, NewPelvisGlobal);
 	
-	// update global transforms below root
+	// update global transforms below pelvis
 	FTargetSkeleton& TargetSkeleton = Processor.GetTargetSkeleton();
-	TargetSkeleton.SetGlobalTransformAndUpdateChildren(Target.BoneIndex, PelvisGlobal, TargetSkeleton.OutputGlobalPose);
+	TargetSkeleton.SetGlobalTransformAndUpdateChildren(Target.BoneIndex, NewPelvisGlobal, TargetSkeleton.OutputGlobalPose);
 
 #if WITH_EDITOR
-	CurrentPelvisTransform = PelvisGlobal;
+	CurrentPelvisTransform = NewPelvisGlobal;
 #endif
 }
 
@@ -288,16 +286,15 @@ void FIKRetargetPelvisMotionOp::Reset()
 	Target = FPelvisTarget();
 }
 
-void FIKRetargetPelvisMotionOp::EncodePose(const TArray<FTransform>& SourceGlobalPose)
+void FIKRetargetPelvisMotionOp::RetargetPelvisTransform(
+	const TArray<FTransform>& SourceGlobalPose,
+	FTransform& OutPelvisGlobalPose)
 {
+	// get source pelvis transform
 	const FTransform& SourceTransform = SourceGlobalPose[Source.BoneIndex];
 	Source.CurrentPosition = SourceTransform.GetTranslation();
-	Source.CurrentPositionNormalized = Source.CurrentPosition * Source.InitialHeightInverse;
-	Source.CurrentRotation = SourceTransform.GetRotation();	
-}
-
-void FIKRetargetPelvisMotionOp::DecodePose(FTransform& OutPelvisGlobalPose)
-{
+	Source.CurrentRotation = SourceTransform.GetRotation();
+	
 	// retarget rotation
 	// NOTE: rotation must be done before translation to get proper orientation for local position offsets
 	FQuat Rotation;
@@ -323,31 +320,47 @@ void FIKRetargetPelvisMotionOp::DecodePose(FTransform& OutPelvisGlobalPose)
 	// retarget position
 	FVector Position;
 	{
-		// generate basic pelvis position by scaling the normalized position by root height
-		FVector RetargetedPosition = Source.CurrentPositionNormalized * Target.InitialHeight;
-		Position = RetargetedPosition;
+		// generate pelvis motion scaled by the relative height ratio of the target to the source height
+		// ie, if source is 2x taller than target, motion is scaled by 0.5
+		const FVector NormalizedSourcePosition = Source.CurrentPosition * Source.InitialHeightInverse;
+		const FVector PositionScaled = NormalizedSourcePosition * Target.InitialHeight;
+
+		// generate pelvis motion as absolute delta
+		const FVector AbsoluteDelta = Source.CurrentPosition - Source.InitialPosition;
+		const FVector PositionAbsolute = Target.InitialPosition + AbsoluteDelta;
+
+		// blend between scaled and absolute motion
+		Position = FMath::Lerp(PositionScaled, PositionAbsolute, Settings.BlendToAbsoluteOffset);
 
 		// apply crotch-to-floor constraint
 		if (!FMath::IsNearlyZero(Settings.FloorConstraintWeight))
 		{
 			if (Source.CurrentPosition.Z <= Source.InitialPosition.Z)
 			{
+				// calculate the constrained height as a proportional target height relative to the source
 				const float SourceInitialCrotchHeight = Source.InitialPosition.Z - Settings.SourceCrotchOffset;
 				const float SourceCurrentCrotchHeight = Source.CurrentPosition.Z - Settings.SourceCrotchOffset;
 				const float SourceCrotchHeightPercent = SourceCurrentCrotchHeight * (1.0f / FMath::Max(1.0f, SourceInitialCrotchHeight));
 				const float TargetInitialCrotchHeight = Target.InitialPosition.Z - Settings.TargetCrotchOffset;
 				const float TargetCurrentCrotchHeight = TargetInitialCrotchHeight * SourceCrotchHeightPercent;
-				const float ConstrainedHeight = TargetCurrentCrotchHeight + Settings.TargetCrotchOffset; 
-				RetargetedPosition.Z = Position.Z = FMath::Lerp(Position.Z, ConstrainedHeight, Settings.FloorConstraintWeight);
+				const float ConstrainedHeight = TargetCurrentCrotchHeight + Settings.TargetCrotchOffset;
+
+				// optionally apply ground proximity falloff to the floor constraint weight
+				float FinalWeight = Settings.FloorConstraintWeight;
+				if (Settings.bUseGroundFalloff)
+				{
+					// get how close we are to the ground (1.0 = Touching Floor, 0.0 = Standing)
+					const float FalloffStartHeight = Target.InitialHeight * Settings.GroundFalloffHeightPercent;
+					const float ProximityFalloff = 1.0f - FMath::Clamp(Position.Z, 0.0f, FalloffStartHeight) / FalloffStartHeight;
+					FinalWeight *= ProximityFalloff;
+				}
+				
+				Position.Z = FMath::Lerp(Position.Z, ConstrainedHeight, FinalWeight);
 			}
-			else
-			{
-				const float InitialSourceTargetDelta = Target.InitialPosition.Z - Source.InitialPosition.Z;
-				const float ConstrainedHeight = Source.CurrentPosition.Z + InitialSourceTargetDelta;
-				RetargetedPosition.Z = Position.Z = FMath::Lerp(Position.Z, ConstrainedHeight, Settings.FloorConstraintWeight);
-			}
-			
 		}
+
+		// record the retargeted position before applying any modifications
+		const FVector RetargetedPosition = Position;
 		
 		// blend the pelvis position towards the source pelvis position
 		const FVector PerAxisAlpha = Settings.BlendToSourceTranslation * Settings.BlendToSourceTranslationWeights;
diff --git a/Engine/Plugins/Animation/IKRig/Source/IKRig/Public/Retargeter/RetargetOps/PelvisMotionOp.h b/Engine/Plugins/Animation/IKRig/Source/IKRig/Public/Retargeter/RetargetOps/PelvisMotionOp.h
--- a/Engine/Plugins/Animation/IKRig/Source/IKRig/Public/Retargeter/RetargetOps/PelvisMotionOp.h
+++ b/Engine/Plugins/Animation/IKRig/Source/IKRig/Public/Retargeter/RetargetOps/PelvisMotionOp.h
@@ -25,6 +25,10 @@ struct FIKRetargetPelvisMotionOpSettings : public FIKRetargetOpSettingsBase
 	UPROPERTY(EditAnywhere, Category=Setup, meta=(ReinitializeOnEdit))
 	FBoneReference TargetPelvisBone;
 
+	/** Blend from motion scaled by the size of the character, to the absolute delta motion from the source.
+	 * NOTE: Typically, this is only applied to vertical motion (Z) to prevent dampening the pelvis bob on short characters. */
+	UPROPERTY(EditAnywhere, BlueprintReadWrite, Category = "Scaled or Absolute", meta = (ClampMin = "0.0", UIMin = "0.0", UIMax = "1.0"))
+	FVector BlendToAbsoluteOffset = FVector::ZeroVector;
 
 	/** Range 0 to 1. Default is 0. Set to a value of 1 to turn the pelvis floor constraint ON.
 	 * When ON, the floor constraint will adjust the vertical Pelvis motion separately according to the following rules:
@@ -42,6 +46,17 @@ struct FIKRetargetPelvisMotionOpSettings : public FIKRetargetOpSettingsBase
 	 * NOTE: Adjust this until the green dot is located roughly at the crotch of the character. */
 	UPROPERTY(EditAnywhere, BlueprintReadWrite, Category = FloorConstraint)
 	double TargetCrotchOffset = 0.0f;
+
+	/** When true, the effect of the floor constraint decays as the pelvis gets further from the floor.
+	 * NOTE: This can help avoid over dampening vertical motion on the pelvis when retargeting to characters that are shorter than the source. */
+	UPROPERTY(EditAnywhere, BlueprintReadWrite, Category = FloorConstraint)
+	bool bUseGroundFalloff = false;
+	
+	/** Controls the height at which the floor constraint will begin blending in as a percentage of the pelvis height in the retarget pose.
+	 * NOTE: Use this to prevent the floor constraint from dampening vertical pelvis motion at normal standing heights.
+	 * At 1, the constraint is active at normal standing height. At 0.1, the constraint will only activate when within 10% of the height off the floor. */
+	UPROPERTY(EditAnywhere, BlueprintReadWrite, Category = FloorConstraint, meta = (ClampMin = "0.1", ClampMax = "1.0", UIMin = "0.1", UIMax = "1.0"))
+	double GroundFalloffHeightPercent = 0.9f;
 	
 	/** Range 0 to 1. Default 1. Blends the amount of retargeted pelvis rotation to apply.
 	*  At 0 the pelvis is left at the rotation from the retarget pose.
@@ -145,7 +160,6 @@ struct FPelvisSource
 	float InitialHeightInverse;
 	FVector InitialPosition;
 	FVector CurrentPosition;
-	FVector CurrentPositionNormalized;
 	FQuat CurrentRotation;
 };
 
@@ -218,9 +232,7 @@ private:
 		const FTargetSkeleton& TargetSkeleton,
 		FIKRigLogger& Log);
 	
-	void EncodePose(const TArray<FTransform> &SourceGlobalPose);
-	
-	void DecodePose(FTransform& OutPelvisGlobalPose);
+	void RetargetPelvisTransform(const TArray<FTransform>& SourceGlobalPose, FTransform& OutPelvisGlobalPose);
 
 	// transient work data
 	FPelvisSource Source;
