[IKRetarget]

#Synced-CL 49791386

diff --git a/Engine/Plugins/Animation/IKRig/Source/IKRig/Private/Retargeter/RetargetOps/PinBoneOp.cpp b/Engine/Plugins/Animation/IKRig/Source/IKRig/Private/Retargeter/RetargetOps/PinBoneOp.cpp
--- a/Engine/Plugins/Animation/IKRig/Source/IKRig/Private/Retargeter/RetargetOps/PinBoneOp.cpp
+++ b/Engine/Plugins/Animation/IKRig/Source/IKRig/Private/Retargeter/RetargetOps/PinBoneOp.cpp
@@ -343,16 +343,18 @@ FTransform FIKRetargetPinBoneOp::GetNewBoneTransform(
 		}
 	case EPinBoneRotationMode::MaintainOffsetFromBoneToCopyFrom:
 		{
-			// get rotation delta between from/to in retarget pose
+			// get global retarget pose of from/to bones
 			const TArray<FTransform>& FromRetargetPose = CopyFrom.RetargetPoses.GetGlobalRetargetPose();
 			const TArray<FTransform>& ToRetargetPose = CopyTo.RetargetPoses.GetGlobalRetargetPose();
-			const FTransform& BoneToCopyFromRefPoseGlobal = FromRetargetPose[PinData.BoneToCopyFrom.BoneIndex];
-			const FTransform& BoneToCopyToRefPoseGlobal = ToRetargetPose[PinData.BoneToCopyTo.BoneIndex];
-			const FQuat Delta = BoneToCopyToRefPoseGlobal.GetRotation() * BoneToCopyFromRefPoseGlobal.GetRotation().Inverse();
+			const FQuat& SourceRefRot = FromRetargetPose[PinData.BoneToCopyFrom.BoneIndex].GetRotation();
+			const FQuat& TargetRefRot = ToRetargetPose[PinData.BoneToCopyTo.BoneIndex].GetRotation();
 
-			// apply "copy from" rotation plus delta
-			const FQuat& BoneToCopyFromRotation = PoseToCopyFrom[PinData.BoneToCopyFrom.BoneIndex].GetRotation();
-			Result.SetRotation(Delta * BoneToCopyFromRotation);
+			// get rotation delta of from bone
+			const FQuat& SourceCurrentRot = PoseToCopyFrom[PinData.BoneToCopyFrom.BoneIndex].GetRotation();
+			const FQuat SourceDelta = SourceCurrentRot * SourceRefRot.Inverse();
+
+			// apply delta to the "to" retarget pose
+			Result.SetRotation(SourceDelta * TargetRefRot);
 			
 			break;
 		}
