[Retargeting] Using Retarget FK Chain rotation mode updates pose in the viewport
[Retargeting] Using Retarget FK Chain translation mode updates pose in the viewport

QMetry: UE-TC-19433, UE-TC-19431

#rnx

#virtualized
#Synced-CL 50095541

diff --git a/Engine/Plugins/Animation/IKRig/Source/IKRig/Public/Retargeter/RetargetOps/PinBoneOp.h b/Engine/Plugins/Animation/IKRig/Source/IKRig/Public/Retargeter/RetargetOps/PinBoneOp.h
--- a/Engine/Plugins/Animation/IKRig/Source/IKRig/Public/Retargeter/RetargetOps/PinBoneOp.h
+++ b/Engine/Plugins/Animation/IKRig/Source/IKRig/Public/Retargeter/RetargetOps/PinBoneOp.h
@@ -139,7 +139,7 @@ struct FIKRetargetPinBoneOpSettings : public FIKRetargetOpSettingsBase
 	UE_API virtual USkeleton* GetSkeleton(const FName InPropertyName) override;
 #endif
 
-	virtual void PostLoad(const FIKRigObjectVersion::Type InVersion) override;
+	UE_API virtual void PostLoad(const FIKRigObjectVersion::Type InVersion) override;
 	bool Serialize(FArchive& Ar);
 };
 
