IKRetarget: Add virtual function to IKRetargetOpBase for directly feeding anim sequence and sample time to ops for direct retargeting without anim instance.

#Synced-CL 50824856

diff --git a/Engine/Plugins/Animation/IKRig/Source/IKRig/Private/Retargeter/IKRetargetProcessor.cpp b/Engine/Plugins/Animation/IKRig/Source/IKRig/Private/Retargeter/IKRetargetProcessor.cpp
--- a/Engine/Plugins/Animation/IKRig/Source/IKRig/Private/Retargeter/IKRetargetProcessor.cpp
+++ b/Engine/Plugins/Animation/IKRig/Source/IKRig/Private/Retargeter/IKRetargetProcessor.cpp
@@ -1999,6 +1999,15 @@ void FIKRetargetProcessor::DebugDrawAllOps(
 	}
 }
 
+void FIKRetargetProcessor::UpdateOpsFromAnimSequence(UAnimSequence* Sequence, float FrameTime)
+{
+	for (FInstancedStruct& OpStruct : OpStack)
+	{
+		FIKRetargetOpBase& Op = OpStruct.GetMutable<FIKRetargetOpBase>();
+		Op.UpdateFromAnimSequence(Sequence, FrameTime);
+	}
+}
+
 #endif
 
 const FRetargetSkeleton& FIKRetargetProcessor::GetSkeleton(ERetargetSourceOrTarget SourceOrTarget) const
diff --git a/Engine/Plugins/Animation/IKRig/Source/IKRig/Public/Retargeter/IKRetargetOps.h b/Engine/Plugins/Animation/IKRig/Source/IKRig/Public/Retargeter/IKRetargetOps.h
--- a/Engine/Plugins/Animation/IKRig/Source/IKRig/Public/Retargeter/IKRetargetOps.h
+++ b/Engine/Plugins/Animation/IKRig/Source/IKRig/Public/Retargeter/IKRetargetOps.h
@@ -15,6 +15,7 @@ struct FIKRigLogger;
 class UIKRetargeter;
 struct FInstancedStruct;
 class USkeleton;
+class UAnimSequence;
 struct FPoseContext;
 struct FRetargetChainMapping;
 enum class ERetargetSourceOrTarget : uint8;
@@ -360,6 +361,10 @@ public:
 	/** (required for any op that stores settings per-chain) return the memory for the settings struct associated with a particular chain
 	 * NOTE: implementing this can automate some boilerplate UI setup, allows systems to display chain settings and get/set them. */
 	virtual uint8* GetChainSettingsMemory(const FName InChainName) { checkNoEntry(); return NULL; };
+	
+	/** (optional) override to update op with anim sequence information when batch retargeting anim sequences
+	* NOTE: this should only be called for batch (offline) retargeting where there is no anim-instance running */
+	virtual void UpdateFromAnimSequence(UAnimSequence* Sequence, float FrameTime) {};
 
 	/** the instance of this op currently running in the editor viewport */
 	mutable FIKRetargetOpBase* EditorInstance;
diff --git a/Engine/Plugins/Animation/IKRig/Source/IKRig/Public/Retargeter/IKRetargetProcessor.h b/Engine/Plugins/Animation/IKRig/Source/IKRig/Public/Retargeter/IKRetargetProcessor.h
--- a/Engine/Plugins/Animation/IKRig/Source/IKRig/Public/Retargeter/IKRetargetProcessor.h
+++ b/Engine/Plugins/Animation/IKRig/Source/IKRig/Public/Retargeter/IKRetargetProcessor.h
@@ -534,6 +534,9 @@ public:
 	/** Run debug drawing on all ops in the stack */
 	UE_API void DebugDrawAllOps(FPrimitiveDrawInterface* InPDI, const FIKRetargetDebugDrawState& EditorState) const;
 	
+	/** Run UpdateFromAnimSequence on all ops */
+	UE_API void UpdateOpsFromAnimSequence(UAnimSequence* Sequence, float FrameTime);
+	
 #endif
 
 private:
diff --git a/Engine/Plugins/Animation/IKRig/Source/IKRigEditor/Private/RetargetEditor/IKRetargetBatchOperation.cpp b/Engine/Plugins/Animation/IKRig/Source/IKRigEditor/Private/RetargetEditor/IKRetargetBatchOperation.cpp
--- a/Engine/Plugins/Animation/IKRig/Source/IKRigEditor/Private/RetargetEditor/IKRetargetBatchOperation.cpp
+++ b/Engine/Plugins/Animation/IKRig/Source/IKRigEditor/Private/RetargetEditor/IKRetargetBatchOperation.cpp
@@ -478,6 +478,9 @@ void UIKRetargetBatchOperation::ConvertAnimation(
 			
 			// let the retargeter scale the input pose
 			OutProcessor.ApplySourceScaleToPose(SourceComponentPose);
+			
+			// Let ops sample data directly from anim sequence
+			OutProcessor.UpdateOpsFromAnimSequence(SourceSequence, TimeAtCurrentFrame);
 
 			// run the retargeter
 			const TArray<FTransform>& TargetComponentPose = OutProcessor.RunRetargeter(SourceComponentPose, SettingsProfile, DeltaTime);
