[IK Retarget] Improved speed planting quality (new method)

#Synced-CL 50894192

#Synced-CL 50894192

#Synced-CL 50894192

diff --git a/Engine/Plugins/Animation/IKRig/Source/IKRig/Private/Retargeter/RetargetOps/SpeedPlantingOp.cpp b/Engine/Plugins/Animation/IKRig/Source/IKRig/Private/Retargeter/RetargetOps/SpeedPlantingOp.cpp
--- a/Engine/Plugins/Animation/IKRig/Source/IKRig/Private/Retargeter/RetargetOps/SpeedPlantingOp.cpp
+++ b/Engine/Plugins/Animation/IKRig/Source/IKRig/Private/Retargeter/RetargetOps/SpeedPlantingOp.cpp
@@ -48,57 +48,108 @@ void FSpeedPlantingGoalState::Update(
 	const double InDeltaTime,
 	const FIKRetargetSpeedPlantingOpSettings& InSettings)
 {
-	// negative speed means no curve was found
-	// (we never want to pin a goal that is missing a curve)
 	if (CurrentSpeedValue < 0.0f)
+    {
+       return; 
+    }
+
+    if (bWasReset)
+    {
+        CurrentPosition = TargetPosition = PlantLocation = LastTargetPosition = IKRigGoal->Position;
+        AccelerationTimer = 0.0f;
+        bIsTrackingTarget = true;
+        bWasReset = false;
+    }
+
+    // update planting State
+    const bool bWasPlanted = bIsPlanted;
+    bIsPlanted = CurrentSpeedValue < InSettings.SpeedThreshold;
+
+    if (bIsPlanted && !bWasPlanted)
+    {
+    	// planted this tick
+    	PlantLocation = IKRigGoal->Position;
+    }
+
+	if (bWasPlanted && !bIsPlanted)
+    {
+		// unplanted this tick
+		bIsTrackingTarget = false;
+    }
+
+	// target is either the planted location or the live animation
+	TargetPosition = bIsPlanted ? PlantLocation : IKRigGoal->Position;
+
+    // update CurrentPosition
+    if (InSettings.BlendMethod == ESpeedPlantBlendMethod::Spring)
+    {
+    	//
+    	// spring method
+    	//
+    	
+        CurrentPosition = UKismetMathLibrary::VectorSpringInterp(
+            CurrentPosition,
+            TargetPosition,
+            PositionSpringState,
+            (InSettings.SpringStrength * 2.0 * PI) * (InSettings.SpringStrength * 2.0 * PI),
+            1.0, 
+            InDeltaTime,
+            1.0, 
+            0.0);
+    }
+	else if (InSettings.BlendMethod == ESpeedPlantBlendMethod::MoveTowards)
 	{
-		return;
-	}
-
-	// forcibly update current position when op is reset
-	if (bWasReset)
-	{
-		CurrentPosition = TargetPosition = IKRigGoal->Position;
-		bWasReset = false;
-	}
+		//
+		// move towards method
+		//
 
-	// update target position
-	{
-		// update if the goal should be planted
-		const bool bWasPlanted = bIsPlanted;
-		bIsPlanted = CurrentSpeedValue < InSettings.SpeedThreshold;
+		// calculate Target Velocity (how fast the animation is moving)
+		FVector TargetVelocity = (IKRigGoal->Position - LastTargetPosition) / FMath::Max(0.0001f, InDeltaTime);
+		LastTargetPosition = IKRigGoal->Position;
 
-		// if we planted this tick, record the location to stay planted at
-		if (bIsPlanted && !bWasPlanted)
+		if (bIsTrackingTarget)
 		{
-			// update target position only once when first planted
-			TargetPosition = IKRigGoal->Position;
+			CurrentPosition = TargetPosition;
 		}
-
-		if (!bIsPlanted)
+		else
 		{
-			// if source is unplanted, continuously update target
-			TargetPosition  = IKRigGoal->Position;
+			// calc distance to target
+			const FVector ToTarget = TargetPosition - CurrentPosition;
+			const float DistanceToTarget = ToTarget.Size();
+			const FVector DirectionToTarget = ToTarget.GetSafeNormal();
+
+			// determine how much "extra" speed we can use to close the gap
+			const float MaxBoostSpeed = FMath::Max(0.0f, InSettings.MaxBlendSpeed - (float)TargetVelocity.Size());
+
+			// factor in the arrival speed gain: as DistanceToTarget -> 0, DistanceLimitedBoost -> 0
+			const float DistanceLimitedBoost = DistanceToTarget * InSettings.ArrivalSpeedGain;
+
+			// current Boost is the lesser of our acceleration ramp or our braking limit
+			// time-based ramp (0-1) to handle the initial lift-off acceleration
+			AccelerationTimer += InDeltaTime;
+			float TimeAlpha = FMath::Clamp(AccelerationTimer / FMath::Max(0.001f, InSettings.TimeToMaxSpeed), 0.0f, 1.0f);
+			TimeAlpha = TimeAlpha * TimeAlpha; // quadratic falloff
+			const float CurrentBoost = FMath::Min(TimeAlpha * MaxBoostSpeed, DistanceLimitedBoost);
+
+			// final velocity is the animation's speed + our corrective boost (which fades off as we approach target)
+			FVector DesiredVelocity = TargetVelocity + (DirectionToTarget * CurrentBoost);
+
+			// step towards target
+			const FVector MoveStep = DesiredVelocity * InDeltaTime;
+			if (MoveStep.Size() >= DistanceToTarget || DistanceToTarget < 0.1f)
+			{
+				// if the step is larger than the distance, we have arrived and should begin tracking animation 1:1
+				CurrentPosition = TargetPosition;
+				bIsTrackingTarget = true;
+				AccelerationTimer = 0.0f;
+			}
+			else
+			{
+				CurrentPosition += MoveStep;
+			}
 		}
 	}
-	
-	// run the goal position through a spring damper to un-plant it
-	const double AngularFrequency = InSettings.SpringStrength * 2.0f * PI;
-	const double Stiffness = AngularFrequency * AngularFrequency;
-	constexpr double Mass = 1.0;
-	constexpr double TargetVelocityAmount = 0.0;
-	constexpr double CriticalDamping = 1.0;
-	CurrentPosition = UKismetMathLibrary::VectorSpringInterp(
-		CurrentPosition,
-		TargetPosition,
-		PositionSpringState,
-		Stiffness,
-		CriticalDamping,
-		InDeltaTime,
-		Mass,
-		TargetVelocityAmount);
-
-	// set the actual goal to the current position from the spring
+
 	IKRigGoal->Position = CurrentPosition;
 }
 
@@ -316,11 +367,17 @@ void FIKRetargetSpeedPlantingOp::DebugDraw(
 {
 	for (const FSpeedPlantingGoalState& GoalToPlant : GoalsToPlant)
 	{
+		FLinearColor Color = GoalToPlant.bIsPlanted ? FLinearColor::Red : FLinearColor::Green;
+		if (Settings.BlendMethod == ESpeedPlantBlendMethod::MoveTowards && !GoalToPlant.bIsTrackingTarget)
+		{
+			Color = FLinearColor::Yellow;
+		}
+		
 		const FVector WorldPosition = InComponentTransform.TransformPosition(GoalToPlant.CurrentPosition);
 		DrawWireSphere(
 				InPDI,
 				WorldPosition,
-				GoalToPlant.bIsPlanted ? FLinearColor::Red : FLinearColor::Green,
+				Color,
 				Settings.DebugDrawSize,
 				12,
 				SDPG_MAX,
diff --git a/Engine/Plugins/Animation/IKRig/Source/IKRig/Public/Retargeter/RetargetOps/SpeedPlantingOp.h b/Engine/Plugins/Animation/IKRig/Source/IKRig/Public/Retargeter/RetargetOps/SpeedPlantingOp.h
--- a/Engine/Plugins/Animation/IKRig/Source/IKRig/Public/Retargeter/RetargetOps/SpeedPlantingOp.h
+++ b/Engine/Plugins/Animation/IKRig/Source/IKRig/Public/Retargeter/RetargetOps/SpeedPlantingOp.h
@@ -34,6 +34,13 @@ struct FRetargetSpeedPlantingSettings
 	void SetName(const FName InName) { TargetChainName = InName; };
 };
 
+UENUM(BlueprintType)
+enum class ESpeedPlantBlendMethod : uint8
+{
+	Spring,
+	MoveTowards
+};
+
 USTRUCT(BlueprintType, meta = (DisplayName = "Speed Plant Settings"))
 struct FIKRetargetSpeedPlantingOpSettings : public FIKRetargetOpSettingsBase
 {
@@ -46,18 +53,44 @@ struct FIKRetargetSpeedPlantingOpSettings : public FIKRetargetOpSettingsBase
 	*  The target IK goal will not be allowed to move whenever the source bone speed drops below this threshold speed. */
 	UPROPERTY(EditAnywhere, BlueprintReadWrite, Category = "Speed Planting", meta = (ClampMin = "0.0", UIMin = "0.0", UIMax = "100.0"))
 	double SpeedThreshold = 15.0f;
+
+	/** The method to use when blending the IK goal in/out of a planted position.
+	 * The Spring method uses a critically damped spring equation to blend into and out of the planted states.
+	 * The Move Towards method can more accurately track fast motion without looking mushy. It only blends out of planted states. */
+	UPROPERTY(EditAnywhere, BlueprintReadWrite, Category = "Speed Planting")
+	ESpeedPlantBlendMethod BlendMethod = ESpeedPlantBlendMethod::Spring;
+
+	/** The maximum speed in cm/s that the goal will move to catch up to the animation after unplanting. This blended towards the animated velocity
+	 * as the goal catches up to animated location after unplanting. */
+	UPROPERTY(EditAnywhere, BlueprintReadWrite, Category = "Speed Planting", meta = (ClampMin = "0.0", UIMin = "0.0", UIMax = "2000.0", 
+		EditCondition = "BlendMethod == ESpeedPlantBlendMethod::MoveTowards", EditConditionHides))
+	double MaxBlendSpeed = 900.0;
+
+	/** Seconds to reach MaxBlendSpeed. Determines how rapidly the goal accelerates after unplanting. */
+	UPROPERTY(EditAnywhere, BlueprintReadWrite, Category = "Speed Planting", meta = (ClampMin = "0.001", UIMin = "0.01", UIMax = "1.0", 
+		EditCondition = "BlendMethod == ESpeedPlantBlendMethod::MoveTowards", EditConditionHides))
+	double TimeToMaxSpeed = 0.2;
+
+	/** Controls how aggressively the goal velocity matches the animated velocity as it arrives at the animated location after unplanting.
+	 * This allows the goal speed to gradually blend towards the target speed as the goal location approaches the animated location.
+	 * Reduce this to remove velocity discontinuities when catching up to the animation.
+	 * 1.0 is very soft, 20.0 is very sharp. Default is 10.0. */
+	UPROPERTY(EditAnywhere, BlueprintReadWrite, Category = "Speed Planting", meta = (ClampMin = "0.001", UIMin = "0.01", UIMax = "20.0", 
+		EditCondition = "BlendMethod == ESpeedPlantBlendMethod::MoveTowards", EditConditionHides))
+	double ArrivalSpeedGain = 10.0f;
 	
 	/**
 	* The spring strength determines how hard it will pull towards the target when planting and unplanting.
 	* The value is the frequency (hz) at which it will oscillate when there is no damping. */
-	UPROPERTY(EditAnywhere, BlueprintReadWrite, Category = "Speed Planting", meta=(Input, ClampMin = "0.0", UIMin = "0.0", UIMax = "10.0"))
+	UPROPERTY(EditAnywhere, BlueprintReadWrite, Category = "Speed Planting", meta=(Input, ClampMin = "0.0", UIMin = "0.0", UIMax = "30.0",
+		EditCondition = "BlendMethod == ESpeedPlantBlendMethod::Spring", EditConditionHides))
 	double SpringStrength = 3.0f;
 	// Deprecated in 5.8, replaced with higher-level "Strength" param
 	UPROPERTY(meta=(DeprecatedProperty))
 	double Stiffness = 250.0f;
 
 	/* Adjust size of goal debug drawing in viewport. Red is planted, green is unplanted. */
-	UPROPERTY(EditAnywhere, BlueprintReadWrite, Category = "Speed Planting", meta = (ClampMin = "0.0", UIMin = "0.0", UIMax = "10.0"))
+	UPROPERTY(EditAnywhere, BlueprintReadWrite, Category = "Debug", meta = (ClampMin = "0.0", UIMin = "0.0", UIMax = "10.0"))
 	double DebugDrawSize = 5.0;
 
 	UE_API virtual const UClass* GetControllerType() const override;
@@ -88,9 +121,13 @@ struct FSpeedPlantingGoalState
 	FVectorSpringState PositionSpringState;
 	FVector CurrentPosition;
 	FVector TargetPosition;
+	FVector LastTargetPosition;
 	double CurrentSpeedValue = -1.0f;
 	bool bIsPlanted = false;
 	bool bWasReset = true;
+	FVector PlantLocation;
+	double AccelerationTimer = 0.0;
+	bool bIsTrackingTarget = true;
 
 #if WITH_EDITOR
 	bool bFoundCurveInSourceComponent = false;
