EditorUsability : GlobalCommands
- Unifying label and description across the editor for FindInContentBrowser command action (ControlRig, IKRig, AssetManagerEditor, ObjectMixer, MaterialEditor, RenderResourceViewer)
- Mapping FindInContentBrowser command to custom action in each class

#Synced-CL 51457231

diff --git a/Engine/Plugins/Animation/ControlRig/Source/ControlRigEditor/Private/Editor/SRigModuleAssetBrowser.cpp b/Engine/Plugins/Animation/ControlRig/Source/ControlRigEditor/Private/Editor/SRigModuleAssetBrowser.cpp
--- a/Engine/Plugins/Animation/ControlRig/Source/ControlRigEditor/Private/Editor/SRigModuleAssetBrowser.cpp
+++ b/Engine/Plugins/Animation/ControlRig/Source/ControlRigEditor/Private/Editor/SRigModuleAssetBrowser.cpp
@@ -13,6 +13,7 @@
 #include "FrontendFilterBase.h"
 #include "Editor/RigVMEditorTools.h"
 #include "Styling/SlateIconFinder.h"
+#include "Toolkits/GlobalEditorCommonCommands.h"
 #include "Widgets/Layout/SSeparator.h"
 #include "Widgets/Layout/SUniformGridPanel.h"
 #include "Widgets/SRigVMVariantWidget.h"
@@ -44,6 +45,14 @@ void SRigModuleAssetBrowser::Construct(
 {
 	bAllowDragging = InArgs._bAllowDragging;
 	AssetViewType = InArgs._AssetViewType;
+	CommandList = MakeShared<FUICommandList>();
+
+	const FGlobalEditorCommonCommands& GlobalCommands = FGlobalEditorCommonCommands::Get();
+	CommandList->MapAction(
+		GlobalCommands.FindInContentBrowser,
+		FExecuteAction::CreateSP(this, &SRigModuleAssetBrowser::ExecuteFindInContentBrowserAction),
+		FCanExecuteAction::CreateSP(this, &SRigModuleAssetBrowser::CanExecuteFindInContentBrowserAction)
+	);
 
 	ChildSlot
 	[
@@ -125,40 +134,45 @@ void SRigModuleAssetBrowser::RefreshView()
 	}
 }
 
-TSharedPtr<SWidget> SRigModuleAssetBrowser::OnGetAssetContextMenu(const TArray<FAssetData>& SelectedAssets) const
+FReply SRigModuleAssetBrowser::OnKeyDown(const FGeometry& InGeometry, const FKeyEvent& InKeyEvent)
 {
-	if (SelectedAssets.Num() <= 0)
+	if (CommandList && CommandList->ProcessCommandBindings(InKeyEvent))
 	{
-		return nullptr;
+		return FReply::Handled();
 	}
+	return SBox::OnKeyDown(InGeometry, InKeyEvent);
+}
 
-	UObject* SelectedAsset = SelectedAssets[0].GetAsset();
-	if (SelectedAsset == nullptr)
+void SRigModuleAssetBrowser::ExecuteFindInContentBrowserAction()
+{
+	if (!GetCurrentSelectionDelegate.IsBound())
+	{
+		return;
+	}
+
+	const TArray<FAssetData> SelectedAssets = GetCurrentSelectionDelegate.Execute();
+	const FContentBrowserModule& ContentBrowserModule = FModuleManager::Get().LoadModuleChecked<FContentBrowserModule>("ContentBrowser");
+	ContentBrowserModule.Get().SyncBrowserToAssets(SelectedAssets);
+}
+
+bool SRigModuleAssetBrowser::CanExecuteFindInContentBrowserAction()
+{
+	return GetCurrentSelectionDelegate.IsBound() && !GetCurrentSelectionDelegate.Execute().IsEmpty();
+}
+
+TSharedPtr<SWidget> SRigModuleAssetBrowser::OnGetAssetContextMenu(const TArray<FAssetData>& SelectedAssets) const
+{
+	if (SelectedAssets.IsEmpty())
 	{
 		return nullptr;
 	}
-	
-	FMenuBuilder MenuBuilder(true, MakeShared<FUICommandList>());
 
+	FMenuBuilder MenuBuilder(true, CommandList);
 	MenuBuilder.BeginSection(TEXT("Asset"), LOCTEXT("AssetSectionLabel", "Asset"));
 	{
-		MenuBuilder.AddMenuEntry(
-			LOCTEXT("Browse", "Browse to Asset"),
-			LOCTEXT("BrowseTooltip", "Browses to the associated asset and selects it in the most recently used Content Browser (summoning one if necessary)"),
-			FSlateIcon(FAppStyle::GetAppStyleSetName(), "SystemWideCommands.FindInContentBrowser.Small"),
-			FUIAction(
-				FExecuteAction::CreateLambda([SelectedAsset] ()
-				{
-					if (SelectedAsset)
-					{
-						const TArray<FAssetData>& Assets = { SelectedAsset };
-						const FContentBrowserModule& ContentBrowserModule = FModuleManager::Get().LoadModuleChecked<FContentBrowserModule>("ContentBrowser");
-						ContentBrowserModule.Get().SyncBrowserToAssets(Assets);
-					}
-				}),
-				FCanExecuteAction::CreateLambda([] () { return true; })
-			)
-		);
+		const TSharedPtr<FUICommandInfo> FindInContentBrowserCommand = FGlobalEditorCommonCommands::Get().FindInContentBrowser;
+		checkf(FindInContentBrowserCommand.IsValid(), TEXT("FGlobalEditorCommonCommands::FindInContentBrowser command not found!"));
+		MenuBuilder.AddMenuEntry(FindInContentBrowserCommand);
 	}
 	MenuBuilder.EndSection();
 
diff --git a/Engine/Plugins/Animation/ControlRig/Source/ControlRigEditor/Private/Editor/SRigModuleAssetBrowser.h b/Engine/Plugins/Animation/ControlRig/Source/ControlRigEditor/Private/Editor/SRigModuleAssetBrowser.h
--- a/Engine/Plugins/Animation/ControlRig/Source/ControlRigEditor/Private/Editor/SRigModuleAssetBrowser.h
+++ b/Engine/Plugins/Animation/ControlRig/Source/ControlRigEditor/Private/Editor/SRigModuleAssetBrowser.h
@@ -31,6 +31,12 @@ public:
 private:
 	void RefreshView();
 
+	//~ Begin SWidget
+	virtual FReply OnKeyDown(const FGeometry& InGeometry, const FKeyEvent& InKeyEvent) override;
+	//~ End SWidget
+
+	void ExecuteFindInContentBrowserAction();
+	bool CanExecuteFindInContentBrowserAction();
 	
 	bool OnShouldFilterAsset(const struct FAssetData& AssetData);
 	TSharedPtr<SWidget> OnGetAssetContextMenu(const TArray<FAssetData>& SelectedAssets) const;
@@ -56,5 +62,8 @@ private:
 	bool bAllowDragging;
 	EAssetViewType::Type AssetViewType;
 
+	/** Commands to execute */
+	TSharedPtr<FUICommandList> CommandList;
+
 	friend class FControlRigBaseEditor;
 };
diff --git a/Engine/Plugins/Animation/IKRig/Source/IKRigEditor/Private/RetargetEditor/SIKRetargetAssetBrowser.cpp b/Engine/Plugins/Animation/IKRig/Source/IKRigEditor/Private/RetargetEditor/SIKRetargetAssetBrowser.cpp
--- a/Engine/Plugins/Animation/IKRig/Source/IKRigEditor/Private/RetargetEditor/SIKRetargetAssetBrowser.cpp
+++ b/Engine/Plugins/Animation/IKRig/Source/IKRigEditor/Private/RetargetEditor/SIKRetargetAssetBrowser.cpp
@@ -18,6 +18,7 @@
 #include "Retargeter/IKRetargetProcessor.h"
 #include "UObject/AssetRegistryTagsContext.h"
 #include "Animation/Skeleton.h"
+#include "Toolkits/GlobalEditorCommonCommands.h"
 
 #define LOCTEXT_NAMESPACE "IKRetargeterAssetBrowser"
 
@@ -27,6 +28,14 @@ void SIKRetargetAssetBrowser::Construct(
 {
 	EditorController = InEditorController;
 	EditorController.Pin()->SetAssetBrowserView(SharedThis(this));
+	CommandList = MakeShared<FUICommandList>();
+
+	const FGlobalEditorCommonCommands& GlobalCommands = FGlobalEditorCommonCommands::Get();
+	CommandList->MapAction(
+		GlobalCommands.FindInContentBrowser,
+		FExecuteAction::CreateSP(this, &SIKRetargetAssetBrowser::ExecuteFindInContentBrowserAction),
+		FCanExecuteAction::CreateSP(this, &SIKRetargetAssetBrowser::CanExecuteFindInContentBrowserAction)
+	);
 	
 	ChildSlot
     [
@@ -103,40 +112,45 @@ void SIKRetargetAssetBrowser::RefreshView()
 	AssetBrowserBox->SetContent(ContentBrowserModule.Get().CreateAssetPicker(AssetPickerConfig));
 }
 
-TSharedPtr<SWidget> SIKRetargetAssetBrowser::OnGetAssetContextMenu(const TArray<FAssetData>& SelectedAssets) const
+FReply SIKRetargetAssetBrowser::OnKeyDown(const FGeometry& InGeometry, const FKeyEvent& InKeyEvent)
 {
-	if (SelectedAssets.Num() <= 0)
+	if (CommandList && CommandList->ProcessCommandBindings(InKeyEvent))
 	{
-		return nullptr;
+		return FReply::Handled();
+	}
+	return SBox::OnKeyDown(InGeometry, InKeyEvent);
+}
+
+void SIKRetargetAssetBrowser::ExecuteFindInContentBrowserAction()
+{
+	if (!GetCurrentSelectionDelegate.IsBound())
+	{
+		return;
 	}
 
-	UObject* SelectedAsset = SelectedAssets[0].GetAsset();
-	if (SelectedAsset == nullptr)
+	const TArray<FAssetData> SelectedAssets = GetCurrentSelectionDelegate.Execute();
+	const FContentBrowserModule& ContentBrowserModule = FModuleManager::Get().LoadModuleChecked<FContentBrowserModule>("ContentBrowser");
+	ContentBrowserModule.Get().SyncBrowserToAssets(SelectedAssets);
+}
+
+bool SIKRetargetAssetBrowser::CanExecuteFindInContentBrowserAction()
+{
+	return GetCurrentSelectionDelegate.IsBound() && !GetCurrentSelectionDelegate.Execute().IsEmpty();
+}
+
+TSharedPtr<SWidget> SIKRetargetAssetBrowser::OnGetAssetContextMenu(const TArray<FAssetData>& SelectedAssets) const
+{
+	if (SelectedAssets.IsEmpty())
 	{
 		return nullptr;
 	}
-	
-	FMenuBuilder MenuBuilder(true, MakeShared<FUICommandList>());
 
+	FMenuBuilder MenuBuilder(true, CommandList);
 	MenuBuilder.BeginSection(TEXT("Asset"), LOCTEXT("AssetSectionLabel", "Asset"));
 	{
-		MenuBuilder.AddMenuEntry(
-			LOCTEXT("Browse", "Browse to Asset"),
-			LOCTEXT("BrowseTooltip", "Browses to the associated asset and selects it in the most recently used Content Browser (summoning one if necessary)"),
-			FSlateIcon(FAppStyle::GetAppStyleSetName(), "SystemWideCommands.FindInContentBrowser.Small"),
-			FUIAction(
-				FExecuteAction::CreateLambda([SelectedAsset] ()
-				{
-					if (SelectedAsset)
-					{
-						const TArray<FAssetData>& Assets = { SelectedAsset };
-						const FContentBrowserModule& ContentBrowserModule = FModuleManager::Get().LoadModuleChecked<FContentBrowserModule>("ContentBrowser");
-						ContentBrowserModule.Get().SyncBrowserToAssets(Assets);
-					}
-				}),
-				FCanExecuteAction::CreateLambda([] () { return true; })
-			)
-		);
+		const TSharedPtr<FUICommandInfo> FindInContentBrowserCommand = FGlobalEditorCommonCommands::Get().FindInContentBrowser;
+		checkf(FindInContentBrowserCommand.IsValid(), TEXT("FGlobalEditorCommonCommands::FindInContentBrowser command not found!"));
+		MenuBuilder.AddMenuEntry(FindInContentBrowserCommand);
 	}
 	MenuBuilder.EndSection();
 
diff --git a/Engine/Plugins/Animation/IKRig/Source/IKRigEditor/Private/RigEditor/SIKRigAssetBrowser.cpp b/Engine/Plugins/Animation/IKRig/Source/IKRigEditor/Private/RigEditor/SIKRigAssetBrowser.cpp
--- a/Engine/Plugins/Animation/IKRig/Source/IKRigEditor/Private/RigEditor/SIKRigAssetBrowser.cpp
+++ b/Engine/Plugins/Animation/IKRig/Source/IKRigEditor/Private/RigEditor/SIKRigAssetBrowser.cpp
@@ -13,6 +13,7 @@
 #include "RigEditor/IKRigController.h"
 #include "RigEditor/IKRigEditorController.h"
 #include "Engine/SkeletalMesh.h"
+#include "Toolkits/GlobalEditorCommonCommands.h"
 #include "UObject/AssetRegistryTagsContext.h"
 
 #define LOCTEXT_NAMESPACE "IKRigAssetBrowser"
@@ -23,6 +24,14 @@ void SIKRigAssetBrowser::Construct(
 {
 	EditorController = InEditorController;
 	EditorController.Pin()->AssetBrowserView = SharedThis(this);
+	CommandList = MakeShared<FUICommandList>();
+
+	const FGlobalEditorCommonCommands& GlobalCommands = FGlobalEditorCommonCommands::Get();
+	CommandList->MapAction(
+		GlobalCommands.FindInContentBrowser,
+		FExecuteAction::CreateSP(this, &SIKRigAssetBrowser::ExecuteFindInContentBrowserAction),
+		FCanExecuteAction::CreateSP(this, &SIKRigAssetBrowser::CanExecuteFindInContentBrowserAction)
+	);
 	
 	ChildSlot
     [
@@ -36,6 +45,32 @@ void SIKRigAssetBrowser::Construct(
 	RefreshView();
 }
 
+FReply SIKRigAssetBrowser::OnKeyDown(const FGeometry& InGeometry, const FKeyEvent& InKeyEvent)
+{
+	if (CommandList && CommandList->ProcessCommandBindings(InKeyEvent))
+	{
+		return FReply::Handled();
+	}
+	return SBox::OnKeyDown(InGeometry, InKeyEvent);
+}
+
+void SIKRigAssetBrowser::ExecuteFindInContentBrowserAction()
+{
+	if (!GetCurrentSelectionDelegate.IsBound())
+	{
+		return;
+	}
+
+	const TArray<FAssetData> SelectedAssets = GetCurrentSelectionDelegate.Execute();
+	const FContentBrowserModule& ContentBrowserModule = FModuleManager::Get().LoadModuleChecked<FContentBrowserModule>("ContentBrowser");
+	ContentBrowserModule.Get().SyncBrowserToAssets(SelectedAssets);
+}
+
+bool SIKRigAssetBrowser::CanExecuteFindInContentBrowserAction()
+{
+	return GetCurrentSelectionDelegate.IsBound() && !GetCurrentSelectionDelegate.Execute().IsEmpty();
+}
+
 void SIKRigAssetBrowser::RefreshView()
 {
 	// setup filtering
@@ -58,6 +93,7 @@ void SIKRigAssetBrowser::RefreshView()
 	AssetPickerConfig.OnGetAssetContextMenu = FOnGetAssetContextMenu::CreateSP(this, &SIKRigAssetBrowser::OnGetAssetContextMenu);
 	AssetPickerConfig.bAllowNullSelection = false;
 	AssetPickerConfig.bFocusSearchBoxWhenOpened = false;
+	AssetPickerConfig.GetCurrentSelectionDelegates.Add(&GetCurrentSelectionDelegate);
 
 	// hide all asset registry columns by default (we only really want the name and path)
 	UObject* AnimSequenceDefaultObject = UAnimSequence::StaticClass()->GetDefaultObject();
@@ -88,27 +124,13 @@ TSharedPtr<SWidget> SIKRigAssetBrowser::OnGetAssetContextMenu(const TArray<FAsse
 		return nullptr;
 	}
 	
-	FMenuBuilder MenuBuilder(true, MakeShared<FUICommandList>());
+	FMenuBuilder MenuBuilder(true, CommandList);
 
 	MenuBuilder.BeginSection(TEXT("Asset"), LOCTEXT("AssetSectionLabel", "Asset"));
 	{
-		MenuBuilder.AddMenuEntry(
-			LOCTEXT("Browse", "Browse to Asset"),
-			LOCTEXT("BrowseTooltip", "Browses to the associated asset and selects it in the most recently used Content Browser (summoning one if necessary)"),
-			FSlateIcon(FAppStyle::GetAppStyleSetName(), "SystemWideCommands.FindInContentBrowser.Small"),
-			FUIAction(
-				FExecuteAction::CreateLambda([SelectedAsset] ()
-				{
-					if (SelectedAsset)
-					{
-						const TArray<FAssetData>& Assets = { SelectedAsset };
-						const FContentBrowserModule& ContentBrowserModule = FModuleManager::Get().LoadModuleChecked<FContentBrowserModule>("ContentBrowser");
-						ContentBrowserModule.Get().SyncBrowserToAssets(Assets);
-					}
-				}),
-				FCanExecuteAction::CreateLambda([] () { return true; })
-			)
-		);
+		const TSharedPtr<FUICommandInfo> FindInContentBrowserCommand = FGlobalEditorCommonCommands::Get().FindInContentBrowser;
+		checkf(FindInContentBrowserCommand.IsValid(), TEXT("FGlobalEditorCommonCommands::FindInContentBrowser command not found!"));
+		MenuBuilder.AddMenuEntry(FindInContentBrowserCommand);
 	}
 	MenuBuilder.EndSection();
 
diff --git a/Engine/Plugins/Animation/IKRig/Source/IKRigEditor/Public/RetargetEditor/SIKRetargetAssetBrowser.h b/Engine/Plugins/Animation/IKRig/Source/IKRigEditor/Public/RetargetEditor/SIKRetargetAssetBrowser.h
--- a/Engine/Plugins/Animation/IKRig/Source/IKRigEditor/Public/RetargetEditor/SIKRetargetAssetBrowser.h
+++ b/Engine/Plugins/Animation/IKRig/Source/IKRigEditor/Public/RetargetEditor/SIKRetargetAssetBrowser.h
@@ -23,6 +23,12 @@ public:
 	void RefreshView();
 	
 private:
+	//~ Begin SWidget
+	virtual FReply OnKeyDown(const FGeometry& InGeometry, const FKeyEvent& InKeyEvent) override;
+	//~ End SWidget
+
+	void ExecuteFindInContentBrowserAction();
+	bool CanExecuteFindInContentBrowserAction();
 
 	FReply OnExportButtonClicked();
 	bool IsExportButtonEnabled() const;
@@ -45,6 +51,9 @@ private:
 
 	/** remember the path the user chose between exports */
 	FString PrevBatchOutputPath = "";
+	
+	/** Commands to execute */
+	TSharedPtr<FUICommandList> CommandList;
 
 	friend FIKRetargetEditorController;
 };
diff --git a/Engine/Plugins/Animation/IKRig/Source/IKRigEditor/Public/RigEditor/SIKRigAssetBrowser.h b/Engine/Plugins/Animation/IKRig/Source/IKRigEditor/Public/RigEditor/SIKRigAssetBrowser.h
--- a/Engine/Plugins/Animation/IKRig/Source/IKRigEditor/Public/RigEditor/SIKRigAssetBrowser.h
+++ b/Engine/Plugins/Animation/IKRig/Source/IKRigEditor/Public/RigEditor/SIKRigAssetBrowser.h
@@ -19,6 +19,12 @@ public:
 	void Construct(const FArguments& InArgs, TSharedRef<FIKRigEditorController> InEditorController);
 
 private:
+	//~ Begin SWidget
+	virtual FReply OnKeyDown(const FGeometry& InGeometry, const FKeyEvent& InKeyEvent) override;
+	//~ End SWidget
+
+	void ExecuteFindInContentBrowserAction();
+	bool CanExecuteFindInContentBrowserAction();
 
 	void RefreshView();
 	void OnPathChange(const FString& NewPath);
@@ -28,6 +34,9 @@ private:
 	
 	TWeakPtr<FIKRigEditorController> EditorController;
 	TSharedPtr<SBox> AssetBrowserBox;
+	TSharedPtr<FUICommandList> CommandList;
+	/** Used to get the currently selected assets */
+	FGetCurrentSelectionDelegate GetCurrentSelectionDelegate;
 
 	friend FIKRigEditorController;
 };
diff --git a/Engine/Plugins/Editor/AssetManagerEditor/Source/AssetManagerEditor/Private/TreeView/SAssetTableTreeView.cpp b/Engine/Plugins/Editor/AssetManagerEditor/Source/AssetManagerEditor/Private/TreeView/SAssetTableTreeView.cpp
--- a/Engine/Plugins/Editor/AssetManagerEditor/Source/AssetManagerEditor/Private/TreeView/SAssetTableTreeView.cpp
+++ b/Engine/Plugins/Editor/AssetManagerEditor/Source/AssetManagerEditor/Private/TreeView/SAssetTableTreeView.cpp
@@ -26,6 +26,7 @@
 #include "Styling/AppStyle.h"
 #include "Styling/StyleColors.h"
 #include "Subsystems/AssetEditorSubsystem.h"
+#include "Toolkits/GlobalEditorCommonCommands.h"
 #include "TreeView/AssetTable.h"
 #include "TreeView/AssetTreeNode.h"
 #include "TreeView/AssetDependencyGrouping.h"
@@ -63,6 +64,13 @@ void SAssetTableTreeView::Construct(const FArguments& InArgs, TSharedPtr<FAssetT
 
 	STableTreeView::ConstructWidget(InTablePtr);
 
+	const FGlobalEditorCommonCommands& GlobalCommands = FGlobalEditorCommonCommands::Get();
+	CommandList->MapAction(
+		GlobalCommands.FindInContentBrowser,
+		FExecuteAction::CreateSP(this, &SAssetTableTreeView::ExecuteFindInContentBrowserAction),
+		FCanExecuteAction::CreateSP(this, &SAssetTableTreeView::CanExecuteFindInContentBrowserAction)
+	);
+
 	CreateGroupings();
 	CreateSortings();
 
@@ -134,6 +142,17 @@ FText SAssetTableTreeView::GetCurrentOperationName() const
 	return STableTreeView::GetCurrentOperationName();
 }
 
+void SAssetTableTreeView::ExecuteFindInContentBrowserAction()
+{
+	const FContentBrowserModule& ContentBrowserModule = FModuleManager::Get().LoadModuleChecked<FContentBrowserModule>("ContentBrowser");
+	ContentBrowserModule.Get().SyncBrowserToAssets(GetAssetDataForSelection());
+}
+
+bool SAssetTableTreeView::CanExecuteFindInContentBrowserAction()
+{
+	return IsRegistrySourceValid() && GetAssetTable() && (SelectedAssetIndices.Num() > 0);
+}
+
 ////////////////////////////////////////////////////////////////////////////////////////////////////
 
 BEGIN_SLATE_FUNCTION_BUILD_OPTIMIZATION
@@ -921,11 +940,6 @@ void SAssetTableTreeView::ExtendMenu(FMenuBuilder& MenuBuilder)
 			return GetAssetTable() && (SelectedAssetIndices.Num() > 0);
 		});
 
-	FCanExecuteAction HasSelectionAndRegistrySourceAndCanExecute = FCanExecuteAction::CreateLambda([this]()
-		{
-			return IsRegistrySourceValid() && GetAssetTable() && (SelectedAssetIndices.Num() > 0);
-		});
-
 	MenuBuilder.BeginSection("Asset", LOCTEXT("ContextMenu_Section_Asset", "Asset"));
 	{
 		//////////////////////////////////////////////////////////////////////////
@@ -956,22 +970,9 @@ void SAssetTableTreeView::ExtendMenu(FMenuBuilder& MenuBuilder)
 		//////////////////////////////////////////////////////////////////////////
 		/// FindInContentBrowser
 
-		FUIAction FindInContentBrowser;
-		FindInContentBrowser.CanExecuteAction = HasSelectionAndRegistrySourceAndCanExecute;
-		FindInContentBrowser.ExecuteAction = FExecuteAction::CreateLambda([this]()
-			{
-				FContentBrowserModule& ContentBrowserModule = FModuleManager::Get().LoadModuleChecked<FContentBrowserModule>("ContentBrowser");
-				ContentBrowserModule.Get().SyncBrowserToAssets(GetAssetDataForSelection());
-		});
-		MenuBuilder.AddMenuEntry
-		(
-			LOCTEXT("ContextMenu_FindInContentBrowserLabel", "Find in Content Browser..."),
-			LOCTEXT("ContextMenu_FindInContentBrowser", "Browses to the associated asset and selects it in the most recently used Content Browser (summoning one if necessary)."),
-			FSlateIcon(),
-			FindInContentBrowser,
-			NAME_None,
-			EUserInterfaceActionType::Button
-		);
+		const TSharedPtr<FUICommandInfo> FindInContentBrowserCommand = FGlobalEditorCommonCommands::Get().FindInContentBrowser;
+		checkf(FindInContentBrowserCommand.IsValid(), TEXT("FGlobalEditorCommonCommands::FindInContentBrowser command not found!"));
+		MenuBuilder.AddMenuEntry(FindInContentBrowserCommand);
 
 		//////////////////////////////////////////////////////////////////////////
 		/// OpenReferenceViewer
diff --git a/Engine/Plugins/Editor/AssetManagerEditor/Source/AssetManagerEditor/Private/TreeView/SAssetTableTreeView.h b/Engine/Plugins/Editor/AssetManagerEditor/Source/AssetManagerEditor/Private/TreeView/SAssetTableTreeView.h
--- a/Engine/Plugins/Editor/AssetManagerEditor/Source/AssetManagerEditor/Private/TreeView/SAssetTableTreeView.h
+++ b/Engine/Plugins/Editor/AssetManagerEditor/Source/AssetManagerEditor/Private/TreeView/SAssetTableTreeView.h
@@ -77,8 +77,9 @@ public:
 
 	////////////////////////////////////////////////////////////////////////////////////////////////////
 
-
 private:
+	void ExecuteFindInContentBrowserAction();
+	bool CanExecuteFindInContentBrowserAction();
 
 	enum class ECheckFilesExistAndHashMatchesResult : uint8
 	{
diff --git a/Engine/Plugins/Editor/ObjectMixer/ObjectMixer/Source/ObjectMixer/Public/Views/List/Modes/SFilterClassMenuItem.h b/Engine/Plugins/Editor/ObjectMixer/ObjectMixer/Source/ObjectMixer/Public/Views/List/Modes/SFilterClassMenuItem.h
--- a/Engine/Plugins/Editor/ObjectMixer/ObjectMixer/Source/ObjectMixer/Public/Views/List/Modes/SFilterClassMenuItem.h
+++ b/Engine/Plugins/Editor/ObjectMixer/ObjectMixer/Source/ObjectMixer/Public/Views/List/Modes/SFilterClassMenuItem.h
@@ -5,6 +5,7 @@
 #include "Algo/Find.h"
 #include "ContentBrowserModule.h"
 #include "IContentBrowserSingleton.h"
+#include "Toolkits/GlobalEditorCommonCommands.h"
 #include "UObject/ObjectPtr.h"
 #include "UObject/UObjectGlobals.h"
 #include "Views/List/Modes/ObjectMixerOutlinerMode.h"
@@ -35,6 +36,19 @@ public:
 		TArray<FObjectMixerOutlinerMode::FFilterClassSelectionInfo>& FilterClassSelectionInfos, const FText TooltipText)
 	{
 		const bool bHasValidAssetData = AssetClassMap.AssetData.IsValid();
+
+		FText ToolTipText;
+		if (bHasValidAssetData)
+		{
+			const TSharedPtr<FUICommandInfo> FindInContentBrowserCommand = FGlobalEditorCommonCommands::Get().FindInContentBrowser;
+			checkf(FindInContentBrowserCommand.IsValid(), TEXT("FGlobalEditorCommonCommands::FindInContentBrowser command not found!"));
+			ToolTipText = FindInContentBrowserCommand->GetDescription();
+		}
+		else
+		{
+			ToolTipText = NSLOCTEXT("SFilterClassMenuItem", "NoBlueprintFilterFound", "This filter class is not a Blueprint class.");
+		}
+
 		ChildSlot
 		[
 			SNew(SHorizontalBox)
@@ -95,9 +109,7 @@ public:
 				.VAlign(VAlign_Center)
 				.ContentPadding(FMargin(2))
 				.IsEnabled(bHasValidAssetData)
-				.ToolTipText(bHasValidAssetData ?
-					NSLOCTEXT("SFilterClassMenuItem", "BrowseTooltip", "Browses to the associated asset and selects it in the most recently used Content Browser (summoning one if necessary)")
-						: NSLOCTEXT("SFilterClassMenuItem", "NoBlueprintFilterFound", "This filter class is not a Blueprint class."))
+				.ToolTipText(ToolTipText)
 				.OnClicked_Lambda([AssetClassMap]()
 				{
 					const FContentBrowserModule& ContentBrowserModule = FModuleManager::Get().LoadModuleChecked<FContentBrowserModule>("ContentBrowser");
diff --git a/Engine/Source/Editor/RenderResourceViewer/Private/SRenderResourceViewer.cpp b/Engine/Source/Editor/RenderResourceViewer/Private/SRenderResourceViewer.cpp
--- a/Engine/Source/Editor/RenderResourceViewer/Private/SRenderResourceViewer.cpp
+++ b/Engine/Source/Editor/RenderResourceViewer/Private/SRenderResourceViewer.cpp
@@ -12,6 +12,7 @@
 #include "IContentBrowserSingleton.h"
 #include "Developer/TreeMap/Public/STreeMap.h"
 #include "RenderResourceViewerSettings.h"
+#include "Toolkits/GlobalEditorCommonCommands.h"
 
 #define LOCTEXT_NAMESPACE "SRenderResourceView"
 
@@ -87,25 +88,6 @@ namespace RenderResourceViewerInternal
 			return SNullWidget::NullWidget;
 		}
 	};
-
-	class FContextMenuCommands : public TCommands<FContextMenuCommands>
-	{
-	public:
-		FContextMenuCommands()
-			: TCommands<FContextMenuCommands>(TEXT("RenderResourceViewer"), NSLOCTEXT("Contexts", "RenderResourceViewer", "Render Resource Viewer"), NAME_None, FAppStyle::GetAppStyleSetName())
-		{}
-
-		virtual void RegisterCommands() override
-		{
-			UI_COMMAND(Command_FindInContentBrowser,
-				"Browse to Asset", 
-				"Browses to the associated asset and selects it in the most recently used Content Browser (summoning one if necessary)", 
-				EUserInterfaceActionType::Button, 
-				FInputChord(EModifierKey::Control, EKeys::B));
-		}
-
-		TSharedPtr<FUICommandInfo> Command_FindInContentBrowser;
-	};
 }
 
 #define RENDER_RESOURCE_VIEWER_ADD_COMBOBOX(DisplayText, Tooltip, Options, SelectedItem, OnSelectionChangedFunc) \
@@ -433,10 +415,12 @@ FReply SRenderResourceViewerWidget::OnKeyDown(const FGeometry& MyGeometry, const
 
 void SRenderResourceViewerWidget::InitCommandList()
 {
-	RenderResourceViewerInternal::FContextMenuCommands::Register();
 	CommandList = MakeShared<FUICommandList>();
 
-	CommandList->MapAction(RenderResourceViewerInternal::FContextMenuCommands::Get().Command_FindInContentBrowser,
+	const TSharedPtr<FUICommandInfo> FindInContentBrowserCommand = FGlobalEditorCommonCommands::Get().FindInContentBrowser;
+	checkf(FindInContentBrowserCommand.IsValid(), TEXT("FGlobalEditorCommonCommands::FindInContentBrowser command not found!"));
+
+	CommandList->MapAction(FindInContentBrowserCommand,
 		FExecuteAction::CreateSP(this, &SRenderResourceViewerWidget::ContextMenu_FindInContentBrowser),
 		FCanExecuteAction::CreateSP(this, &SRenderResourceViewerWidget::ContextMenu_FindInContentBrowser_CanExecute));
 }
@@ -444,9 +428,13 @@ void SRenderResourceViewerWidget::InitCommandList()
 TSharedPtr<SWidget> SRenderResourceViewerWidget::OpenContextMenu()
 {
 	FMenuBuilder MenuBuilder(true, CommandList.ToSharedRef());
+
+	const TSharedPtr<FUICommandInfo> FindInContentBrowserCommand = FGlobalEditorCommonCommands::Get().FindInContentBrowser;
+	checkf(FindInContentBrowserCommand.IsValid(), TEXT("FGlobalEditorCommonCommands::FindInContentBrowser command not found!"));
+
 	MenuBuilder.AddMenuEntry
 	(
-		RenderResourceViewerInternal::FContextMenuCommands::Get().Command_FindInContentBrowser,
+		FindInContentBrowserCommand,
 		NAME_None,
 		TAttribute<FText>(),
 		TAttribute<FText>(),
@@ -484,7 +472,18 @@ void SRenderResourceViewerWidget::ContextMenu_FindInContentBrowser()
 
 bool SRenderResourceViewerWidget::ContextMenu_FindInContentBrowser_CanExecute() const
 {
-	return ResourceListView->GetSelectedItems().Num() > 0;
+	const TArray<TSharedPtr<FRHIResourceStats>>& SelectedNodes = ResourceListView->GetSelectedItems();
+	if (!SelectedNodes.IsEmpty())
+	{
+		for (const TSharedPtr<FRHIResourceStats>& SelectedNode : SelectedNodes)
+		{
+			if (SelectedNode && !SelectedNode->OwnerName.IsNone())
+			{
+				return true;
+			}
+		}
+	}
+	return false;
 }
 
 void SRenderResourceViewerWidget::RefreshSizeMap()
diff --git a/Engine/Source/Editor/RenderResourceViewer/RenderResourceViewer.Build.cs b/Engine/Source/Editor/RenderResourceViewer/RenderResourceViewer.Build.cs
--- a/Engine/Source/Editor/RenderResourceViewer/RenderResourceViewer.Build.cs
+++ b/Engine/Source/Editor/RenderResourceViewer/RenderResourceViewer.Build.cs
@@ -26,6 +26,7 @@ public class RenderResourceViewer : ModuleRules
 				"RHI",
 				"DeveloperSettings",
 				"TreeMap",
+				"UnrealEd"
 			}
 		);
 	}
diff --git a/Engine/Source/Editor/UnrealEd/Private/Toolkits/GlobalEditorCommonCommands.cpp b/Engine/Source/Editor/UnrealEd/Private/Toolkits/GlobalEditorCommonCommands.cpp
--- a/Engine/Source/Editor/UnrealEd/Private/Toolkits/GlobalEditorCommonCommands.cpp
+++ b/Engine/Source/Editor/UnrealEd/Private/Toolkits/GlobalEditorCommonCommands.cpp
@@ -53,7 +53,7 @@ void FGlobalEditorCommonCommands::RegisterCommands()
 
 	UI_COMMAND(SummonOpenAssetDialog, "Open Asset...", "Summons an asset picker", EUserInterfaceActionType::Button, FInputChord(EModifierKey::Control, EKeys::P));
 	UI_COMMAND(SummonOpenAssetDialogAlternate, "Open Asset...", "Summons an asset picker", EUserInterfaceActionType::Button, FInputChord(EModifierKey::Alt | EModifierKey::Shift, EKeys::O));
-	UI_COMMAND(FindInContentBrowser, "Browse to Asset", "Browses to the associated asset and selects it in the most recently used Content Browser (summoning one if necessary)", EUserInterfaceActionType::Button, FInputChord(EModifierKey::Control, EKeys::B));
+	UI_COMMAND(FindInContentBrowser, "Browse to Asset(s)", "Browses and selects the associated asset(s) in the most recently used Content Browser (summoning one if necessary)", EUserInterfaceActionType::Button, FInputChord(EModifierKey::Control, EKeys::B));
 
 	UI_COMMAND(OpenConsoleCommandBox, "Open Console Command Box", "Opens an edit box where you can type in a console command", EUserInterfaceActionType::Button, FInputChord(EKeys::Tilde));
 	UI_COMMAND(SelectNextConsoleExecutor, "Iterate Console Executor", "Iterates through active Console Executors (Python, Cmd, etc.)", EUserInterfaceActionType::Button, FInputChord(EModifierKey::Control, EKeys::Tilde));
