[IK Retarget] Adding Profiles system

#Synced-CL 51723130

diff --git a/Engine/Plugins/Animation/IKRig/Source/IKRig/Private/AnimNodes/AnimNode_RetargetPoseFromMesh.cpp b/Engine/Plugins/Animation/IKRig/Source/IKRig/Private/AnimNodes/AnimNode_RetargetPoseFromMesh.cpp
--- a/Engine/Plugins/Animation/IKRig/Source/IKRig/Private/AnimNodes/AnimNode_RetargetPoseFromMesh.cpp
+++ b/Engine/Plugins/Animation/IKRig/Source/IKRig/Private/AnimNodes/AnimNode_RetargetPoseFromMesh.cpp
@@ -224,6 +224,9 @@ void FAnimNode_RetargetPoseFromMesh::Evaluate_AnyThread(FPoseContext& Output)
 	bool bForceIKOff = LODThresholdForIK != INDEX_NONE && CurrentLODLevel > LODThresholdForIK;
 	FRetargetProfile RetargetProfileToUse = GetMergedRetargetProfile(bForceIKOff);
 
+	// get the property overrides
+	const TArray<FName>* OverridesToUse = bUseCustomOverrideSets ? &OverrideSetsToApply : &IKRetargeterAsset->GetOverrideSetsToApply();
+
 	// give retarget ops a chance to access data from the anim graph before running
 	Processor.OnAnimGraphEvaluateAnyThread(Output);
 
@@ -231,6 +234,7 @@ void FAnimNode_RetargetPoseFromMesh::Evaluate_AnyThread(FPoseContext& Output)
 	const TArray<FTransform>& RetargetedPose = Processor.RunRetargeter(
 		PoseToRetargetFromComponentSpace,
 		RetargetProfileToUse,
+		*OverridesToUse,
 		DeltaTime,
 		CurrentLODLevel);
 
@@ -391,7 +395,7 @@ bool FAnimNode_RetargetPoseFromMesh::EnsureProcessorIsInitialized(const USkeleta
 		// initialize retarget processor with source and target skeletal meshes
 		bool bForceIKOff = false;
 		FRetargetProfile RetargetProfileToUse = GetMergedRetargetProfile(bForceIKOff);
-		Processor.Initialize(SourceMesh, TargetMesh, IKRetargeterAsset, RetargetProfileToUse);
+		Processor.Initialize(SourceMesh, TargetMesh, IKRetargeterAsset, RetargetProfileToUse, OverrideSetsToApply);
 	}
 
 	return Processor.IsInitialized();
diff --git a/Engine/Plugins/Animation/IKRig/Source/IKRig/Private/Retargeter/IKRetargetOps.cpp b/Engine/Plugins/Animation/IKRig/Source/IKRig/Private/Retargeter/IKRetargetOps.cpp
--- a/Engine/Plugins/Animation/IKRig/Source/IKRig/Private/Retargeter/IKRetargetOps.cpp
+++ b/Engine/Plugins/Animation/IKRig/Source/IKRig/Private/Retargeter/IKRetargetOps.cpp
@@ -2,6 +2,9 @@
 
 #include "Retargeter/IKRetargetOps.h"
 
+#include "IKRig.h"
+#include "Retargeter/IKRetargetOverrides.h"
+
 #include UE_INLINE_GENERATED_CPP_BY_NAME(IKRetargetOps)
 
 #define LOCTEXT_NAMESPACE "RetargetOpBase"
@@ -104,6 +107,20 @@ double FOpExecutionTimer::GetRunningAverage() const
 	return RunningAverageTime;
 }
 
+void FIKRetargetOpBase::GetOverrideableProperties(TSharedPtr<FRetargetPropertyOverrideNode>& OutRootNode) const
+{
+	const UScriptStruct* SettingsStruct = GetSettingsType();
+	const uint8* SettingsData = reinterpret_cast<const uint8*>(GetSettingsConst());
+	if (!ensure(SettingsStruct && SettingsData))
+	{
+		return;
+	}
+
+	OutRootNode = MakeShared<FRetargetPropertyOverrideNode>();
+	OutRootNode->NodeName = TEXT("Root");
+	OutRootNode->AddPropertiesFromStruct(SettingsStruct, SettingsData, TEXT(""));
+}
+
 #endif
 
 FName FIKRetargetOpBase::GetName() const
diff --git a/Engine/Plugins/Animation/IKRig/Source/IKRig/Private/Retargeter/IKRetargetOverrides.cpp b/Engine/Plugins/Animation/IKRig/Source/IKRig/Private/Retargeter/IKRetargetOverrides.cpp
new file mode 100644
--- /dev/null
+++ b/Engine/Plugins/Animation/IKRig/Source/IKRig/Private/Retargeter/IKRetargetOverrides.cpp
@@ -0,0 +1,811 @@
+﻿// Copyright Epic Games, Inc. All Rights Reserved.
+
+#include "Retargeter/IKRetargetOverrides.h"
+
+#include "Retargeter/IKRetargeter.h"
+#include "Retargeter/IKRetargetOps.h"
+#include "Retargeter/IKRetargetProcessor.h"
+#include "IKRigLogger.h"
+
+#include "IKRig.h"
+
+#include UE_INLINE_GENERATED_CPP_BY_NAME(IKRetargetOverrides)
+
+#define LOCTEXT_NAMESPACE "IKRetargetOverrides"
+
+TArray<FPropertyPathParser::FPropertyPathStep> FPropertyPathParser::ParsePath(const FString& InPath)
+{
+	TArray<FPropertyPathStep> OutSteps;
+	if (InPath.IsEmpty())
+	{
+		return OutSteps;
+	}
+
+	TArray<FString> RawParts;
+	InPath.ParseIntoArray(RawParts, TEXT("->"), true);
+
+	for (const FString& Part : RawParts)
+	{
+		FPropertyPathStep Step;
+
+		// check for array syntax: "MemberName[0]"
+		if (Part.EndsWith(TEXT("]")))
+		{
+			FString IndexStr;
+			if (Part.Split(TEXT("["), &Step.PropertyName, &IndexStr))
+			{
+				IndexStr.RemoveFromEnd(TEXT("]"));
+				if (IndexStr.IsNumeric())
+				{
+					Step.ArrayIndex = FCString::Atoi(*IndexStr);
+				}
+			}
+		}
+		else
+		{
+			Step.PropertyName = Part;
+		}
+
+		OutSteps.Add(Step);
+	}
+
+	return OutSteps;
+}
+
+#if WITH_EDITOR
+void FRetargetPropertyOverrideNode::AddPropertiesFromStruct(
+	const UScriptStruct* InStruct, 
+	const uint8* InStructData, 
+	FString PathPrefix)
+{
+	if (!ensure(InStruct && InStructData))
+	{
+		return;
+	}
+	
+	static const TSet<UScriptStruct*> AtomicStructTypes = {
+		TBaseStructure<FVector>::Get(), TBaseStructure<FVector2D>::Get(), TBaseStructure<FVector4>::Get(),
+		TBaseStructure<FLinearColor>::Get(), TBaseStructure<FColor>::Get(), 
+		TBaseStructure<FRotator>::Get(), TBaseStructure<FQuat>::Get(), TBaseStructure<FTransform>::Get(),
+	};
+
+	struct FPendingStruct
+	{ 
+		const UScriptStruct* Struct; 
+		const uint8* Data;
+		FString Path; 
+		TSharedPtr<FRetargetPropertyOverrideNode> ParentNode;
+	};
+
+	// recurse through the property tree adding all leaf nodes
+	TArray<FPendingStruct> StructStack;
+	StructStack.Push({ InStruct, InStructData, PathPrefix, SharedThis(this) });
+	while (StructStack.Num() > 0)
+	{
+		FPendingStruct Current = StructStack.Pop();
+
+		for (TFieldIterator<FProperty> PropIt(Current.Struct); PropIt; ++PropIt)
+		{
+			FProperty* Prop = *PropIt;
+			
+			if (!Prop->HasAnyPropertyFlags(CPF_BlueprintVisible) || Prop->HasAnyPropertyFlags(CPF_BlueprintReadOnly))
+			{
+				continue;
+			}
+			
+			if (Prop->HasMetaData(IKRigNonOverrideableMetaLabel))
+			{
+				continue;
+			}
+			
+			// determine the correct parent node based on Category metadata
+			TSharedPtr<FRetargetPropertyOverrideNode> TargetNode = Current.ParentNode;
+			if (Prop->HasMetaData(TEXT("Category")))
+			{
+				FName CategoryName = FName(*Prop->GetMetaData(TEXT("Category")));
+				if (CategoryName == FName(TEXT("Debug")))
+				{
+					continue;
+				}
+				TargetNode = TargetNode->SubNodes.FindOrAdd(CategoryName, MakeShared<FRetargetPropertyOverrideNode>());
+				TargetNode->NodeName = CategoryName;
+			}
+
+			FString FullPath = Current.Path.IsEmpty() ? Prop->GetName() : Current.Path + TEXT("->") + Prop->GetName();
+
+			// CASE 1: ARRAYS
+			if (FArrayProperty* ArrayProp = CastField<FArrayProperty>(Prop))
+			{
+				FScriptArrayHelper ArrayHelper(ArrayProp, ArrayProp->ContainerPtrToValuePtr<void>(Current.Data));
+				if (ArrayHelper.Num() > 0)
+				{
+					FName ArrayName = FName(*Prop->GetName());
+					TSharedPtr<FRetargetPropertyOverrideNode> ArraySubMenu = TargetNode->SubNodes.FindOrAdd(ArrayName, MakeShared<FRetargetPropertyOverrideNode>());
+					ArraySubMenu->NodeName = ArrayName;
+
+					FProperty* InnerProp = ArrayProp->Inner;
+					FStructProperty* InnerStructProp = CastField<FStructProperty>(InnerProp);
+
+					for (int32 i = 0; i < ArrayHelper.Num(); ++i)
+					{
+						FString IndexName = FString::Printf(TEXT("[%d]"), i);
+						FString IndexPath = FString::Printf(TEXT("%s[%d]"), *FullPath, i);
+						
+						TSharedPtr<FRetargetPropertyOverrideNode> IndexNode = ArraySubMenu->SubNodes.FindOrAdd(FName(*IndexName), MakeShared<FRetargetPropertyOverrideNode>());
+						IndexNode->NodeName = FName(*IndexName);
+
+						if (InnerStructProp && !AtomicStructTypes.Contains(InnerStructProp->Struct))
+						{
+							StructStack.Push({ InnerStructProp->Struct, ArrayHelper.GetRawPtr(i), IndexPath, IndexNode });
+						}
+						else
+						{
+							FRetargetPropertyOverrideEntry Entry;
+							Entry.DisplayName = FText::Format(INVTEXT("{0} {1}"), Prop->GetDisplayNameText(), FText::FromString(IndexName));
+							Entry.PropertyPath = IndexPath;
+							IndexNode->LeafProperties.Add(Entry);
+						}
+					}
+				}
+				continue;
+			}
+
+			// CASE 2: STRUCTS
+			FStructProperty* StructProp = CastField<FStructProperty>(Prop);
+			bool bIsLeafStruct = StructProp && AtomicStructTypes.Contains(StructProp->Struct);
+			if (StructProp && !bIsLeafStruct)
+			{
+				FName StructName = FName(*Prop->GetName());
+				TSharedPtr<FRetargetPropertyOverrideNode> NewSubNode = TargetNode->SubNodes.FindOrAdd(StructName, MakeShared<FRetargetPropertyOverrideNode>());
+				NewSubNode->NodeName = StructName;
+				
+				StructStack.Push({ StructProp->Struct, StructProp->ContainerPtrToValuePtr<uint8>(Current.Data), FullPath, NewSubNode });
+			}
+			// CASE 3: LEAF PROPERTIES
+			else
+			{
+				FRetargetPropertyOverrideEntry Entry;
+				Entry.DisplayName = Prop->GetDisplayNameText();
+				Entry.PropertyPath = FullPath;
+				TargetNode->LeafProperties.Add(Entry);
+			}
+		}
+	}
+}
+#endif
+
+bool FRetargetOpPropertyOverride::IsValid(const UScriptStruct* InSettingsStruct) const
+{
+	if (!InSettingsStruct || PropertyPath.IsEmpty())
+	{
+		return false;
+	}
+
+	// try to parse the path into segments
+	TArray<FPropertySegment> Segments;
+	if (!GetSegmentsFromProperyPath(PropertyPath, InSettingsStruct, Segments))
+	{
+		// the path is broken (ie, a property was renamed or deleted)
+		return false;
+	}
+
+	// double check the leaf property exists
+	FProperty* LeafProp = GetLeafProperty(Segments);
+	if (!LeafProp)
+	{
+		return false;
+	}
+
+	// verify it has a serialized value
+	if (ValueAsString.IsEmpty())
+	{
+		return false;
+	}
+	
+	// verify the serialized value is compatible
+	//
+	// allocate a temporary buffer for the property type
+	TArray<uint8> TempStorage;
+	TempStorage.AddZeroed(LeafProp->GetSize());
+	uint8* TempData = TempStorage.GetData();
+	// initialize the property (important for structs/strings to avoid crashes)
+	LeafProp->InitializeValue(TempData);
+	// attempt to import. ImportText returns nullptr on failure.
+	const TCHAR* ImportResult = LeafProp->ImportText_Direct(*ValueAsString, TempData, nullptr, PPF_None, nullptr);
+	// clean up the temporary value before the buffer goes out of scope
+	LeafProp->DestroyValue(TempData);
+	if (ImportResult == nullptr)
+	{
+		return false; // string value is corrupted or incompatible with the current property type
+	}
+
+	return true;
+}
+
+bool FRetargetOpPropertyOverride::IsStructPropertyAtDefault(const UScriptStruct* InStruct, uint8* InStructData, const FString& InPropertyPath)
+{
+	if (!InStruct || !InStructData || InPropertyPath.IsEmpty())
+	{
+		return false;
+	}
+	
+	TArray<FPropertySegment> Segments;
+	if (!GetSegmentsFromProperyPath(InPropertyPath, InStruct, Segments))
+	{
+		return false;
+	}
+	
+	FProperty* LeafProperty = GetLeafProperty(Segments);
+	if (!LeafProperty)
+	{
+		return false;
+	}
+	
+	uint8* InstanceValuePtr = GetDataPointerFromPathSegments(InStructData, Segments, false);
+	if (!InstanceValuePtr)
+	{
+		return false;
+	}
+
+	// instantiate a default struct to compare against
+	TArray<uint8> DefaultBuffer;
+	DefaultBuffer.AddUninitialized(InStruct->GetStructureSize());
+	InStruct->InitializeStruct(DefaultBuffer.GetData());
+	
+	bool bIsAtDefault = false;
+	constexpr bool bResizeArrays = true;
+	uint8* DefaultValuePtr = GetDataPointerFromPathSegments(DefaultBuffer.GetData(), Segments, bResizeArrays);
+	if (DefaultValuePtr)
+	{
+		// use FProperty::Identical() to perform a binary comparison of the two memory locations
+		bIsAtDefault = LeafProperty->Identical(InstanceValuePtr, DefaultValuePtr);
+	}
+
+	// clean up the temporary struct memory
+	InStruct->DestroyStruct(DefaultBuffer.GetData());
+
+	return bIsAtDefault;
+}
+
+bool FRetargetOpPropertyOverride::GetSegmentsFromProperyPath(
+	const FString& InPath,
+	const UScriptStruct* InOwnerStruct,
+	TArray<FPropertySegment>& OutSegments)
+{
+	OutSegments.Reset();
+
+	if (InPath.IsEmpty() || InOwnerStruct == nullptr)
+	{
+		return false;
+	}
+
+	const UStruct* CurrentOwnerStruct = InOwnerStruct;
+	
+	// parse path into steps
+	TArray<FPropertyPathParser::FPropertyPathStep> Steps = FPropertyPathParser::ParsePath(InPath);
+	for (const FPropertyPathParser::FPropertyPathStep& Step : Steps)
+	{
+		// 1. find the property by name in the current struct scope
+		FProperty* Property = FindProperty(CurrentOwnerStruct, Step.PropertyName);
+		if (!Property)
+		{
+			UE_LOG(LogIKRig, Warning, TEXT("Failed to find property '%s' in struct '%s'"),
+				*Step.PropertyName, *CurrentOwnerStruct->GetName());
+			OutSegments.Reset();
+			return false;
+		}
+
+		// 2. add the segment
+		FPropertySegment& Segment = OutSegments.AddDefaulted_GetRef();
+		Segment.Property = Property;
+		Segment.ArrayIndex = Step.ArrayIndex;
+
+		// 3. Update the OwnerStruct for the NEXT step in the path
+		// (if this is a struct, we need its type)
+		if (const FStructProperty* StructProp = CastField<FStructProperty>(Property))
+		{
+			CurrentOwnerStruct = StructProp->Struct;
+		}
+		
+		// if this is an array, the next property will be inside the Inner struct
+		else if (const FArrayProperty* ArrayProp = CastField<FArrayProperty>(Property))
+		{
+			if (const FStructProperty* InnerStructProp = CastField<FStructProperty>(ArrayProp->Inner))
+			{
+				CurrentOwnerStruct = InnerStructProp->Struct;
+			}
+		}
+	}
+
+	// validate we actually parsed something
+	if (OutSegments.IsEmpty())
+	{
+		UE_LOG(LogIKRig, Warning, TEXT("Failed to parse property path '%s'"), *InPath);
+		return false;
+	}
+
+	return true;
+}
+
+uint8* FRetargetOpPropertyOverride::GetDataPointerFromPathSegments(
+    uint8* InStructPtr,
+    const TArray<FPropertySegment>& InSegments,
+    bool bResizeArrays)
+	{
+	if (!InStructPtr || InSegments.IsEmpty())
+	{
+		return nullptr;
+	}
+
+	uint8* CurrentPtr = InStructPtr;
+
+	for (int32 SegmeshIndex = 0; SegmeshIndex < InSegments.Num(); ++SegmeshIndex)
+	{
+		const FPropertySegment& Segment = InSegments[SegmeshIndex];
+		FProperty* Property = Segment.Property.Get();
+		if (!Property)
+		{
+			return nullptr;
+		}
+
+		// move to the property offset within the current container
+		CurrentPtr += Property->GetOffset_ForInternal();
+
+		// handle array indexing
+		if (Segment.ArrayIndex != INDEX_NONE)
+		{
+			FArrayProperty* ArrayProp = CastField<FArrayProperty>(Property);
+			if (!ensure(ArrayProp))
+			{
+				return nullptr;
+			}
+
+			FScriptArrayHelper ArrayHelper(ArrayProp, CurrentPtr);
+
+			if (!ArrayHelper.IsValidIndex(Segment.ArrayIndex))
+			{
+				if (bResizeArrays)
+				{
+					ArrayHelper.ExpandForIndex(Segment.ArrayIndex);
+				}
+				else
+				{
+					return nullptr;
+				}
+			}
+
+			// move pointer to the specific element in the array
+			CurrentPtr = ArrayHelper.GetRawPtr(Segment.ArrayIndex);
+
+			// NOTE: if we have more segments to go, the property we are 
+			// interested in for the next iteration's container is the array's inner property
+			Property = ArrayProp->Inner; 
+		}
+
+		// check if we can continue diving
+		if (SegmeshIndex < InSegments.Num() - 1)
+		{
+			// we check 'Property' which might now be the array's inner property
+			FStructProperty* StructProp = CastField<FStructProperty>(Property);
+			if (!StructProp)
+			{
+				return nullptr; // not a struct, can't dive further
+			}
+			// currentPtr is now the base address of this struct for the next loop iteration
+		}
+	}
+
+	return CurrentPtr;
+}
+
+FProperty* FRetargetOpPropertyOverride::GetLeafProperty(const TArray<FPropertySegment>& InSegments)
+{
+	if (InSegments.Num() == 0)
+	{
+		return nullptr;
+	}
+
+	// the leaf is always the property at the end of the segment chain
+	return InSegments.Last().Property.Get();
+}
+
+FProperty* FRetargetOpPropertyOverride::FindProperty(const UStruct* InStruct, const FString& InNameOrDisplayName)
+{
+	if (!InStruct)
+	{
+		return nullptr;
+	}
+
+	// first try to find by exact internal C++ name first
+	if (FProperty* Prop = InStruct->FindPropertyByName(*InNameOrDisplayName))
+	{
+		return Prop;
+	}
+
+	return nullptr;
+}
+
+bool FRetargetOpOverrides::AddPropertyOverride(
+	const FString& InPropertyPath,
+	const uint8* InSettingsPtr,
+	const UScriptStruct* InSettingsStruct,
+	bool bRequiresReinit)
+{
+	if (!ensure(InSettingsPtr))
+	{
+		return false;
+	}
+
+	FRetargetOpPropertyOverride* PropertyOverride = FindPropertyOverride(InPropertyPath);
+	if (!PropertyOverride)
+	{
+		// add a new one
+		const int32 NewIndex = PropertyOverrides.Emplace(InPropertyPath);
+		PropertyOverride = &PropertyOverrides[NewIndex];
+	}
+
+	UpdateOverrideValue(*PropertyOverride, InSettingsPtr, InSettingsStruct);
+	
+	return true;
+}
+
+bool FRetargetOpOverrides::RemovePropertyOverride(const FString& InPropertyPath)
+{
+	const int32 Index = PropertyOverrides.IndexOfByPredicate([InPropertyPath](const FRetargetOpPropertyOverride& InPropertyOverride)
+		{
+			return InPropertyOverride.GetPropertyPath() == InPropertyPath;
+		});
+
+	if (Index == INDEX_NONE)
+	{
+		return false;
+	}
+	
+	PropertyOverrides.RemoveAt(Index);
+	
+	return true;
+}
+
+bool FRetargetOpOverrides::UpdateOverrideValue(
+	FRetargetOpPropertyOverride& InPropertyOverride,
+	const uint8* InStructData,
+	const UScriptStruct* InScriptStruct)
+{
+	if (!ensure(InStructData && InScriptStruct))
+	{
+		return false;
+	}
+	
+	// parse the string path into property segments
+	TArray<FRetargetOpPropertyOverride::FPropertySegment> Segments;
+	if (!InPropertyOverride.GetSegmentsFromProperyPath(InPropertyOverride.GetPropertyPath(), InScriptStruct, Segments))
+	{
+		return false;
+	}
+
+	// navigate to the specific memory address for this property instance
+	const uint8* PropertyDataPtr = InPropertyOverride.GetDataPointerFromPathSegments(const_cast<uint8*>(InStructData), Segments);
+	if (!PropertyDataPtr)
+	{
+		return false;
+	}
+
+	// get the leaf FProperty metadata
+	FProperty* LeafProp = InPropertyOverride.GetLeafProperty(Segments);
+	if (!LeafProp)
+	{
+		return false;
+	}
+
+	// serialize the actual value at that memory address into our string
+	FString& ValueString = InPropertyOverride.GetValueStringEditable();
+	ValueString.Reset();
+	LeafProp->ExportTextItem_Direct(ValueString, PropertyDataPtr, nullptr, nullptr, PPF_None);
+
+	// trigger runtime to refresh this override value in the memory pool
+	Version++;
+
+	return true;
+}
+
+bool FRetargetOpOverrides::HasPropertyOverride(const FString& InPropertyPath) const
+{
+	for (const FRetargetOpPropertyOverride& PropertyOverride : PropertyOverrides)
+	{
+		if (PropertyOverride.GetPropertyPath() == InPropertyPath)
+		{
+			return true;
+		}
+	}
+	
+	return false;
+}
+
+FRetargetOpPropertyOverride* FRetargetOpOverrides::FindPropertyOverride(const FString& InPropertyPath)
+{
+	for (FRetargetOpPropertyOverride& PropertyOverride : PropertyOverrides)
+	{
+		if (PropertyOverride.GetPropertyPath() == InPropertyPath)
+		{
+			return &PropertyOverride;
+		}
+	}
+	
+	return nullptr;
+}
+
+int32 FRetargetOpOverrides::GetNumPropertyOverrides() const
+{
+	return PropertyOverrides.Num();
+}
+
+FRetargetOpOverrides* FRetargetOverrideSet::FindOpOverrideByOpName(const FName InOpName)
+{
+	for (FRetargetOpOverrides& OpOverride : OpOverrides)
+	{
+		if (OpOverride.OpName == InOpName)
+		{
+			return &OpOverride;
+		}
+	}
+	return nullptr;
+}
+
+void FIKRetargetRuntimeOverrideManager::ApplyPropertyOverridesToOps(FIKRetargetProcessor& InProcessor, const TArrayView<const FName> InOverrideSetsToApply)
+{
+	const UIKRetargeter* Asset = InProcessor.GetRetargetAsset();
+	if (!ensure(Asset))
+	{
+		return;
+	}
+
+	UpdateFullListOfOverrideSets(InProcessor, InOverrideSetsToApply);
+	
+	// check if the list of profiles to apply has changed; if so, RebuildCache()
+	const int32 NewOverrideSetHash = GetOverrideSetArrayHash(OverrideSetsToApply); // this tells us if any new override sets were enabled/disabled
+	if (bIsCacheDirty || (NewOverrideSetHash != OverrideSetHash) || (CachedAssetVersion != Asset->GetOverrideVersion()))
+	{
+		RebuildCache(InProcessor);
+		OverrideSetHash = NewOverrideSetHash;
+		CachedAssetVersion = Asset->GetOverrideVersion();
+		bIsCacheDirty = false;
+	}
+
+#if WITH_EDITOR
+	// update any modified values in the CachedValuePool
+	// NOTE: this is editor only for now, but in the future we could allow this at runtime if we have a BP API for editing override values
+	UpdateCachedValues(InProcessor);
+#endif
+
+	// the "Hot Path", a linear iteration and direct memcpy of all the individual property overrides
+	for (const FCompiledPropertyOverride& Entry : CompiledOverrides)
+	{
+		FMemory::Memcpy(Entry.DestPtr, CachedValuePool.GetData() + Entry.SourcePoolOffset, Entry.Size);
+	}
+}
+
+void FIKRetargetRuntimeOverrideManager::UpdateFullListOfOverrideSets(
+	FIKRetargetProcessor& Processor,
+	const TArrayView<const FName> InOverrideSetsToApply)
+{
+	OverrideSetsToApply.Reset();
+	
+	const UIKRetargeter* Asset = Processor.GetRetargetAsset();
+	if (!ensure(Asset))
+	{
+		return;
+	}
+	
+	// process each set in the input list, note that these could be roots, mid-branch, or leaves
+	for (const FName& OverrideSetName : InOverrideSetsToApply)
+	{
+		// trace the ancestry from the current set, up to the root
+		TArray<FName> CurrentChain;
+		FName CurrentName = OverrideSetName;
+		while (CurrentName != NAME_None)
+		{
+			const FRetargetOverrideSet* OverrideSet = Asset->GetOverrideSetByName(CurrentName);
+			if (!OverrideSet)
+			{
+				break;
+			}
+
+			CurrentChain.Add(CurrentName);
+			CurrentName = OverrideSet->ParentName;
+		}
+
+		// insert the chain into our final, no duplicates
+		for (int32 i = CurrentChain.Num() - 1; i >= 0; --i)
+		{
+			OverrideSetsToApply.AddUnique(CurrentChain[i]);
+		}
+	}
+}
+
+uint32 FIKRetargetRuntimeOverrideManager::GetOverrideSetArrayHash(const TArray<FName>& InProfileNames)
+{
+	uint32 CombinedHash = 0;
+	for (const FName& ProfileName : InProfileNames)
+	{
+		CombinedHash = HashCombine(CombinedHash, GetTypeHash(ProfileName));
+	}
+	return CombinedHash;
+}
+
+void FIKRetargetRuntimeOverrideManager::RebuildCache(FIKRetargetProcessor& Processor)
+{
+    CompiledOverrides.Empty();
+    CachedValuePool.Empty();
+
+	const UIKRetargeter* Asset = Processor.GetRetargetAsset();
+	if (!ensure(Asset))
+	{
+		return;
+	}
+    
+    // first we calculate total memory needed for the pool to avoid reallocations
+	// NOTE: this isn't just for perf, we need non-volatile mem addresses
+    int32 TotalPoolSize = 0;
+    for (const FName& ProfileName : OverrideSetsToApply)
+    {
+    	const FRetargetOverrideSet* OverrideSet = Asset->GetOverrideSetByName(ProfileName);
+        if (!OverrideSet)
+        {
+        	continue;
+        }
+    	
+        for (const FRetargetOpOverrides& OpOverrides : OverrideSet->OpOverrides)
+        {
+            const UScriptStruct* OpScriptStruct = OpOverrides.ScriptStruct;
+            if (!ensure(OpScriptStruct))
+            {
+	            continue;
+            }
+
+            for (const FRetargetOpPropertyOverride& PropOverride : OpOverrides.PropertyOverrides)
+            {
+                TArray<FRetargetOpPropertyOverride::FPropertySegment> Segments;
+                if (FRetargetOpPropertyOverride::GetSegmentsFromProperyPath(PropOverride.GetPropertyPath(), OpScriptStruct, Segments))
+                {
+                    if (FProperty* LeafProp = FRetargetOpPropertyOverride::GetLeafProperty(Segments))
+                    {
+                        TotalPoolSize += LeafProp->GetSize();
+                    }
+                }
+            }
+        }
+    }
+    CachedValuePool.Reserve(TotalPoolSize);
+
+    // now bake values into the pool and cache the pointers
+    for (const FName& ProfileName : OverrideSetsToApply)
+    {
+    	const FRetargetOverrideSet* OverrideSet = Asset->GetOverrideSetByName(ProfileName);
+    	if (!OverrideSet)
+    	{
+    		continue;
+    	}
+
+    	for (const FRetargetOpOverrides& OpOverrides : OverrideSet->OpOverrides)
+    	{
+		    FIKRetargetOpBase* OpInProcessor = Processor.GetRetargetOpByName(OpOverrides.OpName);
+    		if (!ensure(OpInProcessor))
+    		{
+    			continue;
+    		}
+		    FIKRetargetOpSettingsBase* OpSettings = OpInProcessor->GetSettings();
+    		if (!ensure(OpSettings))
+    		{
+    			continue;
+    		}
+            uint8* OpSettingsData = reinterpret_cast<uint8*>(OpSettings);
+            const UScriptStruct* OpScriptStruct = OpOverrides.ScriptStruct;
+
+            if (!ensure(OpInProcessor && OpScriptStruct))
+            {
+	            continue;
+            }
+
+    		for (const FRetargetOpPropertyOverride& PropOverride : OpOverrides.PropertyOverrides)
+    		{
+    			TArray<FRetargetOpPropertyOverride::FPropertySegment> Segments;
+				const bool bPropertyFound = FRetargetOpPropertyOverride::GetSegmentsFromProperyPath(PropOverride.GetPropertyPath(), OpScriptStruct, Segments);
+				if (!ensure(bPropertyFound))
+				{
+					continue;
+				}
+            	
+				FProperty* LeafProp = FRetargetOpPropertyOverride::GetLeafProperty(Segments);
+				uint8* DestAddress = FRetargetOpPropertyOverride::GetDataPointerFromPathSegments(OpSettingsData, Segments);
+
+				if (!ensure(LeafProp && DestAddress))
+				{
+					continue;
+				}
+            	
+				// bake the value into the pool
+				const int32 ValueSize = LeafProp->GetSize();
+				const int32 CurrentOffset = CachedValuePool.Num();
+				CachedValuePool.AddUninitialized(ValueSize);
+                        
+				uint8* PoolEntryAddress = CachedValuePool.GetData() + CurrentOffset;
+                        
+				// convert string to binary at proper offset in pool
+				const TCHAR* ImportResult = LeafProp->ImportText_Direct(*PropOverride.GetValueString(), PoolEntryAddress, nullptr, PPF_None);
+    			if (!ensure(ImportResult))
+    			{
+    				continue; // corrupted string
+    			}
+
+				// record the move
+				FCompiledPropertyOverride& NewMove = CompiledOverrides.AddDefaulted_GetRef();
+				NewMove.SourcePoolOffset = CurrentOffset;
+				NewMove.DestPtr = DestAddress;
+				NewMove.Size = ValueSize;
+				// cache version we stored the value at
+				FOpOverridesKey OpOverrideKey(ProfileName, OpOverrides.OpName);
+				NewMove.OpOverrideKey = OpOverrideKey;
+				NewMove.PropertyPathHash = PropOverride.GetPropertyPathHash();
+    			NewMove.Property = LeafProp;
+			    CachedVersionPerOp.FindOrAdd(OpOverrideKey, OpOverrides.GetValueVersion());
+    		}
+        }
+    }
+}
+
+void FIKRetargetRuntimeOverrideManager::UpdateCachedValues(FIKRetargetProcessor& Processor)
+{
+	const UIKRetargeter* Asset = Processor.GetRetargetAsset();
+	if (!ensure(Asset))
+	{
+		return;
+	}
+
+	if (CachedVersionPerOp.IsEmpty())
+	{
+		return;
+	}
+
+	uint8* PoolBase = CachedValuePool.GetData();
+	
+	for (const FName& OverrideSetName : OverrideSetsToApply)
+	{
+		const FRetargetOverrideSet* OverrideSet = Asset->GetOverrideSetByName(OverrideSetName);
+		if (!ensure(OverrideSet))
+		{
+			continue;
+		}
+
+		for (const FRetargetOpOverrides& OpOverrides : OverrideSet->OpOverrides)
+		{
+			FOpOverridesKey OpOverrideKey(OverrideSetName, OpOverrides.OpName);
+			if (CachedVersionPerOp[OpOverrideKey] == OpOverrides.GetValueVersion())
+			{
+				continue;
+			}
+
+			// uh oh, we need to update the cached values for the overrides stored in this set
+			for (FCompiledPropertyOverride& Entry : CompiledOverrides)
+			{
+				if (Entry.OpOverrideKey != OpOverrideKey)
+				{
+					continue; // not a part of this set, fast early-out for the majority of entries
+				}
+
+				for (const FRetargetOpPropertyOverride& PropertyOverride : OpOverrides.PropertyOverrides)
+				{
+					if (Entry.PropertyPathHash != PropertyOverride.GetPropertyPathHash())
+					{
+						continue;
+					}
+					
+					// update the binary pool address associated with this compiled entry
+					uint8* PoolEntryAddress = PoolBase + Entry.SourcePoolOffset;
+					Entry.Property->ImportText_Direct(*PropertyOverride.GetValueString(), PoolEntryAddress, nullptr, PPF_None);
+				}
+			}
+		}
+	}
+}
+
+#undef LOCTEXT_NAMESPACE
\ No newline at end of file
diff --git a/Engine/Plugins/Animation/IKRig/Source/IKRig/Private/Retargeter/IKRetargetProcessor.cpp b/Engine/Plugins/Animation/IKRig/Source/IKRig/Private/Retargeter/IKRetargetProcessor.cpp
--- a/Engine/Plugins/Animation/IKRig/Source/IKRig/Private/Retargeter/IKRetargetProcessor.cpp
+++ b/Engine/Plugins/Animation/IKRig/Source/IKRig/Private/Retargeter/IKRetargetProcessor.cpp
@@ -1098,17 +1098,21 @@ void FIKRetargetProcessor::Initialize(
 		const USkeletalMesh* SourceSkeletalMesh,
 		const USkeletalMesh* TargetSkeletalMesh,
 		const UIKRetargeter* InRetargeterAsset,
-		const FRetargetProfile& InRetargetProfile,
+		const FRetargetProfile& InCustomProfile,
+		const TArrayView<const FName> InOverrideSetsToApply,
 		const bool bSuppressWarnings)
 {
-	check(InRetargeterAsset && SourceSkeletalMesh && TargetSkeletalMesh);
+	if (!InRetargeterAsset || !SourceSkeletalMesh || !TargetSkeletalMesh)
+	{
+		return;
+	}
 	
 	// don't attempt reinitialization unless inputs have changed
 	if (WasInitializedWithTheseAssets(SourceSkeletalMesh, TargetSkeletalMesh, InRetargeterAsset))
 	{
 		return;
 	}
-	
+
 	// assumed initialization fails unless we get to the bottom of this function
 	bIsInitialized = false;
 
@@ -1149,11 +1153,14 @@ void FIKRetargetProcessor::Initialize(
 		GetPoseScale(ERetargetSourceOrTarget::Target));
 
 	// copy ops from asset 
-	InitialOpStackSetup(RetargeterAsset->GetRetargetOps(), InRetargetProfile);
+	InitialOpStackSetup(RetargeterAsset->GetRetargetOps(), InCustomProfile);
 	
 	// initialize all the retarget ops
 	InitializeRetargetOps();
 
+	// rebuild the override cache with new memory locations
+	OverrideManager.DirtyProfileCache();
+
 	// confirm for the user that the retargeter is initialized
 	Log.LogInfo(FText::Format(
 			LOCTEXT("SuccessfulInit", "Success! The IK Retargeter is ready to transfer animation from the source, {0} to the target, {1}"),
@@ -1282,7 +1289,8 @@ void FIKRetargetProcessor::InitializeRetargetOps()
 
 TArray<FTransform>&  FIKRetargetProcessor::RunRetargeter(
 	TArray<FTransform>& InSourceGlobalPose,
-	const FRetargetProfile& InProfile,
+	const FRetargetProfile& InCustomProfile,
+	const TArrayView<const FName> InOverrideSetsToApply,
 	const float InDeltaTime,
 	const int32 InLOD)
 {
@@ -1303,7 +1311,7 @@ TArray<FTransform>&  FIKRetargetProcessor::RunRetargeter(
 #endif
 
 	// apply the retargeting settings
-	ApplyProfileAtRuntime(InProfile);
+	UpdateSettingsAtRuntime(InCustomProfile, InOverrideSetsToApply);
 
 	// applying settings can cause the retargeter to require reinitialization (when enabling/disabling IK on a chain)
 	// we should be able to safely reinitialize here because we are already initialized or it wouldn't pass the check() at the top
@@ -1314,7 +1322,8 @@ TArray<FTransform>&  FIKRetargetProcessor::RunRetargeter(
 			SourceSkeleton.SkeletalMesh,
 			TargetSkeleton.SkeletalMesh,
 			RetargeterAsset,
-			InProfile,
+			InCustomProfile,
+			InOverrideSetsToApply,
 			bSuppressWarnings);
 
 		if (!IsInitialized())
@@ -2040,7 +2049,7 @@ void FIKRetargetProcessor::SetNeedsInitialized()
 	}
 }
 
-void FIKRetargetProcessor::ApplyProfileAtRuntime(const FRetargetProfile& Profile)
+void FIKRetargetProcessor::UpdateSettingsAtRuntime(const FRetargetProfile& Profile, const TArrayView<const FName> InOverrideSetsToApply)
 {
 	// apply op settings from the profile
 	for (FInstancedStruct& OpStruct : OpStack)
@@ -2048,6 +2057,9 @@ void FIKRetargetProcessor::ApplyProfileAtRuntime(const FRetargetProfile& Profile
 		Profile.ApplyOpProfilesToOpStruct(OpStruct, ECopyOpSettingsContext::Runtime);
 	}
 
+	// apply stored property overrides
+	OverrideManager.ApplyPropertyOverridesToOps(*this, InOverrideSetsToApply);
+
 	// apply flag to force IK off
 	bIKForcedOff = Profile.bForceAllIKOff;
 	
diff --git a/Engine/Plugins/Animation/IKRig/Source/IKRig/Private/Retargeter/IKRetargetProfile.cpp b/Engine/Plugins/Animation/IKRig/Source/IKRig/Private/Retargeter/IKRetargetProfile.cpp
--- a/Engine/Plugins/Animation/IKRig/Source/IKRig/Private/Retargeter/IKRetargetProfile.cpp
+++ b/Engine/Plugins/Animation/IKRig/Source/IKRig/Private/Retargeter/IKRetargetProfile.cpp
@@ -15,7 +15,6 @@
 
 #include UE_INLINE_GENERATED_CPP_BY_NAME(IKRetargetProfile)
 
-
 FRetargetOpProfile::FRetargetOpProfile(
 	const FName InOpName,
 	const UScriptStruct* InSettingsType,
@@ -32,27 +31,32 @@
 
 bool FRetargetOpProfile::CopySettingsToOp(FInstancedStruct& InOutOpStruct, ECopyOpSettingsContext InApplyContext) const
 {
-	if (!SettingsToApply.IsValid())
-	{
-		return false; // no settings stored in the op profile
-	}
-	
 	FIKRetargetOpBase* Op = InOutOpStruct.GetMutablePtr<FIKRetargetOpBase>();
-	if (Op->GetSettingsType() != SettingsToApply.GetScriptStruct())
+	if (!ensure(Op))
 	{
-		return false; // settings were wrong type
+		return false; // should never get here with an invalid op struct
 	}
-
-	const FIKRetargetOpSettingsBase& OpSettingsFromProfile = SettingsToApply.Get<FIKRetargetOpSettingsBase>();
-	switch (InApplyContext)
+	FIKRetargetOpSettingsBase* OpSettings = Op->GetSettings();
+	if (!ensure(OpSettings))
 	{
+		return false; // no settings to modify
+	}
+	const UScriptStruct* OpSettingsType = Op->GetSettingsType();
+	if (!ensure(OpSettingsType))
+	{
+		return false; // unknown settings type
+	}
+	
+	// apply wholesale op settings struct overrides
+	if (SettingsToApply.IsValid() && OpSettingsType == SettingsToApply.GetScriptStruct())
+	{
+		const FIKRetargetOpSettingsBase& OpSettingsFromProfile = SettingsToApply.Get<FIKRetargetOpSettingsBase>();
+		switch (InApplyContext)
+		{
 		case ECopyOpSettingsContext::Runtime:
 			// apply the profile settings to an op in the stack through the CopySettingsAtRuntime interface
 			// this gives each op a chance to digest the new settings in a way that won't require reinitialization
-			if (FIKRetargetOpSettingsBase* OpSettings = Op->GetSettings())
-			{
-				OpSettings->CopySettingsAtRuntime(&OpSettingsFromProfile);
-			}
+			OpSettings->CopySettingsAtRuntime(&OpSettingsFromProfile);
 			break;
 		
 		case ECopyOpSettingsContext::PreInitialize:
@@ -63,8 +67,9 @@
 
 		default:
 			checkNoEntry();
+		}
 	}
-			
+	
 	return true;
 }
 
@@ -135,7 +140,7 @@
 	{
 		if (FRetargetOpProfile* MatchingOpProfile = FindMatchingOpProfile(OtherOpProfile))
 		{
-			// found and op profile for the same op (same index and type) so overwrite it
+			// found an op profile for the same op (same index and type) so overwrite it
 			MatchingOpProfile->CopyFromOtherOpProfile(OtherOpProfile);
 		}
 		else
@@ -145,6 +150,78 @@
 		}
 	}
 
+	MergeDeprecatedProperties(OtherProfile);
+}
+
+FRetargetOpProfile* FRetargetProfile::FindMatchingOpProfile(const FRetargetOpProfile& OtherOpProfile)
+{
+	for (FRetargetOpProfile& OpProfile : RetargetOpProfiles)
+	{
+		if (OpProfile.OpToApplySettingsTo == OtherOpProfile.OpToApplySettingsTo && OpProfile.StaticStruct() == OtherOpProfile.StaticStruct())
+		{
+			return &OpProfile;
+		}
+	}
+
+	return nullptr;
+}
+
+bool FRetargetProfile::ApplyOpProfilesToOpStruct(
+	FInstancedStruct& InOutOpStruct,
+	const ECopyOpSettingsContext InCopyContext) const
+{
+	const FIKRetargetOpBase* OpPtr = InOutOpStruct.GetPtr<FIKRetargetOpBase>();
+	if (!OpPtr)
+	{
+		return false;
+	}
+	const FName OpName = OpPtr->GetName();
+	
+	bool bOpSettingsApplied = true;
+	for (const FRetargetOpProfile& OpProfile : RetargetOpProfiles)
+	{
+		if (OpProfile.OpToApplySettingsTo == NAME_None || OpProfile.OpToApplySettingsTo == OpName)
+		{
+			bOpSettingsApplied &= OpProfile.CopySettingsToOp(InOutOpStruct, InCopyContext);
+		}
+	}
+	return bOpSettingsApplied;
+}
+
+FRetargetOpProfile* FRetargetProfile::GetOpProfileByName(const FName InOpName)
+{
+	for (FRetargetOpProfile& OpProfile : RetargetOpProfiles)
+	{
+		if (OpProfile.OpToApplySettingsTo == InOpName)
+		{
+			return &OpProfile;
+		}
+	}
+	return nullptr;
+}
+
+FRetargetProfile URetargetProfileLibrary::CopyRetargetProfileFromRetargetAsset(const UIKRetargeter* InRetargetAsset)
+{
+	FRetargetProfile RetargetProfile;
+	RetargetProfile.FillProfileWithAssetSettings(InRetargetAsset);
+	return MoveTemp(RetargetProfile);
+}
+
+UIKRetargetOpControllerBase* URetargetProfileLibrary::GetOpControllerFromRetargetProfile(
+	FRetargetProfile& InRetargetProfile,
+	const FName InRetargetOpName)
+{
+	FRetargetOpProfile* OpProfile = InRetargetProfile.GetOpProfileByName(InRetargetOpName);
+	if (!OpProfile)
+	{
+		return nullptr;
+	}
+	
+	return OpProfile->CreateControllerIfNeeded((UObject*)GetTransientPackage());
+}
+
+void FRetargetProfile::MergeDeprecatedProperties(const FRetargetProfile& OtherProfile)
+{
 	//
 	// BEGIN merge deprecated properties from other profile
 	//
@@ -293,64 +370,3 @@
 	// END merge deprecated properties
 	//
 }
-
-FRetargetOpProfile* FRetargetProfile::FindMatchingOpProfile(const FRetargetOpProfile& OtherOpProfile)
-{
-	for (FRetargetOpProfile& OpProfile : RetargetOpProfiles)
-	{
-		if (OpProfile.OpToApplySettingsTo == OtherOpProfile.OpToApplySettingsTo && OpProfile.StaticStruct() == OtherOpProfile.StaticStruct())
-		{
-			return &OpProfile;
-		}
-	}
-
-	return nullptr;
-}
-
-bool FRetargetProfile::ApplyOpProfilesToOpStruct(
-	FInstancedStruct& InOutOpStruct,
-	const ECopyOpSettingsContext InCopyContext) const
-{
-	bool bOpSettingsApplied = true;
-	const FName OpName = InOutOpStruct.GetPtr<FIKRetargetOpBase>()->GetName();
-	for (const FRetargetOpProfile& OpProfile : RetargetOpProfiles)
-	{
-		if (OpProfile.OpToApplySettingsTo == NAME_None || OpProfile.OpToApplySettingsTo == OpName)
-		{
-			bOpSettingsApplied &= OpProfile.CopySettingsToOp(InOutOpStruct, InCopyContext);
-		}
-	}
-	return bOpSettingsApplied;
-}
-
-FRetargetOpProfile* FRetargetProfile::GetOpProfileByName(const FName InOpName)
-{
-	for (FRetargetOpProfile& OpProfile : RetargetOpProfiles)
-	{
-		if (OpProfile.OpToApplySettingsTo == InOpName)
-		{
-			return &OpProfile;
-		}
-	}
-	return nullptr;
-}
-
-FRetargetProfile URetargetProfileLibrary::CopyRetargetProfileFromRetargetAsset(const UIKRetargeter* InRetargetAsset)
-{
-	FRetargetProfile RetargetProfile;
-	RetargetProfile.FillProfileWithAssetSettings(InRetargetAsset);
-	return MoveTemp(RetargetProfile);
-}
-
-UIKRetargetOpControllerBase* URetargetProfileLibrary::GetOpControllerFromRetargetProfile(
-	FRetargetProfile& InRetargetProfile,
-	const FName InRetargetOpName)
-{
-	FRetargetOpProfile* OpProfile = InRetargetProfile.GetOpProfileByName(InRetargetOpName);
-	if (!OpProfile)
-	{
-		return nullptr;
-	}
-	
-	return OpProfile->CreateControllerIfNeeded((UObject*)GetTransientPackage());
-}
diff --git a/Engine/Plugins/Animation/IKRig/Source/IKRig/Private/Retargeter/IKRetargeter.cpp b/Engine/Plugins/Animation/IKRig/Source/IKRig/Private/Retargeter/IKRetargeter.cpp
--- a/Engine/Plugins/Animation/IKRig/Source/IKRig/Private/Retargeter/IKRetargeter.cpp
+++ b/Engine/Plugins/Animation/IKRig/Source/IKRig/Private/Retargeter/IKRetargeter.cpp
@@ -1186,6 +1186,26 @@ const FRetargetProfile* UIKRetargeter::GetProfileByName(const FName& ProfileName
 	return Profiles.Find(ProfileName);
 }
 
+const TArray<FName>& UIKRetargeter::GetOverrideSetsToApply() const
+{
+	OverridesToApply.Reset();
+	
+	for (const TTuple<FName, FRetargetOverrideSet>& OverrideEntry : OverrideSets)
+	{
+		if (OverrideEntry.Value.bActiveByDefault)
+		{
+			OverridesToApply.AddUnique(OverrideEntry.Key);
+		}
+	}
+	
+	return OverridesToApply;
+}
+
+const FRetargetOverrideSet* UIKRetargeter::GetOverrideSetByName(const FName& OverrideName) const
+{
+	return OverrideSets.Find(OverrideName);
+}
+
 //
 // BEGIN DEPRECATED API
 //
diff --git a/Engine/Plugins/Animation/IKRig/Source/IKRig/Private/Retargeter/RetargetOps/AlignPoleVectorOp.cpp b/Engine/Plugins/Animation/IKRig/Source/IKRig/Private/Retargeter/RetargetOps/AlignPoleVectorOp.cpp
--- a/Engine/Plugins/Animation/IKRig/Source/IKRig/Private/Retargeter/RetargetOps/AlignPoleVectorOp.cpp
+++ b/Engine/Plugins/Animation/IKRig/Source/IKRig/Private/Retargeter/RetargetOps/AlignPoleVectorOp.cpp
@@ -380,6 +380,36 @@ uint8* FIKRetargetAlignPoleVectorOp::GetChainSettingsMemory(const FName InChainN
 {
 	return IKRetargetOpUtils::GetChainSettingsMemory(Settings.ChainsToAlign, InChainName);
 }
+
+void FIKRetargetAlignPoleVectorOp::GetOverrideableProperties(TSharedPtr<FRetargetPropertyOverrideNode>& OutRootNode) const
+{
+	// get default settings (populates most of the tree)
+	FIKRetargetOpBase::GetOverrideableProperties(OutRootNode);
+	if (!OutRootNode.IsValid())
+	{
+		return;
+	}
+
+	// remove the chains array
+	OutRootNode->SubNodes.Remove(TEXT("Chains"));
+
+	// replace with submenu for each chain
+	const UScriptStruct* SettingsStruct = GetSettingsType();
+	const uint8* SettingsData = reinterpret_cast<const uint8*>(GetSettingsConst()); 
+	for (int32 ChainIndex = 0; ChainIndex < Settings.ChainsToAlign.Num(); ++ChainIndex)
+	{
+		FName ChainName = Settings.ChainsToAlign[ChainIndex].TargetChainName;
+        
+		TSharedPtr<FRetargetPropertyOverrideNode> ChainNode = MakeShared<FRetargetPropertyOverrideNode>();
+		ChainNode->NodeName = ChainName;
+        
+		// point the path to the specific array index
+		FString ArrayPath = FString::Printf(TEXT("ChainsToAlign[%d]"), ChainIndex);
+		ChainNode->AddPropertiesFromStruct(SettingsStruct, SettingsData, ArrayPath);
+
+		OutRootNode->SubNodes.Add(ChainName, ChainNode);
+	}
+}
 #endif
 
 void FIKRetargetAlignPoleVectorOp::ApplyIKRigs(const UIKRigDefinition* InSourceIKRig, const UIKRigDefinition* InTargetIKRig)
diff --git a/Engine/Plugins/Animation/IKRig/Source/IKRig/Public/AnimNodes/AnimNode_RetargetPoseFromMesh.h b/Engine/Plugins/Animation/IKRig/Source/IKRig/Public/AnimNodes/AnimNode_RetargetPoseFromMesh.h
--- a/Engine/Plugins/Animation/IKRig/Source/IKRig/Public/AnimNodes/AnimNode_RetargetPoseFromMesh.h
+++ b/Engine/Plugins/Animation/IKRig/Source/IKRig/Public/AnimNodes/AnimNode_RetargetPoseFromMesh.h
@@ -47,6 +47,14 @@ struct FAnimNode_RetargetPoseFromMesh : public FAnimNode_Base
 	// Connect a custom retarget profile to modify the retargeter's settings at runtime.
 	UPROPERTY(EditAnywhere, BlueprintReadWrite, Category = Settings, meta=(PinHiddenByDefault))
 	FRetargetProfile CustomRetargetProfile;
+
+	// If true, will use the custom list of named override sets to apply. These are stored in the retarget asset.
+	UPROPERTY(EditAnywhere, BlueprintReadWrite, Category = Settings, meta=(PinHiddenByDefault))
+	bool bUseCustomOverrideSets = false;
+
+	/** Apply one or more override sets stored in the retarget asset. */
+	UPROPERTY(EditAnywhere, BlueprintReadWrite, Category = Settings, meta=(EditCondition="bUseCustomOverrideSets", PinHiddenByDefault))
+	TArray<FName> OverrideSetsToApply;
 	
 	// Max LOD that this node is allowed to run.
 	// For example if you have LODThreshold at 2, it will run until LOD 2 (based on 0 index) when the component LOD becomes 3, it will stop update/evaluate
diff --git a/Engine/Plugins/Animation/IKRig/Source/IKRig/Public/IKRig.h b/Engine/Plugins/Animation/IKRig/Source/IKRig/Public/IKRig.h
--- a/Engine/Plugins/Animation/IKRig/Source/IKRig/Public/IKRig.h
+++ b/Engine/Plugins/Animation/IKRig/Source/IKRig/Public/IKRig.h
@@ -7,6 +7,8 @@
 
 // a UProperty metadata tag used to indicate that a property should trigger reinitialization when edited
 static const TCHAR* IKRigReinitOnEditMetaLabel = TEXT("ReinitializeOnEdit");
+// a UProperty metadata tag used to indicate that a property can not be overridden by a profile
+static const TCHAR* IKRigNonOverrideableMetaLabel = TEXT("NotOverrideable");
 
 class FIKRigModule : public IModuleInterface
 {
diff --git a/Engine/Plugins/Animation/IKRig/Source/IKRig/Public/Retargeter/IKRetargetOps.h b/Engine/Plugins/Animation/IKRig/Source/IKRig/Public/Retargeter/IKRetargetOps.h
--- a/Engine/Plugins/Animation/IKRig/Source/IKRig/Public/Retargeter/IKRetargetOps.h
+++ b/Engine/Plugins/Animation/IKRig/Source/IKRig/Public/Retargeter/IKRetargetOps.h
@@ -19,6 +19,7 @@ class UAnimSequence;
 struct FPoseContext;
 struct FRetargetChainMapping;
 enum class ERetargetSourceOrTarget : uint8;
+struct FRetargetPropertyOverrideNode;
 
 #if WITH_EDITOR
 class FPrimitiveDrawInterface;
@@ -75,7 +76,7 @@ struct FIKRetargetOpSettingsBase : public FIKRetargetOpSkeletonProvider
 	/** The maximum LOD that this Op is allowed to run at.
 	 * For example if you have LODThreshold of 2, the Op will run until LOD 2 (based on 0 index). When the component LOD becomes 3, it will stop running.
 	 * A value of -1 forces the Op to execute at all LOD levels. Default is -1. */
-	UPROPERTY(EditAnywhere, BlueprintReadWrite, Category = "Retarget Op")
+	UPROPERTY(EditAnywhere, BlueprintReadWrite, Category = "Level of Detail")
 	int32 LODThreshold = INDEX_NONE;
 
 	/** toggle on/off all debug drawing on the op. */
@@ -366,6 +367,10 @@ public:
 	* NOTE: this should only be called for batch (offline) retargeting where there is no anim-instance running */
 	virtual void UpdateFromAnimSequence(UAnimSequence* Sequence, float FrameTime) {};
 
+	/** (optional) override this to customize which properties are available to be overridden in a profile.
+	 * NOTE: the default implementation returns ALL properties */
+	UE_API virtual void GetOverrideableProperties(TSharedPtr<FRetargetPropertyOverrideNode>& OutRootNode) const;
+
 	/** the instance of this op currently running in the editor viewport */
 	mutable FIKRetargetOpBase* EditorInstance;
 
diff --git a/Engine/Plugins/Animation/IKRig/Source/IKRig/Public/Retargeter/IKRetargetOverrides.h b/Engine/Plugins/Animation/IKRig/Source/IKRig/Public/Retargeter/IKRetargetOverrides.h
new file mode 100644
--- /dev/null
+++ b/Engine/Plugins/Animation/IKRig/Source/IKRig/Public/Retargeter/IKRetargetOverrides.h
@@ -0,0 +1,311 @@
+﻿// Copyright Epic Games, Inc. All Rights Reserved.
+
+#pragma once
+
+#include "CoreMinimal.h"
+#include "UObject/WeakFieldPtr.h"
+
+#include "IKRetargetOverrides.generated.h"
+
+#define UE_API IKRIG_API
+
+struct FIKRetargetProcessor;
+
+
+// A centralized utility for parsing a property path string
+// Intended for property path strings with properties separated by "->" with support for array syntax like "[0]"
+class FPropertyPathParser
+{
+public:
+	
+	struct FPropertyPathStep
+	{
+		FString PropertyName;
+		int32 ArrayIndex = INDEX_NONE;
+
+		bool IsArray() const { return ArrayIndex != INDEX_NONE; }
+	};
+	
+	IKRIG_API static TArray<FPropertyPathStep> ParsePath(const FString& InPath);
+};
+
+#if WITH_EDITOR
+// What is this for?
+// This type is used to store a customized property tree that ops can supply to tell outside systems which of the op's settings can be overridden.
+// For example, virtual FIKRetargetOpBase::GetOverrideableProperties() takes a root node and allows the op to add a hierarchy of nodes underneath it
+// with customized labels. This allows ops to tell UI code how to group and organize setting overrides.
+//
+// Why not FPropertyNode?
+//		1. FPropertyNode is an editor-only type and we do not want to link Slate
+//		2. This type adds support for custom group labels
+//
+struct FRetargetPropertyOverrideNode: public TSharedFromThis<FRetargetPropertyOverrideNode>
+{
+	/** the name of this submenu (e.g., "Settings" or "IKGoal") */
+	FName NodeName;
+    
+	/** sub-menus under this one */
+	TMap<FName, TSharedPtr<FRetargetPropertyOverrideNode>> SubNodes;
+    
+	/** the actual leaf properties in this specific section */
+	struct FRetargetPropertyOverrideEntry
+	{
+		FText DisplayName;
+		FString PropertyPath; 
+	};
+	TArray<FRetargetPropertyOverrideEntry> LeafProperties;
+
+	/** * helper to populate this node with properties from a specific struct.
+     * @param InStruct The struct type to reflect.
+     * @param InStructData pointer to the struct instance (needed for dynamic content like array elements)
+     * @param PathPrefix The string path leading to this struct (e.g., "Chains[0]")
+     */
+	void AddPropertiesFromStruct(const UScriptStruct* InStruct,	const uint8* InStructData, FString PathPrefix);
+};
+#endif
+
+/** Single property override for a retarget op */
+USTRUCT()
+struct FRetargetOpPropertyOverride
+{
+	GENERATED_BODY()
+
+	FRetargetOpPropertyOverride() = default;
+	FRetargetOpPropertyOverride(const FString& InPropertyPath) : PropertyPath(InPropertyPath)
+	{
+		PropertyPathHash = GetTypeHash(InPropertyPath);
+	} 
+
+	/** get access to the path to this property */
+	const FString& GetPropertyPath() const { return PropertyPath; };
+	/** get read-only access to the value of this override as a string */
+	const FString& GetValueString() const { return ValueAsString; };
+	/** get read/write access to the value of this override as a string */
+	FString& GetValueStringEditable() { return ValueAsString; };
+	/** get the hash of the property path (cached at construction) */
+	int32 GetPropertyPathHash() const { return PropertyPathHash; };
+	
+	/** validate that this property exists and is accessible */
+	bool IsValid(const UScriptStruct* InSettingsStruct) const;
+
+	/** return true if the given property path is at the default value */
+	UE_API static bool IsStructPropertyAtDefault(const UScriptStruct* InStruct, uint8* InStructData, const FString& InPropertyPath);
+
+	struct FPropertySegment
+	{
+		TWeakFieldPtr<FProperty> Property;
+		int32 ArrayIndex = INDEX_NONE;
+	};
+	
+	/** parse a property path string into segments */
+	UE_API static bool GetSegmentsFromProperyPath(
+		const FString& InPath,
+		const UScriptStruct* InOwnerStruct,
+		TArray<FPropertySegment>& OutSegments);
+
+	/** navigate from a struct pointer to the final property pointer */
+	UE_API static uint8* GetDataPointerFromPathSegments(
+		uint8* InStructPtr,
+		const TArray<FPropertySegment>& InSegments,
+		bool bResizeArrays = false);
+
+	/** get the leaf property from the segment chain */
+	UE_API static FProperty* GetLeafProperty(const TArray<FPropertySegment>& InSegments);
+
+	/** find a property by name with redirect support */
+	static FProperty* FindProperty(
+		const UStruct* InStruct,
+		const FString& InNameOrDisplayName);
+
+	/** reflection indirection because these are private */
+	static FName GetPropertyPathName(){ return GET_MEMBER_NAME_CHECKED(FRetargetOpPropertyOverride, PropertyPath); }
+	static FName GetValeAsStringName(){ return GET_MEMBER_NAME_CHECKED(FRetargetOpPropertyOverride, ValueAsString); }
+
+private:
+	
+	/** property path using "->" separator like, "IKGoalSettings->BlendAlpha" and array syntax like, "Chains[0]->Weight"
+	 * NOTE: this allows generic serialization of property overrides */
+	UPROPERTY(EditAnywhere, Category=Override)
+	FString PropertyPath;
+	UPROPERTY()
+	uint32 PropertyPathHash = -1;
+
+	/** Value serialized as string */
+	UPROPERTY(EditAnywhere, Category=Override)
+	FString ValueAsString;
+
+	/** TODO Optional: Mark if this property requires reinitialization when changed */
+	UPROPERTY()
+	bool bRequiresReinitialization = false;
+};
+
+/** A collection of property overrides to assign to a single op. */
+USTRUCT()
+struct FRetargetOpOverrides
+{
+	GENERATED_BODY()
+	
+	/** the retarget op to apply these overrides to */
+	UPROPERTY(EditAnywhere, Category=Override)
+	FName OpName = NAME_None;
+	
+	/** the individual property-level overrides for this op */
+	UPROPERTY(EditAnywhere, Category=Override)
+	TArray<FRetargetOpPropertyOverride> PropertyOverrides;
+
+	/** the reflection data for the struct */
+	UPROPERTY(EditAnywhere, Category=Override)
+	TObjectPtr<UScriptStruct> ScriptStruct;
+
+	/** * registers a new property override for this operation.
+	 * @param InPropertyPath    The string path to the property (e.g., "Settings->Alpha").
+	 * @param InSettingsPtr     Pointer to the memory block containing the source value.
+	 * @param InSettingsStruct  The reflection data (UScriptStruct) for the settings type.
+	 * @param bRequiresReinit   True if applying this property at runtime requires a full op re-initialization.
+	 * @return True if the override was successfully added.
+	 */
+	UE_API bool AddPropertyOverride(const FString& InPropertyPath, const uint8* InSettingsPtr, const UScriptStruct* InSettingsStruct, bool bRequiresReinit);
+
+	/** * Removes an existing property override based on its path.
+	 * @param InPropertyPath    The string path of the override to remove.
+	 * @return True if a matching override was found and removed.
+	 */
+	UE_API bool RemovePropertyOverride(const FString& InPropertyPath);
+
+	/** * updates the serialized value of an existing property override.
+	 * @param InPropertyOverride	The override to update.
+	 * @param InStructData			Pointer to the memory block containing the new value.
+	 * @param InScriptStruct		The reflection data (UScriptStruct) for the settings type.
+	 * @return True if the value was updated, false if no matching override exists.
+	 */
+	UE_API bool UpdateOverrideValue(FRetargetOpPropertyOverride& InPropertyOverride, const uint8* InStructData, const UScriptStruct* InScriptStruct);
+
+	/** @return True if an overrides with the given property path exists in this op profile */
+	UE_API bool HasPropertyOverride(const FString& InPropertyPath) const;
+
+	/** * attempts to locate a specific property override by its path.
+	 * @param InPropertyPath    The string path to search for.
+	 * @return A pointer to the override data if found, otherwise nullptr.
+	 */
+	UE_API FRetargetOpPropertyOverride* FindPropertyOverride(const FString& InPropertyPath);
+
+	/** * returns the total count of properties currently overridden for this operation. 
+	 */
+	UE_API int32 GetNumPropertyOverrides() const;
+
+	/** Get the version of the values stored for this op (incremented on edit) */
+	int32 GetValueVersion() const { return Version; }
+
+private:
+
+	/** incremented whenever override values change, used to update the runtime cache value */
+	uint32 Version = 0;
+};
+
+USTRUCT()
+struct FRetargetOverrideSet
+{
+	GENERATED_BODY()
+	
+	/** the label for this set of overrides. */
+	UPROPERTY(EditAnywhere, Category=Override)
+	FName Name;
+
+	/** used to organize override sets hierarchically */
+	UPROPERTY(EditAnywhere, Category=Override)
+	FName ParentName;
+	
+	/** a collection of sparse property value overrides for ops in the stack. */
+	UPROPERTY(EditAnywhere, Category=Override)
+	TArray<FRetargetOpOverrides> OpOverrides;	
+
+	/** allow an override set to be active by default */
+	UPROPERTY(EditAnywhere, Category=Override)
+	bool bActiveByDefault = false;
+
+	/** get an op override by the op name */
+	UE_API FRetargetOpOverrides* FindOpOverrideByOpName(const FName InOpName);
+};
+
+// This system is used to apply property overrides to retarget ops at runtime in a performant way.
+//
+// The way it works is:
+// 1. The processor calls ApplyPropertyOverridesToOps every frame passing in the names of the profiles to apply.
+// 2. These profiles are then aggregated with their descendent profiles up the hierarchy to form a full list of ProfilesToApply.
+// 3. All the property overrides stored in all the ProfilesToApply are resolved with a source/dest memory address
+// 4. At runtime, we memcpy all the data from the profiles into the op stack that is currently running.
+//
+// NOTE: because we are caching raw pointers, we must be very careful to rebuild the cache any time the op stack is rebuilt!
+// This can be done with: DirtyProfileCache()
+//
+struct FIKRetargetRuntimeOverrideManager
+{
+	// a key associated with all overrides for a given op in a given set of overrides
+	struct FOpOverridesKey
+	{
+		FName OverrideSetName;
+		FName OpName;
+
+		FOpOverridesKey() = default;
+		FOpOverridesKey(FName ProfileName, FName OpName) : OverrideSetName(ProfileName), OpName(OpName){}
+		
+		bool operator==(const FOpOverridesKey& Other) const
+		{
+			return OverrideSetName == Other.OverrideSetName && OpName == Other.OpName;
+		}
+		
+		friend uint32 GetTypeHash(const FOpOverridesKey& Key)
+		{
+			uint32 Hash = GetTypeHash(Key.OverrideSetName);
+			return HashCombine(Hash, GetTypeHash(Key.OpName));
+		}
+	};
+	
+	struct FCompiledPropertyOverride
+	{
+		int32 SourcePoolOffset; // offset into CachedValuePool
+		void* DestPtr;          // pointer to processor op settings memory
+		int32 Size;				// size of the property data
+
+		// value cache invalidate/update stuff
+		FOpOverridesKey OpOverrideKey;	// used to invalidate cached value when any op override in a given profile is edited
+		FProperty* Property;			// used to avoid path parsing when updating values at runtime
+		uint32 PropertyPathHash;		// the property path this override applies to (avoid storing pointers into asset, volatile while editing)
+	};
+
+	void ApplyPropertyOverridesToOps(FIKRetargetProcessor& Processor, const TArrayView<const FName> InOverrideSetsToApply);
+
+	void DirtyProfileCache() { bIsCacheDirty = true; };
+
+private:
+	
+	void UpdateCachedValues(FIKRetargetProcessor& Processor);
+	
+	void UpdateFullListOfOverrideSets(FIKRetargetProcessor& Processor, const TArrayView<const FName> InOverrideSetsToApply);
+
+	void RebuildCache(FIKRetargetProcessor& Processor);
+
+	static uint32 GetOverrideSetArrayHash(const TArray<FName>& InProfileNames);
+
+	// the full list of profiles to apply (includes parents of named profiles)
+	TArray<FName> OverrideSetsToApply; 
+
+	// a single contiguous buffer holding all binary values for the active overrides
+	TArray<uint8> CachedValuePool;
+
+	// cached pointer offsets to memcpy from source to dest value memory
+	// key is a combined hash of the profile name + op name the overrides belong to
+	TArray<FCompiledPropertyOverride> CompiledOverrides;
+
+	// the version of the op override sets we cached (invalidated when value is modified)
+	TMap<FOpOverridesKey, int32> CachedVersionPerOp;
+	int32 CachedAssetVersion;
+
+	// hash of profile names that are currently cached (used to invalidate the CompiledOverrides)
+	uint32 OverrideSetHash;
+
+	// processor forcibly dirties the cache when it re-initializes
+	bool bIsCacheDirty = true;
+};
+
+#undef UE_API
diff --git a/Engine/Plugins/Animation/IKRig/Source/IKRig/Public/Retargeter/IKRetargetProcessor.h b/Engine/Plugins/Animation/IKRig/Source/IKRig/Public/Retargeter/IKRetargetProcessor.h
--- a/Engine/Plugins/Animation/IKRig/Source/IKRig/Public/Retargeter/IKRetargetProcessor.h
+++ b/Engine/Plugins/Animation/IKRig/Source/IKRig/Public/Retargeter/IKRetargetProcessor.h
@@ -367,6 +367,7 @@ public:
 	* @param TargetSkeleton - the skeletal mesh to poses TO
 	* @param InRetargeterAsset - the source asset to use for retargeting settings
 	* @param InRetargetProfile - the collection of settings used to initialize with
+	* @param InOverrideSetsToApply - the names of any sets of property overrides to apply (stored on asset)
 	* @param bSuppressWarnings - if true, will not output warnings during initialization
 	* @warning - Initialization does a lot of validation and can fail for many reasons. Check bIsLoadedAndValid afterwards.
 	*/
@@ -374,13 +375,15 @@ public:
 		const USkeletalMesh *SourceSkeleton,
 		const USkeletalMesh *TargetSkeleton,
 		const UIKRetargeter* InRetargeterAsset,
-		const FRetargetProfile& InRetargetProfile,
+		const FRetargetProfile& InCustomProfile,
+		const TArrayView<const FName> InOverrideSetsToApply,
 		const bool bSuppressWarnings=false);
 
 	/**
 	* Run the retarget to generate a new pose.
 	* @param InSourceGlobalPose -  is the source mesh input pose in Component/Global space
 	* @param InProfile -  the retarget profile to use for this update
+	* @param InOverrideSetsToApply -  an array of override set names stored on the asset to apply
 	* @param InDeltaTime -  time since last tick in seconds (passed to ops)
 	* @param InLOD - the current LOD level for the skeletal mesh
 	* @return The retargeted Component/Global space pose for the target skeleton
@@ -388,6 +391,7 @@ public:
 	UE_API TArray<FTransform>& RunRetargeter(
 		TArray<FTransform>& InSourceGlobalPose,
 		const FRetargetProfile& InProfile,
+		const TArrayView<const FName> InOverrideSetsToApply,
 		const float InDeltaTime,
 		const int32 InLOD = INDEX_NONE);
 
@@ -567,17 +571,22 @@ private:
 	/** The collection of operations to run to transfer animation from source to target. */
 	TArray<FInstancedStruct> OpStack;
 
+	/** Manages the application of stored property overrides at runtime. */
+	FIKRetargetRuntimeOverrideManager OverrideManager;
+
 	/** cached source scaling */
 	mutable FRetargetPoseScaleWithPivot CurrentSourceScale;
 
-	/** Apply the settings stored in a retarget profile. Called inside RunRetargeter(). */
-	void ApplyProfileAtRuntime(const FRetargetProfile& Profile);
+	/** Apply the settings stored in a retarget profile. As well as any stored property overrides. */
+	void UpdateSettingsAtRuntime(const FRetargetProfile& Profile, const TArrayView<const FName> InOverrideSetsToApply);
 
 	/** Internal retarget phase that does simple bone-to-bone copying from source*/
 	void GenerateBasePoses(TArray<FTransform>& InSourceGlobalPose);
 
 	/** Initial setup of the retarget op stack (copies op stack from asset into processor) */
-	void InitialOpStackSetup(const TArray<FInstancedStruct>& OpsFromAsset, const FRetargetProfile& InRetargetProfile);
+	void InitialOpStackSetup(
+		const TArray<FInstancedStruct>& OpsFromAsset,
+		const FRetargetProfile& InRetargetProfile);
 
 	/** Run Initialize() on all ops then collects retargeted bones and call PostInitialize() on each Op */
 	void InitializeRetargetOps();
@@ -658,7 +667,7 @@ public:
 		const FRetargetProfile& Settings,
 		const bool bSuppressWarnings=false)
 	{
-		Processor.Initialize(SourceSkeleton, TargetSkeleton, InRetargeterAsset, Settings, bSuppressWarnings);
+		Processor.Initialize(SourceSkeleton, TargetSkeleton, InRetargeterAsset, Settings, TArray<FName>(),  bSuppressWarnings);
 	}
 	
 	TArray<FTransform>& RunRetargeter(
@@ -667,7 +676,7 @@ public:
 		const float DeltaTime,
 		const FRetargetProfile& Profile)
 	{
-		return Processor.RunRetargeter(InSourceGlobalPose, Profile, DeltaTime);
+		return Processor.RunRetargeter(InSourceGlobalPose, Profile, TArray<FName>(), DeltaTime);
 	}
 
 	void UpdateRetargetPoseAtRuntime(
diff --git a/Engine/Plugins/Animation/IKRig/Source/IKRig/Public/Retargeter/IKRetargetProfile.h b/Engine/Plugins/Animation/IKRig/Source/IKRig/Public/Retargeter/IKRetargetProfile.h
--- a/Engine/Plugins/Animation/IKRig/Source/IKRig/Public/Retargeter/IKRetargetProfile.h
+++ b/Engine/Plugins/Animation/IKRig/Source/IKRig/Public/Retargeter/IKRetargetProfile.h
@@ -23,6 +23,7 @@ enum class ECopyOpSettingsContext : uint8
 	Runtime			// copy only settings that don't require reinitialization
 };
 
+/** A wholesale copy of the settings in a retarget op. */
 USTRUCT(BlueprintType)
 struct FRetargetOpProfile
 {
@@ -30,37 +31,50 @@ struct FRetargetOpProfile
 
 	FRetargetOpProfile() = default;
 
-	FRetargetOpProfile(
+	UE_API FRetargetOpProfile(
 		const FName InOpName,
 		const UScriptStruct* InSettingsType,
 		const FIKRetargetOpSettingsBase* InOpSettings);
 
-	void CopyFromOtherOpProfile(const FRetargetOpProfile& OtherOpProfile);
-
-	/** Copies the settings from the profile into the given Op instanced struct.
-	 * @param InOutOpStruct - the op to apply these settings to
-	 * @param InApplyContext - in 'runtime' mode, will let op apply the settings itself to avoid changes that require reinitialization
-	 * @return true if the supplied op was compatible with the settings, false otherwise */
+	/** * Applies the stored settings from this profile into the provided operation's instance data.
+	 * @param InOutOpStruct     The destination operation struct to receive these settings.
+	 * @param InApplyContext    The execution context (e.g., Runtime vs Editor). In Runtime mode, 
+	 * overrides that require reinitialization will be skipped.
+	 * @return True if the settings were compatible and successfully applied, false otherwise. 
+	 */
 	bool CopySettingsToOp(FInstancedStruct& InOutOpStruct, ECopyOpSettingsContext InApplyContext) const;
+	
+	/**
+	 * Duplicates all settings and property overrides from another operation profile.
+	 * @param OtherOpProfile    The source profile to copy data from.
+	 */
+	void CopyFromOtherOpProfile(const FRetargetOpProfile& OtherOpProfile);
 
-	/** helper function for ops to instantiate their own controller */
+	/** * Helper utility for operations to instantiate and retrieve their associated controller.
+	 * @param Outer             The UObject that will own the newly created controller.
+	 * @return A pointer to the existing or newly created controller base.
+	 */
 	UIKRetargetOpControllerBase* CreateControllerIfNeeded(UObject* Outer);
-	
-	/** The name of the op these settings should be applied to.
-	 * If left at None, will apply to all ops that use the type of settings stored in 'SettingsToApply' */
+
+	/** * The specific name of the operation in the stack these settings should target.
+	 * If NAME_None, settings are applied to all operations matching the 'SettingsToApply' type. 
+	 */
 	UPROPERTY(EditAnywhere, BlueprintReadWrite, Category = "IK Retarget Op Profile")
 	FName OpToApplySettingsTo = NAME_None;
 
-	/** The settings to override for a specific op in the retargeter's op stack.
-	 * NOTE: all settings stored here will be applied to the op. However, at runtime each op independently
-	 * determines which settings are safe to be applied without requiring reinitialization. */
+	/** * The collection of settings to override for a specific operation.
+	 * Note: While all settings here are applied in the editor, operations independently 
+	 * filter which settings are safe to apply at runtime without re-initializing. 
+	 */
 	UPROPERTY(EditAnywhere, BlueprintReadWrite, Category = "IK Retarget Op Profile", meta = (ExcludeBaseStruct, BaseStruct = "/Script/IKRig.IKRetargetOpSettingsBase"))
 	FInstancedStruct SettingsToApply;
-
-	// the controller used to edit these op settings by script/blueprint (lazy instantiated when needed)
+	
+private:
+	/** the controller used to edit these op settings by script/blueprint (lazy instantiated when needed) */
 	TStrongObjectPtr<UIKRetargetOpControllerBase> Controller = nullptr;
 };
 
+/** A collection of property overrides for multiple ops. */
 USTRUCT(BlueprintType)
 struct FRetargetProfile
 {
@@ -78,31 +92,31 @@ public:
 	PRAGMA_ENABLE_DEPRECATION_WARNINGS
 	
 	/** A polymorphic list of override settings to apply to retargeting operations in the stack */
-	UPROPERTY(EditAnywhere, BlueprintReadWrite, Category = "Op Settings")
+	UPROPERTY(VisibleAnywhere, BlueprintReadOnly, Category = "Op Settings")
 	TArray<FRetargetOpProfile> RetargetOpProfiles;
 
-	// If true, the TARGET Retarget Pose specified in this profile will be applied to the Retargeter (when plugged into the Retargeter).
+	/** If true, the TARGET Retarget Pose specified in this profile will be applied to the Retargeter (when plugged into the Retargeter). */
 	UPROPERTY(EditAnywhere, BlueprintReadWrite, Category="Retarget Poses", meta = (DisplayName = "Override Target Retarget Pose"))
 	bool bApplyTargetRetargetPose = false;
 	
-	// Override the TARGET Retarget Pose to use when this profile is active.
-	// The pose must be present in the Retarget Asset and is not applied unless bApplyTargetRetargetPose is true.
+	/** Override the TARGET Retarget Pose to use when this profile is active.
+	 * The pose must be present in the Retarget Asset and is not applied unless bApplyTargetRetargetPose is true. */
 	UPROPERTY(EditAnywhere, BlueprintReadWrite, Category="Retarget Poses", meta=(EditCondition="bApplyTargetRetargetPose"))
 	FName TargetRetargetPoseName;
 
-	// If true, the Source Retarget Pose specified in this profile will be applied to the Retargeter (when plugged into the Retargeter).
+	/** If true, the Source Retarget Pose specified in this profile will be applied to the Retargeter (when plugged into the Retargeter). */
 	UPROPERTY(EditAnywhere, BlueprintReadWrite, Category="Retarget Poses", meta = (DisplayName = "Override Source Retarget Pose"))
 	bool bApplySourceRetargetPose = false;
 
-	// Override the SOURCE Retarget Pose to use when this profile is active.
-	// The pose must be present in the Retarget Asset and is not applied unless bApplySourceRetargetPose is true.
+	/** Override the SOURCE Retarget Pose to use when this profile is active.
+	 * The pose must be present in the Retarget Asset and is not applied unless bApplySourceRetargetPose is true.*/
 	UPROPERTY(EditAnywhere, BlueprintReadWrite, Category="Retarget Poses", meta=(EditCondition="bApplySourceRetargetPose"))
 	FName SourceRetargetPoseName;
 
 	/** Globally forces all IK solving off */
 	UPROPERTY(EditAnywhere, BlueprintReadWrite, Category = "IK Retarget Settings", DisplayName="Force All IK Off")
 	bool bForceAllIKOff = false;
-
+	
 	UE_API void FillProfileWithAssetSettings(const UIKRetargeter* InAsset);
 
 	UE_API void MergeWithOtherProfile(const FRetargetProfile& OtherProfile);
@@ -135,6 +149,10 @@ public:
 	// NOTE: I cannot use the _DEPRECATED suffix on these properties because they live in a USTRUCT
 	// and UBT will complain that deprecated properties cannot be blueprint exposed without a getter but USTRUCT does not support UFUNCTIONS/getters
 	PRAGMA_DISABLE_DEPRECATION_WARNINGS
+
+	// consolidating/isolating deprecation code
+	void MergeDeprecatedProperties(const FRetargetProfile& OtherProfile);
+	
 	// (Deprecated) If true, the Chain Settings stored in this profile will be applied to the Retargeter (when plugged into the Retargeter).
 	UE_DEPRECATED(5.6, "Modifying chain settings must go through a retarget op profile.")
 	UPROPERTY(EditAnywhere, BlueprintReadWrite, Category="Deprecated", meta=(DisplayName="Apply Chain Settings (Deprecated)"))
diff --git a/Engine/Plugins/Animation/IKRig/Source/IKRig/Public/Retargeter/IKRetargeter.h b/Engine/Plugins/Animation/IKRig/Source/IKRig/Public/Retargeter/IKRetargeter.h
--- a/Engine/Plugins/Animation/IKRig/Source/IKRig/Public/Retargeter/IKRetargeter.h
+++ b/Engine/Plugins/Animation/IKRig/Source/IKRig/Public/Retargeter/IKRetargeter.h
@@ -4,6 +4,7 @@
 
 #include "IKRetargetChainMapping.h"
 #include "IKRetargetProfile.h"
+#include "IKRetargetOverrides.h"
 #include "Rig/IKRigDefinition.h"
 #include "IKRetargetSettings.h"
 
@@ -130,6 +131,13 @@ public:
 	// Get the retarget profile by name (may be null) 
 	UE_API const FRetargetProfile* GetProfileByName(const FName& ProfileName) const;
 
+	// Get all the profiles in the asset marked as to be applied
+	UE_API const TArray<FName>& GetOverrideSetsToApply() const;
+	// Get a set of overrides by name (may be null) 
+	UE_API const FRetargetOverrideSet* GetOverrideSetByName(const FName& OverrideName) const;
+	// get current version of the property overrides data (to compare against running processor instances)
+	int32 GetOverrideVersion() const { return OverridesVersion; };
+
 	// get current version of the data (to compare against running processor instances)
 	int32 GetVersion() const { return Version; };
 	// do this after any edit that would require running instance to reinitialize
@@ -430,12 +438,20 @@ private:
 	UPROPERTY(meta = (ExcludeBaseStruct, BaseStruct = "/Script/IKRig.RetargetOpBase"))
 	TArray<FInstancedStruct> RetargetOps;
 	
-	// settings profiles stored in this asset 
+	// whole setting struct profiles stored in this asset
 	UPROPERTY()
 	TMap<FName, FRetargetProfile> Profiles;
 	UPROPERTY()
 	FName CurrentProfile = NAME_None;
 
+	// sparse property overrides to apply at runtime
+	UPROPERTY()
+	TMap<FName, FRetargetOverrideSet> OverrideSets;
+	// the default set of override sets to be applied at runtime
+	mutable TArray<FName> OverridesToApply;
+	// updated anytime overrides are modified
+	mutable int32 OverridesVersion = 0;
+
 	// The set of retarget poses for the SOURCE skeleton.
 	UPROPERTY()
 	TMap<FName, FIKRetargetPose> SourceRetargetPoses;
diff --git a/Engine/Plugins/Animation/IKRig/Source/IKRig/Public/Retargeter/RetargetOps/AlignPoleVectorOp.h b/Engine/Plugins/Animation/IKRig/Source/IKRig/Public/Retargeter/RetargetOps/AlignPoleVectorOp.h
--- a/Engine/Plugins/Animation/IKRig/Source/IKRig/Public/Retargeter/RetargetOps/AlignPoleVectorOp.h
+++ b/Engine/Plugins/Animation/IKRig/Source/IKRig/Public/Retargeter/RetargetOps/AlignPoleVectorOp.h
@@ -65,7 +65,7 @@ struct FRetargetPoleVectorSettings
 	FRetargetPoleVectorSettings(FName InTargetChainName) : TargetChainName(InTargetChainName) {}
 	
 	/** The name of the target chain to align pole vectors on. */
-	UPROPERTY(EditAnywhere, BlueprintReadWrite, Category = "Setup", meta=(ReinitializeOnEdit))
+	UPROPERTY(EditAnywhere, BlueprintReadWrite, Category = "Setup", meta=(ReinitializeOnEdit, NotOverrideable))
 	FName TargetChainName;
 
 	/** Enable pole vector alignment on this chain. */
@@ -100,7 +100,7 @@ struct FIKRetargetAlignPoleVectorOpSettings : public FIKRetargetOpSettingsBase
 	UPROPERTY(EditAnywhere, BlueprintReadWrite, Category = "IK Rig Asset", meta=(ReinitializeOnEdit))
 	TObjectPtr<const UIKRigDefinition> IKRigAsset = nullptr;
 	
-	UPROPERTY(EditAnywhere, BlueprintReadWrite, Category = "Pole Vector Op Settings")
+	UPROPERTY(EditAnywhere, BlueprintReadWrite, Category = "Chains")
 	TArray<FRetargetPoleVectorSettings> ChainsToAlign;
 
 	UE_API void MergePoleVectorSettings(const FRetargetPoleVectorSettings& InSettingsToMerge);
@@ -151,8 +151,8 @@ struct FIKRetargetAlignPoleVectorOp : public FIKRetargetOpBase
 	UE_API virtual void ResetChainSettingsToDefault(const FName& InChainName) override;
 	UE_API virtual bool AreChainSettingsAtDefault(const FName& InChainName) override;
 	UE_API virtual uint8* GetChainSettingsMemory(const FName InChainName) override;
+	UE_API virtual void GetOverrideableProperties(TSharedPtr<FRetargetPropertyOverrideNode>& OutRootNode) const override;
 #endif
-
 	UPROPERTY()
 	FIKRetargetAlignPoleVectorOpSettings Settings;
 
diff --git a/Engine/Plugins/Animation/IKRig/Source/IKRig/Public/Retargeter/RetargetOps/BlendToSourceOp.h b/Engine/Plugins/Animation/IKRig/Source/IKRig/Public/Retargeter/RetargetOps/BlendToSourceOp.h
--- a/Engine/Plugins/Animation/IKRig/Source/IKRig/Public/Retargeter/RetargetOps/BlendToSourceOp.h
+++ b/Engine/Plugins/Animation/IKRig/Source/IKRig/Public/Retargeter/RetargetOps/BlendToSourceOp.h
@@ -25,7 +25,7 @@ struct FIKRetargetBlendToSourceChainSettings
 	void SetName(const FName InName) { TargetChainName = InName; };
 
 	/** The name of the TARGET chain to blend to source on. */
-	UPROPERTY(EditAnywhere, BlueprintReadWrite, Category = "Chain Settings", meta=(ReinitializeOnEdit))
+	UPROPERTY(EditAnywhere, BlueprintReadWrite, Category = "Chain Settings", meta=(ReinitializeOnEdit, NotOverrideable))
 	FName TargetChainName;
 	
 	/** Range 0 to 1. Default 0. Blend the goal on this chain to the location of the equivalent goal on the source chain. */
diff --git a/Engine/Plugins/Animation/IKRig/Source/IKRig/Public/Retargeter/RetargetOps/CopyBasePoseOp.h b/Engine/Plugins/Animation/IKRig/Source/IKRig/Public/Retargeter/RetargetOps/CopyBasePoseOp.h
--- a/Engine/Plugins/Animation/IKRig/Source/IKRig/Public/Retargeter/RetargetOps/CopyBasePoseOp.h
+++ b/Engine/Plugins/Animation/IKRig/Source/IKRig/Public/Retargeter/RetargetOps/CopyBasePoseOp.h
@@ -21,11 +21,11 @@ struct FIKRetargetCopyBasePoseOpSettings : public FIKRetargetOpSettingsBase
 	bool bCopyBasePose = true;
 
 	/** Filters the bones to copy when using "Copy Base Pose". If specified, will only copy all children of the specified bone (inclusive). */
-	UPROPERTY(EditAnywhere, Category = RetargetPhases, meta = (ReinitializeOnEdit, EditCondition = "bCopyBasePose"))
+	UPROPERTY(EditAnywhere, Category = RetargetPhases, meta = (ReinitializeOnEdit, NotOverrideable, EditCondition = "bCopyBasePose"))
 	FBoneReference CopyFromStart;
 
 	/** Bones added to this list (and their children) will not have their pose copied. They will be left at the reference pose.*/
-	UPROPERTY(EditAnywhere, Category = RetargetPhases, meta = (ReinitializeOnEdit, EditCondition = "bCopyBasePose"))
+	UPROPERTY(EditAnywhere, Category = RetargetPhases, meta = (ReinitializeOnEdit, NotOverrideable, EditCondition = "bCopyBasePose"))
 	TArray<FBoneReference> BonesToExclude;
 
 	/** converted on PostLoad */
diff --git a/Engine/Plugins/Animation/IKRig/Source/IKRig/Public/Retargeter/RetargetOps/FKChainsOp.h b/Engine/Plugins/Animation/IKRig/Source/IKRig/Public/Retargeter/RetargetOps/FKChainsOp.h
--- a/Engine/Plugins/Animation/IKRig/Source/IKRig/Public/Retargeter/RetargetOps/FKChainsOp.h
+++ b/Engine/Plugins/Animation/IKRig/Source/IKRig/Public/Retargeter/RetargetOps/FKChainsOp.h
@@ -182,7 +182,7 @@ struct FRetargetFKChainSettings
 	FRetargetFKChainSettings(const FName InTargetChainName) : TargetChainName(InTargetChainName) {};
 
 	/** The name of the TARGET chain to transfer animation onto. */
-	UPROPERTY(EditAnywhere, BlueprintReadWrite, Category = "Setup", meta=(ReinitializeOnEdit))
+	UPROPERTY(EditAnywhere, BlueprintReadWrite, Category = "Setup", meta=(ReinitializeOnEdit, NotOverrideable))
 	FName TargetChainName;
 	
 	/** Whether to copy the shape of the chain from the source skeleton using the Rotation and Translation modes. Default is true.
diff --git a/Engine/Plugins/Animation/IKRig/Source/IKRig/Public/Retargeter/RetargetOps/FloorConstraintOp.h b/Engine/Plugins/Animation/IKRig/Source/IKRig/Public/Retargeter/RetargetOps/FloorConstraintOp.h
--- a/Engine/Plugins/Animation/IKRig/Source/IKRig/Public/Retargeter/RetargetOps/FloorConstraintOp.h
+++ b/Engine/Plugins/Animation/IKRig/Source/IKRig/Public/Retargeter/RetargetOps/FloorConstraintOp.h
@@ -24,7 +24,7 @@ struct FFloorConstraintToeDefinition
 	GENERATED_BODY()
 
 	/** The bone located at the base of the toe (the first knuckle). */
-	UPROPERTY(EditAnywhere, Category=Setup, meta=(ReinitializeOnEdit))
+	UPROPERTY(EditAnywhere, Category=Setup, meta=(ReinitializeOnEdit, NotOverrideable))
 	FBoneReference ToeBone;
 
 	/** The length of the toe from the ToeBone to the tip of the toe mesh. */
@@ -104,7 +104,7 @@ struct FFloorConstraintChainSettings : public FIKRetargetOpSkeletonProvider
 	FFloorConstraintChainSettings(const FName InTargetChainName) : TargetChainName(InTargetChainName) {}
 
 	/** The name of the TARGET chain to transfer animation onto. */
-	UPROPERTY(EditAnywhere, BlueprintReadWrite, Category = "Setup", meta=(ReinitializeOnEdit))
+	UPROPERTY(EditAnywhere, BlueprintReadWrite, Category = "Setup", meta=(ReinitializeOnEdit, NotOverrideable))
 	FName TargetChainName;
 	
 	/** Whether to apply the floor constraint to the location of the IK goal on this chain. Default is false.
diff --git a/Engine/Plugins/Animation/IKRig/Source/IKRig/Public/Retargeter/RetargetOps/IKChainsOp.h b/Engine/Plugins/Animation/IKRig/Source/IKRig/Public/Retargeter/RetargetOps/IKChainsOp.h
--- a/Engine/Plugins/Animation/IKRig/Source/IKRig/Public/Retargeter/RetargetOps/IKChainsOp.h
+++ b/Engine/Plugins/Animation/IKRig/Source/IKRig/Public/Retargeter/RetargetOps/IKChainsOp.h
@@ -70,7 +70,7 @@ struct FRetargetIKChainSettings
 	FRetargetIKChainSettings(const FName InTargetChainName) : TargetChainName(InTargetChainName) {}
 
 	/** The name of the TARGET chain to transfer animation onto. */
-	UPROPERTY(EditAnywhere, BlueprintReadWrite, Category = "Setup", meta=(ReinitializeOnEdit))
+	UPROPERTY(EditAnywhere, BlueprintReadWrite, Category = "Setup", meta=(ReinitializeOnEdit, NotOverrideable))
 	FName TargetChainName;
 	
 	/** Whether to modify the location of the IK goal on this chain. Default is true.
diff --git a/Engine/Plugins/Animation/IKRig/Source/IKRig/Public/Retargeter/RetargetOps/OffsetGoalsOp.h b/Engine/Plugins/Animation/IKRig/Source/IKRig/Public/Retargeter/RetargetOps/OffsetGoalsOp.h
--- a/Engine/Plugins/Animation/IKRig/Source/IKRig/Public/Retargeter/RetargetOps/OffsetGoalsOp.h
+++ b/Engine/Plugins/Animation/IKRig/Source/IKRig/Public/Retargeter/RetargetOps/OffsetGoalsOp.h
@@ -24,7 +24,7 @@ struct FIKRetargetOffsetGoalsChainSettings
 	void SetName(const FName InName) { TargetChainName = InName; };
 
 	/** The name of the TARGET chain for these settings. */
-	UPROPERTY(EditAnywhere, BlueprintReadWrite, Category = "Setup", meta=(ReinitializeOnEdit))
+	UPROPERTY(EditAnywhere, BlueprintReadWrite, Category = "Setup", meta=(ReinitializeOnEdit, NotOverrideable))
 	FName TargetChainName;
 
 	/** Range 0 to 1. Default 0. Blend the offsets on this chain on/off. */
diff --git a/Engine/Plugins/Animation/IKRig/Source/IKRig/Public/Retargeter/RetargetOps/PelvisMotionOp.h b/Engine/Plugins/Animation/IKRig/Source/IKRig/Public/Retargeter/RetargetOps/PelvisMotionOp.h
--- a/Engine/Plugins/Animation/IKRig/Source/IKRig/Public/Retargeter/RetargetOps/PelvisMotionOp.h
+++ b/Engine/Plugins/Animation/IKRig/Source/IKRig/Public/Retargeter/RetargetOps/PelvisMotionOp.h
@@ -18,11 +18,11 @@ struct FIKRetargetPelvisMotionOpSettings : public FIKRetargetOpSettingsBase
 	GENERATED_BODY()
 
 	/** The Pelvis bone on the source skeleton to copy motion FROM. */
-	UPROPERTY(EditAnywhere, Category=Setup, meta=(ReinitializeOnEdit))
+	UPROPERTY(EditAnywhere, Category=Setup, meta=(ReinitializeOnEdit, NotOverrideable))
 	FBoneReference SourcePelvisBone;
 
 	/** The Pelvis bone on the target skeleton to copy motion TO. */
-	UPROPERTY(EditAnywhere, Category=Setup, meta=(ReinitializeOnEdit))
+	UPROPERTY(EditAnywhere, Category=Setup, meta=(ReinitializeOnEdit, NotOverrideable))
 	FBoneReference TargetPelvisBone;
 
 	/** Blend from motion scaled by the size of the character, to the absolute delta motion from the source.
@@ -34,7 +34,7 @@ struct FIKRetargetPelvisMotionOpSettings : public FIKRetargetOpSettingsBase
 	 * When ON, the floor constraint will adjust the vertical Pelvis motion separately according to the following rules:
 	 * 1. When the source pelvis is LOWER than the ref pose, the target pelvis will be lowered proportional to their relative crotch heights.
 	 * 2. When the source pelvis is HIGHER than the ref pose, the target pelvis will be raised an identical amount while maintaining any vertical offset in the ref pose.*/
-	UPROPERTY(EditAnywhere, Category=FloorConstraint, meta = (ClampMin = "0.0", ClampMax = "1.0", UIMin = "0.0", UIMax = "1.0"))
+	UPROPERTY(EditAnywhere, BlueprintReadWrite, Category=FloorConstraint, meta = (ClampMin = "0.0", ClampMax = "1.0", UIMin = "0.0", UIMax = "1.0"))
 	double FloorConstraintWeight = 0.0f;
 	
 	/** A negative vertical offset in cm relative to the Pelvis bone of the SOURCE.
@@ -117,11 +117,11 @@ struct FIKRetargetPelvisMotionOpSettings : public FIKRetargetOpSettingsBase
 	double AffectIKVertical = 0.0f;
 
 	/** Adjust size of the debug drawing */
-	UPROPERTY(EditAnywhere, Category=Debug)
+	UPROPERTY(EditAnywhere, BlueprintReadWrite, Category=Debug)
 	double DebugDrawSize = 5.0f;
 	
 	/** Adjust thickness of the debug drawing */
-	UPROPERTY(EditAnywhere, Category=Debug)
+	UPROPERTY(EditAnywhere, BlueprintReadWrite, Category=Debug)
 	double DebugDrawThickness = 5.0f;
 
 	IKRIG_API virtual const UClass* GetControllerType() const override;
diff --git a/Engine/Plugins/Animation/IKRig/Source/IKRig/Public/Retargeter/RetargetOps/PinBoneOp.h b/Engine/Plugins/Animation/IKRig/Source/IKRig/Public/Retargeter/RetargetOps/PinBoneOp.h
--- a/Engine/Plugins/Animation/IKRig/Source/IKRig/Public/Retargeter/RetargetOps/PinBoneOp.h
+++ b/Engine/Plugins/Animation/IKRig/Source/IKRig/Public/Retargeter/RetargetOps/PinBoneOp.h
@@ -29,10 +29,10 @@ struct FPinBoneData
 	void CachePinData(ERetargetSourceOrTarget InSkeletonToCopyFrom, const FIKRetargetProcessor& Processor);
 
 	// The bone to copy FROM.
-	UPROPERTY(EditAnywhere, Category=Settings, meta=(ReinitializeOnEdit))
+	UPROPERTY(EditAnywhere, Category=Settings, meta=(ReinitializeOnEdit, NotOverrideable))
 	FBoneReference BoneToCopyFrom;
 	// The bone to copy TO. Will have its transform pinned to BoneToCopyFrom
-	UPROPERTY(EditAnywhere, Category=Settings, meta=(ReinitializeOnEdit))
+	UPROPERTY(EditAnywhere, Category=Settings, meta=(ReinitializeOnEdit, NotOverrideable))
 	FBoneReference BoneToCopyTo;
 
 	// the scale factor used to calculate all the cached offset (lazy evaluated when scale factor is modified)
@@ -92,7 +92,7 @@ struct FIKRetargetPinBoneOpSettings : public FIKRetargetOpSettingsBase
 	GENERATED_BODY()
 	
 	/** A list of bone-pairs to copy transforms between */
-	UPROPERTY(EditAnywhere, BlueprintReadWrite, Category=Setup, meta=(ReinitializeOnEdit))
+	UPROPERTY(EditAnywhere, BlueprintReadWrite, Category=Setup, meta=(ReinitializeOnEdit, NotOverrideable))
 	TArray<FPinBoneData> BonesToPin;
 
 	/** Which skeleton to copy animation from. */
diff --git a/Engine/Plugins/Animation/IKRig/Source/IKRig/Public/Retargeter/RetargetOps/RootMotionGeneratorOp.h b/Engine/Plugins/Animation/IKRig/Source/IKRig/Public/Retargeter/RetargetOps/RootMotionGeneratorOp.h
--- a/Engine/Plugins/Animation/IKRig/Source/IKRig/Public/Retargeter/RetargetOps/RootMotionGeneratorOp.h
+++ b/Engine/Plugins/Animation/IKRig/Source/IKRig/Public/Retargeter/RetargetOps/RootMotionGeneratorOp.h
@@ -32,11 +32,11 @@ struct FIKRetargetRootMotionOpSettings : public FIKRetargetOpSettingsBase
 	GENERATED_BODY()
 
 	// The root of the source skeleton.
-	UPROPERTY(EditAnywhere, Category="Bones", meta=(ReinitializeOnEdit))
+	UPROPERTY(EditAnywhere, Category="Bones", meta=(ReinitializeOnEdit, NotOverrideable))
 	FBoneReference SourceRoot;
 
 	// The root of the target skeleton.
-	UPROPERTY(EditAnywhere, Category="Bones", meta=(ReinitializeOnEdit))
+	UPROPERTY(EditAnywhere, Category="Bones", meta=(ReinitializeOnEdit, NotOverrideable))
 	FBoneReference TargetRoot;
 	
 	// Where to copy the root motion from.
@@ -47,7 +47,7 @@ struct FIKRetargetRootMotionOpSettings : public FIKRetargetOpSettingsBase
 	ERootMotionSource RootMotionSource = ERootMotionSource::CopyFromSourceRoot;
 
 	// The pelvis of the target skeleton. Only used when "Root Motion Source" is set to target pelvis.
-	UPROPERTY(EditAnywhere, Category="Motion Source", meta=(ReinitializeOnEdit))
+	UPROPERTY(EditAnywhere, Category="Motion Source", meta=(ReinitializeOnEdit, NotOverrideable))
 	FBoneReference TargetPelvis;
 
 	// Applies only when generating root motion from the Pelvis.
diff --git a/Engine/Plugins/Animation/IKRig/Source/IKRig/Public/Retargeter/RetargetOps/RunIKRigOp.h b/Engine/Plugins/Animation/IKRig/Source/IKRig/Public/Retargeter/RetargetOps/RunIKRigOp.h
--- a/Engine/Plugins/Animation/IKRig/Source/IKRig/Public/Retargeter/RetargetOps/RunIKRigOp.h
+++ b/Engine/Plugins/Animation/IKRig/Source/IKRig/Public/Retargeter/RetargetOps/RunIKRigOp.h
@@ -27,7 +27,7 @@ struct FIKRetargetRunIKRigChainSettings
 	void SetName(const FName InName) { TargetChainName = InName; };
 
 	/** The name of the TARGET chain to run IK on. */
-	UPROPERTY(EditAnywhere, BlueprintReadWrite, Category = "Setup", meta=(ReinitializeOnEdit))
+	UPROPERTY(EditAnywhere, BlueprintReadWrite, Category = "Setup", meta=(ReinitializeOnEdit, NotOverrideable))
 	FName TargetChainName;
 	
 	/** Enable/disable IK on this chain.
diff --git a/Engine/Plugins/Animation/IKRig/Source/IKRig/Public/Retargeter/RetargetOps/ScaleGoalsOp.h b/Engine/Plugins/Animation/IKRig/Source/IKRig/Public/Retargeter/RetargetOps/ScaleGoalsOp.h
--- a/Engine/Plugins/Animation/IKRig/Source/IKRig/Public/Retargeter/RetargetOps/ScaleGoalsOp.h
+++ b/Engine/Plugins/Animation/IKRig/Source/IKRig/Public/Retargeter/RetargetOps/ScaleGoalsOp.h
@@ -23,7 +23,7 @@ struct FIKRetargetScaleGoalsChainSettings
 	void SetName(const FName InName) { TargetChainName = InName; };
 
 	/** The name of the TARGET chain for these settings. */
-	UPROPERTY(EditAnywhere, BlueprintReadWrite, Category = "Setup", meta=(ReinitializeOnEdit))
+	UPROPERTY(EditAnywhere, BlueprintReadWrite, Category = "Setup", meta=(ReinitializeOnEdit, NotOverrideable))
 	FName TargetChainName;
 
 	/** Range 0 to infinity. Default 1. Scales the vertical component of the IK goal's position. */
diff --git a/Engine/Plugins/Animation/IKRig/Source/IKRig/Public/Retargeter/RetargetOps/ScaleSourceOp.h b/Engine/Plugins/Animation/IKRig/Source/IKRig/Public/Retargeter/RetargetOps/ScaleSourceOp.h
--- a/Engine/Plugins/Animation/IKRig/Source/IKRig/Public/Retargeter/RetargetOps/ScaleSourceOp.h
+++ b/Engine/Plugins/Animation/IKRig/Source/IKRig/Public/Retargeter/RetargetOps/ScaleSourceOp.h
@@ -29,7 +29,7 @@ struct FIKRetargetScaleSourceOpSettings : public FIKRetargetOpSettingsBase
 	UPROPERTY(EditAnywhere, BlueprintReadWrite, Category = ScalePivot)
 	EScaleSourcePivot ScalePivot = EScaleSourcePivot::ComponentOrigin;
 
-	UPROPERTY(EditAnywhere, Category = ScalePivot)
+	UPROPERTY(EditAnywhere, Category = ScalePivot, meta=(NotOverrideable))
 	FBoneReference ScalePivotBone;
 
 	UPROPERTY(EditAnywhere, BlueprintReadWrite, Category = ScalePivot)
diff --git a/Engine/Plugins/Animation/IKRig/Source/IKRig/Public/Retargeter/RetargetOps/SpeedPlantingOp.h b/Engine/Plugins/Animation/IKRig/Source/IKRig/Public/Retargeter/RetargetOps/SpeedPlantingOp.h
--- a/Engine/Plugins/Animation/IKRig/Source/IKRig/Public/Retargeter/RetargetOps/SpeedPlantingOp.h
+++ b/Engine/Plugins/Animation/IKRig/Source/IKRig/Public/Retargeter/RetargetOps/SpeedPlantingOp.h
@@ -22,7 +22,7 @@ struct FRetargetSpeedPlantingSettings
 
 	/** The name of the target chain to plant the IK on.
 	 * NOTE: this chain must have an IK Goal assigned to it! */
-	UPROPERTY(EditAnywhere, BlueprintReadWrite, Category = "Speed Planting", meta=(ReinitializeOnEdit))
+	UPROPERTY(EditAnywhere, BlueprintReadWrite, Category = "Speed Planting", meta=(ReinitializeOnEdit, NotOverrideable))
 	FName TargetChainName;
 
 	/** The name of the curve on the source animation that contains the speed of the end effector bone.*/
diff --git a/Engine/Plugins/Animation/IKRig/Source/IKRig/Public/Retargeter/RetargetOps/StretchChainOp.h b/Engine/Plugins/Animation/IKRig/Source/IKRig/Public/Retargeter/RetargetOps/StretchChainOp.h
--- a/Engine/Plugins/Animation/IKRig/Source/IKRig/Public/Retargeter/RetargetOps/StretchChainOp.h
+++ b/Engine/Plugins/Animation/IKRig/Source/IKRig/Public/Retargeter/RetargetOps/StretchChainOp.h
@@ -45,7 +45,7 @@ struct FRetargetStretchChainSettings
 	FRetargetStretchChainSettings(FName InTargetChainName) : TargetChainName(InTargetChainName) {}
 	
 	/** The name of the target chain to stretch. */
-	UPROPERTY(EditAnywhere, BlueprintReadWrite, Category = "Setup", meta=(ReinitializeOnEdit))
+	UPROPERTY(EditAnywhere, BlueprintReadWrite, Category = "Setup", meta=(ReinitializeOnEdit, NotOverrideable))
 	FName TargetChainName;
 
 	/** Enable/disable stretching. */
diff --git a/Engine/Plugins/Animation/IKRig/Source/IKRig/Public/Retargeter/RetargetOps/StrideWarpingOp.h b/Engine/Plugins/Animation/IKRig/Source/IKRig/Public/Retargeter/RetargetOps/StrideWarpingOp.h
--- a/Engine/Plugins/Animation/IKRig/Source/IKRig/Public/Retargeter/RetargetOps/StrideWarpingOp.h
+++ b/Engine/Plugins/Animation/IKRig/Source/IKRig/Public/Retargeter/RetargetOps/StrideWarpingOp.h
@@ -28,7 +28,7 @@ struct FRetargetStrideWarpChainSettings
 	FRetargetStrideWarpChainSettings(const FName InTargetChainName) : TargetChainName(InTargetChainName) {}
 
 	/** The name of the TARGET chain with an IK goal to warp. */
-	UPROPERTY(EditAnywhere, BlueprintReadWrite, Category = "Setup", meta=(ReinitializeOnEdit))
+	UPROPERTY(EditAnywhere, BlueprintReadWrite, Category = "Setup", meta=(ReinitializeOnEdit, NotOverrideable))
 	FName TargetChainName;
 	
 	/** Whether to warp the location of the IK goal on this chain. Default is true. */
diff --git a/Engine/Plugins/Animation/IKRig/Source/IKRig/Public/Retargeter/RetargetOps/WeaponGoalsOp.h b/Engine/Plugins/Animation/IKRig/Source/IKRig/Public/Retargeter/RetargetOps/WeaponGoalsOp.h
--- a/Engine/Plugins/Animation/IKRig/Source/IKRig/Public/Retargeter/RetargetOps/WeaponGoalsOp.h
+++ b/Engine/Plugins/Animation/IKRig/Source/IKRig/Public/Retargeter/RetargetOps/WeaponGoalsOp.h
@@ -19,47 +19,47 @@ struct FIKRetargetWeaponGoalsOpSettings : public FIKRetargetOpSettingsBase
 
 	/* An ancillary weapon bone for the left hand.
 	 * This will be scaled with Weapon Scale. Defaults to "weapon_l". */
-	UPROPERTY(EditAnywhere, Category = "Bones")
+	UPROPERTY(EditAnywhere, Category = "Bones", meta=(NotOverrideable))
 	FBoneReference LeftWeaponBone;
 
 	/* An ancillary weapon bone for the right hand.
 	 * This will be scaled with Weapon Scale. Defaults to "weapon_r". */
-	UPROPERTY(EditAnywhere, Category = "Bones")
+	UPROPERTY(EditAnywhere, Category = "Bones", meta=(NotOverrideable))
 	FBoneReference RightWeaponBone;
 
 	/* An optional bone used to attach props to the left hand.
 	 * This will be scaled with Weapon Scale. Defaults to "hand_attach_l". */
-	UPROPERTY(EditAnywhere, Category = "Bones")
+	UPROPERTY(EditAnywhere, Category = "Bones", meta=(NotOverrideable))
 	FBoneReference LeftHandAttachBone;
 
 	/* An optional bone used to attach props to the right hand.
 	 * This will be scaled with Weapon Scale. Defaults to "hand_attach_r". */
-	UPROPERTY(EditAnywhere, Category = "Bones")
+	UPROPERTY(EditAnywhere, Category = "Bones", meta=(NotOverrideable))
 	FBoneReference RightHandAttachBone;
 
 	/* The bone skinned to the left hand at wrist location.
 	 * This bone must have an associated IK Goal. Defaults to "hand_l". */
-	UPROPERTY(EditAnywhere, Category = "Bones")
+	UPROPERTY(EditAnywhere, Category = "Bones", meta=(NotOverrideable))
 	FBoneReference LeftHandBone;
 
 	/* The bone skinned to the right hand at wrist location.
 	 * This bone must have an associated IK Goal. Defaults to "hand_r". */
-	UPROPERTY(EditAnywhere, Category = "Bones")
+	UPROPERTY(EditAnywhere, Category = "Bones", meta=(NotOverrideable))
 	FBoneReference RightHandBone;
 
 	/* A bone representing the location of the left hand IK position while holding a two-handed prop.
 	 * Typically, a child of the IKHandGun. Defaults to "ik_hand_l". */
-	UPROPERTY(EditAnywhere, Category = "Bones", DisplayName="IK Hand Left Bone")
+	UPROPERTY(EditAnywhere, Category = "Bones", DisplayName="IK Hand Left Bone", meta=(NotOverrideable))
 	FBoneReference IKHandLeftBone;
 
 	/* A bone representing the location of the right hand IK position while holding a two-handed prop.
 	 * Typically, a child of the IKHandGun. Defaults to "ik_hand_r". */
-	UPROPERTY(EditAnywhere, Category = "Bones", DisplayName="IK Hand Right Bone")
+	UPROPERTY(EditAnywhere, Category = "Bones", DisplayName="IK Hand Right Bone", meta=(NotOverrideable))
 	FBoneReference IKHandRightBone;
 
 	/* A bone located at the right hand, and parent of IKHandLeftBone and IKHandRightBone.
 	 * This will be scaled with Weapon Scale. Defaults to "ik_hand_gun". */
-	UPROPERTY(EditAnywhere, Category = "Bones")
+	UPROPERTY(EditAnywhere, Category = "Bones", meta=(NotOverrideable))
 	FBoneReference IKHandGun;
 
 	/* When false, all weapon scaling is skipped and WeaponScale has no effect. */
diff --git a/Engine/Plugins/Animation/IKRig/Source/IKRigEditor/Private/IKRigEditor.cpp b/Engine/Plugins/Animation/IKRig/Source/IKRigEditor/Private/IKRigEditor.cpp
--- a/Engine/Plugins/Animation/IKRig/Source/IKRigEditor/Private/IKRigEditor.cpp
+++ b/Engine/Plugins/Animation/IKRig/Source/IKRigEditor/Private/IKRigEditor.cpp
@@ -16,6 +16,7 @@
 #include "RetargetEditor/IKRetargetEditPoseMode.h"
 #include "RetargetEditor/IKRetargeterThumbnailRenderer.h"
 #include "RetargetEditor/IKRetargetOpDetails.h"
+#include "RetargetEditor/IKRetargetOverrideCustomization.h"
 #include "Retargeter/IKRetargetOps.h"
 #include "RigEditor/IKRigCommands.h"
 #include "RigEditor/IKRigEditMode.h"
@@ -119,6 +120,10 @@ void FIKRigEditor::StartupModule()
 	PropertyEditorModule.RegisterCustomPropertyTypeLayout(
 				"IKRetargetScaleGoalsOpSettings",
 				FOnGetPropertyTypeCustomizationInstance::CreateStatic(&FScaleGoalsOpCustomization::MakeInstance));
+	// retarget profile details
+	PropertyEditorModule.RegisterCustomPropertyTypeLayout(
+				"RetargetOverrideSet",
+				FOnGetPropertyTypeCustomizationInstance::CreateStatic(&FRetargetOverrideCustomization::MakeInstance));
 
 	// register a thumbnail renderer for the assets
 	UThumbnailManager::Get().RegisterCustomRenderer(UIKRigDefinition::StaticClass(), UIKRigThumbnailRenderer::StaticClass());
diff --git a/Engine/Plugins/Animation/IKRig/Source/IKRigEditor/Private/RetargetEditor/IKRetargetApplicationMode.cpp b/Engine/Plugins/Animation/IKRig/Source/IKRigEditor/Private/RetargetEditor/IKRetargetApplicationMode.cpp
--- a/Engine/Plugins/Animation/IKRig/Source/IKRigEditor/Private/RetargetEditor/IKRetargetApplicationMode.cpp
+++ b/Engine/Plugins/Animation/IKRig/Source/IKRigEditor/Private/RetargetEditor/IKRetargetApplicationMode.cpp
@@ -12,6 +12,7 @@
 #include "RetargetEditor/IKRetargetHierarchyTabSummoner.h"
 #include "RetargetEditor/IKRetargetOpStackTabSummoner.h"
 #include "RetargetEditor/IKRetargetOutputLogTabSummoner.h"
+#include "RetargetEditor/IKRetargetOverridesTabSummoner.h"
 
 #define LOCTEXT_NAMESPACE "IKRetargetMode"
 
@@ -42,9 +43,10 @@ FIKRetargetApplicationMode::FIKRetargetApplicationMode(
 	TabFactories.RegisterFactory(MakeShared<FIKRetargetOutputLogTabSummoner>(IKRetargetEditor));
 	TabFactories.RegisterFactory(MakeShared<FIKRetargetHierarchyTabSummoner>(IKRetargetEditor));
 	TabFactories.RegisterFactory(MakeShared<FIKRetargetOpStackTabSummoner>(IKRetargetEditor));
+	TabFactories.RegisterFactory(MakeShared<FIKRetargetOverridesTabSummoner>(IKRetargetEditor));
 
 	// create tab layout
-	TabLayout = FTabManager::NewLayout("Standalone_IKRetargetEditor_Layout_v1.020")
+	TabLayout = FTabManager::NewLayout("Standalone_IKRetargetEditor_Layout_v1.021")
 		->AddArea
 		(
 			FTabManager::NewPrimaryArea()
@@ -103,6 +105,7 @@ FIKRetargetApplicationMode::FIKRetargetApplicationMode(
 				FTabManager::NewStack()
 					->SetSizeCoefficient(0.2f)
 					->AddTab(FPersonaTabs::DetailsID, ETabState::OpenedTab)
+					->AddTab(FIKRetargetOverridesTabSummoner::TabID, ETabState::OpenedTab)
 					->AddTab(FPersonaTabs::AdvancedPreviewSceneSettingsID, ETabState::OpenedTab)
 					->SetForegroundTab(FPersonaTabs::DetailsID)
 					
diff --git a/Engine/Plugins/Animation/IKRig/Source/IKRigEditor/Private/RetargetEditor/IKRetargetBatchOperation.cpp b/Engine/Plugins/Animation/IKRig/Source/IKRigEditor/Private/RetargetEditor/IKRetargetBatchOperation.cpp
--- a/Engine/Plugins/Animation/IKRig/Source/IKRigEditor/Private/RetargetEditor/IKRetargetBatchOperation.cpp
+++ b/Engine/Plugins/Animation/IKRig/Source/IKRigEditor/Private/RetargetEditor/IKRetargetBatchOperation.cpp
@@ -280,7 +280,7 @@ void UIKRetargetBatchOperation::RetargetAssets(
 	FIKRetargetProcessor Processor;
 	FRetargetProfile RetargetProfile;
 	RetargetProfile.FillProfileWithAssetSettings(Context.IKRetargetAsset);
-	Processor.Initialize(Context.SourceMesh, Context.TargetMesh, Context.IKRetargetAsset, RetargetProfile);
+	Processor.Initialize(Context.SourceMesh, Context.TargetMesh, Context.IKRetargetAsset, RetargetProfile, Context.OverrideSetNames);
 	if (!Processor.IsInitialized())
 	{
 		UE_LOG(LogTemp, Warning, TEXT("Unable to initialize the IK Retargeter. Newly created animations were not retargeted!"));
@@ -483,7 +483,7 @@ void UIKRetargetBatchOperation::ConvertAnimation(
 			OutProcessor.UpdateOpsFromAnimSequence(SourceSequence, TimeAtCurrentFrame);
 
 			// run the retargeter
-			const TArray<FTransform>& TargetComponentPose = OutProcessor.RunRetargeter(SourceComponentPose, SettingsProfile, DeltaTime);
+			const TArray<FTransform>& TargetComponentPose = OutProcessor.RunRetargeter(SourceComponentPose, SettingsProfile, Context.OverrideSetNames, DeltaTime);
 
 			// convert to a local-space pose
 			TArray<FTransform> TargetLocalPose = TargetComponentPose;
@@ -833,6 +833,7 @@ TArray<FAssetData> UIKRetargetBatchOperation::DuplicateAndRetarget(
 	USkeletalMesh* SourceMesh,
 	USkeletalMesh* TargetMesh,
 	UIKRetargeter* IKRetargetAsset,
+	const TArray<FName>& InOverrideSetNames,
 	const FString& Search,
 	const FString& Replace,
 	const FString& Prefix,
@@ -865,6 +866,7 @@ TArray<FAssetData> UIKRetargetBatchOperation::DuplicateAndRetarget(
 	Context.bUseSourcePath = bUseSourcePath;
 	Context.bIncludeReferencedAssets = bIncludeReferencedAssets;
 	Context.bOverwriteExistingFiles = bOverwriteExistingFiles;
+	Context.OverrideSetNames = InOverrideSetNames;
 
 	// actually run the batch operation
 	UIKRetargetBatchOperation* BatchOperation = NewObject<UIKRetargetBatchOperation>();
diff --git a/Engine/Plugins/Animation/IKRig/Source/IKRigEditor/Private/RetargetEditor/IKRetargetCommands.cpp b/Engine/Plugins/Animation/IKRig/Source/IKRigEditor/Private/RetargetEditor/IKRetargetCommands.cpp
--- a/Engine/Plugins/Animation/IKRig/Source/IKRigEditor/Private/RetargetEditor/IKRetargetCommands.cpp
+++ b/Engine/Plugins/Animation/IKRig/Source/IKRigEditor/Private/RetargetEditor/IKRetargetCommands.cpp
@@ -11,6 +11,7 @@ void FIKRetargetCommands::RegisterCommands()
 	UI_COMMAND(ShowRetargetPose, "Show Retarget Pose", "Display the retarget pose and retarget to the target skeleton.", EUserInterfaceActionType::Button, FInputChord());
 
 	UI_COMMAND(ShowAssetSettings, "Asset Settings", "Set the source or target IK Rigs, preview meshes. Adjust debug settings.", EUserInterfaceActionType::ToggleButton, FInputChord());
+	UI_COMMAND(ShowOverrideSets, "Override Sets", "Open the Override Sets tab to create / edit sets of op property overrides.", EUserInterfaceActionType::ToggleButton, FInputChord());
 	
 	UI_COMMAND(ResetAllBones, "Reset All", "Sets the selected bones to the mesh reference pose.", EUserInterfaceActionType::Button, FInputChord());
 	UI_COMMAND(ResetSelectedBones, "Reset Selected Bones", "Sets the selected bones to the mesh reference pose.", EUserInterfaceActionType::Button, FInputChord());
diff --git a/Engine/Plugins/Animation/IKRig/Source/IKRigEditor/Private/RetargetEditor/IKRetargetEditor.cpp b/Engine/Plugins/Animation/IKRig/Source/IKRigEditor/Private/RetargetEditor/IKRetargetEditor.cpp
--- a/Engine/Plugins/Animation/IKRig/Source/IKRigEditor/Private/RetargetEditor/IKRetargetEditor.cpp
+++ b/Engine/Plugins/Animation/IKRig/Source/IKRigEditor/Private/RetargetEditor/IKRetargetEditor.cpp
@@ -25,6 +25,7 @@
 #include "RetargetEditor/IKRetargetApplicationMode.h"
 #include "RetargetEditor/IKRetargetDefaultMode.h"
 #include "RetargetEditor/IKRetargetEditorController.h"
+#include "RetargetEditor/IKRetargetOverridesTabSummoner.h"
 
 #include "Styling/AppStyle.h"
 
@@ -146,6 +147,19 @@ void FIKRetargetEditor::BindCommands()
 			return EditorController->IsObjectInDetailsView(Asset);	
 		}));
 
+	ToolkitCommands->MapAction(
+		Commands.ShowOverrideSets,
+		FExecuteAction::CreateLambda([this]()
+		{
+			FGlobalTabmanager::Get()->TryInvokeTab(FIKRetargetOverridesTabSummoner::TabID);
+		}),
+		FCanExecuteAction(),
+		FIsActionChecked::CreateLambda([this]() ->bool
+		{
+			TSharedPtr<SDockTab> ExistingTab = FGlobalTabmanager::Get()->FindExistingLiveTab(FIKRetargetOverridesTabSummoner::TabID);
+			return ExistingTab.IsValid() && ExistingTab->IsForeground();
+		}));
+
 	//
 	// Edit pose commands
 	//
@@ -283,6 +297,13 @@ void FIKRetargetEditor::FillToolbar(FToolBarBuilder& ToolbarBuilder)
 		TAttribute<FText>(),
 		TAttribute<FText>(),
 		FSlateIcon(FIKRetargetEditorStyle::Get().GetStyleSetName(),"IKRetarget.AssetSettings"));
+
+		ToolbarBuilder.AddToolBarButton(
+		FIKRetargetCommands::Get().ShowOverrideSets,
+		NAME_None,
+		TAttribute<FText>(),
+		TAttribute<FText>(),
+		FSlateIcon(FIKRetargetEditorStyle::Get().GetStyleSetName(),"IKRetarget.AssetSettings"));
 	}
 	ToolbarBuilder.EndSection();
 }
diff --git a/Engine/Plugins/Animation/IKRig/Source/IKRigEditor/Private/RetargetEditor/IKRetargetEditorController.cpp b/Engine/Plugins/Animation/IKRig/Source/IKRigEditor/Private/RetargetEditor/IKRetargetEditorController.cpp
--- a/Engine/Plugins/Animation/IKRig/Source/IKRigEditor/Private/RetargetEditor/IKRetargetEditorController.cpp
+++ b/Engine/Plugins/Animation/IKRig/Source/IKRigEditor/Private/RetargetEditor/IKRetargetEditorController.cpp
@@ -36,6 +36,7 @@
 #include "Styling/AppStyle.h"
 
 #include "Misc/ITransaction.h"
+#include "RetargetEditor/SIKRetargetOverrideManager.h"
 
 #define LOCTEXT_NAMESPACE "IKRetargetEditorController"
 
@@ -179,6 +180,11 @@ void FIKRetargetEditorController::Initialize(TSharedPtr<FIKRetargetEditor> InEdi
 	{
 		RefreshOpStackView();
 	});
+	// bind callback when retargeter profiles are modified
+	ProfilesModifiedDelegateHandle = AssetController->OnOverrideSetsModified().AddLambda([this]()
+	{
+		RefreshOverrideSetView();
+	});
 	// bind callback when IK Rig asset is replaced with a different asset
 	IKRigReplacedDelegateHandle = AssetController->OnIKRigReplaced().AddSP(this, &FIKRetargetEditorController::HandleIKRigReplaced);
 	// bind callback when Preview Mesh asset is replaced with a different asset
@@ -188,6 +194,7 @@ void FIKRetargetEditorController::Initialize(TSharedPtr<FIKRetargetEditor> InEdi
 void FIKRetargetEditorController::Close()
 {
 	AssetController->OnRetargeterNeedsInitialized().Remove(RetargeterReInitDelegateHandle);
+	AssetController->OnOverrideSetsModified().Remove(ProfilesModifiedDelegateHandle);
 	AssetController->OnOpStackModified().Remove(OpStackModifiedDelegateHandle);
 	AssetController->OnIKRigReplaced().Remove(IKRigReplacedDelegateHandle);
 	AssetController->OnPreviewMeshReplaced().Remove(PreviewMeshReplacedDelegateHandle);
@@ -796,6 +803,7 @@ void FIKRetargetEditorController::RefreshAllViews() const
 	RefreshDetailsView();
 	RefreshAssetBrowserView();
 	RefreshHierarchyView();
+	RefreshOverrideSetView();
 }
 
 void FIKRetargetEditorController::RefreshDetailsView() const
@@ -832,6 +840,14 @@ void FIKRetargetEditorController::RefreshOpStackView() const
 	}
 }
 
+void FIKRetargetEditorController::RefreshOverrideSetView() const
+{
+	if (OverrideSetView.IsValid())
+	{
+		OverrideSetView->RefreshTreeView();
+	}
+}
+
 void FIKRetargetEditorController::RefreshPoseList() const
 {
 	if (HierarchyView.IsValid())
@@ -1156,6 +1172,16 @@ int32 FIKRetargetEditorController::GetSelectedOpIndex() const
 	return AssetController->GetIndexOfOpByName(Selection.LastSelectedOpName);
 }
 
+FName FIKRetargetEditorController::GetSelectedOverrideSet() const
+{
+	if (!OverrideSetView.IsValid())
+	{
+		return NAME_None;
+	}
+
+	return OverrideSetView.Get()->GetSelectedOverrideSetName();
+}
+
 void FIKRetargetEditorController::SetRetargeterMode(ERetargeterOutputMode Mode)
 {
 	if (OutputMode == Mode)
diff --git a/Engine/Plugins/Animation/IKRig/Source/IKRigEditor/Private/RetargetEditor/IKRetargetOverrideCustomization.cpp b/Engine/Plugins/Animation/IKRig/Source/IKRigEditor/Private/RetargetEditor/IKRetargetOverrideCustomization.cpp
new file mode 100644
--- /dev/null
+++ b/Engine/Plugins/Animation/IKRig/Source/IKRigEditor/Private/RetargetEditor/IKRetargetOverrideCustomization.cpp
@@ -0,0 +1,440 @@
+﻿// Copyright Epic Games, Inc. All Rights Reserved.
+
+#include "RetargetEditor/IKRetargetOverrideCustomization.h"
+
+#include "DetailLayoutBuilder.h"
+#include "DetailWidgetRow.h"
+#include "IDetailChildrenBuilder.h"
+#include "IDetailGroup.h"
+#include "IPropertyUtilities.h"
+#include "RetargetEditor/IKRetargeterController.h"
+#include "Retargeter/IKRetargeter.h"
+#include "Retargeter/IKRetargetOverrides.h"
+#include "RigEditor/IKRigStructViewer.h"
+#include "StructUtils/InstancedStruct.h"
+#include "Widgets/Input/SButton.h"
+#include "Widgets/Input/SEditableTextBox.h"
+#include "Widgets/Text/STextBlock.h"
+
+#define LOCTEXT_NAMESPACE "RetargetOverrideCustomization"
+
+void FRetargetOverrideCustomization::CustomizeHeader(
+    TSharedRef<IPropertyHandle> PropertyHandle,
+    FDetailWidgetRow& HeaderRow,
+    IPropertyTypeCustomizationUtils& CustomizationUtils)
+{
+}
+
+void FRetargetOverrideCustomization::CustomizeChildren(TSharedRef<IPropertyHandle> RetargetOverrideSetHandle, IDetailChildrenBuilder& ChildBuilder, IPropertyTypeCustomizationUtils& CustomizationUtils)
+{
+	// NOTE:
+	// The top object being customized here is an FRetargetOverrideSet.
+	// FRetargetOverrideSet contains a TArray<FRetargetOpOverrides> with an entry for each op that has property overrides in this set.
+	// The FRetargetOpOverrides has a TArray<FRetargetOpPropertyOverride> which contains the override value for a single property in the op settings
+	// This customization lists all the property overrides with each op in its own category
+
+	// get the controller for the set being edited
+	AssetController = GetAssetController(CustomizationUtils);
+	if (!AssetController.IsValid())
+	{
+		// the struct is being viewed outside the editor, so display the default view
+		uint32 NumChildren;
+		RetargetOverrideSetHandle->GetNumChildren(NumChildren);
+		for (uint32 ChildIndex = 0; ChildIndex < NumChildren; ++ChildIndex)
+		{
+			TSharedPtr<IPropertyHandle> ChildHandle = RetargetOverrideSetHandle->GetChildHandle(ChildIndex);
+			if (!ChildHandle.IsValid())
+			{
+				continue;
+			}
+
+			const FProperty* Property = ChildHandle->GetProperty();
+			if (!Property)
+			{
+				continue;
+			}
+
+			FString Category = Property->GetMetaData(TEXT("Category"));
+			if (Category.Equals(TEXT("Deprecated")))
+			{
+				continue; // skip if it belongs to the "Deprecated" category
+			}
+
+			ChildBuilder.AddProperty(ChildHandle.ToSharedRef());
+		}
+
+		// skip editor-side customization
+		return;
+	}
+
+	// get the name of this set
+	FRetargetOverrideSet* SetPtr = nullptr;
+	void* RawData = nullptr;
+	if (RetargetOverrideSetHandle->GetValueData(RawData) == FPropertyAccess::Success)
+	{
+		SetPtr = static_cast<FRetargetOverrideSet*>(RawData);
+	}
+	const FName SetName = AssetController->GetOverrideSetName(SetPtr);
+	if (!ensure(SetName != NAME_None))
+	{
+		return;
+	}
+
+	// get handle to the RetargetOpOverrides array property
+	TSharedPtr<IPropertyHandle> OpOverridesArrayHandle = RetargetOverrideSetHandle->GetChildHandle(GET_MEMBER_NAME_CHECKED(FRetargetOverrideSet, OpOverrides));
+	if (!OpOverridesArrayHandle.IsValid())
+	{
+		return;
+	}
+
+	// how many ops are stored in this override set?
+	uint32 NumOpOverrides = 0;
+	OpOverridesArrayHandle->GetNumChildren(NumOpOverrides);
+
+	// make copy of the weak ptr for lambda capture
+	TWeakObjectPtr<UIKRetargeterController> LocalController = AssetController;
+
+	// for each op that has overrides in this set...
+	for (uint32 OpIdx = 0; OpIdx < NumOpOverrides; ++OpIdx)
+	{
+		TSharedPtr<IPropertyHandle> OpOverridesHandle = OpOverridesArrayHandle->GetChildHandle(OpIdx);
+
+		// get the reflection data for the settings struct associated with this op
+		const UScriptStruct* OpSettingsType = GetSettingsStruct(OpOverridesHandle);
+		if (!ensure(OpSettingsType))
+		{
+			continue;
+		}
+
+		// get the op name
+		TSharedPtr<IPropertyHandle> OpNameHandle = OpOverridesHandle->GetChildHandle(GET_MEMBER_NAME_CHECKED(FRetargetOpOverrides, OpName));
+		if (!ensure(OpNameHandle.IsValid()))
+		{
+			continue;
+		}
+		FName OpName;
+		OpNameHandle->GetValue(OpName);
+
+		// generate a scoped struct to display
+		TSharedPtr<FStructOnScope> OpSettingsScopedStruct = GetCachedScopedStruct(OpName, OpSettingsType);
+		if (!ensure(OpSettingsScopedStruct.IsValid()))
+		{
+			continue;
+		}
+
+		// create a group for this specific op
+		IDetailGroup& OpGroup = ChildBuilder.AddGroup(OpName, FText::FromName(OpName));
+
+		// get a handle to the array of property overrides
+		TSharedPtr<IPropertyHandle> OverridesArrayHandle = OpOverridesHandle->GetChildHandle(GET_MEMBER_NAME_CHECKED(FRetargetOpOverrides, PropertyOverrides));
+		uint32 NumOverrides = 0;
+		OverridesArrayHandle->GetNumChildren(NumOverrides);
+
+		// update temp struct with all property override values
+		// NOTE: I had to do this BEFORE AddExternalStructure because otherwise the EditConditions evaluated with the wrong values
+		for (uint32 OverrideIndex = 0; OverrideIndex < NumOverrides; ++OverrideIndex)
+		{
+			FString PropertyPath;
+			GetOverridePropertyPath(OpOverridesHandle, OverrideIndex, PropertyPath);
+			TSharedPtr<IPropertyHandle> ValueHandle = GetOverrideValueHandle(OpOverridesHandle, OverrideIndex);
+			CopyStoredValueToScopedStruct(OpSettingsScopedStruct, OpSettingsType, PropertyPath, ValueHandle);
+		}
+
+		// add the struct as a child (this registers the memory with the details system)
+		IDetailPropertyRow* ExternalRoot = ChildBuilder.AddExternalStructure(OpSettingsScopedStruct.ToSharedRef());
+		if (!ensure(ExternalRoot))
+		{
+			continue;
+		}
+
+		// hide the whole struct so we can manually add only the overridden properties below
+		ExternalRoot->Visibility(EVisibility::Hidden);
+
+		// split properties into groups so that we can distinguish between different elements of the same array
+		TMap<FString, TArray<int32>> PropertyGroups = MapOverridesToGroups(OpOverridesHandle);
+
+		for (TTuple<FString, TArray<int>>& Group : PropertyGroups)
+		{
+			const FString& GroupName = Group.Key;
+			const TArray<int32>& OverrideIndices = Group.Value;
+
+			// add the group
+			IDetailGroup& PropertyGroup = OpGroup.AddGroup(FName(GroupName), FText::FromString(GroupName));
+
+			// add all the property overrides in this group
+			for (int32 OverrideIndex : OverrideIndices)
+			{
+				// get the path to this property
+				FString PropertyPath;
+				GetOverridePropertyPath(OpOverridesHandle, OverrideIndex, PropertyPath);
+
+				// resolve the final leaf property so we can use its widget generation
+				TSharedPtr<IPropertyHandle> PropertyHandle = ResolveHandleFromPath(ExternalRoot->GetPropertyHandle(), PropertyPath);
+				if (!ensure(PropertyHandle.IsValid()))
+				{
+					continue;
+				}
+
+				// a callback to marshall edits back into the override string
+				auto UpdateValueAction = [LocalController, SetName, OpName, PropertyPath, OpSettingsScopedStruct]()
+				{
+					if (LocalController.IsValid())
+					{
+						ensure(LocalController->UpdateOverrideValue(SetName, OpName, PropertyPath, *OpSettingsScopedStruct.Get()));
+					}
+				};
+
+				// bind to the main property (handles bool/float)
+				PropertyHandle->SetOnPropertyValueChanged(FSimpleDelegate::CreateLambda(UpdateValueAction));
+				// bind to the children (handles struct properties like FVector X, Y, Z changes)
+				PropertyHandle->SetOnChildPropertyValueChanged(FSimpleDelegate::CreateLambda(UpdateValueAction));
+
+				// add the widget for the property to the group
+				AddRowForPropertyOverride(PropertyHandle, PropertyGroup, SetName, OpName, PropertyPath);
+			}
+		}
+	}
+}
+
+TObjectPtr<UIKRetargeterController> FRetargetOverrideCustomization::GetAssetController(
+    const IPropertyTypeCustomizationUtils& CustomizationUtils)
+{
+    TArray<TWeakObjectPtr<UObject>> SelectedObjects = CustomizationUtils.GetPropertyUtilities()->GetSelectedObjects();
+    if (!ensure(!SelectedObjects.IsEmpty()))
+    {
+       return nullptr;
+    }
+
+    UIKRigStructViewer* StructViewer = Cast<UIKRigStructViewer>(SelectedObjects[0].Get());
+    if (!StructViewer)
+    {
+       return nullptr; // this can happen if an FRetargetOverrideSet is viewed in a blueprint details panel
+    }
+    
+    UObject* StructOwner = StructViewer->GetStructOwner();
+    if (!ensure(StructOwner))
+    {
+       return nullptr;
+    }
+    
+    UIKRetargeter* Retargeter = Cast<UIKRetargeter>(StructOwner);
+    if (!ensure(Retargeter))
+    {
+       return nullptr;
+    }
+    
+    return UIKRetargeterController::GetController(Retargeter);
+}
+
+void FRetargetOverrideCustomization::AddRowForPropertyOverride(
+	TSharedPtr<IPropertyHandle> InPropertyHandle,
+	IDetailGroup& InPropertyGroup,
+	const FName& InOverrideSetName,
+	const FName& InOpName,
+	const FString& InPropertyPath)
+{
+	// add the property row
+	// NOTE: cannot just use AddWidgetRow() here because it won't use struct customizations
+	IDetailPropertyRow& PropertyRow = InPropertyGroup.AddPropertyRow(InPropertyHandle.ToSharedRef());
+	TSharedPtr<SWidget> DefaultNameWidget;
+	TSharedPtr<SWidget> DefaultValueWidget;
+	PropertyRow.GetDefaultWidgets(DefaultNameWidget, DefaultValueWidget);
+
+	// copy of controller ptr for lambda
+	TWeakObjectPtr<UIKRetargeterController> LocalController = AssetController;
+	
+	PropertyRow.CustomWidget()
+	.NameContent()
+	[
+	   DefaultNameWidget.ToSharedRef()
+	]
+	.ValueContent()
+	[
+	   DefaultValueWidget.ToSharedRef()
+	]
+	.ExtensionContent()
+	[
+		// TODO add curve / variable binding UI here
+
+		SNew(SButton)
+		.ButtonStyle(FAppStyle::Get(), "SimpleButton")
+		.OnClicked_Lambda([LocalController, InOverrideSetName, InOpName, InPropertyPath]()
+		{
+			if (LocalController.IsValid())
+			{
+				LocalController->RemovePropertyOverrideFromOp(InOverrideSetName, InOpName, InPropertyPath);
+			}
+			return FReply::Handled();
+		})
+		[
+			SNew(SImage)
+			.Image(FAppStyle::GetBrush("Icons.Delete"))
+			.ColorAndOpacity(FSlateColor::UseForeground())
+		]
+	];
+}
+
+TMap<FString, TArray<int32>> FRetargetOverrideCustomization::MapOverridesToGroups(const TSharedPtr<IPropertyHandle>& OpOverridesHandle) const
+{
+    TMap<FString, TArray<int32>> GroupMap;
+
+    if (!ensure(OpOverridesHandle.IsValid()))
+    {
+       return GroupMap;
+    }
+
+    // get the handle to the array of property overrides
+    TSharedPtr<IPropertyHandle> OverridesArrayHandle = OpOverridesHandle->GetChildHandle(GET_MEMBER_NAME_CHECKED(FRetargetOpOverrides, PropertyOverrides));
+    if (!ensure(OverridesArrayHandle.IsValid()))
+    {
+       return GroupMap;
+    }
+
+    // iterate over all the property overrides
+    uint32 NumOverrides = 0;
+    OverridesArrayHandle->GetNumChildren(NumOverrides);
+    for (uint32 OverrideIndex = 0; OverrideIndex < NumOverrides; ++OverrideIndex)
+    {
+       FString FullPath;
+       GetOverridePropertyPath(OpOverridesHandle, OverrideIndex, FullPath);
+
+       FString GroupLabel;
+       FString RemainingPath;
+
+       // check if the path contains an array element (e.g., "ChainsToAlign[3]->bEnabled")
+       // we split at the last "]->" to ensure we group by the deepest array index
+       if (FullPath.Split(TEXT("]->"), &GroupLabel, &RemainingPath, ESearchCase::IgnoreCase, ESearchDir::FromEnd))
+       {
+          // restore the bracket that was consumed by Split
+          GroupLabel += TEXT("]");
+       }
+       else
+       {
+          // if there's no array index, it's a "loose" property, so just put these in a common group
+          GroupLabel = TEXT("Op Settings");
+       }
+
+       // add the index to the corresponding group
+       GroupMap.FindOrAdd(GroupLabel).Add(static_cast<int32>(OverrideIndex));
+    }
+
+    return MoveTemp(GroupMap);
+}
+
+void FRetargetOverrideCustomization::GetOverridePropertyPath(
+    const TSharedPtr<IPropertyHandle>& InOpOverridesHandle,
+    const int32 InOverrideIndex,
+    FString& OutPath) const
+{
+    TSharedPtr<IPropertyHandle> OverridesArrayHandle = InOpOverridesHandle->GetChildHandle(GET_MEMBER_NAME_CHECKED(FRetargetOpOverrides, PropertyOverrides));
+    TSharedPtr<IPropertyHandle> OverrideHandle = OverridesArrayHandle->GetChildHandle(InOverrideIndex);
+    TSharedPtr<IPropertyHandle> PathHandle =  OverrideHandle->GetChildHandle(FRetargetOpPropertyOverride::GetPropertyPathName());
+    PathHandle->GetValue(OutPath);
+}
+
+TSharedPtr<IPropertyHandle> FRetargetOverrideCustomization::GetOverrideValueHandle(
+    const TSharedPtr<IPropertyHandle>& InOpOverridesHandle,
+    const int32 InOverrideIndex) const
+{
+    TSharedPtr<IPropertyHandle> OverridesArrayHandle = InOpOverridesHandle->GetChildHandle(GET_MEMBER_NAME_CHECKED(FRetargetOpOverrides, PropertyOverrides));
+    TSharedPtr<IPropertyHandle> OverrideHandle = OverridesArrayHandle->GetChildHandle(InOverrideIndex);
+    return OverrideHandle->GetChildHandle(FRetargetOpPropertyOverride::GetValeAsStringName());
+}
+
+TSharedPtr<FStructOnScope>& FRetargetOverrideCustomization::GetCachedScopedStruct(const FName& InOpName, const UScriptStruct* OpSettingsType)
+{
+    TSharedPtr<FStructOnScope>& ScopedStruct = ScopedInstances.FindOrAdd(InOpName);
+
+    if (!ScopedStruct.IsValid())
+    {
+       // initialize the struct with the override value
+       ScopedStruct = MakeShared<FStructOnScope>(OpSettingsType);
+    }
+            
+    return ScopedStruct;
+}
+
+bool FRetargetOverrideCustomization::CopyStoredValueToScopedStruct(
+    const TSharedPtr<FStructOnScope>& InScopedStruct,
+    const UScriptStruct* OpSettingsType,
+    const FString& InPropertyPath,
+    const TSharedPtr<IPropertyHandle>& InRawValueHandle)
+{
+    if (!ensure(InScopedStruct.IsValid()))
+    {
+       return false;
+    }
+
+    TArray<FRetargetOpPropertyOverride::FPropertySegment> Segments;
+    if (!ensure(FRetargetOpPropertyOverride::GetSegmentsFromProperyPath(InPropertyPath, OpSettingsType, Segments)))
+    {
+       return false;
+    }
+    
+    uint8* ValuePtr = FRetargetOpPropertyOverride::GetDataPointerFromPathSegments(InScopedStruct->GetStructMemory(), Segments, true);
+    FProperty* LeafProp = FRetargetOpPropertyOverride::GetLeafProperty(Segments);
+    FString CurrentValue;
+    InRawValueHandle->GetValue(CurrentValue);
+    if (!ensure(ValuePtr && LeafProp && !CurrentValue.IsEmpty()))
+    {
+       return false;
+    }
+    
+    LeafProp->ImportText_Direct(*CurrentValue, ValuePtr, nullptr, PPF_None);
+    
+    return true;
+}
+
+const UScriptStruct* FRetargetOverrideCustomization::GetSettingsStruct(TSharedPtr<IPropertyHandle>& OpOverridesHandle) const
+{
+    TSharedPtr<IPropertyHandle> SettingsHandle = OpOverridesHandle->GetChildHandle(GET_MEMBER_NAME_CHECKED(FRetargetOpOverrides, ScriptStruct));
+        
+    const UScriptStruct* SettingsStruct = nullptr;
+    void* StructData = nullptr;
+    if (SettingsHandle->GetValueData(StructData) == FPropertyAccess::Success)
+    {
+       if (TObjectPtr<UScriptStruct>* StructPtr = static_cast<TObjectPtr<UScriptStruct>*>(StructData))
+       {
+          SettingsStruct = StructPtr->Get();
+       }
+    }
+
+    return SettingsStruct;
+}
+
+TSharedPtr<IPropertyHandle> FRetargetOverrideCustomization::ResolveHandleFromPath(TSharedPtr<IPropertyHandle> InRootHandle, const FString& InPath)
+{
+	// validation
+	if (!InRootHandle.IsValid() || InPath.IsEmpty())
+	{
+		return nullptr;
+	}
+
+	// tokenize the property path into steps
+	TArray<FPropertyPathParser::FPropertyPathStep> Steps = FPropertyPathParser::ParsePath(InPath);
+
+	// follow each step in the path by traversing the property handles
+	TSharedPtr<IPropertyHandle> CurrentHandle = InRootHandle;
+	for (const FPropertyPathParser::FPropertyPathStep& Step : Steps)
+	{
+		// navigate to the member property by name
+		CurrentHandle = CurrentHandle->GetChildHandle(*Step.PropertyName);
+        
+		// if valid and this step requires an array index, dive deeper
+		if (CurrentHandle.IsValid() && Step.IsArray())
+		{
+			CurrentHandle = CurrentHandle->GetChildHandle(Step.ArrayIndex);
+		}
+
+		// early out if any step in the chain fails
+		if (!CurrentHandle.IsValid())
+		{
+			return nullptr;
+		}
+	}
+
+	return CurrentHandle;
+}
+
+
+#undef LOCTEXT_NAMESPACE
\ No newline at end of file
diff --git a/Engine/Plugins/Animation/IKRig/Source/IKRigEditor/Private/RetargetEditor/IKRetargetOverridesTabSummoner.cpp b/Engine/Plugins/Animation/IKRig/Source/IKRigEditor/Private/RetargetEditor/IKRetargetOverridesTabSummoner.cpp
new file mode 100644
--- /dev/null
+++ b/Engine/Plugins/Animation/IKRig/Source/IKRigEditor/Private/RetargetEditor/IKRetargetOverridesTabSummoner.cpp
@@ -0,0 +1,36 @@
+﻿// Copyright Epic Games, Inc. All Rights Reserved.
+#include "RetargetEditor/IKRetargetOverridesTabSummoner.h"
+
+#include "IDocumentation.h"
+#include "RetargetEditor/IKRetargetEditor.h"
+#include "RetargetEditor/SIKRetargetOverrideManager.h"
+
+#define LOCTEXT_NAMESPACE "IKRetargetOverrideTabSummoner"
+
+const FName FIKRetargetOverridesTabSummoner::TabID(TEXT("OverrideSets"));
+
+FIKRetargetOverridesTabSummoner::FIKRetargetOverridesTabSummoner(
+	const TSharedRef<FIKRetargetEditor>& InRetargetEditor)
+	: FWorkflowTabFactory(TabID, InRetargetEditor),
+	IKRetargetEditor(InRetargetEditor)
+{
+	bIsSingleton = true; // only allow a single instance of this tab
+	
+	TabLabel = LOCTEXT("IKRetargetOverridesTabLabel", "Override Sets");
+	TabIcon = FSlateIcon(FAppStyle::GetAppStyleSetName(), "Kismet.Tabs.CompilerResults");
+
+	ViewMenuDescription = LOCTEXT("IKRetargetOverrides_ViewMenu_Desc", "Override Sets");
+	ViewMenuTooltip = LOCTEXT("IKRetargetOverrides_ViewMenu_ToolTip", "Show the Retargeting Overrides Tab");
+}
+
+TSharedPtr<SToolTip> FIKRetargetOverridesTabSummoner::CreateTabToolTipWidget(const FWorkflowTabSpawnInfo& Info) const
+{
+	return  IDocumentation::Get()->CreateToolTip(LOCTEXT("IKRetargetOverridesTooltip", "Create and edit property override sets."), NULL, TEXT("Shared/Editors/Persona"), TEXT("IKRetargetProfiles_Window"));
+}
+
+TSharedRef<SWidget> FIKRetargetOverridesTabSummoner::CreateTabBody(const FWorkflowTabSpawnInfo& Info) const
+{
+	return SNew(SIKRetargetOverrideManager, IKRetargetEditor.Pin()->GetController());
+}
+
+#undef LOCTEXT_NAMESPACE 
diff --git a/Engine/Plugins/Animation/IKRig/Source/IKRigEditor/Private/RetargetEditor/IKRetargeterController.cpp b/Engine/Plugins/Animation/IKRig/Source/IKRigEditor/Private/RetargetEditor/IKRetargeterController.cpp
--- a/Engine/Plugins/Animation/IKRig/Source/IKRigEditor/Private/RetargetEditor/IKRetargeterController.cpp
+++ b/Engine/Plugins/Animation/IKRig/Source/IKRigEditor/Private/RetargetEditor/IKRetargeterController.cpp
@@ -266,6 +266,24 @@ FName UIKRetargeterController::SetOpName(const FName InName, const int32 InOpInd
 			}
 		}
 	}
+
+	// update any profiles pointing at the old name
+	for (TTuple<FName, FRetargetProfile>& ProfilePair : Asset->Profiles)
+	{
+		if (FRetargetOpProfile* OpProfile = ProfilePair.Value.GetOpProfileByName(OldOpName))
+		{
+			OpProfile->OpToApplySettingsTo = UniqueName;
+		}
+	}
+
+	// update any override sets pointing to the old name
+	for (TTuple<FName, FRetargetOverrideSet>& OverrideSetPair : Asset->OverrideSets)
+	{
+		if (FRetargetOpOverrides* OpOverride = OverrideSetPair.Value.FindOpOverrideByOpName(OldOpName))
+		{
+			OpOverride->OpName = UniqueName;
+		}
+	}
 	
 	return UniqueName;
 }
@@ -1457,6 +1475,497 @@ TWeakPtr<const FIKRigSkeleton> UIKRetargeterController::GetCachedIKRigSkeleton(c
 	return NewIKRigSkeleton.ToWeakPtr();
 }
 
+FName UIKRetargeterController::AddNewRetargetOverrideSet(FName InOverrideSetName)
+{
+    constexpr bool bShouldTransact = true;
+    FScopedTransaction Transaction(LOCTEXT("AddNewRetargetOverrideSet", "Add New Retarget Override Set"), bShouldTransact);
+    Asset->Modify();
+
+    const FName UniqueName = GetUniqueOverrideSetName(InOverrideSetName);
+    FRetargetOverrideSet& NewSet = Asset->OverrideSets.Add(UniqueName);
+    NewSet.Name = UniqueName;
+
+	PostEditOverrideSets();
+	
+    return UniqueName;
+}
+
+bool UIKRetargeterController::RemoveRetargetOverrideSet(FName InOverrideSetName)
+{
+    if (!Asset->OverrideSets.Contains(InOverrideSetName))
+    {
+       return false;
+    }
+    
+    constexpr bool bShouldTransact = true;
+    FScopedTransaction Transaction(LOCTEXT("RemoveRetargetOverrideSet", "Remove Retarget Override Set"), bShouldTransact);
+    Asset->Modify();
+    
+    Asset->OverrideSets.Remove(InOverrideSetName);
+
+	PostEditOverrideSets();
+	
+    return true;
+}
+
+bool UIKRetargeterController::RenameRetargetOverrideSet(FName InOverrideSetName, FName InNewName)
+{
+    FString CleanString = InNewName.ToString().TrimStartAndEnd();
+    CleanString = SlugStringForValidName(CleanString);
+    InNewName = FName(CleanString);
+    
+    if (InOverrideSetName == InNewName || InNewName == NAME_None)
+    {
+       return false;
+    }
+    
+    FRetargetOverrideSet* SetToRename = Asset->OverrideSets.Find(InOverrideSetName);
+    if (!SetToRename || Asset->OverrideSets.Contains(InNewName))
+    {
+       return false;
+    }
+
+    constexpr bool bShouldTransact = true;
+    FScopedTransaction Transaction(LOCTEXT("RenameRetargetOverrideSet", "Rename Retarget Override Set"), bShouldTransact);
+    Asset->Modify();
+
+    FRetargetOverrideSet TempValue = MoveTemp(*SetToRename);
+    Asset->OverrideSets.Remove(InOverrideSetName);
+    
+    const FName NewName = GetUniqueOverrideSetName(InNewName);
+    TempValue.Name = NewName;
+    Asset->OverrideSets.Add(NewName, MoveTemp(TempValue));
+
+    // update parent/child references
+	for (TTuple<FName, FRetargetOverrideSet>& It : Asset->OverrideSets)
+	{
+		if (It.Value.ParentName == InOverrideSetName)
+		{
+			It.Value.ParentName = NewName;
+		}
+	}
+
+	PostEditOverrideSets();
+    return true;
+}
+
+bool UIKRetargeterController::SetParentOverrideSet(FName InOverrideSetName, FName InParentName)
+{
+    if (InOverrideSetName == InParentName)
+    {
+       return false;
+    }
+
+    FRetargetOverrideSet* SetToReparent = Asset->OverrideSets.Find(InOverrideSetName);
+    if (!SetToReparent)
+    {
+       return false;
+    }
+
+    if (InParentName != NAME_None && !Asset->OverrideSets.Contains(InParentName))
+    {
+       return false;
+    }  
+
+    constexpr bool bShouldTransact = true;
+    FScopedTransaction Transaction(LOCTEXT("ReparentRetargetOverrideSet", "Reparent Retarget Override Set"), bShouldTransact);
+    Asset->Modify();
+
+    // prevent circular dependency by moving children of the current set to the current set's parent
+    const TArray<FName> AllChildren = GetAllChildrenOverrideSets(InOverrideSetName);
+    if (AllChildren.Contains(InParentName))
+    {
+       const TArray<FName> ChildrenToMove = GetDirectChildrenOverrideSets(InOverrideSetName);
+       const FName GrandParent = GetParentOverrideSet(InOverrideSetName);
+       for (const FName& ChildName : ChildrenToMove)
+       {
+          Asset->OverrideSets[ChildName].ParentName = GrandParent;
+       }
+    }
+
+    SetToReparent->ParentName = InParentName;
+	PostEditOverrideSets();
+    
+    return true;
+}
+
+FName UIKRetargeterController::GetParentOverrideSet(FName InOverrideSetName)
+{
+    const FRetargetOverrideSet* OverrideSet = Asset->OverrideSets.Find(InOverrideSetName);
+    return OverrideSet ? OverrideSet->ParentName : NAME_None;
+}
+
+FName UIKRetargeterController::GetUniqueOverrideSetName(const FName InNameToMakeUnique) const
+{
+    auto NameInUse = [this](const FName& Name) { return Asset->OverrideSets.Contains(Name); };
+
+    if (!NameInUse(InNameToMakeUnique))
+    {
+       return InNameToMakeUnique;
+    }
+    
+    int32 Number = InNameToMakeUnique.GetNumber() + 1;
+    while(NameInUse(FName(InNameToMakeUnique, Number)))
+    {
+       Number++;
+    }
+
+    return FName(InNameToMakeUnique, Number);
+}
+
+const TMap<FName, FRetargetOverrideSet>& UIKRetargeterController::GetAllOverrideSets() const
+{
+    return Asset->OverrideSets;
+}
+
+FRetargetOverrideSet* UIKRetargeterController::FindOverrideSet(const FName InOverrideSetName) const
+{
+    return Asset->OverrideSets.Find(InOverrideSetName);
+}
+
+const TArray<FName> UIKRetargeterController::GetAllChildrenOverrideSets(const FName InOverrideSetName) const
+{
+    TArray<FName> AllChildren;
+	
+    TArray<FName> Stack;
+    Stack.Push(InOverrideSetName);
+    while (!Stack.IsEmpty())
+    {
+       FName CurrentParent = Stack.Pop(EAllowShrinking::No);
+       for (const TTuple<FName, FRetargetOverrideSet>& OverrideSetPair : Asset->OverrideSets)
+       {
+          if (OverrideSetPair.Value.ParentName == CurrentParent)
+          {
+             AllChildren.Add(OverrideSetPair.Key);
+             Stack.Push(OverrideSetPair.Key);
+          }
+       }
+    }
+
+    return AllChildren;
+}
+
+const TArray<FName> UIKRetargeterController::GetDirectChildrenOverrideSets(const FName InOverrideSetName) const
+{
+    TArray<FName> Children;
+    for (const TTuple<FName, FRetargetOverrideSet>& OverrideSetPair : Asset->OverrideSets)
+    {
+       if (OverrideSetPair.Value.ParentName == InOverrideSetName)
+       {
+          Children.Add(OverrideSetPair.Key);
+       }
+    }
+    return Children;
+}
+
+FRetargetOpOverrides* UIKRetargeterController::FindOpOverrides(const FName InOverrideSetName, const FName InOpName) const
+{
+    FRetargetOverrideSet* OverrideSet = FindOverrideSet(InOverrideSetName);
+    if (!OverrideSet)
+    {
+	    return nullptr;
+    }
+
+    for (FRetargetOpOverrides& OpOverrides : OverrideSet->OpOverrides)
+    {
+       if (OpOverrides.OpName == InOpName)
+       {
+          return &OpOverrides;
+       }
+    }
+    return nullptr;
+}
+
+FRetargetOpOverrides* UIKRetargeterController::AddOpOverrides(const FName InOverrideSetName, const FName InOpName) const
+{
+    FRetargetOverrideSet* OverrideSet = FindOverrideSet(InOverrideSetName);
+    if (!OverrideSet)
+    {
+	    return nullptr;
+    }
+
+    if (FRetargetOpOverrides* Existing = FindOpOverrides(InOverrideSetName, InOpName))
+    {
+        return Existing;
+    }
+
+    FIKRetargetOpBase* Op = GetRetargetOpByName(InOpName);
+    if (!Op)
+    {
+	    return nullptr;
+    }
+
+    constexpr bool bShouldTransact = true;
+    FScopedTransaction Transaction(LOCTEXT("AddOpOverrides", "Added Op Overrides to Set"), bShouldTransact);
+    Asset->Modify();
+
+    FRetargetOpOverrides& NewOpOverrides = OverrideSet->OpOverrides.AddDefaulted_GetRef();
+    NewOpOverrides.OpName = InOpName;
+    NewOpOverrides.ScriptStruct = const_cast<UScriptStruct*>(Op->GetSettingsType());
+
+	PostEditOverrideSets();
+    return &NewOpOverrides;
+}
+
+bool UIKRetargeterController::RemoveOpOverrides(const FName InOverrideSetName, const FName InOpName) const
+{
+    FRetargetOverrideSet* OverrideSet = FindOverrideSet(InOverrideSetName);
+    if (!OverrideSet)
+    {
+	    return false;
+    }
+
+    int32 Index = OverrideSet->OpOverrides.IndexOfByPredicate([&](const FRetargetOpOverrides& O) { return O.OpName == InOpName; });
+    if (Index == INDEX_NONE)
+    {
+	    return false;
+    }
+
+    constexpr bool bShouldTransact = true;
+    FScopedTransaction Transaction(LOCTEXT("RemoveOpOverrides", "Removed Op Overrides from Set"), bShouldTransact);
+    Asset->Modify();
+    
+    OverrideSet->OpOverrides.RemoveAt(Index);
+
+	PostEditOverrideSets();
+    return true;
+}
+
+bool UIKRetargeterController::UpdateOverrideValue(
+	const FName InOverrideSetName,
+	const FName InOpName,
+	const FString& InPropertyPath,
+	const FStructOnScope& InNewValue) const
+{
+	FRetargetOpOverrides* OpOverrides = FindOpOverrides(InOverrideSetName, InOpName);
+	if (!OpOverrides)
+	{
+		return false;
+	}
+
+	FRetargetOpPropertyOverride* PropertyOverride = OpOverrides->FindPropertyOverride(InPropertyPath);
+	if (!PropertyOverride)
+	{
+		return false;
+	}
+	
+	constexpr bool bShouldTransact = true;
+	FScopedTransaction Transaction(LOCTEXT("UpdateOverrideValue", "Edit Override Value"), bShouldTransact);
+	Asset->Modify();
+
+	const UStruct* BaseStruct = InNewValue.GetStruct();
+	const UScriptStruct* ScriptStruct = CastChecked<UScriptStruct>(BaseStruct);
+	OpOverrides->UpdateOverrideValue(*PropertyOverride, InNewValue.GetStructMemory(), ScriptStruct);
+	
+	return true;
+}
+
+bool UIKRetargeterController::AddPropertyOverrideToOp(const FName InOverrideSetName, const FName InOpName, const FString& InPropertyPath) const
+{
+    FIKRetargetOpBase* Op = GetRetargetOpByName(InOpName);
+    if (!Op)
+    {
+	    return false;
+    }
+    
+    FRetargetOpOverrides* OpOverrides = AddOpOverrides(InOverrideSetName, InOpName);
+    if (!OpOverrides)
+    {
+	    return false;
+    }
+
+    const uint8* SettingsPtr = reinterpret_cast<const uint8*>(Op->GetSettings());
+    const UScriptStruct* SettingsStruct = Op->GetSettingsType();
+
+    constexpr bool bShouldTransact = true;
+    FScopedTransaction Transaction(LOCTEXT("AddPropertyOverride", "Added Property Override"), bShouldTransact);
+    Asset->Modify();
+    
+    const bool bWasAdded = OpOverrides->AddPropertyOverride(InPropertyPath, SettingsPtr, SettingsStruct, false);
+
+	PostEditOverrideSets();
+    return bWasAdded;
+}
+
+bool UIKRetargeterController::RemovePropertyOverrideFromOp(const FName InOverrideSetName, const FName InOpName, const FString& InPropertyPath) const
+{
+    FRetargetOpOverrides* OpOverrides = FindOpOverrides(InOverrideSetName, InOpName);
+    if (!OpOverrides)
+    {
+	    return false;
+    }
+
+    constexpr bool bShouldTransact = true;
+    FScopedTransaction Transaction(LOCTEXT("RemovePropertyOverride", "Removed Property Override"), bShouldTransact);
+    Asset->Modify();
+    
+    const bool bWasRemoved = OpOverrides->RemovePropertyOverride(InPropertyPath);
+
+    // clean up Op entry if no overrides remain
+    if (OpOverrides->GetNumPropertyOverrides() == 0)
+    {
+		RemoveOpOverrides(InOverrideSetName, InOpName);
+    }
+
+	PostEditOverrideSets();
+    return bWasRemoved;
+}
+
+bool UIKRetargeterController::RemoveAllPropertyOverridesFromSet(const FName InOverrideSetName) const
+{
+    FRetargetOverrideSet* OverrideSet = FindOverrideSet(InOverrideSetName);
+    if (!OverrideSet)
+    {
+	    return false;
+    }
+
+    constexpr bool bShouldTransact = true;
+    FScopedTransaction Transaction(LOCTEXT("ClearOverrideSet", "Clear Override Set"), bShouldTransact);
+    Asset->Modify();
+
+    OverrideSet->OpOverrides.Empty();
+    
+	PostEditOverrideSets();
+    return true;
+}
+
+bool UIKRetargeterController::SetOverrideSetActiveByDefault(const FName InOverrideSetName, const bool bActiveByDefault) const
+{
+    FRetargetOverrideSet* OverrideSet = FindOverrideSet(InOverrideSetName);
+    if (!OverrideSet)
+    {
+	    return false;
+    }
+
+    constexpr bool bShouldTransact = true;
+    FScopedTransaction Transaction(LOCTEXT("SetOverrideSetActive", "Set Override Set Active By Default"), bShouldTransact);
+    Asset->Modify();
+
+    OverrideSet->bActiveByDefault = bActiveByDefault;
+    
+	PostEditOverrideSets();
+    return true;
+}
+
+bool UIKRetargeterController::GetOverrideSetActiveByDefault(const FName InOverrideSetName) const
+{
+    FRetargetOverrideSet* OverrideSet = FindOverrideSet(InOverrideSetName);
+    return OverrideSet ? OverrideSet->bActiveByDefault : false;
+}
+
+bool UIKRetargeterController::OverrideAllProperties(const FName InOverrideSetName, const FName InOpName) const
+{
+    FIKRetargetOpBase* Op = GetRetargetOpByName(InOpName);
+    if (!Op)
+    {
+	    return false;
+    }
+
+    TSharedPtr<FRetargetPropertyOverrideNode> RootNode;
+    Op->GetOverrideableProperties(RootNode);
+    if (!RootNode.IsValid())
+    {
+	    return false;
+    }
+
+    constexpr bool bShouldTransact = true;
+    FScopedTransaction Transaction(LOCTEXT("OverrideAllProperties", "Override All Properties"), bShouldTransact);
+    Asset->Modify();
+    
+    TArray<TSharedPtr<FRetargetPropertyOverrideNode>> NodeStack;
+    NodeStack.Push(RootNode);
+    while (!NodeStack.IsEmpty())
+	{
+		TSharedPtr<FRetargetPropertyOverrideNode> CurrentNode = NodeStack.Pop();
+		for (const FRetargetPropertyOverrideNode::FRetargetPropertyOverrideEntry& Leaf : CurrentNode->LeafProperties)
+		{
+			AddPropertyOverrideToOp(InOverrideSetName, InOpName, Leaf.PropertyPath);
+		}
+		for (TTuple<FName, TSharedPtr<FRetargetPropertyOverrideNode>>& It : CurrentNode->SubNodes)
+		{
+			NodeStack.Push(It.Value);
+		}
+	}
+
+	PostEditOverrideSets();
+    return true;
+}
+
+bool UIKRetargeterController::OverrideAllEditedProperties(const FName InOverrideSetName, const FName InOpName) const
+{
+    FIKRetargetOpBase* Op = GetRetargetOpByName(InOpName);
+    if (!Op)
+    {
+	    return false;
+    }
+
+    TSharedPtr<FRetargetPropertyOverrideNode> RootNode;
+    Op->GetOverrideableProperties(RootNode);
+    if (!RootNode.IsValid())
+    {
+	    return false;
+    }
+
+    const UScriptStruct* SettingsStruct = Op->GetSettingsType();
+    uint8* SettingsData = reinterpret_cast<uint8*>(Op->GetSettings());
+
+    constexpr bool bShouldTransact = true;
+    FScopedTransaction Transaction(LOCTEXT("OverrideAllEditedProperties", "Override All Edited Properties"), bShouldTransact);
+    Asset->Modify();
+    
+    TArray<TSharedPtr<FRetargetPropertyOverrideNode>> NodeStack;
+    NodeStack.Push(RootNode);
+	while (!NodeStack.IsEmpty())
+	{
+		TSharedPtr<FRetargetPropertyOverrideNode> CurrentNode = NodeStack.Pop();
+		for (const FRetargetPropertyOverrideNode::FRetargetPropertyOverrideEntry& Leaf : CurrentNode->LeafProperties)
+		{
+			if (!FRetargetOpPropertyOverride::IsStructPropertyAtDefault(SettingsStruct, SettingsData, Leaf.PropertyPath))
+			{
+				AddPropertyOverrideToOp(InOverrideSetName, InOpName, Leaf.PropertyPath);    
+			}
+		}
+		for (TTuple<FName, TSharedPtr<FRetargetPropertyOverrideNode>>& It : CurrentNode->SubNodes)
+		{
+			NodeStack.Push(It.Value);
+		}
+	}
+
+	PostEditOverrideSets();
+    return true;
+}
+
+bool UIKRetargeterController::HasPropertyOverride(const FName InOverrideSetName, const FName InOpName, const FString& InPropertyPath) const
+{
+    FRetargetOpOverrides* OpOverrides = FindOpOverrides(InOverrideSetName, InOpName);
+    return OpOverrides ? OpOverrides->HasPropertyOverride(InPropertyPath) : false;
+}
+
+FName UIKRetargeterController::GetOverrideSetName(const FRetargetOverrideSet* InOverrideSet) const
+{
+    if (!InOverrideSet)
+    {
+	    return NAME_None;
+    }
+    
+    for (const TTuple<FName, FRetargetOverrideSet>& It : Asset->OverrideSets)
+    {
+       if (&It.Value == InOverrideSet)
+       {
+	       return It.Key;
+       }
+    }
+    return NAME_None;
+}
+
+void UIKRetargeterController::PostEditOverrideSets() const
+{
+	// informs any attached runtime processors to rebuild their override cache
+	Asset->OverridesVersion++;
+	// informs editor to refresh the override UI
+	OverrideSetsModified.Broadcast();
+}
+
 //
 // BEGIN DEPRECATED API
 //
diff --git a/Engine/Plugins/Animation/IKRig/Source/IKRigEditor/Private/RetargetEditor/IKRetargeterPoseGenerator.cpp b/Engine/Plugins/Animation/IKRig/Source/IKRigEditor/Private/RetargetEditor/IKRetargeterPoseGenerator.cpp
--- a/Engine/Plugins/Animation/IKRig/Source/IKRigEditor/Private/RetargetEditor/IKRetargeterPoseGenerator.cpp
+++ b/Engine/Plugins/Animation/IKRig/Source/IKRigEditor/Private/RetargetEditor/IKRetargeterPoseGenerator.cpp
@@ -632,7 +632,7 @@ bool FRetargetAutoPoseGenerator::CheckReadyToAlignBones() const
 	constexpr bool bSuppressWarnings = true;
 	FRetargetProfile RetargetProfile;
 	RetargetProfile.FillProfileWithAssetSettings(RetargeterAsset);
-	Processor->Initialize(SourceSkeletalMesh, TargetSkeletalMesh, RetargeterAsset, RetargetProfile, bSuppressWarnings);
+	Processor->Initialize(SourceSkeletalMesh, TargetSkeletalMesh, RetargeterAsset, RetargetProfile, TArray<FName>(), bSuppressWarnings);
 	
 	// can't auto align until processor is initialized because we need a fully resolved chain mapping on source/target
 	// this could happen if the user attempts to edit the retarget pose on an incompatible IK Rig
diff --git a/Engine/Plugins/Animation/IKRig/Source/IKRigEditor/Private/RetargetEditor/SIKRetargetOverrideManager.cpp b/Engine/Plugins/Animation/IKRig/Source/IKRigEditor/Private/RetargetEditor/SIKRetargetOverrideManager.cpp
new file mode 100644
--- /dev/null
+++ b/Engine/Plugins/Animation/IKRig/Source/IKRigEditor/Private/RetargetEditor/SIKRetargetOverrideManager.cpp
@@ -0,0 +1,785 @@
+﻿// Copyright Epic Games, Inc. All Rights Reserved.
+
+#include "RetargetEditor/SIKRetargetOverrideManager.h"
+
+#include "RetargetEditor/IKRetargetEditorController.h"
+#include "Preferences/PersonaOptions.h"
+#include "Widgets/Input/SSegmentedControl.h"
+#include "Retargeter/IKRetargetProcessor.h"
+#include "RigEditor/IKRigEditorStyle.h"
+#include "SPositiveActionButton.h"
+#include "GeometryCollection/Facades/CollectionConstraintOverrideFacade.h"
+#include "Misc/TextFilterExpressionEvaluator.h"
+#include "Net/ReplayPlaylistTracker.h"
+#include "RetargetEditor/IKRetargetEditorStyle.h"
+#include "RigEditor/IKRigStructViewer.h"
+#include "Widgets/Input/SSearchBox.h"
+#include "Widgets/Text/SInlineEditableTextBlock.h"
+
+#define LOCTEXT_NAMESPACE "SIKRetargetOverrideManager"
+
+static const FName OverrideSetActiveColumn(TEXT("OverrideSetActive"));
+static const FName OverrideSetNameColumn(TEXT("OverrideSetName"));
+static const FName OverrideSetOpsColumn(TEXT("NumOps"));
+static const FName OverrideSetOverridesColumn(TEXT("NumOverrides"));
+
+FIKRetargetOverrideSetElement::FIKRetargetOverrideSetElement(
+	const FName& InName,
+	const TSharedRef<FIKRetargetEditorController>& InEditorController)
+	: Key(FText::FromName(InName)), Name(InName), EditorController(InEditorController)
+{}
+
+FText FIKRetargetOverrideSetElement::GetNumOpsLabel() const
+{
+	TSharedPtr<FIKRetargetEditorController> Controller = EditorController.Pin();
+	if (!Controller)
+	{
+	   return FText::AsNumber(0);
+	}
+
+	FRetargetOverrideSet* OverrideSet = Controller->AssetController->FindOverrideSet(Name);
+	if (!ensure(OverrideSet))
+	{
+	   return FText::AsNumber(0);
+	}
+
+	return FText::AsNumber(OverrideSet->OpOverrides.Num());
+}
+
+FText FIKRetargetOverrideSetElement::GetNumOverridesLabel() const
+{
+	TSharedPtr<FIKRetargetEditorController> Controller = EditorController.Pin();
+	if (!Controller)
+	{
+	   return FText::AsNumber(0);
+	}
+
+	FRetargetOverrideSet* OverrideSet = Controller->AssetController->FindOverrideSet(Name);
+	if (!ensure(OverrideSet))
+	{
+	   return FText::AsNumber(0);
+	}
+
+	int32 NumOverrides = 0;
+	for (const FRetargetOpOverrides& OpOverrides : OverrideSet->OpOverrides)
+	{
+	   NumOverrides += OpOverrides.GetNumPropertyOverrides();
+	}
+
+	return FText::AsNumber(NumOverrides);
+}
+
+void SIKRetargetOverrideSetRow::Construct(
+	const FArguments& InArgs,
+	const TSharedRef<STableViewBase>& InOwnerTable)
+{
+	OwnerTable = InOwnerTable;
+	WeakTreeElement = InArgs._TreeElement;
+	EditorController = InArgs._EditorController;
+	OverrideManager = InArgs._OverrideManager;
+	
+	SMultiColumnTableRow<TSharedPtr<FIKRetargetOverrideSetElement>>::Construct(
+	   FSuperRowType::FArguments()
+	   .OnDragDetected(OverrideManager.Pin().Get(), &SIKRetargetOverrideManager::OnDragDetected)
+	   .OnCanAcceptDrop(OverrideManager.Pin().Get(), &SIKRetargetOverrideManager::OnCanAcceptDrop)
+	   .OnAcceptDrop(OverrideManager.Pin().Get(), &SIKRetargetOverrideManager::OnAcceptDrop), InOwnerTable);
+}
+
+TSharedRef<SWidget> SIKRetargetOverrideSetRow::GenerateWidgetForColumn(const FName& ColumnName)
+{
+	const FSlateBrush* Brush = FIKRetargetEditorStyle::Get().GetBrush("IKRetarget.PostSettingsSmall");
+
+	FTextBlockStyle NormalText = FIKRigEditorStyle::Get().GetWidgetStyle<FTextBlockStyle>("IKRig.Tree.NormalText");
+	FSlateFontInfo TextFont = NormalText.Font;
+	FSlateColor TextColor = NormalText.ColorAndOpacity;
+
+	if (ColumnName == OverrideSetActiveColumn)
+	{
+		return SNew(SBox)
+		.HAlign(HAlign_Center)
+		.VAlign(VAlign_Center)
+		[
+			SNew(SCheckBox)
+			.IsChecked(this, &SIKRetargetOverrideSetRow::IsOverrideSetActive)
+			.OnCheckStateChanged(this, &SIKRetargetOverrideSetRow::OnOverrideSetActiveChanged)
+			.ToolTipText(LOCTEXT("OverrideSetActiveTooltip", "Toggle if this override set is active in the current stack."))
+		];
+	}
+
+	if (ColumnName == OverrideSetNameColumn)
+	{
+		TSharedPtr< SHorizontalBox > RowBox;
+		SAssignNew(RowBox, SHorizontalBox)
+		+ SHorizontalBox::Slot()
+		.AutoWidth()
+		[
+			SNew( SExpanderArrow, SharedThis(this) )
+			.ShouldDrawWires(true)
+		];
+
+		RowBox->AddSlot()
+		.MaxWidth(18)
+		.FillWidth(1.0)
+		.HAlign(HAlign_Left)
+		.VAlign(VAlign_Center)
+		[
+			SNew(SImage)
+			.Image(Brush)
+		];
+
+		RowBox->AddSlot()
+		.AutoWidth()
+		.VAlign(VAlign_Center)
+		[
+			SAssignNew(InlineTextBlock, SInlineEditableTextBlock)
+			.Text(this, &SIKRetargetOverrideSetRow::GetItemText)
+			.OnVerifyTextChanged(this, &SIKRetargetOverrideSetRow::OnVerifyTextChanged)
+			.OnTextCommitted(this, &SIKRetargetOverrideSetRow::OnTextCommitted)
+		];
+
+		return RowBox.ToSharedRef();
+	}
+
+	if (ColumnName == OverrideSetOpsColumn)
+	{
+		TWeakPtr<FIKRetargetOverrideSetElement> TreeElementPtr = WeakTreeElement;
+		
+		TSharedPtr< SHorizontalBox > RowBox;
+		SAssignNew(RowBox, SHorizontalBox)
+		+ SHorizontalBox::Slot()
+		.AutoWidth()
+		[
+			SNew(STextBlock)
+			.Text_Lambda([TreeElementPtr](){ return TreeElementPtr.Pin()->GetNumOpsLabel(); })
+			.Font(TextFont)
+			.ColorAndOpacity(TextColor)
+		];
+
+		return RowBox.ToSharedRef();
+	}
+
+	if (ColumnName == OverrideSetOverridesColumn)
+	{
+		TWeakPtr<FIKRetargetOverrideSetElement> TreeElementPtr = WeakTreeElement;
+		
+		TSharedPtr< SHorizontalBox > RowBox;
+		SAssignNew(RowBox, SHorizontalBox)
+		+ SHorizontalBox::Slot()
+		.AutoWidth()
+		[
+			SNew(STextBlock)
+			.Text_Lambda([TreeElementPtr]()
+				{
+					if (TreeElementPtr.IsValid())
+					{
+						return TreeElementPtr.Pin()->GetNumOverridesLabel();	
+					}
+					return FText::GetEmpty();
+				})
+			.Font(TextFont)
+			.ColorAndOpacity(TextColor)
+		];
+
+		return RowBox.ToSharedRef();
+	}
+
+	return SNullWidget::NullWidget;
+}
+
+void SIKRetargetOverrideSetRow::OnRequestRename()
+{
+	InlineTextBlock.Get()->EnterEditingMode();
+}
+
+FText SIKRetargetOverrideSetRow::GetItemText() const
+{
+	return FText::FromName(WeakTreeElement.Pin()->Name);
+}
+
+bool SIKRetargetOverrideSetRow::OnVerifyTextChanged(const FText& InNewText, FText& OutErrorMessage)
+{
+	if (InNewText.IsEmpty())
+	{
+	   OutErrorMessage = LOCTEXT("OverrideSetNameEmpty", "New override set name cannot be empty.");
+	   return false;
+	}
+
+	FName NewName = FName(*InNewText.ToString());
+	FRetargetOverrideSet* ExistingSet = EditorController.Pin()->AssetController->FindOverrideSet(NewName);
+	if (ExistingSet)
+	{
+	   OutErrorMessage = LOCTEXT("OverrideSetNameNotUnique", "New override set name already in use.");
+	   return false;
+	}
+
+	return true;
+}
+
+void SIKRetargetOverrideSetRow::OnTextCommitted(const FText& InText, ETextCommit::Type CommitInfo)
+{
+	const FName OldName = WeakTreeElement.Pin()->Name;
+	const FName NewName = FName(*InText.ToString());
+	EditorController.Pin()->AssetController->RenameRetargetOverrideSet(OldName, NewName);
+}
+
+ECheckBoxState SIKRetargetOverrideSetRow::IsOverrideSetActive() const
+{
+	if (TSharedPtr<FIKRetargetOverrideSetElement> Item = WeakTreeElement.Pin())
+	{
+		TObjectPtr<UIKRetargeterController> Controller = EditorController.Pin()->AssetController;
+		const bool bIsActiveByDefault = Controller->GetOverrideSetActiveByDefault(Item->Name);
+		return bIsActiveByDefault ? ECheckBoxState::Checked : ECheckBoxState::Unchecked;
+	}
+	return ECheckBoxState::Unchecked;
+}
+
+void SIKRetargetOverrideSetRow::OnOverrideSetActiveChanged(ECheckBoxState NewState) const
+{
+	if (TSharedPtr<FIKRetargetOverrideSetElement> Item = WeakTreeElement.Pin())
+	{
+		TObjectPtr<UIKRetargeterController> Controller = EditorController.Pin()->AssetController;
+		const bool bNewState = NewState == ECheckBoxState::Checked;
+		Controller->SetOverrideSetActiveByDefault(Item->Name, bNewState);
+	}
+}
+
+TSharedRef<FRetargetOverrideSetDragDropOp> FRetargetOverrideSetDragDropOp::New(TWeakPtr<FIKRetargetOverrideSetElement> InElement)
+{
+	TSharedRef<FRetargetOverrideSetDragDropOp> Operation = MakeShared<FRetargetOverrideSetDragDropOp>();
+	Operation->Element = InElement;
+	Operation->Construct();
+	return Operation;
+}
+
+TSharedPtr<SWidget> FRetargetOverrideSetDragDropOp::GetDefaultDecorator() const
+{
+	return FDecoratedDragDropOp::GetDefaultDecorator();
+}
+
+void SIKRetargetOverrideManager::Construct(
+	const FArguments& InArgs,
+	TSharedRef<FIKRetargetEditorController> InEditorController)
+{
+	EditorController = InEditorController;
+	EditorController.Pin()->SetOverrideSetsView(SharedThis(this));
+	
+	TextFilter = MakeShareable(new FTextFilterExpressionEvaluator(ETextFilterExpressionEvaluatorMode::BasicString));
+
+	FDetailsViewArgs DetailsViewArgs;
+	DetailsViewArgs.bAllowSearch = false;
+	DetailsViewArgs.bLockable = false;
+	DetailsViewArgs.bUpdatesFromSelection = false;
+	DetailsViewArgs.bAllowMultipleTopLevelObjects = false;
+	DetailsViewArgs.NameAreaSettings = FDetailsViewArgs::ENameAreaSettings::HideNameArea;
+	FPropertyEditorModule& PropertyModule = FModuleManager::LoadModuleChecked<FPropertyEditorModule>("PropertyEditor");
+	DetailsView = PropertyModule.CreateDetailView(DetailsViewArgs);
+	
+	ChildSlot
+	[
+		SNew(SVerticalBox)
+
+		+ SVerticalBox::Slot()
+		.AutoHeight()
+		.VAlign(VAlign_Top)
+		.HAlign(HAlign_Fill)
+		[
+			SNew(SHorizontalBox)
+			
+			+SHorizontalBox::Slot()
+			[
+				SNew(SPositiveActionButton)
+				.Icon(FAppStyle::Get().GetBrush("Icons.Plus"))
+				.Text(LOCTEXT("AddRetargetOverrideSetLabel", "Add Override Set"))
+				.ToolTipText(LOCTEXT("AddOverrideSetToolTip", "Create a new retarget override set."))
+				.OnClicked_Lambda([this]()
+				{
+					const FName NewOverrideSet = FName("NewRetargetOverrideSet");
+					EditorController.Pin()->AssetController->AddNewRetargetOverrideSet(NewOverrideSet);
+					return FReply::Handled();
+				})
+			]
+			
+			+SHorizontalBox::Slot()
+			.FillWidth(1.0f)
+			[
+				SNew(SSearchBox)
+				.SelectAllTextWhenFocused(true)
+				.OnTextChanged( this, &SIKRetargetOverrideManager::OnFilterTextChanged )
+				.HintText( LOCTEXT( "SearchBoxHint", "Filter Hierarchy Tree...") )
+			]
+		]
+
+		+SVerticalBox::Slot()
+		.Padding(2.0f)
+		[
+			SNew(SSplitter)
+			.Orientation(Orient_Vertical)
+			
+			+ SSplitter::Slot()
+			.Value(0.3)
+			[
+				SNew(SBorder)
+				.Padding(2.0f)
+				.BorderImage(FAppStyle::GetBrush("SCSEditor.TreePanel"))
+				[
+					SAssignNew(TreeView, SIKRetargetOverrideSetTreeView)
+					.TreeItemsSource(&RootElements)
+					.SelectionMode(ESelectionMode::Multi)
+					.OnGenerateRow_Lambda( [this](
+						TSharedPtr<FIKRetargetOverrideSetElement> InItem,
+						const TSharedRef<STableViewBase>& OwnerTable) -> TSharedRef<ITableRow>
+					{
+						return SNew(SIKRetargetOverrideSetRow, OwnerTable)
+						.EditorController(EditorController.Pin())
+						.TreeElement(InItem)
+						.OverrideManager(SharedThis(this));
+					})
+					.OnGetChildren(this, &SIKRetargetOverrideManager::HandleGetChildrenForTree)
+					.OnSelectionChanged(this, &SIKRetargetOverrideManager::OnSelectionChanged)
+					.OnMouseButtonDoubleClick(this, &SIKRetargetOverrideManager::OnItemDoubleClicked)
+					.OnSetExpansionRecursive(this, &SIKRetargetOverrideManager::OnSetExpansionRecursive)
+					.OnContextMenuOpening(this, &SIKRetargetOverrideManager::OnGenerateContextMenu)
+					.HighlightParentNodesForSelection(false)
+					.HeaderRow
+					(
+						SNew(SHeaderRow)
+
+						+ SHeaderRow::Column(OverrideSetActiveColumn)
+						.DefaultLabel(LOCTEXT("OverrideSetActiveLabel", "Active"))
+						.FillWidth(0.1f)
+						
+						+ SHeaderRow::Column(OverrideSetNameColumn)
+						.DefaultLabel(LOCTEXT("OverrideSetNameLabel", "Override Set Name"))
+						.FillWidth(0.7f)
+
+						+ SHeaderRow::Column(OverrideSetOpsColumn)
+						.DefaultLabel(LOCTEXT("NumOpsLabel", "Ops"))
+						.FillWidth(0.1f)
+						
+						+ SHeaderRow::Column(OverrideSetOverridesColumn)
+						.DefaultLabel(LOCTEXT("NumOverridesLabel", "Properties"))
+						.FillWidth(0.1f)
+					)
+				]
+			]
+			
+			+ SSplitter::Slot()
+			.Value(0.7f)
+			[
+				SNew(SVerticalBox)
+				+ SVerticalBox::Slot()
+				.FillHeight(1.0f)
+				[
+					DetailsView.ToSharedRef()
+				]
+			]
+		]
+	];
+
+	constexpr bool IsInitialSetup = true;
+	RefreshTreeView(IsInitialSetup);
+}
+
+FName SIKRetargetOverrideManager::GetSelectedOverrideSetName() const
+{
+	TArray<TSharedPtr<FIKRetargetOverrideSetElement>> SelectedItems = TreeView->GetSelectedItems();
+	if (SelectedItems.IsEmpty())
+	{
+	   return NAME_None;
+	}
+
+	return SelectedItems[0]->Name;
+}
+
+void SIKRetargetOverrideManager::OnFilterTextChanged(const FText& SearchText)
+{
+	TextFilter->SetFilterText(SearchText);
+	RefreshTreeView(false);
+}
+
+void SIKRetargetOverrideManager::RefreshTreeView(bool IsInitialSetup)
+{
+	const TSharedPtr<FIKRetargetEditorController> Controller = EditorController.Pin();
+	if (!Controller.IsValid())
+	{
+		return;
+	}
+
+	TreeView->SaveAndClearState();
+
+	RootElements.Reset();
+	AllElements.Reset();
+
+	const TMap<FName, FRetargetOverrideSet>& AllOverrideSets = Controller->AssetController->GetAllOverrideSets();
+	
+	TMap<FName, TSharedPtr<FIKRetargetOverrideSetElement>> SetToElementMap;
+
+	const bool bHasTextFilter = !TextFilter->GetFilterText().IsEmpty();
+	auto PassesTextFilter = [&](const FName& NameToCheck) -> bool
+	{
+		return TextFilter->TestTextFilter(FBasicStringFilterExpressionContext(NameToCheck.ToString()));
+	};
+
+	for (const auto& Pair : AllOverrideSets)
+	{
+		const FName& SetName = Pair.Key;
+		const FRetargetOverrideSet& SetData = Pair.Value;
+
+		TSharedPtr<FIKRetargetOverrideSetElement> NewElement = MakeShared<FIKRetargetOverrideSetElement>(SetName, Controller.ToSharedRef());
+		
+		AllElements.Add(NewElement);
+		SetToElementMap.Add(SetName, NewElement);
+
+		NewElement->bIsHidden= false;
+		if (FilterOptions.bHideEmptySets && SetData.OpOverrides.IsEmpty())
+		{
+			NewElement->bIsHidden = true;
+		}
+		else if (bHasTextFilter && !PassesTextFilter(SetName))
+		{
+			NewElement->bIsHidden = true;
+		}
+	}
+
+	for (TSharedPtr<FIKRetargetOverrideSetElement>& CurrentElement : AllElements)
+	{
+		if (CurrentElement->bIsHidden)
+		{
+			continue;
+		}
+
+		const FRetargetOverrideSet* SetData = AllOverrideSets.Find(CurrentElement->Name);
+		if (!ensure(SetData))
+		{
+			continue;
+		}
+
+		TSharedPtr<FIKRetargetOverrideSetElement> ValidParent = nullptr;
+		FName NextParentName = SetData->ParentName;
+
+		while (NextParentName != NAME_None)
+		{
+			TSharedPtr<FIKRetargetOverrideSetElement>* ParentPtr = SetToElementMap.Find(NextParentName);
+			if (!ParentPtr)
+			{
+			   break;
+			}
+			TSharedPtr<FIKRetargetOverrideSetElement> Parent = *ParentPtr;
+			
+			if (!Parent->bIsHidden)
+			{
+			   ValidParent = Parent;
+			   break;
+			}
+
+			const FRetargetOverrideSet* ParentSet = AllOverrideSets.Find(NextParentName);
+			NextParentName = ParentSet ? ParentSet->ParentName : NAME_None;
+		}
+
+		if (ValidParent.IsValid())
+		{
+			ValidParent->Children.Add(CurrentElement);
+			CurrentElement->Parent = ValidParent;
+		}
+		else
+		{
+			RootElements.Add(CurrentElement);
+		}
+	}
+
+	if (IsInitialSetup)
+	{
+		for (const TSharedPtr<FIKRetargetOverrideSetElement>& Root : RootElements)
+		{
+			SetExpansionRecursive(Root, false, true);
+		}
+	}
+	else
+	{
+		for (const TSharedPtr<FIKRetargetOverrideSetElement>& Element : AllElements)
+		{
+			TreeView->RestoreState(Element);
+		}
+	}
+
+	TreeView->RequestTreeRefresh();
+}
+
+void SIKRetargetOverrideManager::HandleGetChildrenForTree(
+	TSharedPtr<FIKRetargetOverrideSetElement> InItem,
+	TArray<TSharedPtr<FIKRetargetOverrideSetElement>>& OutChildren)
+{
+	for (TWeakPtr<FIKRetargetOverrideSetElement>& Child : InItem->Children)
+	{
+	   OutChildren.Add(Child.Pin());
+	}
+}
+
+void SIKRetargetOverrideManager::OnSelectionChanged(
+	TSharedPtr<FIKRetargetOverrideSetElement> Selection,
+	ESelectInfo::Type SelectInfo)
+{
+	TArray<TSharedPtr<FIKRetargetOverrideSetElement>> SelectedItems = TreeView->GetSelectedItems();
+	if (SelectedItems.IsEmpty())
+	{
+	   DetailsView->SetObject(nullptr, true /*bForceRefresh*/);
+	   return;
+	}
+
+	const TObjectPtr<UIKRetargeterController> AssetController = EditorController.Pin()->AssetController;
+
+	const FName SetToViewName = SelectedItems[0].Get()->Name;
+	FRetargetOverrideSet* SetToView = AssetController->FindOverrideSet(SetToViewName);
+	if (!SetToView)
+	{
+	   DetailsView->SetObject(nullptr, true /*bForceRefresh*/);
+	   return;
+	}
+
+	UIKRetargeter* Asset = AssetController->GetAsset();
+	auto MemoryProvider = [Asset, SetToViewName]() -> uint8*
+	   {
+		  UIKRetargeterController* Controller = UIKRetargeterController::GetController(Asset);
+		  if (FRetargetOverrideSet* SetToEdit = Controller->FindOverrideSet(SetToViewName))
+		  {
+			 return reinterpret_cast<uint8*>(SetToEdit);
+		  }
+	   
+		  return nullptr;
+	   };
+	
+	FIKRigStructToView StructToView;
+	StructToView.Type = FRetargetOverrideSet::StaticStruct();
+	StructToView.MemoryProvider = MemoryProvider;
+	StructToView.Owner = AssetController->GetAsset();
+	StructToView.UniqueName = SetToViewName;
+	
+	UIKRigStructViewer* StructViewer = AssetController->GetStructViewer();
+	StructViewer->SetStructToView(StructToView);
+	
+	constexpr bool bForceRefresh = true;
+	DetailsView->SetObject(StructViewer, bForceRefresh);
+}
+
+void SIKRetargetOverrideManager::OnItemDoubleClicked(TSharedPtr<FIKRetargetOverrideSetElement> InItem)
+{
+	if (TreeView->IsItemExpanded(InItem))
+	{
+	   SetExpansionRecursive(InItem, false, false);
+	}
+	else
+	{
+	   SetExpansionRecursive(InItem, false, true);
+	}
+}
+
+void SIKRetargetOverrideManager::OnSetExpansionRecursive(TSharedPtr<FIKRetargetOverrideSetElement> InItem, bool bShouldBeExpanded)
+{
+	SetExpansionRecursive(InItem, false, bShouldBeExpanded);
+}
+
+void SIKRetargetOverrideManager::SetExpansionRecursive(
+	TSharedPtr<FIKRetargetOverrideSetElement> InElement,
+	bool bTowardsParent,
+	bool bShouldBeExpanded)
+{
+	TreeView->SetItemExpansion(InElement, bShouldBeExpanded);
+
+	if (bTowardsParent)
+	{
+		if (InElement->Parent.IsValid())
+		{
+			SetExpansionRecursive(InElement->Parent.Pin(), bTowardsParent, bShouldBeExpanded);
+		}
+	}
+	else
+	{
+		for (int32 ChildIndex = 0; ChildIndex < InElement->Children.Num(); ++ChildIndex)
+		{
+			SetExpansionRecursive(InElement->Children[ChildIndex].Pin(), bTowardsParent, bShouldBeExpanded);
+		}
+	}
+}
+
+TSharedPtr<SWidget> SIKRetargetOverrideManager::OnGenerateContextMenu()
+{
+	constexpr bool bShouldCloseWindowAfterMenuSelection = true;
+	FMenuBuilder MenuBuilder(bShouldCloseWindowAfterMenuSelection, nullptr);
+
+	MenuBuilder.BeginSection("ItemActions", INVTEXT("Override Set Actions"));
+	{
+		MenuBuilder.AddMenuEntry(
+		LOCTEXT("RenameItem", "Rename"),
+		LOCTEXT("RenameLabel", "Rename the selected override set."),
+		FSlateIcon(FAppStyle::GetAppStyleSetName(), "GenericCommands.Rename"),
+		FUIAction(
+		FExecuteAction::CreateLambda([this]()
+		{
+			TArray<TSharedPtr<FIKRetargetOverrideSetElement>> Selection = TreeView->GetSelectedItems();
+			if (Selection.IsEmpty())
+			{
+				return;
+			}
+
+			TSharedPtr<ITableRow> RowWidget = TreeView->WidgetFromItem(Selection[0]);
+			if (RowWidget.IsValid())
+			{
+				TSharedPtr<SIKRetargetOverrideSetRow> SetRow = StaticCastSharedRef<SIKRetargetOverrideSetRow>(RowWidget->AsWidget());
+				if (SetRow.IsValid())
+				{
+					SetRow->OnRequestRename();
+				}
+			}
+		}),
+		FCanExecuteAction::CreateLambda([this]()
+		{
+			return TreeView.IsValid() && TreeView->GetNumItemsSelected() == 1;
+		})));
+
+		MenuBuilder.AddMenuEntry(
+		LOCTEXT("UnparentItem", "Unparent"),
+		LOCTEXT("UnparentLabel", "Unparent the selected override set."),
+		FSlateIcon(FAppStyle::GetAppStyleSetName(), "GenericCommands.Rename"),
+		FUIAction(FExecuteAction::CreateLambda([this]()
+		{
+			TArray<TSharedPtr<FIKRetargetOverrideSetElement>> Selection = TreeView->GetSelectedItems();
+			if (Selection.IsEmpty())
+			{
+				return;
+			}
+			
+			TSharedPtr<FIKRetargetOverrideSetElement>& SetElement = Selection[0];
+			EditorController.Pin()->AssetController->SetParentOverrideSet(SetElement->Name, NAME_None);
+		}),
+		FCanExecuteAction::CreateLambda([this]()
+		{
+			return TreeView.IsValid() && TreeView->GetNumItemsSelected() == 1;
+		})));
+
+		MenuBuilder.AddMenuEntry(
+		LOCTEXT("RemoveAllLabel", "Remove All Overrides"),
+		LOCTEXT("RemoveAllTip", "Remove all overrides from the selected override set."),
+		FSlateIcon(FAppStyle::GetAppStyleSetName(), "GenericCommands.SelectAll"),
+		FUIAction(FExecuteAction::CreateLambda([this]()
+		{
+			TArray<TSharedPtr<FIKRetargetOverrideSetElement>> Selection = TreeView->GetSelectedItems();
+			if (Selection.IsEmpty())
+			{
+				return;
+			}
+
+			TSharedPtr<FIKRetargetOverrideSetElement> SelectedSetElement = Selection[0];
+			if (!SelectedSetElement.IsValid())
+			{
+				return;
+			}
+			
+			const FName SetToClear = SelectedSetElement.Get()->Name;
+			EditorController.Pin()->AssetController->RemoveAllPropertyOverridesFromSet(SetToClear);
+		}),
+		FCanExecuteAction::CreateLambda([this]()
+		{
+			return TreeView.IsValid() && TreeView->GetNumItemsSelected() > 0;
+		})));
+
+		MenuBuilder.AddMenuEntry(
+		LOCTEXT("DeleteItem", "Delete"),
+		LOCTEXT("DeleteLabel", "Delete the selected override set."),
+		FSlateIcon(FAppStyle::GetAppStyleSetName(), "GenericCommands.Delete"),
+		FUIAction(FExecuteAction::CreateLambda([this]()
+		{
+			TArray<TSharedPtr<FIKRetargetOverrideSetElement>> Selection = TreeView->GetSelectedItems();
+			if (Selection.IsEmpty())
+			{
+				return;
+			}
+
+			TSharedPtr<FIKRetargetOverrideSetElement> SelectedSetElement = Selection[0];
+			if (!SelectedSetElement.IsValid())
+			{
+				return;
+			}
+			
+			const FName SetToDelete = SelectedSetElement.Get()->Name;
+			EditorController.Pin()->AssetController->RemoveRetargetOverrideSet(SetToDelete);
+		}),
+		FCanExecuteAction::CreateLambda([this]()
+		{
+			return TreeView.IsValid() && TreeView->GetNumItemsSelected() > 0;
+		})));
+	}
+	MenuBuilder.EndSection();
+
+	return MenuBuilder.MakeWidget();
+}
+
+FReply SIKRetargetOverrideManager::OnDragDetected(
+	const FGeometry& MyGeometry,
+	const FPointerEvent& MouseEvent)
+{
+	const TArray<TSharedPtr<FIKRetargetOverrideSetElement>> CurrentSelectedItems = TreeView->GetSelectedItems();
+	if (CurrentSelectedItems.Num() != 1)
+	{
+		return FReply::Unhandled();
+	}
+
+	if (MouseEvent.IsMouseButtonDown(EKeys::LeftMouseButton))
+	{
+		const TSharedPtr<FIKRetargetOverrideSetElement> DraggedElement = CurrentSelectedItems[0];
+		const TSharedRef<FRetargetOverrideSetDragDropOp> DragDropOp = FRetargetOverrideSetDragDropOp::New(DraggedElement);
+		return FReply::Handled().BeginDragDrop(DragDropOp);
+	}
+
+	return FReply::Unhandled();
+}
+
+TOptional<EItemDropZone> SIKRetargetOverrideManager::OnCanAcceptDrop(
+	const FDragDropEvent& DragDropEvent,
+	EItemDropZone DropZone,
+	TSharedPtr<FIKRetargetOverrideSetElement> TargetElement)
+{
+	TOptional<EItemDropZone> ReturnedDropZone;
+	const TSharedPtr<FRetargetOverrideSetDragDropOp> DragDropOp = DragDropEvent.GetOperationAs<FRetargetOverrideSetDragDropOp>();
+	if (!DragDropOp.IsValid())
+	{
+		return ReturnedDropZone;
+	}
+
+	const TSharedPtr<FIKRetargetOverrideSetElement>& DraggedElement = DragDropOp.Get()->Element.Pin();
+	if (!ensure(DraggedElement))
+	{
+		return ReturnedDropZone;
+	}
+
+	if (TargetElement == DraggedElement)
+	{
+		return ReturnedDropZone;
+	}
+
+	ReturnedDropZone = DropZone;
+	return ReturnedDropZone;
+}
+
+FReply SIKRetargetOverrideManager::OnAcceptDrop(
+	const FDragDropEvent& DragDropEvent,
+	EItemDropZone DropZone,
+	TSharedPtr<FIKRetargetOverrideSetElement> TargetElement)
+{
+	const TSharedPtr<FRetargetOverrideSetDragDropOp> DragDropOp = DragDropEvent.GetOperationAs<FRetargetOverrideSetDragDropOp>();
+	if (!ensure(DragDropOp.IsValid()))
+	{
+		return FReply::Unhandled();
+	}
+
+	const TSharedPtr<FIKRetargetOverrideSetElement>& DraggedElement = DragDropOp.Get()->Element.Pin();
+	if (!ensure(DraggedElement))
+	{
+		return FReply::Unhandled();
+	}
+
+	EditorController.Pin()->AssetController->SetParentOverrideSet(DraggedElement.Get()->Name, TargetElement.Get()->Name);
+
+	return FReply::Handled();
+}
+
+#undef LOCTEXT_NAMESPACE
\ No newline at end of file
diff --git a/Engine/Plugins/Animation/IKRig/Source/IKRigEditor/Private/RetargetEditor/SRetargetOpStack.cpp b/Engine/Plugins/Animation/IKRig/Source/IKRigEditor/Private/RetargetEditor/SRetargetOpStack.cpp
--- a/Engine/Plugins/Animation/IKRig/Source/IKRigEditor/Private/RetargetEditor/SRetargetOpStack.cpp
+++ b/Engine/Plugins/Animation/IKRig/Source/IKRigEditor/Private/RetargetEditor/SRetargetOpStack.cpp
@@ -8,6 +8,7 @@
 #include "RetargetEditor/IKRetargetEditorController.h"
 #include "RetargetEditor/IKRetargetEditorStyle.h"
 #include "Retargeter/IKRetargetOps.h"
+#include "Retargeter/IKRetargetProfile.h"
 #include "RigEditor/IKRigEditorStyle.h"
 #include "Widgets/Text/SInlineEditableTextBlock.h"
 
@@ -49,221 +50,247 @@ const UScriptStruct* FRetargetOpStackElement::GetType() const
 void SRetargetOpItem::Construct(const FArguments& InArgs)
 {
 	ListView = InArgs._InOpListWidget;
-    Element = InArgs._InStackElement;
+	Element = InArgs._InStackElement;
 
-    ChildSlot
-    [
-        SNew(SHorizontalBox)
-    	.ToolTipText(TAttribute<FText>::CreateLambda([this]()
-		{
-			const UScriptStruct* OpType = Element.Pin()->GetType();
-    		if (!OpType)
-    		{
-    			return FText::FromString(TEXT("Unknown Op type. Resave asset to clean."));
-    		}
-    			
-    		return FText::Format(
-			FText::FromString(TEXT("Op Type: {0} \nIndex in Stack: {1}\n")),
-			FText::FromString(OpType->GetStructCPPName()),
-			FText::AsNumber(Element.Pin()->GetIndexInStack()) );
-		}))
-	    
-        // drag Icon
-        + SHorizontalBox::Slot()
-        .MaxWidth(18)
-        .AutoWidth()
-        .HAlign(HAlign_Left)
-        .VAlign(VAlign_Center)
-        .Padding(FMargin(OpHorizontalPadding, OpVerticalPadding))
-        [
-            SNew(SImage)
-            .Image(FIKRigEditorStyle::Get().GetBrush("IKRig.DragSolver"))
-        ]
-
-        // enable Checkbox
-        + SHorizontalBox::Slot()
-        .AutoWidth()
-        .HAlign(HAlign_Left)
-        .VAlign(VAlign_Center)
-        .Padding(FMargin(OpHorizontalPadding, OpVerticalPadding))
-        [
-            SNew(SCheckBox)
-            .IsChecked_Lambda([this]() -> ECheckBoxState
-            {
-                FIKRetargetEditorController* EditorController = ListView.Pin()->EditorController.Pin().Get();
-            	bool bEnabled = EditorController->AssetController->GetRetargetOpEnabled(Element.Pin()->GetIndexInStack());
-                return bEnabled ? ECheckBoxState::Checked : ECheckBoxState::Unchecked;
-            })
-            .OnCheckStateChanged_Lambda([this](ECheckBoxState InCheckBoxState)
-            {
-            	FIKRetargetEditorController* EditorController = ListView.Pin()->EditorController.Pin().Get();
-            	const bool bIsChecked = InCheckBoxState == ECheckBoxState::Checked;
-            	EditorController->AssetController->SetRetargetOpEnabled( Element.Pin()->GetIndexInStack(), bIsChecked);
-				EditorController->ReinitializeRetargeterNoUIRefresh();
-            })
-        ]
-
-	    // toggle debug draw button
-    	+ SHorizontalBox::Slot()
-		.AutoWidth()
-		.HAlign(HAlign_Left)
-		.VAlign(VAlign_Center)
-		.Padding(FMargin(OpHorizontalPadding, OpVerticalPadding))
+	ChildSlot
+	[
+		SNew(SBorder)
+		 .BorderImage(FAppStyle::GetBrush("NoBorder")) // Invisible but hit-testable
+		 .Visibility(EVisibility::Visible)
+		 .OnMouseButtonUp_Lambda([this](const FGeometry& MyGeometry, const FPointerEvent& MouseEvent) -> FReply
+		 {
+			 if (MouseEvent.GetEffectingButton() != EKeys::RightMouseButton)
+			 {
+				return FReply::Unhandled();
+			 }
+
+			// build a menu for this op
+			FMenuBuilder MenuBuilder(true, nullptr);
+			GenerateRightClickMenu(MenuBuilder);
+
+			// summon the menu at the cursor's location
+			FSlateApplication::Get().PushMenu(
+				SharedThis(this),
+				MouseEvent.GetEventPath() != nullptr ? *MouseEvent.GetEventPath() : FWidgetPath(),
+				MenuBuilder.MakeWidget(),
+				MouseEvent.GetScreenSpacePosition(),
+				FPopupTransitionEffect(FPopupTransitionEffect::ContextMenu));
+
+			 return FReply::Handled();
+		 })
 		[
-			SNew(SButton)
-			.NormalPaddingOverride(FMargin(2, 2))
-			.PressedPaddingOverride(FMargin(2, 2))
-			.ToolTipText(LOCTEXT("ToggleOpDebugDraw", "Toggle this Ops debug drawing in the viewport."))
-			.Visibility_Lambda([this]()
+			 SNew(SHorizontalBox)
+			.ToolTipText(TAttribute<FText>::CreateLambda([this]()
 			{
-				FIKRetargetEditorController* EditorController = ListView.Pin()->EditorController.Pin().Get();
-				FIKRetargetOpBase* Op = EditorController->AssetController->GetRetargetOpByIndex(Element.Pin()->GetIndexInStack());
-				if (!ensure(Op))
+				const UScriptStruct* OpType = Element.Pin()->GetType();
+				if (!OpType)
 				{
-					return EVisibility::Collapsed;
+					return FText::FromString(TEXT("Unknown Op type. Resave asset to clean."));
 				}
-				return Op->HasDebugDrawing() ? EVisibility::Visible : EVisibility::Collapsed;
-			})
-			.OnClicked_Lambda([this]() -> FReply
-			{
-				FIKRetargetEditorController* EditorController = ListView.Pin()->EditorController.Pin().Get();
-				FIKRetargetOpBase* Op = EditorController->AssetController->GetRetargetOpByIndex(Element.Pin()->GetIndexInStack());
-				if (!ensure(Op))
+					
+				return FText::Format(
+				FText::FromString(TEXT("Op Type: {0} \nIndex in Stack: {1}\n")),
+				FText::FromString(OpType->GetStructCPPName()),
+				FText::AsNumber(Element.Pin()->GetIndexInStack()) );
+			}))
+			
+			// drag Icon
+			+ SHorizontalBox::Slot()
+			.MaxWidth(18)
+			.AutoWidth()
+			.HAlign(HAlign_Left)
+			.VAlign(VAlign_Center)
+			.Padding(FMargin(OpHorizontalPadding, OpVerticalPadding))
+			[
+				SNew(SImage)
+				.Image(FIKRigEditorStyle::Get().GetBrush("IKRig.DragSolver"))
+			]
+
+			// enable Checkbox
+			+ SHorizontalBox::Slot()
+			.AutoWidth()
+			.HAlign(HAlign_Left)
+			.VAlign(VAlign_Center)
+			.Padding(FMargin(OpHorizontalPadding, OpVerticalPadding))
+			[
+				SNew(SCheckBox)
+				.IsChecked_Lambda([this]() -> ECheckBoxState
 				{
-					return FReply::Handled();
-				}
+					FIKRetargetEditorController* EditorController = GetEditorController();
+					bool bEnabled = EditorController->AssetController->GetRetargetOpEnabled(Element.Pin()->GetIndexInStack());
+					return bEnabled ? ECheckBoxState::Checked : ECheckBoxState::Unchecked;
+				})
+				.OnCheckStateChanged_Lambda([this](ECheckBoxState InCheckBoxState)
+				{
+					FIKRetargetEditorController* EditorController = GetEditorController();
+					const bool bIsChecked = InCheckBoxState == ECheckBoxState::Checked;
+					EditorController->AssetController->SetRetargetOpEnabled( Element.Pin()->GetIndexInStack(), bIsChecked);
+					EditorController->ReinitializeRetargeterNoUIRefresh();
+				})
+			]
 
-				Op->GetSettings()->bDebugDraw = !Op->GetSettings()->bDebugDraw;
-				return FReply::Handled();
-			})
-			.Content()
+			// toggle debug draw button
+			+ SHorizontalBox::Slot()
+			.AutoWidth()
+			.HAlign(HAlign_Left)
+			.VAlign(VAlign_Center)
+			.Padding(FMargin(OpHorizontalPadding, OpVerticalPadding))
 			[
-				SNew(SImage)
-				.Image(FIKRetargetEditorStyle::Get().GetBrush("IKRetarget.Debug"))
-				.ColorAndOpacity_Lambda([this]()
+				SNew(SButton)
+				.NormalPaddingOverride(FMargin(2, 2))
+				.PressedPaddingOverride(FMargin(2, 2))
+				.ToolTipText(LOCTEXT("ToggleOpDebugDraw", "Toggle this Ops debug drawing in the viewport."))
+				.Visibility_Lambda([this]()
+				{
+					FIKRetargetEditorController* EditorController = GetEditorController();
+					FIKRetargetOpBase* Op = EditorController->AssetController->GetRetargetOpByIndex(Element.Pin()->GetIndexInStack());
+					if (!ensure(Op))
+					{
+						return EVisibility::Collapsed;
+					}
+					return Op->HasDebugDrawing() ? EVisibility::Visible : EVisibility::Collapsed;
+				})
+				.OnClicked_Lambda([this]() -> FReply
 				{
-					FIKRetargetEditorController* EditorController = ListView.Pin()->EditorController.Pin().Get();
+					FIKRetargetEditorController* EditorController = GetEditorController();
 					FIKRetargetOpBase* Op = EditorController->AssetController->GetRetargetOpByIndex(Element.Pin()->GetIndexInStack());
-					if (!Op)
+					if (!ensure(Op))
 					{
-						return FLinearColor::White;
+						return FReply::Handled();
 					}
-					return Op->GetSettings()->bDebugDraw ?FLinearColor(0,1,0,1.0) : FLinearColor::White;
+
+					Op->GetSettings()->bDebugDraw = !Op->GetSettings()->bDebugDraw;
+					return FReply::Handled();
 				})
+				.Content()
+				[
+					SNew(SImage)
+					.Image(FIKRetargetEditorStyle::Get().GetBrush("IKRetarget.Debug"))
+					.ColorAndOpacity_Lambda([this]()
+					{
+						FIKRetargetEditorController* EditorController = GetEditorController();
+						FIKRetargetOpBase* Op = EditorController->AssetController->GetRetargetOpByIndex(Element.Pin()->GetIndexInStack());
+						if (!Op)
+						{
+							return FLinearColor::White;
+						}
+						return Op->GetSettings()->bDebugDraw ?FLinearColor(0,1,0,1.0) : FLinearColor::White;
+					})
+				]
 			]
-		]
 
-        // display Name
-        + SHorizontalBox::Slot()
-        .AutoWidth()
-        .HAlign(HAlign_Left)
-        .VAlign(VAlign_Center)
-        .Padding(FMargin(OpHorizontalPadding, OpVerticalPadding))
-        [
-			SAssignNew(EditNameWidget, SInlineEditableTextBlock)
-			.Text(this, &SRetargetOpItem::GetName)
-        	.OnVerifyTextChanged_Lambda([](const FText& InText, FText& OutErrorMessage)
-			{
-        		static FString IllegalNameCharacters = "^<>:\"/\\|?*";
-				return FName::IsValidXName(InText.ToString(), FString(IllegalNameCharacters), &OutErrorMessage);
-			})
-			.OnTextCommitted(this, &SRetargetOpItem::OnNameCommitted)
-			.MultiLine(false)
-        ]
-
-        // spacer
-        + SHorizontalBox::Slot()
-        .HAlign(HAlign_Fill)
-        .FillWidth(1.0f)
-        [
-            SNew(SSpacer)
-            .Size(FVector2D(0.0f, 1.0f))
-        ]
-
-    	// profiling
-		+ SHorizontalBox::Slot()
-		.AutoWidth()
-		.HAlign(HAlign_Right)
-		.VAlign(VAlign_Center)
-		.Padding(FMargin(OpHorizontalPadding, OpVerticalPadding))
-		[
-			SNew(STextBlock)
-			.Text_Lambda([this]() -> FText
-			{
-				FIKRetargetEditorController* EditorController = ListView.Pin()->EditorController.Pin().Get();
-				FIKRetargetOpBase* Op = EditorController->AssetController->GetRetargetOpByIndex(Element.Pin()->GetIndexInStack());
-				const double Milliseconds = (Op && Op->EditorInstance ? Op->EditorInstance->Timer.GetRunningAverage() : 0.0f) * 1000.0;
-				FString Formatted = FString::Printf(TEXT("%.3f ms"), Milliseconds);
-				return FText::FromString(Formatted);
-			})
-			.IsEnabled(false)
-			.Visibility_Lambda([this]()
-			{
-				const bool Enabled = ListView.Pin()->EditorController.Pin()->AssetController->GetAsset()->bProfileOps;
-				return Enabled ? EVisibility::Visible : EVisibility::Collapsed;
-			})
-		]
+			// display Name
+			+ SHorizontalBox::Slot()
+			.AutoWidth()
+			.HAlign(HAlign_Left)
+			.VAlign(VAlign_Center)
+			.Padding(FMargin(OpHorizontalPadding, OpVerticalPadding))
+			[
+				SAssignNew(EditNameWidget, SInlineEditableTextBlock)
+				.Text(this, &SRetargetOpItem::GetName)
+				.OnVerifyTextChanged_Lambda([](const FText& InText, FText& OutErrorMessage)
+				{
+					static FString IllegalNameCharacters = "^<>:\"/\\|?*";
+					return FName::IsValidXName(InText.ToString(), FString(IllegalNameCharacters), &OutErrorMessage);
+				})
+				.OnTextCommitted(this, &SRetargetOpItem::OnNameCommitted)
+				.MultiLine(false)
+			]
 
-    	// add sub-op button
-		+ SHorizontalBox::Slot()
-		.AutoWidth()
-		.HAlign(HAlign_Right)
-		.VAlign(VAlign_Center)
-		.Padding(FMargin(OpHorizontalPadding, OpVerticalPadding))
-		[
-			SNew(SPositiveActionButton)
-			.Visibility(Element.Pin()->GetCanHaveChildren() ? EVisibility::Visible : EVisibility::Collapsed)
-			.Icon(FAppStyle::Get().GetBrush("Icons.Plus"))
-			.ToolTipText(LOCTEXT("AddChildToolTip", "Add a child op to run before this op."))
-			.OnGetMenuContent_Lambda([this]()
-			{
-				return ListView.Pin()->CreateAddNewOpMenu();	
-			})
-		]
+			// spacer
+			+ SHorizontalBox::Slot()
+			.HAlign(HAlign_Fill)
+			.FillWidth(1.0f)
+			[
+				SNew(SSpacer)
+				.Size(FVector2D(0.0f, 1.0f))
+			]
 
-        // spin Box
-        + SHorizontalBox::Slot()
-        .AutoWidth()
-        .HAlign(HAlign_Right)
-        .VAlign(VAlign_Center)
-        .Padding(FMargin(OpHorizontalPadding, OpVerticalPadding))
-        [
-            SNew(SSpinBox<double>)
-	        .Value_Lambda([this]() -> double
-	        {
-		        return 1.0f;							// TODO add support for alpha value on op
-	        })
-	        .IsEnabled(false)					// TODO enable this once alpha is supported
-	        .Visibility(EVisibility::Collapsed)	// TODO enable this once alpha is supported
-            .MinValue(TOptional<double>())
-            .MaxValue(TOptional<double>())
-            .Style(&FAppStyle::Get().GetWidgetStyle<FSpinBoxStyle>("Sequencer.HyperlinkSpinBox"))
-        ]
-
-        // delete Button
-        + SHorizontalBox::Slot()
-        .AutoWidth()
-        .HAlign(HAlign_Right)
-        .VAlign(VAlign_Center)
-        .Padding(FMargin(OpHorizontalPadding, OpVerticalPadding))
-        [
-            SNew(SButton)
-        	.NormalPaddingOverride(FMargin(2, 2))
-			.PressedPaddingOverride(FMargin(2, 2))
-            .ToolTipText(LOCTEXT("DeleteOp", "Delete retarget op and remove from stack."))
-            .OnClicked_Lambda([this]() -> FReply
-            {
-            	ListView.Pin()->DeleteRetargetOp(Element.Pin());
-                return FReply::Handled();
-            })
-            .Content()
-            [
-            	SNew(SImage)
-					.Image(FAppStyle::Get().GetBrush("Icons.Delete"))
-					.ColorAndOpacity(FSlateColor::UseForeground())
-            ]
-        ]
-    ];
+			// profiling
+			+ SHorizontalBox::Slot()
+			.AutoWidth()
+			.HAlign(HAlign_Right)
+			.VAlign(VAlign_Center)
+			.Padding(FMargin(OpHorizontalPadding, OpVerticalPadding))
+			[
+				SNew(STextBlock)
+				.Text_Lambda([this]() -> FText
+				{
+					FIKRetargetEditorController* EditorController = GetEditorController();
+					FIKRetargetOpBase* Op = EditorController->AssetController->GetRetargetOpByIndex(Element.Pin()->GetIndexInStack());
+					const double Milliseconds = (Op && Op->EditorInstance ? Op->EditorInstance->Timer.GetRunningAverage() : 0.0f) * 1000.0;
+					FString Formatted = FString::Printf(TEXT("%.3f ms"), Milliseconds);
+					return FText::FromString(Formatted);
+				})
+				.IsEnabled(false)
+				.Visibility_Lambda([this]()
+				{
+					const bool Enabled = ListView.Pin()->EditorController.Pin()->AssetController->GetAsset()->bProfileOps;
+					return Enabled ? EVisibility::Visible : EVisibility::Collapsed;
+				})
+			]
+
+			// add sub-op button
+			+ SHorizontalBox::Slot()
+			.AutoWidth()
+			.HAlign(HAlign_Right)
+			.VAlign(VAlign_Center)
+			.Padding(FMargin(OpHorizontalPadding, OpVerticalPadding))
+			[
+				SNew(SPositiveActionButton)
+				.Visibility(Element.Pin()->GetCanHaveChildren() ? EVisibility::Visible : EVisibility::Collapsed)
+				.Icon(FAppStyle::Get().GetBrush("Icons.Plus"))
+				.ToolTipText(LOCTEXT("AddChildToolTip", "Add a child op to run before this op."))
+				.OnGetMenuContent_Lambda([this]()
+				{
+					return ListView.Pin()->CreateAddNewOpMenu();	
+				})
+			]
+
+			// spin Box
+			+ SHorizontalBox::Slot()
+			.AutoWidth()
+			.HAlign(HAlign_Right)
+			.VAlign(VAlign_Center)
+			.Padding(FMargin(OpHorizontalPadding, OpVerticalPadding))
+			[
+				SNew(SSpinBox<double>)
+				.Value_Lambda([this]() -> double
+				{
+					return 1.0f;							// TODO add support for alpha value on op
+				})
+				.IsEnabled(false)					// TODO enable this once alpha is supported
+				.Visibility(EVisibility::Collapsed)	// TODO enable this once alpha is supported
+				.MinValue(TOptional<double>())
+				.MaxValue(TOptional<double>())
+				.Style(&FAppStyle::Get().GetWidgetStyle<FSpinBoxStyle>("Sequencer.HyperlinkSpinBox"))
+			]
+
+			// delete Button
+			+ SHorizontalBox::Slot()
+			.AutoWidth()
+			.HAlign(HAlign_Right)
+			.VAlign(VAlign_Center)
+			.Padding(FMargin(OpHorizontalPadding, OpVerticalPadding))
+			[
+				SNew(SButton)
+				.NormalPaddingOverride(FMargin(2, 2))
+				.PressedPaddingOverride(FMargin(2, 2))
+				.ToolTipText(LOCTEXT("DeleteOp", "Delete retarget op and remove from stack."))
+				.OnClicked_Lambda([this]() -> FReply
+				{
+					ListView.Pin()->DeleteRetargetOp(Element.Pin());
+					return FReply::Handled();
+				})
+				.Content()
+				[
+					SNew(SImage)
+						.Image(FAppStyle::Get().GetBrush("Icons.Delete"))
+						.ColorAndOpacity(FSlateColor::UseForeground())
+				]
+			]   
+		]
+	];
 
 	// bind the rename callback to enter text editing mode when item is slow double-clicked
 	Element.Pin()->OnRenameRequested.BindSP(EditNameWidget.Get(), &SInlineEditableTextBlock::EnterEditingMode);
@@ -272,7 +299,7 @@ void SRetargetOpItem::Construct(const FArguments& InArgs)
 void SRetargetOpItem::OnNameCommitted(const FText& InText, ETextCommit::Type InCommitType) const
 {
 	FName NewName = FName(InText.ToString());
-	FIKRetargetEditorController* EditorController = ListView.Pin()->EditorController.Pin().Get();
+	FIKRetargetEditorController* EditorController = GetEditorController();
 	EditorController->AssetController->SetOpName(NewName, Element.Pin()->GetIndexInStack());
 }
 
@@ -282,6 +309,132 @@ FText SRetargetOpItem::GetName() const
 	return FText::FromName(OpName);
 }
 
+void SRetargetOpItem::GenerateRightClickMenu(FMenuBuilder& InMenuBuilder)
+{
+	FIKRetargetEditorController* EditorController = GetEditorController();
+	TObjectPtr<UIKRetargeterController> AssetController = EditorController->AssetController;
+	FIKRetargetOpBase* Op = AssetController->GetRetargetOpByIndex(Element.Pin()->GetIndexInStack());
+	if (!ensure(Op))
+	{
+		return;
+	}
+
+	const FName SelectedOverrideSet = EditorController->GetSelectedOverrideSet();
+	const FName SelectedOpName = Op->GetName();
+
+	if (SelectedOverrideSet == NAME_None)
+	{
+		InMenuBuilder.BeginSection("WarningSection", LOCTEXT("WarningHeading", "Notice"));
+		{
+			InMenuBuilder.AddWidget(
+				SNew(STextBlock)
+				.Text(LOCTEXT("NoSelectionWarning", "No profile selected to add overrides to. \nOpen the Retarget Profiles tab to edit a profile."))
+				.ColorAndOpacity(FAppStyle::GetColor("ErrorReporting.WarningBackgroundColor"))
+				.Margin(FMargin(16.0f, 4.0f))
+				.Font(FAppStyle::GetFontStyle("BoldFont")),
+				FText::GetEmpty(),
+				true
+			);
+		}
+		InMenuBuilder.EndSection();
+		return;
+	}
+	
+	InMenuBuilder.BeginSection("ProfileActions", FText::Format(LOCTEXT("ProfileHeading", "Edit Profile: {0}"), FText::FromName(SelectedOverrideSet)));
+
+	InMenuBuilder.AddMenuEntry(
+		LOCTEXT("AddChangedLabel", "Add Changed Properties"),
+		LOCTEXT("AddChangedTooltip", "Record all modifications to the op settings into the currently selected override set."),
+		FSlateIcon(FAppStyle::GetAppStyleSetName(), "Icons.Plus"),
+		FUIAction(FExecuteAction::CreateLambda([this, SelectedOverrideSet, SelectedOpName]()
+			{
+				GetEditorController()->AssetController->OverrideAllEditedProperties(SelectedOverrideSet, SelectedOpName);
+			})));
+
+	InMenuBuilder.AddMenuEntry(
+		LOCTEXT("AddAllLabel", "Add All Properties"),
+		LOCTEXT("AddAllTooltip", "Record all the op settings into the currently selected override set."),
+		FSlateIcon(FAppStyle::GetAppStyleSetName(), "Icons.Plus"),
+		FUIAction(FExecuteAction::CreateLambda([this, SelectedOverrideSet, SelectedOpName]()
+			{
+				GetEditorController()->AssetController->OverrideAllProperties(SelectedOverrideSet, SelectedOpName);
+			})));
+
+	InMenuBuilder.AddMenuEntry(
+		LOCTEXT("ClearAllLabel", "Clear All Overrides"),
+		LOCTEXT("ClearAllTooltip", "Removes all sparse overrides from this overrise set."),
+		FSlateIcon(FAppStyle::GetAppStyleSetName(), "Icons.Delete"),
+		FUIAction(FExecuteAction::CreateLambda([this, SelectedOverrideSet, SelectedOpName]()
+			{
+				GetEditorController()->AssetController->RemoveOpOverrides(SelectedOverrideSet, SelectedOpName);
+			})));
+	
+	InMenuBuilder.EndSection();
+
+	InMenuBuilder.BeginSection("OpProperties", LOCTEXT("ProfileHeading", "Add Overrides to Profile"));
+
+	TSharedPtr<FRetargetPropertyOverrideNode> RootNode;
+	Op->GetOverrideableProperties(RootNode);
+	MakeOverrideSubMenu(InMenuBuilder, RootNode);
+
+	InMenuBuilder.EndSection();
+}
+
+FIKRetargetEditorController* SRetargetOpItem::GetEditorController() const
+{
+	return ListView.IsValid() ? ListView.Pin()->EditorController.Pin().Get() : nullptr;
+}
+
+void SRetargetOpItem::MakeOverrideSubMenu(FMenuBuilder& SubMenuBuilder, TSharedPtr<FRetargetPropertyOverrideNode> Node)
+{
+	const FName ProfileName = GetEditorController()->GetSelectedOverrideSet();
+	FName OpName = GetEditorController()->AssetController->GetOpName(Element.Pin()->GetIndexInStack());
+	
+	// add leaf properties first
+	for (const FRetargetPropertyOverrideNode::FRetargetPropertyOverrideEntry& Entry : Node->LeafProperties)
+	{
+		SubMenuBuilder.AddMenuEntry(
+			Entry.DisplayName,
+			FText::Format(LOCTEXT("OverrideTipLabel", "Add or Remove {0} override in current override set."), FText::FromString(Entry.PropertyPath)),
+			FSlateIcon(),
+			FUIAction(
+				FExecuteAction::CreateLambda([this, ProfileName, OpName, Entry]()
+				{
+					TObjectPtr<UIKRetargeterController> AssetController = GetEditorController()->AssetController;
+					if (AssetController->HasPropertyOverride(ProfileName, OpName, Entry.PropertyPath))
+					{
+						// remove it
+						AssetController->RemovePropertyOverrideFromOp(ProfileName, OpName, Entry.PropertyPath);
+					}
+					else
+					{
+						// add it
+						AssetController->AddPropertyOverrideToOp(ProfileName, OpName, Entry.PropertyPath);	
+					}
+				}),
+				FCanExecuteAction(),
+				FIsActionChecked::CreateLambda([this, ProfileName, OpName, Entry]()
+				{
+					TObjectPtr<UIKRetargeterController> AssetController = GetEditorController()->AssetController;
+					return AssetController->HasPropertyOverride(ProfileName, OpName, Entry.PropertyPath);
+				})
+			),
+			NAME_None,
+			EUserInterfaceActionType::ToggleButton
+		);
+	}
+
+	// add sub-nodes (nested submenus)
+	for (TTuple<FName, TSharedPtr<FRetargetPropertyOverrideNode>>& Pair : Node->SubNodes)
+	{
+		SubMenuBuilder.AddSubMenu(
+			FText::FromName(Pair.Key),
+			FText::GetEmpty(),
+			FNewMenuDelegate::CreateSP(this, &SRetargetOpItem::MakeOverrideSubMenu, Pair.Value)
+		);
+	}
+}
+
 bool SRetargetOpItem::IsOpEnabled() const
 {
 	if (const FIKRetargetOpBase* Op = GetRetargetOp())
@@ -298,7 +451,7 @@ FIKRetargetOpBase* SRetargetOpItem::GetRetargetOp() const
 	{
 		return nullptr;
 	}
-	FIKRetargetEditorController* EditorController = ListView.Pin()->EditorController.Pin().Get();
+	FIKRetargetEditorController* EditorController = GetEditorController();
 	if (!EditorController)
 	{
 		return nullptr;
diff --git a/Engine/Plugins/Animation/IKRig/Source/IKRigEditor/Public/RetargetEditor/IKRetargetBatchOperation.h b/Engine/Plugins/Animation/IKRig/Source/IKRigEditor/Public/RetargetEditor/IKRetargetBatchOperation.h
--- a/Engine/Plugins/Animation/IKRig/Source/IKRigEditor/Public/RetargetEditor/IKRetargetBatchOperation.h
+++ b/Engine/Plugins/Animation/IKRig/Source/IKRigEditor/Public/RetargetEditor/IKRetargetBatchOperation.h
@@ -85,6 +85,9 @@ public:
 	// NOTE: results may not be WYSIWYG with the editor, but the behavior should remain intact.
 	bool bRetainAdditiveFlags = true;
 
+	// An array of property override sets to apply. These must be present in the supplied IK Retarget asset.
+	TArray<FName> OverrideSetNames;
+
 	// Reset all data (called when window re-opened
 	void Reset()
 	{
@@ -96,6 +99,7 @@ public:
 		NameRule.Suffix = "";
 		NameRule.ReplaceFrom = "";
 		NameRule.ReplaceTo = "";
+		OverrideSetNames.Reset();
 	}
 
 	// Is the data configured in such a way that we could run the retarget?
@@ -130,15 +134,17 @@ public:
 	 * @param bUseSourcePath Place the assets in the same folder as their source location instead of the given target path.
 	 * @param bIncludeReferencedAssets Whether to retarget animation assets referenced by the AssetsToRetarget
 	 * @param bOverwriteExistingFiles Whether to overwrite any existing output files with the same name, or create new files with a unique name
+	 * @param InOverrideSetNames An array of names of the property override sets to apply to the retarget (must be stored in the IK Retarget Asset)
 	 * 
 	 * Return: list of new animation files created.
 	 */
-	UFUNCTION(BlueprintCallable, Category=IKBatchRetarget)
+	UFUNCTION(BlueprintCallable, Category=IKBatchRetarget, meta=(AutoCreateRefTerm="InOverrideSetNames"))
 	static UE_API TArray<FAssetData> DuplicateAndRetarget(
 		const TArray<FAssetData>& AssetsToRetarget,
 		USkeletalMesh* SourceMesh,
 		USkeletalMesh* TargetMesh,
 		UIKRetargeter* IKRetargetAsset,
+		const TArray<FName>& InOverrideSetNames,
 		const FString& Search = "",
 		const FString& Replace = "",
 		const FString& Prefix = "",
diff --git a/Engine/Plugins/Animation/IKRig/Source/IKRigEditor/Public/RetargetEditor/IKRetargetCommands.h b/Engine/Plugins/Animation/IKRig/Source/IKRigEditor/Public/RetargetEditor/IKRetargetCommands.h
--- a/Engine/Plugins/Animation/IKRig/Source/IKRigEditor/Public/RetargetEditor/IKRetargetCommands.h
+++ b/Engine/Plugins/Animation/IKRig/Source/IKRigEditor/Public/RetargetEditor/IKRetargetCommands.h
@@ -27,6 +27,8 @@ public:
 
 	// open asset settings
 	TSharedPtr< FUICommandInfo > ShowAssetSettings;
+	// open the override set manager
+	TSharedPtr< FUICommandInfo > ShowOverrideSets;
 
 	// enable Root retarget pass
 	TSharedPtr< FUICommandInfo > EnableRoot;
diff --git a/Engine/Plugins/Animation/IKRig/Source/IKRigEditor/Public/RetargetEditor/IKRetargetEditorController.h b/Engine/Plugins/Animation/IKRig/Source/IKRigEditor/Public/RetargetEditor/IKRetargetEditorController.h
--- a/Engine/Plugins/Animation/IKRig/Source/IKRigEditor/Public/RetargetEditor/IKRetargetEditorController.h
+++ b/Engine/Plugins/Animation/IKRig/Source/IKRigEditor/Public/RetargetEditor/IKRetargetEditorController.h
@@ -14,6 +14,7 @@
 
 #include "Misc/TransactionObjectEvent.h"
 
+class SIKRetargetOverrideManager;
 class SRetargetOpStack;
 enum class ERetargetSourceOrTarget : uint8;
 class SIKRetargetHierarchy;
@@ -108,9 +109,10 @@ public:
 	void HandleRetargeterNeedsInitialized() const;
 	// reinitialize retargeter without refreshing UI
 	void ReinitializeRetargeterNoUIRefresh() const;
+	// delegate handles to un-bind
 	FDelegateHandle RetargeterReInitDelegateHandle;
-	// refresh op stack whenever it's modified
 	FDelegateHandle OpStackModifiedDelegateHandle;
+	FDelegateHandle ProfilesModifiedDelegateHandle;
 	// callback when IK Rig asset has been swapped out
 	void HandleIKRigReplaced(ERetargetSourceOrTarget SourceOrTarget);
 	FDelegateHandle IKRigReplacedDelegateHandle;
@@ -153,6 +155,7 @@ public:
 	void SetOutputLogView(const TSharedPtr<SIKRigOutputLog>& InOutputLogView) { OutputLogView = InOutputLogView; };
 	void SetHierarchyView(const TSharedPtr<SIKRetargetHierarchy>& InHierarchyView) { HierarchyView = InHierarchyView; };
 	void SetOpStackView(const TSharedPtr<SRetargetOpStack>& InOpStackView) { OpStackView = InOpStackView; };
+	void SetOverrideSetsView(const TSharedPtr<SIKRetargetOverrideManager>& InOverrideSetView) { OverrideSetView = InOverrideSetView; };
 	bool IsObjectInDetailsView(const UObject* Object);
 	
 	// force refresh all views in the editor
@@ -161,6 +164,7 @@ public:
 	void RefreshAssetBrowserView() const;
 	void RefreshHierarchyView() const;
 	void RefreshOpStackView() const;
+	void RefreshOverrideSetView() const;
 	void RefreshPoseList() const;
 	void SetDetailsObject(UObject* DetailsObject) const;
 	void SetDetailsObjects(const TArray<UObject*>& DetailsObjects) const;
@@ -257,6 +261,10 @@ public:
 	int32 GetSelectedOpIndex() const;
 	// END op selection
 
+	// SELECTION - PROFILES
+	FName GetSelectedOverrideSet() const;
+	// END profile selection
+
 	// ------------------------- END SELECTION -----------------------------
 
 	// determine if bone in the specified skeleton is part of the retarget (in a mapped chain)
@@ -337,6 +345,8 @@ private:
 	TSharedPtr<SIKRetargetHierarchy> HierarchyView;
 	// op stack widget
 	TSharedPtr<SRetargetOpStack> OpStackView;
+	// overrides widget
+	TSharedPtr<SIKRetargetOverrideManager> OverrideSetView;
 
 	// when prompting user to assign an IK Rig
 	TSharedPtr<SWindow> IKRigPickerWindow;
diff --git a/Engine/Plugins/Animation/IKRig/Source/IKRigEditor/Public/RetargetEditor/IKRetargetOverrideCustomization.h b/Engine/Plugins/Animation/IKRig/Source/IKRigEditor/Public/RetargetEditor/IKRetargetOverrideCustomization.h
new file mode 100644
--- /dev/null
+++ b/Engine/Plugins/Animation/IKRig/Source/IKRigEditor/Public/RetargetEditor/IKRetargetOverrideCustomization.h
@@ -0,0 +1,122 @@
+﻿// Copyright Epic Games, Inc. All Rights Reserved.
+#pragma once
+
+#include "IPropertyTypeCustomization.h"
+
+class UIKRetargeterController;
+class IDetailGroup;
+class FStructOnScope;
+
+// This builds a custom details view for FRetargetOverrideSet types.
+// It supports editing retarget op settings overrides stored within a retarget override set
+class FRetargetOverrideCustomization : public IPropertyTypeCustomization
+{
+	
+public:
+	
+	static TSharedRef<IPropertyTypeCustomization> MakeInstance() 
+	{ 
+		return MakeShared<FRetargetOverrideCustomization>(); 
+	}
+
+	/** IPropertyTypeCustomization interface */
+	virtual void CustomizeHeader(TSharedRef<IPropertyHandle> PropertyHandle, FDetailWidgetRow& HeaderRow, IPropertyTypeCustomizationUtils& CustomizationUtils) override;
+	virtual void CustomizeChildren(TSharedRef<IPropertyHandle> PropertyHandle, IDetailChildrenBuilder& ChildBuilder, IPropertyTypeCustomizationUtils& CustomizationUtils) override;
+
+private:
+
+	/** We customize the FRetargetOverrideSet details differently when outside the retarget editor context. */
+	void CustomizeOutsideEditor(TSharedRef<IPropertyHandle> OverrideSetHandle, IDetailChildrenBuilder& ChildBuilder);
+
+	/** Get the controller for the retarget asset by traversing through the UIKRigStructViewer being customized
+	 * @return a pointer to the asset controller this set belongs to (or nullptr if outside editor) */
+	TObjectPtr<UIKRetargeterController> GetAssetController(const IPropertyTypeCustomizationUtils& CustomizationUtils);
+	
+	/**
+	* Injects a custom row for a property override into a detail group.
+	* @param InPropertyHandle  Handle to the property being displayed.
+	* @param InPropertyGroup   The detail group where this row will be added.
+	* @param InOverrideSetName The name of the current override set.
+	* @param InOpName          The name of the retarget op the override belongs to.
+	* @param InPropertyPath    The full path to the property within the op's settings.
+	*/
+	void AddRowForPropertyOverride(
+		TSharedPtr<IPropertyHandle> InPropertyHandle,
+		IDetailGroup& InPropertyGroup,
+		const FName& InOverrideSetName,
+		const FName& InOpName,
+		const FString& InPropertyPath);
+
+	/**
+	* Iterates through the overrides in an override set and groups them by their 
+	* logical sub-menu/category defined in the metadata.
+	* @param OpOverridesHandle  Handle to the FRetargetOpOverrides struct.
+	* @return A map where keys are group names and values are indices into the override array.
+	*/
+	TMap<FString, TArray<int32>> MapOverridesToGroups(const TSharedPtr<IPropertyHandle>& OpOverridesHandle) const;
+
+	/**
+	* Retrieves or creates a cached FStructOnScope for a specific retarget op.
+	* Used for temporary value manipulation or default value comparisons.
+	* @param InOpName          The name of the operation.
+	* @param OpSettingsType    The script struct type to instantiate if not cached.
+	* @return A reference to the shared pointer containing the scoped struct.
+	*/
+	TSharedPtr<FStructOnScope>& GetCachedScopedStruct(const FName& InOpName, const UScriptStruct* OpSettingsType);
+
+	/**
+	* Copies a serialized string value from a property handle into the corresponding property of a live struct instance.
+	* @param InScopedStruct     The destination struct instance.
+	* @param OpSettingsType     The struct type definition.
+	* @param InPropertyPath     The path to the property to be updated.
+	* @param InRawValueHandle   The handle containing the source serialized value.
+	* @return True if the value was successfully resolved and copied.
+	*/
+	bool CopyStoredValueToScopedStruct(
+		const TSharedPtr<FStructOnScope>& InScopedStruct,
+		const UScriptStruct* OpSettingsType,
+		const FString& InPropertyPath,
+		const TSharedPtr<IPropertyHandle>& InRawValueHandle);
+
+	/**
+	* Resolves the underlying UScriptStruct type for a retarget operation's settings.
+	* @param OpOverridesHandle   Handle to the FRetargetOpOverrides.
+	* @return The UScriptStruct pointer if found, otherwise nullptr.
+	*/
+	const UScriptStruct* GetSettingsStruct(TSharedPtr<IPropertyHandle>& OpOverridesHandle) const;
+
+	/**
+	* Navigates a property handle hierarchy based on a string path.
+	* @param InRootHandle      The starting point for the search.
+	* @param InPath            The string path (e.g., "Settings->SubStruct->Value").
+	* @return A handle to the leaf property if the path is valid, otherwise null.
+	*/
+	TSharedPtr<IPropertyHandle> ResolveHandleFromPath(TSharedPtr<IPropertyHandle> InRootHandle, const FString& InPath);
+
+	/**
+	* Extracts the property path string from a specific override entry in the override set.
+	* @param InOpOverridesHandle Handle to the parent FRetargetOpOverrides.
+	* @param InOverrideIndex   Index of the override in the TArray.
+	* @param OutPath           [Out] The resulting property path string.
+	*/
+	void GetOverridePropertyPath(
+		const TSharedPtr<IPropertyHandle>& InOpOverridesHandle,
+		const int32 InOverrideIndex,
+		FString& OutPath) const;
+
+	/**
+	* Gets the property handle for the 'Value' field of a specific override entry.
+	* @param InOpOverridesHandle Handle to the parent FRetargetOpOverrides.
+	* @param InOverrideIndex   Index of the override in the TArray.
+	* @return The handle to the override's value string.
+	*/
+	TSharedPtr<IPropertyHandle> GetOverrideValueHandle(
+		const TSharedPtr<IPropertyHandle>& InOpOverridesHandle,
+		const int32 InOverrideIndex) const;
+	
+	/** Stores the binary data for temp op settings structs while being edited <OpName, OpSettingsMemory> */
+	TMap<FName, TSharedPtr<FStructOnScope>> ScopedInstances;
+
+	/** Stored and referenced by UI callbacks for transactional edits to override sets. */
+	TWeakObjectPtr<UIKRetargeterController> AssetController;
+};
diff --git a/Engine/Plugins/Animation/IKRig/Source/IKRigEditor/Public/RetargetEditor/IKRetargetOverridesTabSummoner.h b/Engine/Plugins/Animation/IKRig/Source/IKRigEditor/Public/RetargetEditor/IKRetargetOverridesTabSummoner.h
new file mode 100644
--- /dev/null
+++ b/Engine/Plugins/Animation/IKRig/Source/IKRigEditor/Public/RetargetEditor/IKRetargetOverridesTabSummoner.h
@@ -0,0 +1,23 @@
+﻿// Copyright Epic Games, Inc. All Rights Reserved.
+
+#pragma once
+
+#include "CoreMinimal.h"
+#include "WorkflowOrientedApp/WorkflowTabFactory.h"
+
+class FIKRetargetEditor;
+
+/** Summons the "Retarget Override Sets" tab in the IK Retarget editor. */
+struct FIKRetargetOverridesTabSummoner : public FWorkflowTabFactory
+{
+public:
+	static const FName TabID;
+	
+	FIKRetargetOverridesTabSummoner(const TSharedRef<FIKRetargetEditor>& InIKRetargetEditor);
+	virtual TSharedRef<SWidget> CreateTabBody(const FWorkflowTabSpawnInfo& Info) const override;
+	virtual TSharedPtr<SToolTip> CreateTabToolTipWidget(const FWorkflowTabSpawnInfo& Info) const override;
+
+protected:
+	TWeakPtr<FIKRetargetEditor> IKRetargetEditor;
+};
+
diff --git a/Engine/Plugins/Animation/IKRig/Source/IKRigEditor/Public/RetargetEditor/IKRetargeterController.h b/Engine/Plugins/Animation/IKRig/Source/IKRigEditor/Public/RetargetEditor/IKRetargeterController.h
--- a/Engine/Plugins/Animation/IKRig/Source/IKRigEditor/Public/RetargetEditor/IKRetargeterController.h
+++ b/Engine/Plugins/Animation/IKRig/Source/IKRigEditor/Public/RetargetEditor/IKRetargeterController.h
@@ -451,8 +451,108 @@ public:
 		const ERetargetSourceOrTarget SourceOrTarget) const;
 
 	TWeakPtr<const FIKRigSkeleton> GetCachedIKRigSkeleton(const USkeletalMesh* InSkeletalMesh) const;
+
+	//
+	// RETARGET OVERRIDES PUBLIC/SCRIPTING API
+	//
+
+	/** Add a new override set to store retarget op setting overrides.
+	* @param InOverrideSetName    The name of the new set.
+	* @return The name of the new retarget override set (guaranteed unique). */
+	UFUNCTION(BlueprintCallable, Category="IK Retarget")
+	UE_API FName AddNewRetargetOverrideSet(FName InOverrideSetName);
+
+	/** Remove an override set by name.
+	* @param InOverrideSetName    The name of the override set to remove.
+	* @return True if the override set was removed, false otherwise. */
+	UFUNCTION(BlueprintCallable, Category="IK Retarget")
+	UE_API bool RemoveRetargetOverrideSet(FName InOverrideSetName);
+
+	/** Rename an existing override set.
+	* @param InOverrideSetName    The name of the override set to rename.
+	* @param InNewName            The new name.
+	* @return True if the override set was renamed, false otherwise. */
+	UFUNCTION(BlueprintCallable, Category="IK Retarget")
+	UE_API bool RenameRetargetOverrideSet(FName InOverrideSetName, FName InNewName);
+
+	/** Assign a parent to the given override set.
+	* NOTE: Override sets can be applied hierarchically to make for easy composition of property overrides.
+	* @param InOverrideSetName    The name of the override set to re-parent.
+	* @param InParentName         The new parent name or NAME_None.
+	* @return True if the override set was re-parented, false otherwise. */
+	UFUNCTION(BlueprintCallable, Category="IK Retarget")
+	UE_API bool SetParentOverrideSet(FName InOverrideSetName, FName InParentName);
+
+	/** Get the name of the parent for the given override set.
+	* @param InOverrideSetName    The name of the override set.
+	* @return The name of the parent override set or NAME_None. */
+	UFUNCTION(BlueprintCallable, Category="IK Retarget")
+	UE_API FName GetParentOverrideSet(FName InOverrideSetName);
+
+	//
+	// RETARGET OVERRIDES C++ ONLY API
+	//
+	
+	/** Return NameToMakeUnique with a numbered suffix that makes it unique among existing override sets. */
+	UE_API FName GetUniqueOverrideSetName(const FName NameToMakeUnique) const;
+
+	/** Get read-only access to all override sets in the retargeter. */
+	UE_API const TMap<FName, FRetargetOverrideSet>& GetAllOverrideSets() const;
+
+	/** Get access to an override set by name. */
+	UE_API FRetargetOverrideSet* FindOverrideSet(const FName InOverrideSetName) const;
+
+	/** Get a list of all override sets that are children of the given set (recursive). */
+	UE_API const TArray<FName> GetAllChildrenOverrideSets(const FName InOverrideSetName) const;
+
+	/** Get a list of all override sets that are direct children of the given set (non-recursive). */
+	UE_API const TArray<FName> GetDirectChildrenOverrideSets(const FName InOverrideSetName) const;
+
+	/** Find the specific op overrides within a given override set. */
+	UE_API FRetargetOpOverrides* FindOpOverrides(const FName InOverrideSetName, const FName InOpName) const;
+
+	/** Add an Op override entry to the given override set. */
+	UE_API FRetargetOpOverrides* AddOpOverrides(const FName InOverrideSetName, const FName InOpName) const;
+
+	/** Remove all property overrides for the given op in the given override set. */
+	UE_API bool RemoveOpOverrides(const FName InOverrideSetName, const FName InOpName) const;
+
+	/** Update the value of a stored property override inside the given override set, for the given op and property path. */
+	UE_API bool UpdateOverrideValue(const FName InOverrideSetName, const FName InOpName, const FString& InPropertyPath, const FStructOnScope& InNewValue) const;
+
+	/** Add a property override for a specific op within a specific override set. */
+	UE_API bool AddPropertyOverrideToOp(const FName InOverrideSetName, const FName InOpName, const FString& InPropertyPath) const;
+
+	/** Add a property override for all properties within the given op. */
+	UE_API bool OverrideAllProperties(const FName InOverrideSetName, const FName InOpName) const;
+
+	/** Add a property override for all edited properties within the given op (properties not at their default value). */
+	UE_API bool OverrideAllEditedProperties(const FName InOverrideSetName, const FName InOpName) const;
+
+	/** Remove a given property override for a specific op within a specific override set. */
+	UE_API bool RemovePropertyOverrideFromOp(const FName InOverrideSetName, const FName InOpName, const FString& InPropertyPath) const;
+
+	/** Remove all property overrides from all ops stored in the given override set. */
+	UE_API bool RemoveAllPropertyOverridesFromSet(const FName InOverrideSetName) const;
+
+	/** Set the given override set to be active by default (unless otherwise overridden by the user).
+	 * @return True if the override set was found and modified, false otherwise. */
+	UE_API bool SetOverrideSetActiveByDefault(const FName InOverrideSetName, const bool bActiveByDefault) const;
+
+	/** Get if the given override set is configured to be active by default.
+	 * @return True if the override set was found and is active by default, false otherwise. */
+	UE_API bool GetOverrideSetActiveByDefault(const FName InOverrideSetName) const;
+
+	/** Returns true if the given override set has an override for the given property on the given op. */
+	UE_API bool HasPropertyOverride(const FName InOverrideSetName, const FName InOpName, const FString& InPropertyPath) const;
+
+	/** Returns the name of the given override set. */
+	UE_API FName GetOverrideSetName(const FRetargetOverrideSet* InOverrideSet) const;
 	
 private:
+
+	/** Controller functions call this anytime the override sets are modified */
+	void PostEditOverrideSets() const;
 	
 	// Remove bones from retarget poses that are no longer in skeleton 
 	UE_API void CleanPoseList(const ERetargetSourceOrTarget SourceOrTarget) const;
@@ -522,6 +622,9 @@ private:
 
 	DECLARE_MULTICAST_DELEGATE(FOnOpStackModified);
 	FOnOpStackModified OpStackModified;
+
+	DECLARE_MULTICAST_DELEGATE(FOnOverrideSetsModified);
+	FOnOverrideSetsModified OverrideSetsModified;
 	
 	DECLARE_MULTICAST_DELEGATE_OneParam(FOnIKRigReplaced, ERetargetSourceOrTarget);
 	FOnIKRigReplaced IKRigReplaced;
@@ -559,6 +662,9 @@ public:
 	// Attach a delegate to be notified whenever the op stack is modified.
 	FOnOpStackModified& OnOpStackModified(){ return OpStackModified; };
 
+	// Attach a delegate to be notified whenever the override sets are modified.
+	FOnOverrideSetsModified& OnOverrideSetsModified(){ return OverrideSetsModified; };
+
 	// get the singleton struct viewer for editing structs belonging to the asset this controller controls
 	UE_API UIKRigStructViewer* GetStructViewer();
 	
diff --git a/Engine/Plugins/Animation/IKRig/Source/IKRigEditor/Public/RetargetEditor/SIKRetargetOverrideManager.h b/Engine/Plugins/Animation/IKRig/Source/IKRigEditor/Public/RetargetEditor/SIKRetargetOverrideManager.h
new file mode 100644
--- /dev/null
+++ b/Engine/Plugins/Animation/IKRig/Source/IKRigEditor/Public/RetargetEditor/SIKRetargetOverrideManager.h
@@ -0,0 +1,153 @@
+﻿// Copyright Epic Games, Inc. All Rights Reserved.
+
+#pragma once
+
+#include "CoreMinimal.h"
+#include "EditorUndoClient.h"
+#include "SBaseHierarchyTreeView.h"
+#include "SIKRetargetPoseEditor.h"
+#include "DragAndDrop/DecoratedDragDropOp.h"
+#include "Widgets/SCompoundWidget.h"
+
+class FTextFilterExpressionEvaluator;
+class SIKRetargetOverrideManager;
+class SInlineEditableTextBlock;
+
+/** Tree element representing a single override set in the hierarchy */
+class FIKRetargetOverrideSetElement : public TSharedFromThis<FIKRetargetOverrideSetElement>
+{
+public:
+    
+    FIKRetargetOverrideSetElement(
+       const FName& InName,
+       const TSharedRef<FIKRetargetEditorController>& InEditorController);
+
+    FText Key;
+    TWeakPtr<FIKRetargetOverrideSetElement> Parent;
+    TArray<TWeakPtr<FIKRetargetOverrideSetElement>> Children;
+    TWeakPtr<SIKRetargetOverrideManager> OverrideManager;
+    FName Name;
+    bool bIsHidden = false;
+
+    FText GetNumOpsLabel() const;
+    FText GetNumOverridesLabel() const;
+
+private:
+    
+    TWeakPtr<FIKRetargetEditorController> EditorController;
+};
+
+/** Row widget for the override set tree view */
+class SIKRetargetOverrideSetRow : public SMultiColumnTableRow<TSharedPtr<FIKRetargetOverrideSetElement>>
+{
+    
+public:
+
+    SLATE_BEGIN_ARGS(SIKRetargetOverrideSetRow) {}
+    SLATE_ARGUMENT(TSharedPtr<FIKRetargetEditorController>, EditorController)
+    SLATE_ARGUMENT(TSharedPtr<FIKRetargetOverrideSetElement>, TreeElement)
+    SLATE_ARGUMENT(TSharedPtr<SIKRetargetOverrideManager>, OverrideManager)
+    SLATE_END_ARGS()
+    
+    void Construct(const FArguments& InArgs, const TSharedRef<STableViewBase>& InOwnerTable);
+
+    /** Overridden from SMultiColumnTableRow. Generates a widget for this column of the tree view. */
+    virtual TSharedRef<SWidget> GenerateWidgetForColumn(const FName& ColumnName) override;
+
+    /** Requests that the inline text block enters edit mode */
+    void OnRequestRename();
+    
+private:
+    
+    FText GetItemText() const;
+    bool OnVerifyTextChanged(const FText& InText, FText& OutErrorMessage);
+    void OnTextCommitted(const FText& InText, ETextCommit::Type CommitInfo);
+
+    ECheckBoxState IsOverrideSetActive() const;
+    void OnOverrideSetActiveChanged(ECheckBoxState NewState) const;
+
+    TSharedPtr<STableViewBase> OwnerTable;
+    TWeakPtr<FIKRetargetOverrideSetElement> WeakTreeElement;
+    TWeakPtr<FIKRetargetEditorController> EditorController;
+    TWeakPtr<SIKRetargetOverrideManager> OverrideManager;
+    TSharedPtr<SInlineEditableTextBlock> InlineTextBlock;
+};
+
+/** Drag and drop operation for re-parenting override sets in the tree */
+class FRetargetOverrideSetDragDropOp : public FDecoratedDragDropOp
+{
+public:
+    
+    DRAG_DROP_OPERATOR_TYPE(FRetargetOverrideSetDragDropOp, FDecoratedDragDropOp)
+
+    static TSharedRef<FRetargetOverrideSetDragDropOp> New(TWeakPtr<FIKRetargetOverrideSetElement> InElement);
+    
+    virtual TSharedPtr<SWidget> GetDefaultDecorator() const override;
+    
+    TWeakPtr<FIKRetargetOverrideSetElement> Element;
+};
+
+struct FRetargetOverrideSetFilterOptions
+{
+    bool bHideEmptySets;
+};
+
+typedef SBaseHierarchyTreeView<FIKRetargetOverrideSetElement> SIKRetargetOverrideSetTreeView;
+
+/** Widget for managing and hierarchy of retarget override sets */
+class SIKRetargetOverrideManager : public SCompoundWidget, public FEditorUndoClient
+{
+public:
+    SLATE_BEGIN_ARGS(SIKRetargetOverrideManager) {}
+    SLATE_END_ARGS()
+
+    void Construct(const FArguments& InArgs, TSharedRef<FIKRetargetEditorController> InEditorController);
+
+    /** Gets the name of the override set currently selected in the tree */
+    FName GetSelectedOverrideSetName() const;
+
+private:
+    
+    /** centralized editor controls (facilitate cross-communication between multiple UI elements)*/
+    TWeakPtr<FIKRetargetEditorController> EditorController;
+
+    /** the details view for editing selected override set properties */
+    TSharedPtr<IDetailsView> DetailsView;
+    
+    /** tree view widget */
+    TSharedPtr<SIKRetargetOverrideSetTreeView> TreeView;
+    TArray<TSharedPtr<FIKRetargetOverrideSetElement>> RootElements;
+    TArray<TSharedPtr<FIKRetargetOverrideSetElement>> AllElements;
+
+    /** filtering the tree with search box */
+    void OnFilterTextChanged(const FText& SearchText);
+    TSharedPtr<FTextFilterExpressionEvaluator> TextFilter;
+    FRetargetOverrideSetFilterOptions FilterOptions;
+    
+    /** tree view callbacks */
+    void RefreshTreeView(bool IsInitialSetup=false);
+    void HandleGetChildrenForTree(TSharedPtr<FIKRetargetOverrideSetElement> InElement, TArray<TSharedPtr<FIKRetargetOverrideSetElement>>& OutChildren);
+    void OnSelectionChanged(TSharedPtr<FIKRetargetOverrideSetElement> Selection, ESelectInfo::Type SelectInfo);
+    void OnItemDoubleClicked(TSharedPtr<FIKRetargetOverrideSetElement> InElement);
+    void OnSetExpansionRecursive(TSharedPtr<FIKRetargetOverrideSetElement> InElement, bool bShouldBeExpanded);
+    void SetExpansionRecursive(TSharedPtr<FIKRetargetOverrideSetElement> InElement, bool bTowardsParent, bool bShouldBeExpanded);
+    TSharedPtr<SWidget> OnGenerateContextMenu();
+    /** END tree view callbacks */
+
+    /** drag and drop operations */
+    FReply OnDragDetected(
+       const FGeometry& MyGeometry,
+       const FPointerEvent& MouseEvent);
+    TOptional<EItemDropZone> OnCanAcceptDrop(
+       const FDragDropEvent& DragDropEvent,
+       EItemDropZone DropZone,
+       TSharedPtr<FIKRetargetOverrideSetElement> TargetElement);
+    FReply OnAcceptDrop(
+       const FDragDropEvent& DragDropEvent,
+       EItemDropZone DropZone,
+       TSharedPtr<FIKRetargetOverrideSetElement> TargetElement);
+    /** END drag and drop operations */
+
+    friend SIKRetargetOverrideSetRow;
+    friend FIKRetargetEditorController;
+};
\ No newline at end of file
diff --git a/Engine/Plugins/Animation/IKRig/Source/IKRigEditor/Public/RetargetEditor/SRetargetOpStack.h b/Engine/Plugins/Animation/IKRig/Source/IKRigEditor/Public/RetargetEditor/SRetargetOpStack.h
--- a/Engine/Plugins/Animation/IKRig/Source/IKRigEditor/Public/RetargetEditor/SRetargetOpStack.h
+++ b/Engine/Plugins/Animation/IKRig/Source/IKRigEditor/Public/RetargetEditor/SRetargetOpStack.h
@@ -16,6 +16,7 @@ class SRetargetOpStack;
 struct FIKRetargetOpBase;
 class FIKRetargetEditorController;
 class SRetargetOpList;
+struct FRetargetPropertyOverrideNode;
 
 // SListView item-type for a single retarget op
 class FRetargetOpStackElement
@@ -180,10 +181,16 @@ public:
 	FText GetName() const;
 	// ENd renaming op
 
+	void GenerateRightClickMenu(FMenuBuilder& InMenuBuilder);
+
 	static float OpHorizontalPadding;
 	static float OpVerticalPadding;
 	
 private:
+
+	FIKRetargetEditorController* GetEditorController() const;
+
+	void MakeOverrideSubMenu(FMenuBuilder& SubMenuBuilder, TSharedPtr<FRetargetPropertyOverrideNode> Node);
 	
 	bool IsOpEnabled() const;
 	
diff --git a/Engine/Plugins/Animation/IKRig/Source/IKRigUAF/Private/RigUnit_UAFIKRetargeter.cpp b/Engine/Plugins/Animation/IKRig/Source/IKRigUAF/Private/RigUnit_UAFIKRetargeter.cpp
--- a/Engine/Plugins/Animation/IKRig/Source/IKRigUAF/Private/RigUnit_UAFIKRetargeter.cpp
+++ b/Engine/Plugins/Animation/IKRig/Source/IKRigUAF/Private/RigUnit_UAFIKRetargeter.cpp
@@ -40,7 +40,7 @@ FRigUnit_UAFIKRetargeter_Execute()
 	if (!Processor.WasInitializedWithTheseAssets(SourceMesh, TargetMesh, IKRetargetAsset))
 	{
 		// initialize retarget processor with source and target skeletal meshes
-		Processor.Initialize(SourceMesh, TargetMesh, IKRetargetAsset, WorkData.Profile);
+		Processor.Initialize(SourceMesh, TargetMesh, IKRetargetAsset, WorkData.Profile, OverrideSetNames);
 	}
 
 	// bail out if not initialized
@@ -67,6 +67,7 @@ FRigUnit_UAFIKRetargeter_Execute()
 	const TArray<FTransform>& TargetGlobalPose = Processor.RunRetargeter(
 	WorkData.SourceMeshPoseGlobal,
 	WorkData.Profile,
+	OverrideSetNames,
 	ExecuteContext.GetDeltaTime(),
 	TargetRefPose.GetSourceLODLevel());
 
diff --git a/Engine/Plugins/Animation/IKRig/Source/IKRigUAF/Public/RigUnit_UAFIKRetargeter.h b/Engine/Plugins/Animation/IKRig/Source/IKRigUAF/Public/RigUnit_UAFIKRetargeter.h
--- a/Engine/Plugins/Animation/IKRig/Source/IKRigUAF/Public/RigUnit_UAFIKRetargeter.h
+++ b/Engine/Plugins/Animation/IKRig/Source/IKRigUAF/Public/RigUnit_UAFIKRetargeter.h
@@ -55,6 +55,11 @@ struct FRigUnit_UAFIKRetargeter : public FRigUnit_AnimNextBase
 	UPROPERTY(EditAnywhere, DisplayName = "Custom Profile", Category = "Graph", meta = (Input))
 	FRetargetProfile CustomRetargetProfile;
 
+	// An optional list of property override sets to apply to op settings
+	// NOTE: these names must match the names of override sets stored in the retarget asset
+	UPROPERTY(EditAnywhere, DisplayName = "Custom Profile", Category = "Graph", meta = (Input))
+	TArray<FName> OverrideSetNames;
+
 	// The output pose for the target skeletal mesh.
 	UPROPERTY(EditAnywhere, DisplayName = "Retargeted Pose", Category = "Graph", meta = (Output))
 	FAnimNextGraphLODPose TargetPose;
