Follow up fix from CL51862986 that broke back compatibility in IK retarget processor.

#Synced-CL 51862986

diff --git a/Engine/Plugins/Animation/IKRig/Source/IKRig/Private/AnimNodes/AnimNode_RetargetPoseFromMesh.cpp b/Engine/Plugins/Animation/IKRig/Source/IKRig/Private/AnimNodes/AnimNode_RetargetPoseFromMesh.cpp
--- a/Engine/Plugins/Animation/IKRig/Source/IKRig/Private/AnimNodes/AnimNode_RetargetPoseFromMesh.cpp
+++ b/Engine/Plugins/Animation/IKRig/Source/IKRig/Private/AnimNodes/AnimNode_RetargetPoseFromMesh.cpp
@@ -231,12 +231,13 @@ void FAnimNode_RetargetPoseFromMesh::Evaluate_AnyThread(FPoseContext& Output)
 	Processor.OnAnimGraphEvaluateAnyThread(Output);
 
 	// run the retargeter
-	const TArray<FTransform>& RetargetedPose = Processor.RunRetargeter(
-		PoseToRetargetFromComponentSpace,
-		RetargetProfileToUse,
-		*OverridesToUse,
-		DeltaTime,
-		CurrentLODLevel);
+	FRetargetRunParameters Params;
+	Params.SourceGlobalPose = &PoseToRetargetFromComponentSpace;
+	Params.Profile = &RetargetProfileToUse;
+	Params.OverrideSetsToApply = *OverridesToUse;
+	Params.DeltaTime = DeltaTime;
+	Params.LOD = CurrentLODLevel;
+	const TArray<FTransform>& RetargetedPose = Processor.RunRetargeter(Params);
 
 	// reset delta time (it's accumulated)
 	DeltaTime = 0.f;
@@ -392,10 +393,18 @@ bool FAnimNode_RetargetPoseFromMesh::EnsureProcessorIsInitialized(const USkeleta
 	// try initializing the processor
 	if (!Processor.WasInitializedWithTheseAssets(SourceMesh, TargetMesh, IKRetargeterAsset))
 	{
-		// initialize retarget processor with source and target skeletal meshes
-		bool bForceIKOff = false;
+		constexpr bool bForceIKOff = false;
 		FRetargetProfile RetargetProfileToUse = GetMergedRetargetProfile(bForceIKOff);
-		Processor.Initialize(SourceMesh, TargetMesh, IKRetargeterAsset, RetargetProfileToUse, OverrideSetsToApply);
+		
+		// initialize retarget processor with source and target skeletal meshes
+		FRetargetInitParameters Params;
+		Params.SourceSkeletalMesh = SourceMesh;
+		Params.TargetSkeletalMesh = TargetMesh;
+		Params.RetargeterAsset = IKRetargeterAsset;
+		Params.CustomProfile = &RetargetProfileToUse;
+		Params.OverrideSetsToApply = OverrideSetsToApply;
+		Params.bSuppressWarnings = false;
+		Processor.Initialize(Params);
 	}
 
 	return Processor.IsInitialized();
diff --git a/Engine/Plugins/Animation/IKRig/Source/IKRig/Private/Retargeter/IKRetargetProcessor.cpp b/Engine/Plugins/Animation/IKRig/Source/IKRig/Private/Retargeter/IKRetargetProcessor.cpp
--- a/Engine/Plugins/Animation/IKRig/Source/IKRig/Private/Retargeter/IKRetargetProcessor.cpp
+++ b/Engine/Plugins/Animation/IKRig/Source/IKRig/Private/Retargeter/IKRetargetProcessor.cpp
@@ -1094,89 +1094,9 @@ void FRetargeterBoneChains::Reset()
 	TargetBoneChains.Reset();
 }
 
-void FIKRetargetProcessor::Initialize(
-		const USkeletalMesh* SourceSkeletalMesh,
-		const USkeletalMesh* TargetSkeletalMesh,
-		const UIKRetargeter* InRetargeterAsset,
-		const FRetargetProfile& InCustomProfile,
-		const TArrayView<const FName> InOverrideSetsToApply,
-		const bool bSuppressWarnings)
-{
-	if (!InRetargeterAsset || !SourceSkeletalMesh || !TargetSkeletalMesh)
-	{
-		bIsInitialized = false;
-		return;
-	}
-	
-	// don't attempt reinitialization unless inputs have changed
-	if (WasInitializedWithTheseAssets(SourceSkeletalMesh, TargetSkeletalMesh, InRetargeterAsset))
-	{
-		return;
-	}
-
-	// assumed initialization fails unless we get to the bottom of this function
-	bIsInitialized = false;
-
-	// record source asset
-	RetargeterAsset = InRetargeterAsset;
-
-	// check prerequisite assets
-	if (!RetargeterAsset)
-	{
-		Log.LogError(LOCTEXT("MissingRetargetAsset", "IK Retargeter unable to initialize because no IK Retargeter asset was supplied."));
-		return;
-	}
-	if (!SourceSkeletalMesh)
-	{
-		Log.LogError(LOCTEXT("MissingSourceMesh", "IK Retargeter unable to initialize. Missing source Skeletal Mesh asset."));
-		return;
-	}
-	if (!TargetSkeletalMesh)
-	{
-		Log.LogError(LOCTEXT("MissingTargetMesh", "IK Retargeter unable to initialize. Missing target Skeletal Mesh asset."));
-		return;
-	}
-	
-	// initialize source skeleton
-	SourceSkeleton.Initialize(
-		SourceSkeletalMesh,
-		ERetargetSourceOrTarget::Source,
-		RetargeterAsset,
-		GetPelvisBone(ERetargetSourceOrTarget::Source, ERetargetOpsToSearch::AssetOps),
-		GetPoseScale(ERetargetSourceOrTarget::Source));
-	
-	// initialize target skeleton
-	TargetSkeleton.Initialize(
-		TargetSkeletalMesh,
-		ERetargetSourceOrTarget::Target,
-		RetargeterAsset,
-		GetPelvisBone(ERetargetSourceOrTarget::Target, ERetargetOpsToSearch::AssetOps),
-		GetPoseScale(ERetargetSourceOrTarget::Target));
-
-	// copy ops from asset 
-	InitialOpStackSetup(RetargeterAsset->GetRetargetOps(), InCustomProfile);
-	
-	// initialize all the retarget ops
-	InitializeRetargetOps();
-
-	// rebuild the override cache with new memory locations
-	OverrideManager.DirtyProfileCache();
-
-	// confirm for the user that the retargeter is initialized
-	Log.LogInfo(FText::Format(
-			LOCTEXT("SuccessfulInit", "Success! The IK Retargeter is ready to transfer animation from the source, {0} to the target, {1}"),
-			FText::FromString(SourceSkeletalMesh->GetName()), FText::FromString(TargetSkeletalMesh->GetName())));
-	
-	bIsInitialized = true;
-	AssetVersionInitializedWith = RetargeterAsset->GetVersion();
-#if WITH_EDITOR
-	RetargeterInitialized.Broadcast();
-#endif
-}
-
 void FIKRetargetProcessor::InitialOpStackSetup(
 	const TArray<FInstancedStruct>& OpsFromAsset,
-	const FRetargetProfile& InRetargetProfile)
+	const FRetargetProfile* InRetargetProfile)
 {
 	// create copies of all the ops in the asset
 	OpStack.Reset(OpsFromAsset.Num());
@@ -1192,7 +1112,10 @@ void FIKRetargetProcessor::InitialOpStackSetup(
 		// create a copy of the op
 		int32 NewOpIndex = OpStack.Emplace(AssetOpStruct);
 		// apply the profile settings to the op (may be different than what was loaded from the asset)
-		InRetargetProfile.ApplyOpProfilesToOpStruct(OpStack[NewOpIndex], ECopyOpSettingsContext::PreInitialize);
+		if (InRetargetProfile)
+		{
+			InRetargetProfile->ApplyOpProfilesToOpStruct(OpStack[NewOpIndex], ECopyOpSettingsContext::PreInitialize);
+		}
 		
 #if WITH_EDITOR
 		FIKRetargetOpBase& NewOp = OpStack[NewOpIndex].GetMutable<FIKRetargetOpBase>();
@@ -1288,23 +1211,97 @@ void FIKRetargetProcessor::InitializeRetargetOps()
 	}
 }
 
-TArray<FTransform>&  FIKRetargetProcessor::RunRetargeter(
-	TArray<FTransform>& InSourceGlobalPose,
-	const FRetargetProfile& InCustomProfile,
-	const TArrayView<const FName> InOverrideSetsToApply,
-	const float InDeltaTime,
-	const int32 InLOD)
+void FIKRetargetProcessor::Initialize(const FRetargetInitParameters& Params)
+{
+	// check prerequisite assets
+	if (!Params.RetargeterAsset)
+	{
+		Log.LogError(LOCTEXT("MissingRetargetAsset", "IK Retargeter unable to initialize because no IK Retargeter asset was supplied."));
+		bIsInitialized = false;
+		return;
+	}
+	if (!Params.SourceSkeletalMesh)
+	{
+		Log.LogError(LOCTEXT("MissingSourceMesh", "IK Retargeter unable to initialize. Missing source Skeletal Mesh asset."));
+		bIsInitialized = false;
+		return;
+	}
+	if (!Params.TargetSkeletalMesh)
+	{
+		Log.LogError(LOCTEXT("MissingTargetMesh", "IK Retargeter unable to initialize. Missing target Skeletal Mesh asset."));
+		bIsInitialized = false;
+		return;
+	}
+	
+	// don't attempt reinitialization unless inputs have changed
+	if (WasInitializedWithTheseAssets(Params.SourceSkeletalMesh, Params.TargetSkeletalMesh, Params.RetargeterAsset))
+	{
+		return;
+	}
+
+	// assumed initialization fails unless we get to the bottom of this function
+	bIsInitialized = false;
+
+	// record source asset
+	RetargeterAsset = Params.RetargeterAsset;
+	
+	// initialize source skeleton
+	SourceSkeleton.Initialize(
+		Params.SourceSkeletalMesh,
+		ERetargetSourceOrTarget::Source,
+		RetargeterAsset,
+		GetPelvisBone(ERetargetSourceOrTarget::Source, ERetargetOpsToSearch::AssetOps),
+		GetPoseScale(ERetargetSourceOrTarget::Source));
+	
+	// initialize target skeleton
+	TargetSkeleton.Initialize(
+		Params.TargetSkeletalMesh,
+		ERetargetSourceOrTarget::Target,
+		RetargeterAsset,
+		GetPelvisBone(ERetargetSourceOrTarget::Target, ERetargetOpsToSearch::AssetOps),
+		GetPoseScale(ERetargetSourceOrTarget::Target));
+
+	// copy ops from asset 
+	InitialOpStackSetup(RetargeterAsset->GetRetargetOps(), Params.CustomProfile);
+	
+	// initialize all the retarget ops
+	InitializeRetargetOps();
+
+	// rebuild the override cache with new memory locations
+	OverrideManager.DirtyProfileCache();
+
+	// confirm for the user that the retargeter is initialized
+	Log.LogInfo(FText::Format(
+			LOCTEXT("SuccessfulInit", "Success! The IK Retargeter is ready to transfer animation from the source, {0} to the target, {1}"),
+			FText::FromString(Params.SourceSkeletalMesh->GetName()), FText::FromString(Params.TargetSkeletalMesh->GetName())));
+	
+	bIsInitialized = true;
+	AssetVersionInitializedWith = RetargeterAsset->GetVersion();
+#if WITH_EDITOR
+	RetargeterInitialized.Broadcast();
+#endif
+}
+
+TArray<FTransform>& FIKRetargetProcessor::RunRetargeter(const FRetargetRunParameters& Params)
 {
 	SCOPE_CYCLE_COUNTER(STAT_Anim_RetargetProcessorRun);
+
+	if (!ensure(Params.SourceGlobalPose))
+	{
+		Log.LogError(LOCTEXT("MissingInputPose", "IK Retargeter unable to run. Missing input source pose."));
+		return TargetSkeleton.OutputGlobalPose;
+	}
 	
 	if (!ensure(bIsInitialized))
 	{
-		return InSourceGlobalPose;
+		Log.LogError(LOCTEXT("NotInitialized", "IK Retargeter unable to run. Was not successfully initialized first."));
+		return TargetSkeleton.OutputGlobalPose;
 	}
 	
 #if WITH_EDITOR
 	// validate system running the retargeter has stripped all the scale out of the incoming pose
-	for (const FTransform& Transform : InSourceGlobalPose)
+	const TArray<FTransform>& SourcePose = *Params.SourceGlobalPose;
+	for (const FTransform& Transform : SourcePose)
 	{
 		const bool bHasNoScale = Transform.GetScale3D().Equals(FVector::OneVector);
 		ensureMsgf(bHasNoScale, TEXT("Found scale values on incoming pose in retarget processor. Scale should be baked into translation and set to 1,1,1. "));
@@ -1312,20 +1309,19 @@ TArray<FTransform>&  FIKRetargetProcessor::RunRetargeter(
 #endif
 
 	// apply the retargeting settings
-	UpdateSettingsAtRuntime(InCustomProfile, InOverrideSetsToApply);
+	UpdateSettingsAtRuntime(Params.Profile, Params.OverrideSetsToApply);
 
 	// applying settings can cause the retargeter to require reinitialization (when enabling/disabling IK on a chain)
 	// we should be able to safely reinitialize here because we are already initialized or it wouldn't pass the check() at the top
 	if (!IsInitialized())
 	{
-		constexpr bool bSuppressWarnings = false;
-		Initialize(
-			SourceSkeleton.SkeletalMesh,
-			TargetSkeleton.SkeletalMesh,
-			RetargeterAsset,
-			InCustomProfile,
-			InOverrideSetsToApply,
-			bSuppressWarnings);
+		FRetargetInitParameters InitParams;
+		InitParams.SourceSkeletalMesh = SourceSkeleton.SkeletalMesh;
+		InitParams.TargetSkeletalMesh = TargetSkeleton.SkeletalMesh;
+		InitParams.RetargeterAsset = RetargeterAsset;
+		InitParams.CustomProfile = Params.Profile;
+		InitParams.bSuppressWarnings = false;
+		Initialize(InitParams);
 
 		if (!IsInitialized())
 		{
@@ -1335,14 +1331,48 @@ TArray<FTransform>&  FIKRetargetProcessor::RunRetargeter(
 	}
 	
 	// generate the pose to start retargeting from
-	GenerateBasePoses(InSourceGlobalPose);
+	GenerateBasePoses(*Params.SourceGlobalPose);
 	
 	// run the stack of retargeting operations
-	RunRetargetOps(InSourceGlobalPose, TargetSkeleton.OutputGlobalPose, InDeltaTime, InLOD);
+	RunRetargetOps(*Params.SourceGlobalPose, TargetSkeleton.OutputGlobalPose, Params.DeltaTime, Params.LOD);
 	
 	return TargetSkeleton.OutputGlobalPose;
 }
 
+PRAGMA_DISABLE_DEPRECATION_WARNINGS
+void FIKRetargetProcessor::Initialize(
+    const USkeletalMesh* SourceSkeletalMesh, 
+    const USkeletalMesh* TargetSkeletalMesh,
+    const UIKRetargeter* InRetargeterAsset, 
+    const FRetargetProfile& InCustomProfile, 
+    const bool bSuppressWarnings)
+{
+	// forward to the new struct-based API
+    FRetargetInitParameters Params;
+    Params.SourceSkeletalMesh = SourceSkeletalMesh;
+    Params.TargetSkeletalMesh = TargetSkeletalMesh;
+    Params.RetargeterAsset = InRetargeterAsset;
+    Params.CustomProfile = &InCustomProfile;
+    Params.bSuppressWarnings = bSuppressWarnings;
+    Initialize(Params);
+}
+
+TArray<FTransform>& FIKRetargetProcessor::RunRetargeter(
+	TArray<FTransform>& InSourceGlobalPose,
+	const FRetargetProfile& InProfile,
+	const float InDeltaTime, 
+	const int32 InLOD)
+{
+	// forward to the new struct-based API
+	FRetargetRunParameters Params;
+	Params.SourceGlobalPose = &InSourceGlobalPose;
+	Params.Profile = &InProfile;
+	Params.DeltaTime = InDeltaTime;
+	Params.LOD = InLOD;
+	return RunRetargeter(Params);
+}
+PRAGMA_ENABLE_DEPRECATION_WARNINGS
+
 void FIKRetargetProcessor::GenerateBasePoses(TArray<FTransform>& InSourceGlobalPose)
 {
 	// start from the retarget pose
@@ -2050,24 +2080,30 @@ void FIKRetargetProcessor::SetNeedsInitialized()
 	}
 }
 
-void FIKRetargetProcessor::UpdateSettingsAtRuntime(const FRetargetProfile& Profile, const TArrayView<const FName> InOverrideSetsToApply)
+void FIKRetargetProcessor::UpdateSettingsAtRuntime(const FRetargetProfile* InProfile, const TArrayView<const FName> InOverrideSetsToApply)
 {
 	// apply op settings from the profile
-	for (FInstancedStruct& OpStruct : OpStack)
+	if (InProfile)
 	{
-		Profile.ApplyOpProfilesToOpStruct(OpStruct, ECopyOpSettingsContext::Runtime);
+		for (FInstancedStruct& OpStruct : OpStack)
+		{
+			InProfile->ApplyOpProfilesToOpStruct(OpStruct, ECopyOpSettingsContext::Runtime);
+		}
 	}
 
 	// apply stored property overrides
 	OverrideManager.ApplyPropertyOverridesToOps(*this, InOverrideSetsToApply);
-
-	// apply flag to force IK off
-	bIKForcedOff = Profile.bForceAllIKOff;
 	
-	// apply retarget poses specified in the profile
-	// NOTE: must be done AFTER applying op settings because it uses the SourceScaleFactor in the ScaleSourceOp settings
-	UpdateRetargetPoseAtRuntime(Profile.SourceRetargetPoseName, ERetargetSourceOrTarget::Source);
-	UpdateRetargetPoseAtRuntime(Profile.TargetRetargetPoseName, ERetargetSourceOrTarget::Target);
+	if (InProfile)
+	{
+		// apply flag to force IK off
+		bIKForcedOff = InProfile->bForceAllIKOff;
+	
+		// apply retarget poses specified in the profile
+		// NOTE: must be done AFTER applying op settings because it uses the SourceScaleFactor in the ScaleSourceOp settings
+		UpdateRetargetPoseAtRuntime(InProfile->SourceRetargetPoseName, ERetargetSourceOrTarget::Source);
+		UpdateRetargetPoseAtRuntime(InProfile->TargetRetargetPoseName, ERetargetSourceOrTarget::Target);
+	}
 }
 
 void FIKRetargetProcessor::UpdateRetargetPoseAtRuntime(
diff --git a/Engine/Plugins/Animation/IKRig/Source/IKRig/Public/Retargeter/IKRetargetProcessor.h b/Engine/Plugins/Animation/IKRig/Source/IKRig/Public/Retargeter/IKRetargetProcessor.h
--- a/Engine/Plugins/Animation/IKRig/Source/IKRig/Public/Retargeter/IKRetargetProcessor.h
+++ b/Engine/Plugins/Animation/IKRig/Source/IKRig/Public/Retargeter/IKRetargetProcessor.h
@@ -349,6 +349,41 @@ enum class ERetargetOpsToSearch : uint8
 	AssetOps
 };
 
+/** * Parameters for initializing the IK Retargeter.
+* Initialization does a lot of validation and can fail for many reasons. 
+* Check IsInitialized() afterwards.
+*/
+struct FRetargetInitParameters
+{
+	/** The skeletal mesh to pose FROM */
+	const USkeletalMesh* SourceSkeletalMesh = nullptr;
+	/** The skeletal mesh to pose TO */
+	const USkeletalMesh* TargetSkeletalMesh = nullptr;
+	/** The source asset to use for retargeting settings */
+	const UIKRetargeter* RetargeterAsset = nullptr;
+	/** The collection of settings used to initialize with */
+	const FRetargetProfile* CustomProfile = nullptr;
+	/** The names of any sets of overrides sets to apply (stored on asset) */
+	TArrayView<const FName> OverrideSetsToApply;
+	/** If true, will not output warnings during initialization */
+	bool bSuppressWarnings = false;
+};
+
+/** Parameters for running the retargeter to generate a new pose. */
+struct FRetargetRunParameters
+{
+	/** The source mesh input pose in Component/Global space */
+	TArray<FTransform>* SourceGlobalPose = nullptr;
+	/** The retarget profile to use for this update */
+	const FRetargetProfile* Profile = nullptr;
+	/** An array of override set names stored on the asset to apply */
+	TArrayView<const FName> OverrideSetsToApply;
+	/** Time since last tick in seconds (passed to ops) */
+	float DeltaTime = 0.0f;
+	/** The current LOD level for the skeletal mesh */
+	int32 LOD = INDEX_NONE;
+};
+
 /** The runtime processor that converts an input pose from a source skeleton into an output pose on a target skeleton.
  * To use:
  * 1. Initialize a processor with a Source/Target skeletal mesh and a UIKRetargeter asset.
@@ -360,40 +395,19 @@ struct FIKRetargetProcessor
 {
 
 public:
-	
+
 	/**
 	* Initialize the retargeter to enable running it.
-	* @param SourceSkeleton - the skeletal mesh to poses FROM
-	* @param TargetSkeleton - the skeletal mesh to poses TO
-	* @param InRetargeterAsset - the source asset to use for retargeting settings
-	* @param InCustomProfile - the collection of settings used to initialize with
-	* @param InOverrideSetsToApply - the names of any sets of overrides sets to apply (stored on asset)
-	* @param bSuppressWarnings - if true, will not output warnings during initialization
-	* @warning - Initialization does a lot of validation and can fail for many reasons. Check IsInitialized() afterwards.
+	* @param Params - The initialization settings and source/target skeletons.
 	*/
-	UE_API void Initialize(
-		const USkeletalMesh *SourceSkeleton,
-		const USkeletalMesh *TargetSkeleton,
-		const UIKRetargeter* InRetargeterAsset,
-		const FRetargetProfile& InCustomProfile,
-		const TArrayView<const FName> InOverrideSetsToApply,
-		const bool bSuppressWarnings=false);
+	UE_API void Initialize(const FRetargetInitParameters& Params);
 
 	/**
 	* Run the retarget to generate a new pose.
-	* @param InSourceGlobalPose -  is the source mesh input pose in Component/Global space
-	* @param InProfile -  the retarget profile to use for this update
-	* @param InOverrideSetsToApply -  an array of override set names stored on the asset to apply
-	* @param InDeltaTime -  time since last tick in seconds (passed to ops)
-	* @param InLOD - the current LOD level for the skeletal mesh
+	* @param Params - The input pose and runtime configuration.
 	* @return The retargeted Component/Global space pose for the target skeleton
 	*/
-	UE_API TArray<FTransform>& RunRetargeter(
-		TArray<FTransform>& InSourceGlobalPose,
-		const FRetargetProfile& InProfile,
-		const TArrayView<const FName> InOverrideSetsToApply,
-		const float InDeltaTime,
-		const int32 InLOD = INDEX_NONE);
+	UE_API TArray<FTransform>& RunRetargeter(const FRetargetRunParameters& Params);
 
 	/** Get whether this processor is ready to call RunRetargeter() and generate new poses. */
 	UE_API bool IsInitialized() const;
@@ -543,6 +557,22 @@ public:
 	
 #endif
 
+	// BEGIN DEPRECATED API
+	UE_DEPRECATED(5.8, "Please use the Initialize overload that takes FRetargetInitParameters.")
+	UE_API void Initialize(
+		const USkeletalMesh* SourceSkeleton,
+		const USkeletalMesh* TargetSkeleton,
+		const UIKRetargeter* InRetargeterAsset,
+		const FRetargetProfile& InCustomProfile,
+		const bool bSuppressWarnings = false);
+	UE_DEPRECATED(5.8, "Please use the RunRetargeter overload that takes FRetargetRunParameters.")
+	UE_API TArray<FTransform>& RunRetargeter(
+		TArray<FTransform>& InSourceGlobalPose,
+		const FRetargetProfile& InProfile,
+		const float InDeltaTime,
+		const int32 InLOD = -1);
+	// END DEPRECATED API
+	
 private:
 
 	/** If true, all IK operations are skipped
@@ -578,7 +608,7 @@ private:
 	mutable FRetargetPoseScaleWithPivot CurrentSourceScale;
 
 	/** Apply the settings stored in a retarget profile. As well as any stored property overrides. */
-	void UpdateSettingsAtRuntime(const FRetargetProfile& Profile, const TArrayView<const FName> InOverrideSetsToApply);
+	void UpdateSettingsAtRuntime(const FRetargetProfile* Profile, const TArrayView<const FName> InOverrideSetsToApply);
 
 	/** Internal retarget phase that does simple bone-to-bone copying from source*/
 	void GenerateBasePoses(TArray<FTransform>& InSourceGlobalPose);
@@ -586,7 +616,7 @@ private:
 	/** Initial setup of the retarget op stack (copies op stack from asset into processor) */
 	void InitialOpStackSetup(
 		const TArray<FInstancedStruct>& OpsFromAsset,
-		const FRetargetProfile& InRetargetProfile);
+		const FRetargetProfile* InRetargetProfile);
 
 	/** Run Initialize() on all ops then collects retargeted bones and call PostInitialize() on each Op */
 	void InitializeRetargetOps();
@@ -667,7 +697,7 @@ public:
 		const FRetargetProfile& Settings,
 		const bool bSuppressWarnings=false)
 	{
-		Processor.Initialize(SourceSkeleton, TargetSkeleton, InRetargeterAsset, Settings, TArray<FName>(),  bSuppressWarnings);
+		Processor.Initialize(SourceSkeleton, TargetSkeleton, InRetargeterAsset, Settings, bSuppressWarnings);
 	}
 	
 	TArray<FTransform>& RunRetargeter(
@@ -676,7 +706,7 @@ public:
 		const float DeltaTime,
 		const FRetargetProfile& Profile)
 	{
-		return Processor.RunRetargeter(InSourceGlobalPose, Profile, TArray<FName>(), DeltaTime);
+		return Processor.RunRetargeter(InSourceGlobalPose, Profile, DeltaTime);
 	}
 
 	void UpdateRetargetPoseAtRuntime(
diff --git a/Engine/Plugins/Animation/IKRig/Source/IKRigEditor/Private/RetargetEditor/IKRetargetBatchOperation.cpp b/Engine/Plugins/Animation/IKRig/Source/IKRigEditor/Private/RetargetEditor/IKRetargetBatchOperation.cpp
--- a/Engine/Plugins/Animation/IKRig/Source/IKRigEditor/Private/RetargetEditor/IKRetargetBatchOperation.cpp
+++ b/Engine/Plugins/Animation/IKRig/Source/IKRigEditor/Private/RetargetEditor/IKRetargetBatchOperation.cpp
@@ -276,11 +276,18 @@ void UIKRetargetBatchOperation::RetargetAssets(
 	}
 
 	// convert the animation using the IK retargeter	
-	// initialize the retarget processor
 	FIKRetargetProcessor Processor;
 	FRetargetProfile RetargetProfile;
 	RetargetProfile.FillProfileWithAssetSettings(Context.IKRetargetAsset);
-	Processor.Initialize(Context.SourceMesh, Context.TargetMesh, Context.IKRetargetAsset, RetargetProfile, Context.OverrideSetNames);
+	// initialize the retarget processor
+	FRetargetInitParameters Params;
+	Params.SourceSkeletalMesh = Context.SourceMesh;
+	Params.TargetSkeletalMesh = Context.TargetMesh;
+	Params.RetargeterAsset = Context.IKRetargetAsset;
+	Params.CustomProfile = &RetargetProfile;
+	Params.OverrideSetsToApply = Context.OverrideSetNames;
+	Params.bSuppressWarnings = false;
+	Processor.Initialize(Params);
 	if (!Processor.IsInitialized())
 	{
 		UE_LOG(LogTemp, Warning, TEXT("Unable to initialize the IK Retargeter. Newly created animations were not retargeted!"));
@@ -483,7 +490,12 @@ void UIKRetargetBatchOperation::ConvertAnimation(
 			OutProcessor.UpdateOpsFromAnimSequence(SourceSequence, TimeAtCurrentFrame);
 
 			// run the retargeter
-			const TArray<FTransform>& TargetComponentPose = OutProcessor.RunRetargeter(SourceComponentPose, SettingsProfile, Context.OverrideSetNames, DeltaTime);
+			FRetargetRunParameters Params;
+			Params.SourceGlobalPose = &SourceComponentPose;
+			Params.Profile = &SettingsProfile;
+			Params.OverrideSetsToApply = Context.OverrideSetNames;
+			Params.DeltaTime = DeltaTime;
+			const TArray<FTransform>& TargetComponentPose = OutProcessor.RunRetargeter(Params);
 
 			// convert to a local-space pose
 			TArray<FTransform> TargetLocalPose = TargetComponentPose;
diff --git a/Engine/Plugins/Animation/IKRig/Source/IKRigEditor/Private/RetargetEditor/IKRetargeterPoseGenerator.cpp b/Engine/Plugins/Animation/IKRig/Source/IKRigEditor/Private/RetargetEditor/IKRetargeterPoseGenerator.cpp
--- a/Engine/Plugins/Animation/IKRig/Source/IKRigEditor/Private/RetargetEditor/IKRetargeterPoseGenerator.cpp
+++ b/Engine/Plugins/Animation/IKRig/Source/IKRigEditor/Private/RetargetEditor/IKRetargeterPoseGenerator.cpp
@@ -625,14 +625,17 @@ bool FRetargetAutoPoseGenerator::CheckReadyToAlignBones() const
 		return false;
 	}
 
-	// we have to initialize a processor to "resolve" the retarget asset setup onto actual skeletons to do the alignment
-	USkeletalMesh* SourceSkeletalMesh = Controller->GetPreviewMesh(ERetargetSourceOrTarget::Source);
-	USkeletalMesh* TargetSkeletalMesh = Controller->GetPreviewMesh(ERetargetSourceOrTarget::Target);
-	UIKRetargeter* RetargeterAsset = Controller->GetAsset();
-	constexpr bool bSuppressWarnings = true;
+	// initialize retarget processor with source and target skeletal meshes
+	// NOTE: we have to initialize a processor to "resolve" the retarget asset setup onto actual skeletons to do the alignment
 	FRetargetProfile RetargetProfile;
-	RetargetProfile.FillProfileWithAssetSettings(RetargeterAsset);
-	Processor->Initialize(SourceSkeletalMesh, TargetSkeletalMesh, RetargeterAsset, RetargetProfile, TArray<FName>(), bSuppressWarnings);
+	RetargetProfile.FillProfileWithAssetSettings(Controller->GetAsset());
+	FRetargetInitParameters Params;
+	Params.SourceSkeletalMesh = Controller->GetPreviewMesh(ERetargetSourceOrTarget::Source);
+	Params.TargetSkeletalMesh = Controller->GetPreviewMesh(ERetargetSourceOrTarget::Target);
+	Params.RetargeterAsset = Controller->GetAsset();
+	Params.CustomProfile = &RetargetProfile;
+	Params.bSuppressWarnings = true;
+	Processor->Initialize(Params);
 	
 	// can't auto align until processor is initialized because we need a fully resolved chain mapping on source/target
 	// this could happen if the user attempts to edit the retarget pose on an incompatible IK Rig
diff --git a/Engine/Plugins/Animation/IKRig/Source/IKRigUAF/Private/RigUnit_UAFIKRetargeter.cpp b/Engine/Plugins/Animation/IKRig/Source/IKRigUAF/Private/RigUnit_UAFIKRetargeter.cpp
--- a/Engine/Plugins/Animation/IKRig/Source/IKRigUAF/Private/RigUnit_UAFIKRetargeter.cpp
+++ b/Engine/Plugins/Animation/IKRig/Source/IKRigUAF/Private/RigUnit_UAFIKRetargeter.cpp
@@ -40,7 +40,14 @@ FRigUnit_UAFIKRetargeter_Execute()
 	if (!Processor.WasInitializedWithTheseAssets(SourceMesh, TargetMesh, IKRetargetAsset))
 	{
 		// initialize retarget processor with source and target skeletal meshes
-		Processor.Initialize(SourceMesh, TargetMesh, IKRetargetAsset, WorkData.Profile, OverrideSetNames);
+		FRetargetInitParameters Params;
+		Params.SourceSkeletalMesh = SourceMesh;
+		Params.TargetSkeletalMesh = TargetMesh;
+		Params.RetargeterAsset = IKRetargetAsset;
+		Params.CustomProfile = &WorkData.Profile;
+		Params.OverrideSetsToApply = OverrideSetNames;
+		Params.bSuppressWarnings = false;
+		Processor.Initialize(Params);
 	}
 
 	// bail out if not initialized
@@ -64,12 +71,12 @@ FRigUnit_UAFIKRetargeter_Execute()
 	FAnimationRuntime::FillUpComponentSpaceTransforms(SourceRefSkeleton, WorkData.SourceMeshPoseLocal, WorkData.SourceMeshPoseGlobal);
 	
 	// run the retargeter
-	const TArray<FTransform>& TargetGlobalPose = Processor.RunRetargeter(
-	WorkData.SourceMeshPoseGlobal,
-	WorkData.Profile,
-	OverrideSetNames,
-	ExecuteContext.GetDeltaTime(),
-	TargetRefPose.GetSourceLODLevel());
+	FRetargetRunParameters Params;
+	Params.SourceGlobalPose = &WorkData.SourceMeshPoseGlobal;
+	Params.Profile = &WorkData.Profile;
+	Params.DeltaTime = ExecuteContext.GetDeltaTime();
+	Params.LOD = TargetRefPose.GetSourceLODLevel();
+	const TArray<FTransform>& TargetGlobalPose = Processor.RunRetargeter(Params);
 
 	// fill the target pose with a ref pose at the appropriate LOD level
 	TargetPose.LODPose.PrepareForLOD(TargetRefPose, TargetRefPose.GetSourceLODLevel(), true/*bSetRefPose*/, false/*bAdditive*/);
