[IK Retarget] Add support for headless mode in batch retargeter.

#Synced-CL 51992877

diff --git a/Engine/Plugins/Animation/IKRig/Source/IKRigEditor/Private/RetargetEditor/IKRetargetBatchOperation.cpp b/Engine/Plugins/Animation/IKRig/Source/IKRigEditor/Private/RetargetEditor/IKRetargetBatchOperation.cpp
--- a/Engine/Plugins/Animation/IKRig/Source/IKRigEditor/Private/RetargetEditor/IKRetargetBatchOperation.cpp
+++ b/Engine/Plugins/Animation/IKRig/Source/IKRigEditor/Private/RetargetEditor/IKRetargetBatchOperation.cpp
@@ -725,6 +725,12 @@ void UIKRetargetBatchOperation::NotifyUserOfResults(
 	const FIKRetargetBatchOperationContext& Context,
 	FScopedSlowTask& Progress) const
 {
+	// skip UI interaction if batch retargeter is ran in headless mode
+	if (IsRunningCommandlet() || FApp::IsUnattended() )
+	{
+		return;
+	}
+	
 	// gather newly created objects
 	TArray<UObject*> NewAssets;
 	GetNewAssets(NewAssets);
