[IK Retarget] Adding variable inputs to drive property overrides.

#Synced-CL 52117481

diff --git a/Engine/Plugins/Animation/IKRig/Source/IKRig/Private/AnimNodes/AnimNode_RetargetPoseFromMesh.cpp b/Engine/Plugins/Animation/IKRig/Source/IKRig/Private/AnimNodes/AnimNode_RetargetPoseFromMesh.cpp
--- a/Engine/Plugins/Animation/IKRig/Source/IKRig/Private/AnimNodes/AnimNode_RetargetPoseFromMesh.cpp
+++ b/Engine/Plugins/Animation/IKRig/Source/IKRig/Private/AnimNodes/AnimNode_RetargetPoseFromMesh.cpp
@@ -65,6 +65,9 @@ void FAnimNode_RetargetPoseFromMesh::Update_AnyThread(const FAnimationUpdateCont
 		// update pose input
 		Source.Update(Context);
 	}
+
+	// get the variable values from the input pins
+	CopyInputPropertiesToVariables(Context.AnimInstanceProxy->GetAnimInstanceObject());
 }
 
 void FAnimNode_RetargetPoseFromMesh::Evaluate_AnyThread(FPoseContext& Output)
@@ -235,6 +238,8 @@ void FAnimNode_RetargetPoseFromMesh::Evaluate_AnyThread(FPoseContext& Output)
 	Params.SourceGlobalPose = &PoseToRetargetFromComponentSpace;
 	Params.Profile = &RetargetProfileToUse;
 	Params.OverrideSetsToApply = *OverridesToUse;
+	Params.Variables = &RuntimeVariables;
+	Params.BoundCurveValues = &BoundCurveValues;
 	Params.DeltaTime = DeltaTime;
 	Params.LOD = CurrentLODLevel;
 	const TArray<FTransform>& RetargetedPose = Processor.RunRetargeter(Params);
@@ -326,6 +331,14 @@ void FAnimNode_RetargetPoseFromMesh::PreUpdate(const UAnimInstance* InAnimInstan
 			Processor.OnAnimGraphPreUpdateMainThread(*SourceMeshComponent.Get(), *TargetMeshComponent);
 		}
 	}
+
+	// get the curve values used by property overrides that are bound to curves
+	const bool bCopyingFromOtherComponent = RetargetFrom != ERetargetSourceMode::SourcePosePin;
+	TWeakObjectPtr<USkeletalMeshComponent> ComponentWithCurves = bCopyingFromOtherComponent ? SourceMeshComponent : TargetMeshComponent;
+	if (ComponentWithCurves.IsValid())
+	{
+		CopyBoundCurveValues(ComponentWithCurves->GetAnimInstance());
+	}
 }
 
 FIKRetargetProcessor* FAnimNode_RetargetPoseFromMesh::GetRetargetProcessor()
@@ -402,7 +415,6 @@ bool FAnimNode_RetargetPoseFromMesh::EnsureProcessorIsInitialized(const USkeleta
 		Params.TargetSkeletalMesh = TargetMesh;
 		Params.RetargeterAsset = IKRetargeterAsset;
 		Params.CustomProfile = &RetargetProfileToUse;
-		Params.OverrideSetsToApply = OverrideSetsToApply;
 		Params.bSuppressWarnings = false;
 		Processor.Initialize(Params);
 	}
@@ -451,3 +463,133 @@ double FAnimNode_RetargetPoseFromMesh::GetAverageExecutionTime() const
 }
 #endif
 
+void FAnimNode_RetargetPoseFromMesh::CopyInputPropertiesToVariables(const UObject* InSourceInstance)
+{
+	if (!InSourceInstance || !IKRetargeterAsset)
+	{
+		return;
+	}
+
+	// start the bag with the hardcoded values from the asset (possibly overridden by pin values)
+	RuntimeVariables = IKRetargeterAsset->GetVariables().Bag;
+
+	const UPropertyBag* BagStruct = RuntimeVariables.GetPropertyBagStruct();
+	if (!BagStruct)
+	{
+		return;
+	}
+
+	UClass* SourceClass = InSourceInstance->GetClass();
+
+	// get the mutable memory address of the Property Bag so we can write to it
+	uint8* BagMemory = RuntimeVariables.GetMutableValue().GetMemory();
+	if (!BagMemory)
+	{
+		return;
+	}
+
+	if (!ensure(SourcePropertyNames.Num() == DestPropertyNames.Num()))
+	{
+		return; // should not happen
+	}
+
+	// iterate through the compiler-generated mappings and copy pin values to the variables property bag
+	for (int32 PropIdx = 0; PropIdx < SourcePropertyNames.Num(); ++PropIdx)
+	{
+		const FName& SourceName = SourcePropertyNames[PropIdx];
+		const FName& DestName = DestPropertyNames[PropIdx]; 
+
+		FProperty* SourceProp = SourceClass->FindPropertyByName(SourceName);
+		FProperty* BagProp = BagStruct->FindPropertyByName(DestName);
+		if (!SourceProp || !BagProp)
+		{
+			continue;
+		}
+		
+		const void* SourceValuePtr = SourceProp->ContainerPtrToValuePtr<void>(InSourceInstance);
+		void* BagValuePtr = BagProp->ContainerPtrToValuePtr<void>(BagMemory);
+
+		// EXACT MATCH: fast mem copy (ie, bool to bool, int32 to int32)
+		if (SourceProp->SameType(BagProp))
+		{
+			SourceProp->CopyCompleteValue(BagValuePtr, SourceValuePtr);
+		}
+		// NUMERIC MISMATCH: safely coerce numbers...
+		else if (FNumericProperty* SourceNum = CastField<FNumericProperty>(SourceProp))
+		{
+			if (FNumericProperty* BagNum = CastField<FNumericProperty>(BagProp))
+			{
+				// handle float type bridging (double, float)
+				if (SourceNum->IsFloatingPoint() && BagNum->IsFloatingPoint())
+				{
+					double Val = SourceNum->GetFloatingPointPropertyValue(SourceValuePtr);
+					BagNum->SetFloatingPointPropertyValue(BagValuePtr, Val);
+				}
+				// handle integer type bridging (int64, int32, byte)
+				else if (SourceNum->IsInteger() && BagNum->IsInteger())
+				{
+					int64 Val = SourceNum->GetSignedIntPropertyValue(SourceValuePtr);
+					BagNum->SetIntPropertyValue(BagValuePtr, Val);
+				}
+			}
+		}
+		// COMPLEX MISMATCH: handle struct bridging (ie, FVector3d -> FVector3f)
+		else
+		{
+			// string export/import handles struct precision differences (ie, LWC vector to float vector etc)
+			FString TempStr;
+			SourceProp->ExportTextItem_Direct(TempStr, SourceValuePtr, nullptr, nullptr, PPF_None);
+			BagProp->ImportText_Direct(*TempStr, BagValuePtr, nullptr, PPF_None);
+		}
+	}
+}
+
+void FAnimNode_RetargetPoseFromMesh::CopyBoundCurveValues(const UObject* InSourceInstance)
+{
+	const UAnimInstance* AnimInstance = Cast<UAnimInstance>(InSourceInstance);
+	if (!AnimInstance)
+	{
+		return;
+	}
+	
+	if (!IKRetargeterAsset)
+	{
+		return;
+	}
+
+	// rebuild the TMap with the latest bound curve names
+	if (CachedOverrideVersion != IKRetargeterAsset->GetOverrideVersion())
+	{
+		BoundCurveValues.Reset();
+		
+		const TMap<FName, FRetargetOverrideSet>& AllOverrideSets = IKRetargeterAsset->GetOverrideSets();
+		for (const TTuple<FName, FRetargetOverrideSet>& OverrideSet : AllOverrideSets)
+		{
+			for (const FRetargetOpOverrides& OpOverrides : OverrideSet.Value.OpOverrides)
+			{
+				for (const FRetargetOpPropertyOverride& PropertyOverride : OpOverrides.PropertyOverrides)
+				{
+					const FName BoundCurveName = PropertyOverride.GetBoundCurveName();
+					if (BoundCurveName.IsNone())
+					{
+						continue;
+					}
+					BoundCurveValues.Emplace(BoundCurveName, 0.0f);
+				}
+			}
+		}
+		
+		CachedOverrideVersion = IKRetargeterAsset->GetOverrideVersion();
+	}
+
+	// fill the map with values from the curves
+	const TMap<FName, float>& AnimCurveList = AnimInstance->GetAnimationCurveList(EAnimCurveType::AttributeCurve);
+	for (TTuple<FName, float>& BoundCurve : BoundCurveValues)
+	{
+		if (const float* AnimCurveValue = AnimCurveList.Find(BoundCurve.Key))
+		{
+			BoundCurve.Value = *AnimCurveValue;
+		}
+	}
+}
+
diff --git a/Engine/Plugins/Animation/IKRig/Source/IKRig/Private/Retargeter/IKRetargetOverrideManager.cpp b/Engine/Plugins/Animation/IKRig/Source/IKRig/Private/Retargeter/IKRetargetOverrideManager.cpp
new file mode 100644
--- /dev/null
+++ b/Engine/Plugins/Animation/IKRig/Source/IKRig/Private/Retargeter/IKRetargetOverrideManager.cpp
@@ -0,0 +1,356 @@
+﻿// Copyright Epic Games, Inc. All Rights Reserved.
+
+#include "Retargeter/IKRetargetOverrideManager.h"
+
+#include "StructUtils/PropertyBag.h"
+#include "Retargeter/IKRetargetProcessor.h"
+
+void FIKRetargetRuntimeOverrideManager::ApplyPropertyOverridesToOps(
+	FIKRetargetProcessor& InProcessor,
+	const TArrayView<const FName> InOverrideSetsToApply,
+	FInstancedPropertyBag* InVariables,
+	TMap<FName, float>* InBoundCurveValues)
+{
+	const UIKRetargeter* Asset = InProcessor.GetRetargetAsset();
+	if (!ensure(Asset))
+	{
+		return;
+	}
+
+	UpdateFullListOfOverrideSets(InProcessor, InOverrideSetsToApply);
+	
+	const int32 NewOverrideSetHash = GetOverrideSetArrayHash(OverrideSetsToApply); 
+	if (bIsCacheDirty || (NewOverrideSetHash != OverrideSetHash) || (CachedAssetVersion != Asset->GetOverrideVersion()))
+	{
+		RebuildCache(InProcessor, InVariables);
+		OverrideSetHash = NewOverrideSetHash;
+		CachedAssetVersion = Asset->GetOverrideVersion();
+		bIsCacheDirty = false;
+	}
+
+#if WITH_EDITOR
+	UpdateCachedValues(InProcessor);
+#endif
+
+	const uint8* BagMemory = InVariables ? InVariables->GetValue().GetMemory() : nullptr;
+
+	// the "Hot Path", a linear iteration and direct memory write of all individual property overrides
+	for (const FCompiledPropertyOverride& Entry : CompiledOverrides)
+	{
+		// 1. Dynamic Bound Curve
+		if (Entry.BoundCurveName != NAME_None && InBoundCurveValues)
+		{
+			if (const float* CurveValuePtr = InBoundCurveValues->Find(Entry.BoundCurveName))
+			{
+				const float CurveVal = *CurveValuePtr;
+				switch (Entry.CurveCoercionType)
+				{
+					case ECurveCoercionType::Float:
+						*(float*)Entry.DestPtr = CurveVal;
+						break;
+					case ECurveCoercionType::Double:
+						*(double*)Entry.DestPtr = (double)CurveVal;
+						break;
+					case ECurveCoercionType::Int32:
+						*(int32*)Entry.DestPtr = FMath::RoundToInt(CurveVal);
+						break;
+					case ECurveCoercionType::Bool:
+						// NOTE: Unreal can compact multiple UProperty bools into a single byte!
+						// So bools must use the property setter to safely respect bitmask boundaries in structs
+						((FBoolProperty*)Entry.Property)->SetPropertyValue(Entry.DestPtr, CurveVal > 0.5f);
+						break;
+					case ECurveCoercionType::Byte:
+						*(uint8*)Entry.DestPtr = (uint8)FMath::RoundToInt(FMath::Clamp(CurveVal, 0, 255));
+						break;
+					case ECurveCoercionType::Int64:
+						*(int64*)Entry.DestPtr = (int64)FMath::RoundToInt(CurveVal);
+						break;
+					case ECurveCoercionType::None:
+					default:
+						break;
+				}
+			}
+		}
+		// 2. Dynamic Bound Variable
+		else if (Entry.BagMemoryOffset != INDEX_NONE && BagMemory)
+		{
+			FMemory::Memcpy(Entry.DestPtr, BagMemory + Entry.BagMemoryOffset, Entry.Size);
+		}
+		// 3. Static Hardcoded String
+		else if (Entry.SourcePoolOffset != INDEX_NONE)
+		{
+			FMemory::Memcpy(Entry.DestPtr, CachedValuePool.GetData() + Entry.SourcePoolOffset, Entry.Size);
+		}
+	}
+}
+
+void FIKRetargetRuntimeOverrideManager::UpdateFullListOfOverrideSets(FIKRetargetProcessor& Processor, const TArrayView<const FName> InOverrideSetsToApply)
+{
+	OverrideSetsToApply.Reset();
+	
+	const UIKRetargeter* Asset = Processor.GetRetargetAsset();
+	if (!ensure(Asset))
+	{
+		return;
+	}
+	
+	for (const FName& OverrideSetName : InOverrideSetsToApply)
+	{
+		TArray<FName> CurrentChain;
+		FName CurrentName = OverrideSetName;
+		while (CurrentName != NAME_None)
+		{
+			const FRetargetOverrideSet* OverrideSet = Asset->GetOverrideSetByName(CurrentName);
+			if (!OverrideSet) break;
+
+			CurrentChain.Add(CurrentName);
+			CurrentName = OverrideSet->ParentName;
+		}
+
+		for (int32 i = CurrentChain.Num() - 1; i >= 0; --i)
+		{
+			OverrideSetsToApply.AddUnique(CurrentChain[i]);
+		}
+	}
+}
+
+uint32 FIKRetargetRuntimeOverrideManager::GetOverrideSetArrayHash(const TArray<FName>& InProfileNames)
+{
+	uint32 CombinedHash = 0;
+	for (const FName& ProfileName : InProfileNames)
+	{
+		CombinedHash = HashCombine(CombinedHash, GetTypeHash(ProfileName));
+	}
+	return CombinedHash;
+}
+
+void FIKRetargetRuntimeOverrideManager::RebuildCache(FIKRetargetProcessor& Processor, FInstancedPropertyBag* InVariables)
+{
+	CompiledOverrides.Reset();
+	CachedValuePool.Reset();
+	CachedVersionPerOp.Reset();
+
+	const UIKRetargeter* Asset = Processor.GetRetargetAsset();
+	if (!ensure(Asset))
+	{
+		return;
+	}
+	
+	const UPropertyBag* BagStruct = InVariables ? InVariables->GetPropertyBagStruct() : nullptr;
+
+	// first we calculate total memory needed for the pool
+	int32 TotalPoolSize = 0;
+	for (const FName& ProfileName : OverrideSetsToApply)
+	{
+		const FRetargetOverrideSet* OverrideSet = Asset->GetOverrideSetByName(ProfileName);
+		if (!OverrideSet)
+		{
+			continue;
+		}
+		
+		for (const FRetargetOpOverrides& OpOverrides : OverrideSet->OpOverrides)
+		{
+			const UScriptStruct* OpScriptStruct = OpOverrides.ScriptStruct;
+			if (!ensure(OpScriptStruct))
+			{
+				continue;
+			}
+
+			for (const FRetargetOpPropertyOverride& PropOverride : OpOverrides.PropertyOverrides)
+			{
+				TArray<FRetargetOpPropertyOverride::FPropertySegment> Segments;
+				if (FRetargetOpPropertyOverride::GetSegmentsFromProperyPath(PropOverride.GetPropertyPath(), OpScriptStruct, Segments))
+				{
+					if (FProperty* LeafProp = FRetargetOpPropertyOverride::GetLeafProperty(Segments))
+					{
+						TotalPoolSize += LeafProp->GetSize();
+					}
+				}
+			}
+		}
+	}
+	CachedValuePool.Reserve(TotalPoolSize);
+
+	// now bake values into the pool and cache the pointers
+	for (const FName& ProfileName : OverrideSetsToApply)
+	{
+		const FRetargetOverrideSet* OverrideSet = Asset->GetOverrideSetByName(ProfileName);
+		if (!OverrideSet)
+		{
+			continue;
+		}
+
+		for (const FRetargetOpOverrides& OpOverrides : OverrideSet->OpOverrides)
+		{
+			FIKRetargetOpBase* OpInProcessor = Processor.GetRetargetOpByName(OpOverrides.OpName);
+			if (!ensure(OpInProcessor))
+			{
+				continue;
+			}
+			
+			FIKRetargetOpSettingsBase* OpSettings = OpInProcessor->GetSettings();
+			if (!ensure(OpSettings))
+			{
+				continue;
+			}
+			
+			uint8* OpSettingsData = reinterpret_cast<uint8*>(OpSettings);
+			const UScriptStruct* OpScriptStruct = OpOverrides.ScriptStruct;
+
+			if (!ensure(OpScriptStruct))
+			{
+				continue;
+			}
+
+			for (const FRetargetOpPropertyOverride& PropOverride : OpOverrides.PropertyOverrides)
+			{
+				TArray<FRetargetOpPropertyOverride::FPropertySegment> Segments;
+				if (!ensure(FRetargetOpPropertyOverride::GetSegmentsFromProperyPath(PropOverride.GetPropertyPath(), OpScriptStruct, Segments))) continue;
+				
+				FProperty* LeafProp = FRetargetOpPropertyOverride::GetLeafProperty(Segments);
+				uint8* DestAddress = FRetargetOpPropertyOverride::GetDataPointerFromPathSegments(OpSettingsData, Segments);
+
+				if (!ensure(LeafProp && DestAddress))
+				{
+					continue;
+				}
+				
+				const int32 ValueSize = LeafProp->GetSize();
+
+				// record the move parameters
+				FCompiledPropertyOverride& NewMove = CompiledOverrides.AddDefaulted_GetRef();
+				NewMove.DestPtr = DestAddress;
+				NewMove.Size = ValueSize;
+				FOpOverridesKey OpOverrideKey(ProfileName, OpOverrides.OpName);
+				NewMove.OpOverrideKey = OpOverrideKey;
+				NewMove.PropertyPathHash = PropOverride.GetPropertyPathHash();
+				NewMove.Property = LeafProp;
+
+				FName BoundCurveName = PropOverride.GetBoundCurveName();
+				FName BoundVarName = PropOverride.GetBoundVariableName();
+				FProperty* BagProp = (BagStruct && BoundVarName != NAME_None) ? BagStruct->FindPropertyByName(BoundVarName) : nullptr;
+
+				// 1. bound to an anim curve
+				if (BoundCurveName != NAME_None)
+				{
+					NewMove.BoundCurveName = BoundCurveName;
+					
+					// cache the target data type to avoid reflection costs in the hot path
+					if (LeafProp->IsA<FFloatProperty>()) NewMove.CurveCoercionType = ECurveCoercionType::Float;
+					else if (LeafProp->IsA<FDoubleProperty>()) NewMove.CurveCoercionType = ECurveCoercionType::Double;
+					else if (LeafProp->IsA<FIntProperty>()) NewMove.CurveCoercionType = ECurveCoercionType::Int32;
+					else if (LeafProp->IsA<FBoolProperty>()) NewMove.CurveCoercionType = ECurveCoercionType::Bool;
+					else if (LeafProp->IsA<FByteProperty>()) NewMove.CurveCoercionType = ECurveCoercionType::Byte;
+					else if (LeafProp->IsA<FInt64Property>()) NewMove.CurveCoercionType = ECurveCoercionType::Int64;
+				}
+				// 2. bound to a variable
+				else if (BagProp && BagProp->GetSize() == ValueSize)
+				{
+					NewMove.BagMemoryOffset = BagProp->GetOffset_ForInternal();
+				}
+				// 3. unbound / hardcoded
+				else
+				{
+					const int32 CurrentOffset = CachedValuePool.Num();
+					CachedValuePool.AddUninitialized(ValueSize);
+							
+					uint8* PoolEntryAddress = CachedValuePool.GetData() + CurrentOffset;
+							
+					const TCHAR* ImportResult = LeafProp->ImportText_Direct(*PropOverride.GetValueString(), PoolEntryAddress, nullptr, PPF_None);
+					if (!ensure(ImportResult))
+					{
+						CompiledOverrides.Pop();
+						continue; 
+					}
+
+					NewMove.SourcePoolOffset = CurrentOffset;
+				}
+
+				CachedVersionPerOp.FindOrAdd(OpOverrideKey, OpOverrides.GetValueVersion());
+			}
+		}
+	}
+}
+
+void FIKRetargetRuntimeOverrideManager::DebugPrintVariables(FInstancedPropertyBag* InVariables)
+{
+	if (!InVariables || !InVariables->GetPropertyBagStruct())
+	{
+		return;
+	}
+	
+	const UPropertyBag* BagStruct = InVariables->GetPropertyBagStruct();
+	const uint8* ContainerMemory = InVariables->GetValue().GetMemory();
+
+	if (!BagStruct || !ContainerMemory)
+	{
+		return;
+	}
+
+	UE_LOG(LogTemp, Warning, TEXT("=== Dumping FInstancedPropertyBag ==="));
+	for (const FPropertyBagPropertyDesc& Desc : BagStruct->GetPropertyDescs())
+	{
+		if (FProperty* Prop = BagStruct->FindPropertyByName(Desc.Name))
+		{
+			const void* ValuePtr = Prop->ContainerPtrToValuePtr<void>(ContainerMemory);
+			FString ExportedValueStr;
+			Prop->ExportTextItem_Direct(ExportedValueStr, ValuePtr, nullptr, nullptr, PPF_None);
+			UE_LOG(LogTemp, Warning, TEXT("  [%s] %s : %s"), *Prop->GetCPPType(), *Desc.Name.ToString(), *ExportedValueStr);
+		}
+	}
+	UE_LOG(LogTemp, Warning, TEXT("====================================="));
+}
+
+void FIKRetargetRuntimeOverrideManager::UpdateCachedValues(FIKRetargetProcessor& Processor)
+{
+	const UIKRetargeter* Asset = Processor.GetRetargetAsset();
+	if (!ensure(Asset) || CachedVersionPerOp.IsEmpty())
+	{
+		return;
+	}
+
+	uint8* PoolBase = CachedValuePool.GetData();
+	
+	for (const FName& OverrideSetName : OverrideSetsToApply)
+	{
+		const FRetargetOverrideSet* OverrideSet = Asset->GetOverrideSetByName(OverrideSetName);
+		if (!ensure(OverrideSet))
+		{
+			continue;
+		}
+
+		for (const FRetargetOpOverrides& OpOverrides : OverrideSet->OpOverrides)
+		{
+			FOpOverridesKey OpOverrideKey(OverrideSetName, OpOverrides.OpName);
+			if (CachedVersionPerOp[OpOverrideKey] == OpOverrides.GetValueVersion())
+			{
+				continue;
+			}
+
+			for (FCompiledPropertyOverride& Entry : CompiledOverrides)
+			{
+				if (Entry.OpOverrideKey != OpOverrideKey)
+				{
+					continue;
+				}
+				
+				// ignore bound variables and bound curves, they update automatically every frame
+				if (Entry.SourcePoolOffset == INDEX_NONE)
+				{
+					continue;
+				}
+
+				for (const FRetargetOpPropertyOverride& PropertyOverride : OpOverrides.PropertyOverrides)
+				{
+					if (Entry.PropertyPathHash != PropertyOverride.GetPropertyPathHash())
+					{
+						continue;
+					}
+					
+					uint8* PoolEntryAddress = PoolBase + Entry.SourcePoolOffset;
+					Entry.Property->ImportText_Direct(*PropertyOverride.GetValueString(), PoolEntryAddress, nullptr, PPF_None);
+				}
+			}
+		}
+	}
+}
\ No newline at end of file
diff --git a/Engine/Plugins/Animation/IKRig/Source/IKRig/Private/Retargeter/IKRetargetOverrides.cpp b/Engine/Plugins/Animation/IKRig/Source/IKRig/Private/Retargeter/IKRetargetOverrides.cpp
--- a/Engine/Plugins/Animation/IKRig/Source/IKRig/Private/Retargeter/IKRetargetOverrides.cpp
+++ b/Engine/Plugins/Animation/IKRig/Source/IKRig/Private/Retargeter/IKRetargetOverrides.cpp
@@ -551,261 +551,4 @@ FRetargetOpOverrides* FRetargetOverrideSet::FindOpOverrideByOpName(const FName I
 	return nullptr;
 }
 
-void FIKRetargetRuntimeOverrideManager::ApplyPropertyOverridesToOps(FIKRetargetProcessor& InProcessor, const TArrayView<const FName> InOverrideSetsToApply)
-{
-	const UIKRetargeter* Asset = InProcessor.GetRetargetAsset();
-	if (!ensure(Asset))
-	{
-		return;
-	}
-
-	UpdateFullListOfOverrideSets(InProcessor, InOverrideSetsToApply);
-	
-	// check if the list of profiles to apply has changed; if so, RebuildCache()
-	const int32 NewOverrideSetHash = GetOverrideSetArrayHash(OverrideSetsToApply); // this tells us if any new override sets were enabled/disabled
-	if (bIsCacheDirty || (NewOverrideSetHash != OverrideSetHash) || (CachedAssetVersion != Asset->GetOverrideVersion()))
-	{
-		RebuildCache(InProcessor);
-		OverrideSetHash = NewOverrideSetHash;
-		CachedAssetVersion = Asset->GetOverrideVersion();
-		bIsCacheDirty = false;
-	}
-
-#if WITH_EDITOR
-	// update any modified values in the CachedValuePool
-	// NOTE: this is editor only for now, but in the future we could allow this at runtime if we have a BP API for editing override values
-	UpdateCachedValues(InProcessor);
-#endif
-
-	// the "Hot Path", a linear iteration and direct memcpy of all the individual property overrides
-	for (const FCompiledPropertyOverride& Entry : CompiledOverrides)
-	{
-		FMemory::Memcpy(Entry.DestPtr, CachedValuePool.GetData() + Entry.SourcePoolOffset, Entry.Size);
-	}
-}
-
-void FIKRetargetRuntimeOverrideManager::UpdateFullListOfOverrideSets(
-	FIKRetargetProcessor& Processor,
-	const TArrayView<const FName> InOverrideSetsToApply)
-{
-	OverrideSetsToApply.Reset();
-	
-	const UIKRetargeter* Asset = Processor.GetRetargetAsset();
-	if (!ensure(Asset))
-	{
-		return;
-	}
-	
-	// process each set in the input list, note that these could be roots, mid-branch, or leaves
-	for (const FName& OverrideSetName : InOverrideSetsToApply)
-	{
-		// trace the ancestry from the current set, up to the root
-		TArray<FName> CurrentChain;
-		FName CurrentName = OverrideSetName;
-		while (CurrentName != NAME_None)
-		{
-			const FRetargetOverrideSet* OverrideSet = Asset->GetOverrideSetByName(CurrentName);
-			if (!OverrideSet)
-			{
-				break;
-			}
-
-			CurrentChain.Add(CurrentName);
-			CurrentName = OverrideSet->ParentName;
-		}
-
-		// insert the chain into our final, no duplicates
-		for (int32 i = CurrentChain.Num() - 1; i >= 0; --i)
-		{
-			OverrideSetsToApply.AddUnique(CurrentChain[i]);
-		}
-	}
-}
-
-uint32 FIKRetargetRuntimeOverrideManager::GetOverrideSetArrayHash(const TArray<FName>& InProfileNames)
-{
-	uint32 CombinedHash = 0;
-	for (const FName& ProfileName : InProfileNames)
-	{
-		CombinedHash = HashCombine(CombinedHash, GetTypeHash(ProfileName));
-	}
-	return CombinedHash;
-}
-
-void FIKRetargetRuntimeOverrideManager::RebuildCache(FIKRetargetProcessor& Processor)
-{
-    CompiledOverrides.Empty();
-    CachedValuePool.Empty();
-
-	const UIKRetargeter* Asset = Processor.GetRetargetAsset();
-	if (!ensure(Asset))
-	{
-		return;
-	}
-    
-    // first we calculate total memory needed for the pool to avoid reallocations
-	// NOTE: this isn't just for perf, we need non-volatile mem addresses
-    int32 TotalPoolSize = 0;
-    for (const FName& ProfileName : OverrideSetsToApply)
-    {
-    	const FRetargetOverrideSet* OverrideSet = Asset->GetOverrideSetByName(ProfileName);
-        if (!OverrideSet)
-        {
-        	continue;
-        }
-    	
-        for (const FRetargetOpOverrides& OpOverrides : OverrideSet->OpOverrides)
-        {
-            const UScriptStruct* OpScriptStruct = OpOverrides.ScriptStruct;
-            if (!ensure(OpScriptStruct))
-            {
-	            continue;
-            }
-
-            for (const FRetargetOpPropertyOverride& PropOverride : OpOverrides.PropertyOverrides)
-            {
-                TArray<FRetargetOpPropertyOverride::FPropertySegment> Segments;
-                if (FRetargetOpPropertyOverride::GetSegmentsFromProperyPath(PropOverride.GetPropertyPath(), OpScriptStruct, Segments))
-                {
-                    if (FProperty* LeafProp = FRetargetOpPropertyOverride::GetLeafProperty(Segments))
-                    {
-                        TotalPoolSize += LeafProp->GetSize();
-                    }
-                }
-            }
-        }
-    }
-    CachedValuePool.Reserve(TotalPoolSize);
-
-    // now bake values into the pool and cache the pointers
-    for (const FName& ProfileName : OverrideSetsToApply)
-    {
-    	const FRetargetOverrideSet* OverrideSet = Asset->GetOverrideSetByName(ProfileName);
-    	if (!OverrideSet)
-    	{
-    		continue;
-    	}
-
-    	for (const FRetargetOpOverrides& OpOverrides : OverrideSet->OpOverrides)
-    	{
-		    FIKRetargetOpBase* OpInProcessor = Processor.GetRetargetOpByName(OpOverrides.OpName);
-    		if (!ensure(OpInProcessor))
-    		{
-    			continue;
-    		}
-		    FIKRetargetOpSettingsBase* OpSettings = OpInProcessor->GetSettings();
-    		if (!ensure(OpSettings))
-    		{
-    			continue;
-    		}
-            uint8* OpSettingsData = reinterpret_cast<uint8*>(OpSettings);
-            const UScriptStruct* OpScriptStruct = OpOverrides.ScriptStruct;
-
-            if (!ensure(OpInProcessor && OpScriptStruct))
-            {
-	            continue;
-            }
-
-    		for (const FRetargetOpPropertyOverride& PropOverride : OpOverrides.PropertyOverrides)
-    		{
-    			TArray<FRetargetOpPropertyOverride::FPropertySegment> Segments;
-				const bool bPropertyFound = FRetargetOpPropertyOverride::GetSegmentsFromProperyPath(PropOverride.GetPropertyPath(), OpScriptStruct, Segments);
-				if (!ensure(bPropertyFound))
-				{
-					continue;
-				}
-            	
-				FProperty* LeafProp = FRetargetOpPropertyOverride::GetLeafProperty(Segments);
-				uint8* DestAddress = FRetargetOpPropertyOverride::GetDataPointerFromPathSegments(OpSettingsData, Segments);
-
-				if (!ensure(LeafProp && DestAddress))
-				{
-					continue;
-				}
-            	
-				// bake the value into the pool
-				const int32 ValueSize = LeafProp->GetSize();
-				const int32 CurrentOffset = CachedValuePool.Num();
-				CachedValuePool.AddUninitialized(ValueSize);
-                        
-				uint8* PoolEntryAddress = CachedValuePool.GetData() + CurrentOffset;
-                        
-				// convert string to binary at proper offset in pool
-				const TCHAR* ImportResult = LeafProp->ImportText_Direct(*PropOverride.GetValueString(), PoolEntryAddress, nullptr, PPF_None);
-    			if (!ensure(ImportResult))
-    			{
-    				continue; // corrupted string
-    			}
-
-				// record the move
-				FCompiledPropertyOverride& NewMove = CompiledOverrides.AddDefaulted_GetRef();
-				NewMove.SourcePoolOffset = CurrentOffset;
-				NewMove.DestPtr = DestAddress;
-				NewMove.Size = ValueSize;
-				// cache version we stored the value at
-				FOpOverridesKey OpOverrideKey(ProfileName, OpOverrides.OpName);
-				NewMove.OpOverrideKey = OpOverrideKey;
-				NewMove.PropertyPathHash = PropOverride.GetPropertyPathHash();
-    			NewMove.Property = LeafProp;
-			    CachedVersionPerOp.FindOrAdd(OpOverrideKey, OpOverrides.GetValueVersion());
-    		}
-        }
-    }
-}
-
-void FIKRetargetRuntimeOverrideManager::UpdateCachedValues(FIKRetargetProcessor& Processor)
-{
-	const UIKRetargeter* Asset = Processor.GetRetargetAsset();
-	if (!ensure(Asset))
-	{
-		return;
-	}
-
-	if (CachedVersionPerOp.IsEmpty())
-	{
-		return;
-	}
-
-	uint8* PoolBase = CachedValuePool.GetData();
-	
-	for (const FName& OverrideSetName : OverrideSetsToApply)
-	{
-		const FRetargetOverrideSet* OverrideSet = Asset->GetOverrideSetByName(OverrideSetName);
-		if (!ensure(OverrideSet))
-		{
-			continue;
-		}
-
-		for (const FRetargetOpOverrides& OpOverrides : OverrideSet->OpOverrides)
-		{
-			FOpOverridesKey OpOverrideKey(OverrideSetName, OpOverrides.OpName);
-			if (CachedVersionPerOp[OpOverrideKey] == OpOverrides.GetValueVersion())
-			{
-				continue;
-			}
-
-			// uh oh, we need to update the cached values for the overrides stored in this set
-			for (FCompiledPropertyOverride& Entry : CompiledOverrides)
-			{
-				if (Entry.OpOverrideKey != OpOverrideKey)
-				{
-					continue; // not a part of this set, fast early-out for the majority of entries
-				}
-
-				for (const FRetargetOpPropertyOverride& PropertyOverride : OpOverrides.PropertyOverrides)
-				{
-					if (Entry.PropertyPathHash != PropertyOverride.GetPropertyPathHash())
-					{
-						continue;
-					}
-					
-					// update the binary pool address associated with this compiled entry
-					uint8* PoolEntryAddress = PoolBase + Entry.SourcePoolOffset;
-					Entry.Property->ImportText_Direct(*PropertyOverride.GetValueString(), PoolEntryAddress, nullptr, PPF_None);
-				}
-			}
-		}
-	}
-}
-
 #undef LOCTEXT_NAMESPACE
\ No newline at end of file
diff --git a/Engine/Plugins/Animation/IKRig/Source/IKRig/Private/Retargeter/IKRetargetProcessor.cpp b/Engine/Plugins/Animation/IKRig/Source/IKRig/Private/Retargeter/IKRetargetProcessor.cpp
--- a/Engine/Plugins/Animation/IKRig/Source/IKRig/Private/Retargeter/IKRetargetProcessor.cpp
+++ b/Engine/Plugins/Animation/IKRig/Source/IKRig/Private/Retargeter/IKRetargetProcessor.cpp
@@ -1309,7 +1309,7 @@ TArray<FTransform>& FIKRetargetProcessor::RunRetargeter(const FRetargetRunParame
 #endif
 
 	// apply the retargeting settings
-	UpdateSettingsAtRuntime(Params.Profile, Params.OverrideSetsToApply);
+	UpdateSettingsAtRuntime(Params.Profile, Params.OverrideSetsToApply, Params.Variables, Params.BoundCurveValues);
 
 	// applying settings can cause the retargeter to require reinitialization (when enabling/disabling IK on a chain)
 	// we should be able to safely reinitialize here because we are already initialized or it wouldn't pass the check() at the top
@@ -2080,7 +2080,11 @@ void FIKRetargetProcessor::SetNeedsInitialized()
 	}
 }
 
-void FIKRetargetProcessor::UpdateSettingsAtRuntime(const FRetargetProfile* InProfile, const TArrayView<const FName> InOverrideSetsToApply)
+void FIKRetargetProcessor::UpdateSettingsAtRuntime(
+	const FRetargetProfile* InProfile,
+	const TArrayView<const FName> InOverrideSetsToApply,
+	FInstancedPropertyBag* InVariables,
+	TMap<FName, float>* InBoundCurveValues)
 {
 	// apply op settings from the profile
 	if (InProfile)
@@ -2092,7 +2096,7 @@ void FIKRetargetProcessor::UpdateSettingsAtRuntime(const FRetargetProfile* InPro
 	}
 
 	// apply stored property overrides
-	OverrideManager.ApplyPropertyOverridesToOps(*this, InOverrideSetsToApply);
+	OverrideManager.ApplyPropertyOverridesToOps(*this, InOverrideSetsToApply, InVariables, InBoundCurveValues);
 	
 	if (InProfile)
 	{
diff --git a/Engine/Plugins/Animation/IKRig/Source/IKRig/Public/AnimNodes/AnimNode_RetargetPoseFromMesh.h b/Engine/Plugins/Animation/IKRig/Source/IKRig/Public/AnimNodes/AnimNode_RetargetPoseFromMesh.h
--- a/Engine/Plugins/Animation/IKRig/Source/IKRig/Public/AnimNodes/AnimNode_RetargetPoseFromMesh.h
+++ b/Engine/Plugins/Animation/IKRig/Source/IKRig/Public/AnimNodes/AnimNode_RetargetPoseFromMesh.h
@@ -5,6 +5,7 @@
 #include "CoreMinimal.h"
 
 #include "Animation/AnimNodeBase.h"
+#include "Animation/AnimNode_CustomProperty.h"
 #include "Retargeter/IKRetargeter.h"
 #include "Retargeter/IKRetargetProcessor.h"
 #include "Retargeter/IKRetargetProfile.h"
@@ -24,7 +25,7 @@ enum class ERetargetSourceMode : uint8
 };
 
 USTRUCT(BlueprintInternalUseOnly)
-struct FAnimNode_RetargetPoseFromMesh : public FAnimNode_Base
+struct FAnimNode_RetargetPoseFromMesh : public FAnimNode_CustomProperty
 {
 	GENERATED_BODY()
 	
@@ -86,6 +87,10 @@ struct FAnimNode_RetargetPoseFromMesh : public FAnimNode_Base
 	virtual int32 GetLODThreshold() const override { return LODThreshold; }
 	// End of FAnimNode_Base interface
 
+	// FAnimNode_CustomProperty interface
+	virtual UClass* GetTargetClass() const override {return nullptr;};
+	// END FAnimNode_CustomProperty interface
+
 	// access to the runtime processor
 	UE_API FIKRetargetProcessor* GetRetargetProcessor();
 
@@ -102,6 +107,21 @@ struct FAnimNode_RetargetPoseFromMesh : public FAnimNode_Base
 
 private:
 	
+	// the runtime storage for variable values
+	UPROPERTY(Transient)
+	FInstancedPropertyBag RuntimeVariables;
+
+	// a map of curves bound to property overrides (Name, float)
+	UPROPERTY(Transient)
+	TMap<FName, float> BoundCurveValues;
+	int32 CachedOverrideVersion = INDEX_NONE;
+
+	// pulls runtime values of variables out of anim instance and stores in RuntimeVariables
+	void CopyInputPropertiesToVariables(const UObject* InSourceInstance);
+
+	// fills BoundCurves from the anim instance
+	void CopyBoundCurveValues(const UObject* InSourceInstance);
+	
 	// get a retarget profile that includes settings from the asset, plus any custom profile overrides
 	UE_API FRetargetProfile GetMergedRetargetProfile(bool bEnableIK) const;
 
diff --git a/Engine/Plugins/Animation/IKRig/Source/IKRig/Public/Retargeter/IKRetargetOverrideManager.h b/Engine/Plugins/Animation/IKRig/Source/IKRig/Public/Retargeter/IKRetargetOverrideManager.h
new file mode 100644
--- /dev/null
+++ b/Engine/Plugins/Animation/IKRig/Source/IKRig/Public/Retargeter/IKRetargetOverrideManager.h
@@ -0,0 +1,89 @@
+﻿// Copyright Epic Games, Inc. All Rights Reserved.
+
+#pragma once
+
+#include "CoreMinimal.h"
+
+struct FIKRetargetProcessor;
+struct FInstancedPropertyBag;
+
+struct FIKRetargetRuntimeOverrideManager
+{
+	struct FOpOverridesKey
+	{
+	   FName OverrideSetName;
+	   FName OpName;
+
+	   FOpOverridesKey() = default;
+	   FOpOverridesKey(FName ProfileName, FName OpName) : OverrideSetName(ProfileName), OpName(OpName){}
+	   
+	   bool operator==(const FOpOverridesKey& Other) const
+	   {
+		  return OverrideSetName == Other.OverrideSetName && OpName == Other.OpName;
+	   }
+	   
+	   friend uint32 GetTypeHash(const FOpOverridesKey& Key)
+	   {
+		  uint32 Hash = GetTypeHash(Key.OverrideSetName);
+		  return HashCombine(Hash, GetTypeHash(Key.OpName));
+	   }
+	};
+	
+	// fast cache for safe numeric casting in the hot path
+	enum class ECurveCoercionType : uint8
+	{
+	   None,
+	   Bool,
+	   Float,
+	   Double,
+	   Int32,
+	   Int64,
+	   Byte
+	};
+
+	struct FCompiledPropertyOverride
+	{
+	   int32 SourcePoolOffset = INDEX_NONE; // offset into CachedValuePool (if hardcoded)
+	   int32 BagMemoryOffset = INDEX_NONE;  // offset into FInstancedPropertyBag memory (if bound to var)
+	   
+	   FName BoundCurveName = NAME_None;    // name of the curve (if bound to curve)
+	   ECurveCoercionType CurveCoercionType = ECurveCoercionType::None; // the target type for curve casting
+
+	   void* DestPtr = nullptr;             // pointer to processor op settings memory
+	   int32 Size = 0;                      // size of the property data
+
+	   // value cache invalidate/update stuff
+	   FOpOverridesKey OpOverrideKey; 
+	   FProperty* Property = nullptr;         
+	   uint32 PropertyPathHash = 0;      
+	};
+
+	void ApplyPropertyOverridesToOps(
+	   FIKRetargetProcessor& Processor,
+	   const TArrayView<const FName> InOverrideSetsToApply,
+	   FInstancedPropertyBag* InVariables,
+	   TMap<FName, float>* InBoundCurveValues);
+
+	void DirtyProfileCache() { bIsCacheDirty = true; };
+
+private:
+
+	static void DebugPrintVariables(FInstancedPropertyBag* InVariables);
+	
+	void UpdateCachedValues(FIKRetargetProcessor& Processor);
+	
+	void UpdateFullListOfOverrideSets(FIKRetargetProcessor& Processor, const TArrayView<const FName> InOverrideSetsToApply);
+
+	void RebuildCache(FIKRetargetProcessor& Processor, FInstancedPropertyBag* InVariables);
+
+	static uint32 GetOverrideSetArrayHash(const TArray<FName>& InProfileNames);
+
+	TArray<FName> OverrideSetsToApply; 
+	TArray<uint8> CachedValuePool;
+	TArray<FCompiledPropertyOverride> CompiledOverrides;
+	TMap<FOpOverridesKey, int32> CachedVersionPerOp;
+	
+	int32 CachedAssetVersion;
+	uint32 OverrideSetHash;
+	bool bIsCacheDirty = true;
+};
\ No newline at end of file
diff --git a/Engine/Plugins/Animation/IKRig/Source/IKRig/Public/Retargeter/IKRetargetOverrides.h b/Engine/Plugins/Animation/IKRig/Source/IKRig/Public/Retargeter/IKRetargetOverrides.h
--- a/Engine/Plugins/Animation/IKRig/Source/IKRig/Public/Retargeter/IKRetargetOverrides.h
+++ b/Engine/Plugins/Animation/IKRig/Source/IKRig/Public/Retargeter/IKRetargetOverrides.h
@@ -85,6 +85,16 @@ struct FRetargetOpPropertyOverride
 	/** get the hash of the property path (cached at construction) */
 	uint32 GetPropertyPathHash() const { return PropertyPathHash; };
 	
+	/** set the name of the variable to bind to this property
+	 * NOTE: this must go through the retargeter controller */
+	void SetBoundVariableName(const FName InVariableName) { BoundVariableName = InVariableName; };
+	FName GetBoundVariableName() const { return BoundVariableName; };
+
+	/** set the name of the curve to bind to this property
+	 * NOTE: this must go through the retargeter controller */
+	void SetBoundCurveName(const FName InCurveName) { BoundCurveName = InCurveName; };
+	FName GetBoundCurveName() const { return BoundCurveName; };
+	
 	/** validate that this property exists and is accessible */
 	bool IsValid(const UScriptStruct* InSettingsStruct) const;
 
@@ -125,18 +135,22 @@ private:
 	
 	/** property path using "->" separator like, "IKGoalSettings->BlendAlpha" and array syntax like, "Chains[0]->Weight"
 	 * NOTE: this allows generic serialization of property overrides */
-	UPROPERTY(EditAnywhere, Category=Override)
+	UPROPERTY()
 	FString PropertyPath;
 	UPROPERTY()
 	uint32 PropertyPathHash = 0;
 
 	/** Value serialized as string */
-	UPROPERTY(EditAnywhere, Category=Override)
+	UPROPERTY()
 	FString ValueAsString;
 
-	/** TODO Optional: Mark if this property requires reinitialization when changed */
+	/** The name of the variable this override is bound to, if any. */
+	UPROPERTY()
+	FName BoundVariableName = NAME_None;
+
+	/** The name of the curve this override is bound to, if any. */
 	UPROPERTY()
-	bool bRequiresReinitialization = false;
+	FName BoundCurveName = NAME_None;
 };
 
 /** A collection of property overrides to assign to a single op. */
@@ -227,85 +241,4 @@ struct FRetargetOverrideSet
 	UE_API FRetargetOpOverrides* FindOpOverrideByOpName(const FName InOpName);
 };
 
-// This system is used to apply property overrides to retarget ops at runtime in a performant way.
-//
-// The way it works is:
-// 1. The processor calls ApplyPropertyOverridesToOps every frame passing in the names of the profiles to apply.
-// 2. These profiles are then aggregated with their descendent profiles up the hierarchy to form a full list of ProfilesToApply.
-// 3. All the property overrides stored in all the ProfilesToApply are resolved with a source/dest memory address
-// 4. At runtime, we memcpy all the data from the profiles into the op stack that is currently running.
-//
-// NOTE: because we are caching raw pointers, we must be very careful to rebuild the cache any time the op stack is rebuilt!
-// This can be done with: DirtyProfileCache()
-//
-struct FIKRetargetRuntimeOverrideManager
-{
-	// a key associated with all overrides for a given op in a given set of overrides
-	struct FOpOverridesKey
-	{
-		FName OverrideSetName;
-		FName OpName;
-
-		FOpOverridesKey() = default;
-		FOpOverridesKey(FName ProfileName, FName OpName) : OverrideSetName(ProfileName), OpName(OpName){}
-		
-		bool operator==(const FOpOverridesKey& Other) const
-		{
-			return OverrideSetName == Other.OverrideSetName && OpName == Other.OpName;
-		}
-		
-		friend uint32 GetTypeHash(const FOpOverridesKey& Key)
-		{
-			uint32 Hash = GetTypeHash(Key.OverrideSetName);
-			return HashCombine(Hash, GetTypeHash(Key.OpName));
-		}
-	};
-	
-	struct FCompiledPropertyOverride
-	{
-		int32 SourcePoolOffset; // offset into CachedValuePool
-		void* DestPtr;          // pointer to processor op settings memory
-		int32 Size;				// size of the property data
-
-		// value cache invalidate/update stuff
-		FOpOverridesKey OpOverrideKey;	// used to invalidate cached value when any op override in a given profile is edited
-		FProperty* Property;			// used to avoid path parsing when updating values at runtime
-		uint32 PropertyPathHash;		// the property path this override applies to (avoid storing pointers into asset, volatile while editing)
-	};
-
-	void ApplyPropertyOverridesToOps(FIKRetargetProcessor& Processor, const TArrayView<const FName> InOverrideSetsToApply);
-
-	void DirtyProfileCache() { bIsCacheDirty = true; };
-
-private:
-	
-	void UpdateCachedValues(FIKRetargetProcessor& Processor);
-	
-	void UpdateFullListOfOverrideSets(FIKRetargetProcessor& Processor, const TArrayView<const FName> InOverrideSetsToApply);
-
-	void RebuildCache(FIKRetargetProcessor& Processor);
-
-	static uint32 GetOverrideSetArrayHash(const TArray<FName>& InProfileNames);
-
-	// the full list of profiles to apply (includes parents of named profiles)
-	TArray<FName> OverrideSetsToApply; 
-
-	// a single contiguous buffer holding all binary values for the active overrides
-	TArray<uint8> CachedValuePool;
-
-	// cached pointer offsets to memcpy from source to dest value memory
-	// key is a combined hash of the profile name + op name the overrides belong to
-	TArray<FCompiledPropertyOverride> CompiledOverrides;
-
-	// the version of the op override sets we cached (invalidated when value is modified)
-	TMap<FOpOverridesKey, int32> CachedVersionPerOp;
-	int32 CachedAssetVersion;
-
-	// hash of profile names that are currently cached (used to invalidate the CompiledOverrides)
-	uint32 OverrideSetHash;
-
-	// processor forcibly dirties the cache when it re-initializes
-	bool bIsCacheDirty = true;
-};
-
 #undef UE_API
diff --git a/Engine/Plugins/Animation/IKRig/Source/IKRig/Public/Retargeter/IKRetargetProcessor.h b/Engine/Plugins/Animation/IKRig/Source/IKRig/Public/Retargeter/IKRetargetProcessor.h
--- a/Engine/Plugins/Animation/IKRig/Source/IKRig/Public/Retargeter/IKRetargetProcessor.h
+++ b/Engine/Plugins/Animation/IKRig/Source/IKRig/Public/Retargeter/IKRetargetProcessor.h
@@ -5,6 +5,7 @@
 #include "CoreMinimal.h"
 #include "IKRetargeter.h"
 #include "IKRetargetOps.h"
+#include "IKRetargetOverrideManager.h"
 #include "IKRetargetSettings.h"
 #include "IKRigLogger.h"
 #include "Kismet/KismetMathLibrary.h"
@@ -363,8 +364,6 @@ struct FRetargetInitParameters
 	const UIKRetargeter* RetargeterAsset = nullptr;
 	/** The collection of settings used to initialize with */
 	const FRetargetProfile* CustomProfile = nullptr;
-	/** The names of any sets of overrides sets to apply (stored on asset) */
-	TArrayView<const FName> OverrideSetsToApply;
 	/** If true, will not output warnings during initialization */
 	bool bSuppressWarnings = false;
 };
@@ -378,6 +377,10 @@ struct FRetargetRunParameters
 	const FRetargetProfile* Profile = nullptr;
 	/** An array of override set names stored on the asset to apply */
 	TArrayView<const FName> OverrideSetsToApply;
+	/** A property bag housing variable values that may be bound to overrides */
+	FInstancedPropertyBag* Variables = nullptr;
+	/** A map of curve names to float values for any curves bound to an override */
+	TMap<FName, float>* BoundCurveValues = nullptr;;
 	/** Time since last tick in seconds (passed to ops) */
 	float DeltaTime = 0.0f;
 	/** The current LOD level for the skeletal mesh */
@@ -608,7 +611,11 @@ private:
 	mutable FRetargetPoseScaleWithPivot CurrentSourceScale;
 
 	/** Apply the settings stored in a retarget profile. As well as any stored property overrides. */
-	void UpdateSettingsAtRuntime(const FRetargetProfile* Profile, const TArrayView<const FName> InOverrideSetsToApply);
+	void UpdateSettingsAtRuntime(
+		const FRetargetProfile* InProfile,
+		const TArrayView<const FName> InOverrideSetsToApply,
+		FInstancedPropertyBag* InVariables,
+		TMap<FName, float>* InBoundCurveValues);
 
 	/** Internal retarget phase that does simple bone-to-bone copying from source*/
 	void GenerateBasePoses(TArray<FTransform>& InSourceGlobalPose);
diff --git a/Engine/Plugins/Animation/IKRig/Source/IKRig/Public/Retargeter/IKRetargeter.h b/Engine/Plugins/Animation/IKRig/Source/IKRig/Public/Retargeter/IKRetargeter.h
--- a/Engine/Plugins/Animation/IKRig/Source/IKRig/Public/Retargeter/IKRetargeter.h
+++ b/Engine/Plugins/Animation/IKRig/Source/IKRig/Public/Retargeter/IKRetargeter.h
@@ -7,6 +7,7 @@
 #include "IKRetargetOverrides.h"
 #include "Rig/IKRigDefinition.h"
 #include "IKRetargetSettings.h"
+#include "StructUtils/PropertyBag.h"
 
 #include "IKRetargeter.generated.h"
 
@@ -18,6 +19,15 @@ class FIKRetargetEditorController;
 
 class URetargetOpStack;
 
+USTRUCT()
+struct FRetargetVariableContainer
+{
+	GENERATED_BODY()
+	
+	UPROPERTY()
+	FInstancedPropertyBag Bag;
+};
+
 USTRUCT(BlueprintType)
 struct FIKRetargetPose
 {
@@ -131,6 +141,8 @@ public:
 	// Get the retarget profile by name (may be null) 
 	UE_API const FRetargetProfile* GetProfileByName(const FName& ProfileName) const;
 
+	// Get read-only access to all the override sets
+	const TMap<FName, FRetargetOverrideSet>& GetOverrideSets() const { return OverrideSets; };
 	// Get all the override sets that are enabled
 	UE_API const TArray<FName>& GetOverrideSetsToApply() const;
 	// Get a set of overrides by name (may be null) 
@@ -138,6 +150,9 @@ public:
 	// get current version of the property overrides data (to compare against running processor instances)
 	int32 GetOverrideVersion() const { return OverridesVersion; };
 
+	// get read-write access to the variable bag (caution here, all editing should go through controller) 
+	FRetargetVariableContainer& GetVariables() { return Variables; };
+
 	// get current version of the data (to compare against running processor instances)
 	int32 GetVersion() const { return Version; };
 	// do this after any edit that would require running instance to reinitialize
@@ -451,6 +466,10 @@ private:
 	mutable TArray<FName> OverridesToApply;
 	// updated anytime overrides are modified
 	mutable int32 OverridesVersion = 0;
+	
+	// storage for all variables (used to dynamically adjust property overrides)
+	UPROPERTY()
+	FRetargetVariableContainer Variables;
 
 	// The set of retarget poses for the SOURCE skeleton.
 	UPROPERTY()
diff --git a/Engine/Plugins/Animation/IKRig/Source/IKRigDeveloper/IKRigDeveloper.Build.cs b/Engine/Plugins/Animation/IKRig/Source/IKRigDeveloper/IKRigDeveloper.Build.cs
--- a/Engine/Plugins/Animation/IKRig/Source/IKRigDeveloper/IKRigDeveloper.Build.cs
+++ b/Engine/Plugins/Animation/IKRig/Source/IKRigDeveloper/IKRigDeveloper.Build.cs
@@ -32,7 +32,8 @@ namespace UnrealBuildTool.Rules
                         "BlueprintGraph",
                         "Slate",
                         "SlateCore",
-                        "Persona"
+                        "Persona",
+                        "StructUtilsEditor",
                     }
                 );
             }
diff --git a/Engine/Plugins/Animation/IKRig/Source/IKRigDeveloper/Private/AnimGraphNode_RetargetPoseFromMesh.cpp b/Engine/Plugins/Animation/IKRig/Source/IKRigDeveloper/Private/AnimGraphNode_RetargetPoseFromMesh.cpp
--- a/Engine/Plugins/Animation/IKRig/Source/IKRigDeveloper/Private/AnimGraphNode_RetargetPoseFromMesh.cpp
+++ b/Engine/Plugins/Animation/IKRig/Source/IKRigDeveloper/Private/AnimGraphNode_RetargetPoseFromMesh.cpp
@@ -1,6 +1,8 @@
 // Copyright Epic Games, Inc. All Rights Reserved.
 
 #include "AnimGraphNode_RetargetPoseFromMesh.h"
+
+#include "PropertyBagDetails.h"
 #include "Animation/AnimInstance.h"
 #include "Kismet2/CompilerResultsLog.h"
 
@@ -50,6 +52,152 @@ void UAnimGraphNode_RetargetPoseFromMesh::CustomizePinData(UEdGraphPin* Pin, FNa
 	}
 }
 
+void UAnimGraphNode_RetargetPoseFromMesh::CreateCustomPins(TArray<UEdGraphPin*>* InOldPins)
+{
+	Super::CreateCustomPins(InOldPins);
+
+	// if we don't have an asset, we can't have custom variable pins
+	if (!IsValid(Node.IKRetargeterAsset))
+	{
+		return;
+	}
+
+	// handle the 'PostLoad' race condition (if ABP is loaded before the asset)
+	if (Node.IKRetargeterAsset->HasAllFlags(RF_NeedPostLoad)) 
+	{
+		CreateCustomPinsFromUnloadedAsset(InOldPins);
+		return;
+	}
+
+	// asset is valid and loaded, generate current pins
+	CreateCustomPinsFromValidAsset();
+}
+
+void UAnimGraphNode_RetargetPoseFromMesh::AddVariablesToDetailsPanel(IDetailLayoutBuilder& DetailBuilder)
+{
+	// NOTE: this lists all the variables in the asset and adds/removes from the TArray<FName> ExposedVariableNames
+	// Variables stored in ExposedVariableNames will become pins on the anim node.
+
+	
+	UIKRetargeter* RetargeterAsset = Node.IKRetargeterAsset;
+	if (!RetargeterAsset)
+	{
+		return;
+	}
+
+	const FInstancedPropertyBag& Bag = RetargeterAsset->GetVariables().Bag;
+	const UPropertyBag* BagStruct = Bag.GetPropertyBagStruct();
+	if (!BagStruct)
+	{
+		return;
+	}
+
+	// create a dedicated category for pin exposure
+	IDetailCategoryBuilder& PinCategory = DetailBuilder.EditCategory("Variable Pins", LOCTEXT("VariablePins", "Variable Pins"));
+
+	for (const FPropertyBagPropertyDesc& Desc : BagStruct->GetPropertyDescs())
+	{
+		const FName VarName = Desc.Name;
+
+		PinCategory.AddCustomRow(FText::FromName(VarName))
+		.NameContent()
+		[
+			SNew(SHorizontalBox)
+			+ SHorizontalBox::Slot()
+			.VAlign(VAlign_Center)
+			[
+				SNew(STextBlock)
+				.Text(FText::FromName(VarName))
+				.Font(IDetailLayoutBuilder::GetDetailFont())
+			]
+		]
+		.ValueContent()
+		[
+			SNew(SCheckBox)
+			.IsChecked_Lambda([this, VarName]()
+			{
+				return ExposedVariableNames.Contains(VarName) ? ECheckBoxState::Checked : ECheckBoxState::Unchecked;
+			})
+			.OnCheckStateChanged_Lambda([this, VarName](ECheckBoxState NewState)
+			{
+				if (NewState == ECheckBoxState::Checked)
+				{
+					ExposedVariableNames.Add(VarName);
+				}
+				else
+				{
+					ExposedVariableNames.Remove(VarName);
+				}
+				
+				// force the node to recreate pins to show/hide the change immediately
+				ReconstructNode();
+			})
+			[
+				SNew(STextBlock)
+				.Text(LOCTEXT("ExposePinLabel", "Expose as Pin"))
+				.Font(IDetailLayoutBuilder::GetDetailFont())
+			]
+		];
+	}
+}
+
+void UAnimGraphNode_RetargetPoseFromMesh::CreateCustomPinsFromValidAsset()
+{
+	const FInstancedPropertyBag& Bag = Node.IKRetargeterAsset->GetVariables().Bag;
+	const UPropertyBag* BagStruct = Bag.GetPropertyBagStruct();
+	const uint8* BagMemory = Bag.GetValue().GetMemory();
+
+	if (!BagStruct || !BagMemory)
+	{
+		return;
+	}
+	
+	for (const FName& VarName : ExposedVariableNames)
+	{
+		const FPropertyBagPropertyDesc* Desc = BagStruct->FindPropertyDescByName(VarName);
+		if (!Desc)
+		{
+			continue;
+		}
+		
+		FEdGraphPinType PinType = UE::StructUtils::GetPropertyDescAsPin(*Desc);
+		UEdGraphPin* NewPin = CreatePin(EGPD_Input, PinType, VarName);
+		NewPin->PinFriendlyName = FText::FromName(VarName);
+		
+		// extract the default value from the asset and assign it to the pin
+		// NOTE: assigning a default value is necessary to avoid BP compile errors
+		if (FProperty* BagProp = BagStruct->FindPropertyByName(VarName))
+		{
+			const void* ValuePtr = BagProp->ContainerPtrToValuePtr<void>(BagMemory);
+			
+			FString DefaultString;
+			BagProp->ExportTextItem_Direct(DefaultString, ValuePtr, nullptr, nullptr, PPF_None);
+			NewPin->DefaultValue = DefaultString;
+		}
+	}
+}
+
+void UAnimGraphNode_RetargetPoseFromMesh::CreateCustomPinsFromUnloadedAsset(TArray<UEdGraphPin*>* InOldPins)
+{
+	if (!InOldPins)
+	{
+		return;
+	}
+
+	for (UEdGraphPin* OldPin : *InOldPins)
+	{
+		// only recreate pins that belonged to our dynamic variable set
+		if (OldPin->Direction == EGPD_Input && ExposedVariableNames.Contains(OldPin->GetFName()))
+		{
+			// check if this pin already exists (avoid duplicates during reconstruction)
+			if (!FindPin(OldPin->GetFName(), EGPD_Input))
+			{
+				CreatePin(EGPD_Input, OldPin->PinType, OldPin->GetFName());
+			}
+		}
+	}
+}
+
 UObject* UAnimGraphNode_RetargetPoseFromMesh::GetJumpTargetForDoubleClick() const
 {
 	return Node.IKRetargeterAsset;
@@ -57,14 +205,20 @@ UObject* UAnimGraphNode_RetargetPoseFromMesh::GetJumpTargetForDoubleClick() cons
 
 void UAnimGraphNode_RetargetPoseFromMesh::PostEditChangeProperty(struct FPropertyChangedEvent& PropertyChangedEvent)
 {
+	Super::PostEditChangeProperty(PropertyChangedEvent); // required for UAnimGraphNode_CustomProperty base
+
 	const FName PropertyName = (PropertyChangedEvent.Property ? PropertyChangedEvent.Property->GetFName() : NAME_None);
 	if ((PropertyName == GET_MEMBER_NAME_STRING_CHECKED(FAnimNode_RetargetPoseFromMesh, RetargetFrom)))
 	{
 		ReconstructNode();
 	}
+	else if (PropertyName == GET_MEMBER_NAME_STRING_CHECKED(FAnimNode_RetargetPoseFromMesh, IKRetargeterAsset))
+	{
+		ReconstructNode();
+	}
 }
 
-void UAnimGraphNode_RetargetPoseFromMesh::ValidateAnimNodeDuringCompilation(USkeleton* ForSkeleton,	FCompilerResultsLog& MessageLog)
+void UAnimGraphNode_RetargetPoseFromMesh::ValidateAnimNodeDuringCompilation(USkeleton* ForSkeleton, FCompilerResultsLog& MessageLog)
 {
 	Super::ValidateAnimNodeDuringCompilation(ForSkeleton, MessageLog);
 
@@ -121,7 +275,7 @@ void UAnimGraphNode_RetargetPoseFromMesh::ValidateAnimNodeDuringCompilation(USke
 			{
 				MessageLog.Warning(*Warning.ToString());
 			}
-		
+			
 			const TArray<FText>& Errors = Processor->Log.GetErrors();
 			for (const FText& Error : Errors)
 			{
@@ -150,6 +304,13 @@ void UAnimGraphNode_RetargetPoseFromMesh::ValidateAnimNodeDuringCompilation(USke
 	}
 }
 
+void UAnimGraphNode_RetargetPoseFromMesh::CustomizeDetails(IDetailLayoutBuilder& DetailBuilder)
+{
+	Super::CustomizeDetails(DetailBuilder);
+
+	AddVariablesToDetailsPanel(DetailBuilder);
+}
+
 void UAnimGraphNode_RetargetPoseFromMesh::PreloadRequiredAssets()
 {
 	Super::PreloadRequiredAssets();
@@ -162,5 +323,4 @@ void UAnimGraphNode_RetargetPoseFromMesh::PreloadRequiredAssets()
 	}
 }
 
-#undef LOCTEXT_NAMESPACE
-
+#undef LOCTEXT_NAMESPACE
\ No newline at end of file
diff --git a/Engine/Plugins/Animation/IKRig/Source/IKRigDeveloper/Public/AnimGraphNode_RetargetPoseFromMesh.h b/Engine/Plugins/Animation/IKRig/Source/IKRigDeveloper/Public/AnimGraphNode_RetargetPoseFromMesh.h
--- a/Engine/Plugins/Animation/IKRig/Source/IKRigDeveloper/Public/AnimGraphNode_RetargetPoseFromMesh.h
+++ b/Engine/Plugins/Animation/IKRig/Source/IKRigDeveloper/Public/AnimGraphNode_RetargetPoseFromMesh.h
@@ -5,6 +5,7 @@
 #include "CoreMinimal.h"
 #include "UObject/ObjectMacros.h"
 #include "AnimGraphNode_Base.h"
+#include "AnimGraphNode_CustomProperty.h"
 #include "AnimNodes/AnimNode_RetargetPoseFromMesh.h"
 #include "AnimGraphNode_RetargetPoseFromMesh.generated.h"
 
@@ -13,10 +14,8 @@
 class FPrimitiveDrawInterface;
 class USkeletalMeshComponent;
 
-// Facilitates runtime retargeting of an input pose using a user-specified IK Retargeter asset
-// Supports input poses from the animation graph, or from other skeletal mesh components
 UCLASS(MinimalAPI)
-class UAnimGraphNode_RetargetPoseFromMesh : public UAnimGraphNode_Base
+class UAnimGraphNode_RetargetPoseFromMesh : public UAnimGraphNode_CustomProperty
 {
 	GENERATED_BODY()
 
@@ -27,15 +26,17 @@ public:
 	// UEdGraphNode interface
 	UE_API virtual FText GetNodeTitle(ENodeTitleType::Type TitleType) const override;
 	UE_API virtual void PostEditChangeProperty(struct FPropertyChangedEvent& PropertyChangedEvent) override;
-	UE_API virtual void ValidateAnimNodeDuringCompilation(USkeleton* ForSkeleton, FCompilerResultsLog& MessageLog);
+	UE_API virtual void ValidateAnimNodeDuringCompilation(USkeleton* ForSkeleton, FCompilerResultsLog& MessageLog) override;
 	// End of UEdGraphNode interface
 
 	// UAnimGraphNode_Base interface
+	virtual void CustomizeDetails(IDetailLayoutBuilder& DetailBuilder) override;
 	UE_API virtual void PreloadRequiredAssets() override;
 	UE_API virtual void CopyNodeDataToPreviewNode(FAnimNode_Base* AnimNode) override;
 	UE_API virtual FEditorModeID GetEditorMode() const override;
 	UE_API virtual void Draw(FPrimitiveDrawInterface* PDI, USkeletalMeshComponent* PreviewSkelMeshComp) const override;
 	UE_API virtual void CustomizePinData(UEdGraphPin* Pin, FName SourcePropertyName, int32 ArrayIndex) const override;
+	virtual void CreateCustomPins(TArray<UEdGraphPin*>* OldPins) override;
 	virtual bool UsingCopyPoseFromMesh() const override { return true; };
 	// End of UAnimGraphNode_Base interface
 
@@ -44,6 +45,25 @@ public:
 	// End of UK2Node interface
 
 	static UE_API const FName AnimModeName;
+
+private:
+	// UAnimGraphNode_CustomProperty interface
+	virtual FAnimNode_CustomProperty* GetCustomPropertyNode() override { return &Node; }
+	virtual const FAnimNode_CustomProperty* GetCustomPropertyNode() const override { return &Node; }
+	// return nullptr here to completely disable the base class's automatic pin manager
+	// we are manually handling the UI and Pin generation ourselves via CustomizeDetails
+	virtual UStruct* GetTargetStruct() const override { return nullptr; }
+	virtual UStruct* GetTargetSkeletonStruct() const override { return nullptr; }
+	virtual bool NeedsToSpecifyValidTargetClass() const override { return false; }
+	// End of UAnimGraphNode_CustomProperty interface
+
+	void AddVariablesToDetailsPanel(IDetailLayoutBuilder& DetailBuilder);
+	void CreateCustomPinsFromValidAsset();
+	void CreateCustomPinsFromUnloadedAsset(TArray<UEdGraphPin*>* InOldPins);
+	
+	// Set of variable names from the asset that should be exposed as input pins
+	UPROPERTY()
+	TSet<FName> ExposedVariableNames;
 };
 
-#undef UE_API
+#undef UE_API
\ No newline at end of file
diff --git a/Engine/Plugins/Animation/IKRig/Source/IKRigEditor/IKRigEditor.Build.cs b/Engine/Plugins/Animation/IKRig/Source/IKRigEditor/IKRigEditor.Build.cs
--- a/Engine/Plugins/Animation/IKRig/Source/IKRigEditor/IKRigEditor.Build.cs
+++ b/Engine/Plugins/Animation/IKRig/Source/IKRigEditor/IKRigEditor.Build.cs
@@ -38,6 +38,7 @@ namespace UnrealBuildTool.Rules
                     "ContentBrowser",
                     "AnimationBlueprintLibrary",
                     "MessageLog",
+                    "StructUtilsEditor",
 
                     "PropertyEditor",
 					"BlueprintGraph",
diff --git a/Engine/Plugins/Animation/IKRig/Source/IKRigEditor/Private/RetargetEditor/IKRetargetApplicationMode.cpp b/Engine/Plugins/Animation/IKRig/Source/IKRigEditor/Private/RetargetEditor/IKRetargetApplicationMode.cpp
--- a/Engine/Plugins/Animation/IKRig/Source/IKRigEditor/Private/RetargetEditor/IKRetargetApplicationMode.cpp
+++ b/Engine/Plugins/Animation/IKRig/Source/IKRigEditor/Private/RetargetEditor/IKRetargetApplicationMode.cpp
@@ -13,6 +13,7 @@
 #include "RetargetEditor/IKRetargetOpStackTabSummoner.h"
 #include "RetargetEditor/IKRetargetOutputLogTabSummoner.h"
 #include "RetargetEditor/IKRetargetOverridesTabSummoner.h"
+#include "RetargetEditor/IKRetargetVariablesTabSummoner.h"
 
 #define LOCTEXT_NAMESPACE "IKRetargetMode"
 
@@ -44,9 +45,10 @@ FIKRetargetApplicationMode::FIKRetargetApplicationMode(
 	TabFactories.RegisterFactory(MakeShared<FIKRetargetHierarchyTabSummoner>(IKRetargetEditor));
 	TabFactories.RegisterFactory(MakeShared<FIKRetargetOpStackTabSummoner>(IKRetargetEditor));
 	TabFactories.RegisterFactory(MakeShared<FIKRetargetOverridesTabSummoner>(IKRetargetEditor));
+	TabFactories.RegisterFactory(MakeShared<FIKRetargetVariablesTabSummoner>(IKRetargetEditor));
 
 	// create tab layout
-	TabLayout = FTabManager::NewLayout("Standalone_IKRetargetEditor_Layout_v1.021")
+	TabLayout = FTabManager::NewLayout("Standalone_IKRetargetEditor_Layout_v1.022")
 		->AddArea
 		(
 			FTabManager::NewPrimaryArea()
@@ -78,6 +80,8 @@ FIKRetargetApplicationMode::FIKRetargetApplicationMode(
 							FTabManager::NewStack()
 							->SetSizeCoefficient(0.4f)
 							->AddTab(FIKRetargetAssetBrowserTabSummoner::TabID, ETabState::OpenedTab)
+							->AddTab(FIKRetargetOverridesTabSummoner::TabID, ETabState::OpenedTab)
+							->AddTab(FIKRetargetVariablesTabSummoner::TabID, ETabState::OpenedTab)
 						)
 					)
 					->Split
@@ -105,7 +109,6 @@ FIKRetargetApplicationMode::FIKRetargetApplicationMode(
 				FTabManager::NewStack()
 					->SetSizeCoefficient(0.2f)
 					->AddTab(FPersonaTabs::DetailsID, ETabState::OpenedTab)
-					->AddTab(FIKRetargetOverridesTabSummoner::TabID, ETabState::OpenedTab)
 					->AddTab(FPersonaTabs::AdvancedPreviewSceneSettingsID, ETabState::OpenedTab)
 					->SetForegroundTab(FPersonaTabs::DetailsID)
 					
diff --git a/Engine/Plugins/Animation/IKRig/Source/IKRigEditor/Private/RetargetEditor/IKRetargetBatchOperation.cpp b/Engine/Plugins/Animation/IKRig/Source/IKRigEditor/Private/RetargetEditor/IKRetargetBatchOperation.cpp
--- a/Engine/Plugins/Animation/IKRig/Source/IKRigEditor/Private/RetargetEditor/IKRetargetBatchOperation.cpp
+++ b/Engine/Plugins/Animation/IKRig/Source/IKRigEditor/Private/RetargetEditor/IKRetargetBatchOperation.cpp
@@ -285,7 +285,6 @@ void UIKRetargetBatchOperation::RetargetAssets(
 	Params.TargetSkeletalMesh = Context.TargetMesh;
 	Params.RetargeterAsset = Context.IKRetargetAsset;
 	Params.CustomProfile = &RetargetProfile;
-	Params.OverrideSetsToApply = Context.OverrideSetNames;
 	Params.bSuppressWarnings = false;
 	Processor.Initialize(Params);
 	if (!Processor.IsInitialized())
diff --git a/Engine/Plugins/Animation/IKRig/Source/IKRigEditor/Private/RetargetEditor/IKRetargetCommands.cpp b/Engine/Plugins/Animation/IKRig/Source/IKRigEditor/Private/RetargetEditor/IKRetargetCommands.cpp
--- a/Engine/Plugins/Animation/IKRig/Source/IKRigEditor/Private/RetargetEditor/IKRetargetCommands.cpp
+++ b/Engine/Plugins/Animation/IKRig/Source/IKRigEditor/Private/RetargetEditor/IKRetargetCommands.cpp
@@ -12,6 +12,7 @@ void FIKRetargetCommands::RegisterCommands()
 
 	UI_COMMAND(ShowAssetSettings, "Asset Settings", "Set the source or target IK Rigs, preview meshes. Adjust debug settings.", EUserInterfaceActionType::ToggleButton, FInputChord());
 	UI_COMMAND(ShowOverrideSets, "Override Sets", "Open the Override Sets tab to create / edit sets of op property overrides.", EUserInterfaceActionType::ToggleButton, FInputChord());
+	UI_COMMAND(ShowVariables, "Variables", "Open the Variables tab to create / edit variables to drive property overrides.", EUserInterfaceActionType::ToggleButton, FInputChord());
 	
 	UI_COMMAND(ResetAllBones, "Reset All", "Sets the selected bones to the mesh reference pose.", EUserInterfaceActionType::Button, FInputChord());
 	UI_COMMAND(ResetSelectedBones, "Reset Selected Bones", "Sets the selected bones to the mesh reference pose.", EUserInterfaceActionType::Button, FInputChord());
diff --git a/Engine/Plugins/Animation/IKRig/Source/IKRigEditor/Private/RetargetEditor/IKRetargetEditor.cpp b/Engine/Plugins/Animation/IKRig/Source/IKRigEditor/Private/RetargetEditor/IKRetargetEditor.cpp
--- a/Engine/Plugins/Animation/IKRig/Source/IKRigEditor/Private/RetargetEditor/IKRetargetEditor.cpp
+++ b/Engine/Plugins/Animation/IKRig/Source/IKRigEditor/Private/RetargetEditor/IKRetargetEditor.cpp
@@ -26,6 +26,7 @@
 #include "RetargetEditor/IKRetargetDefaultMode.h"
 #include "RetargetEditor/IKRetargetEditorController.h"
 #include "RetargetEditor/IKRetargetOverridesTabSummoner.h"
+#include "RetargetEditor/IKRetargetVariablesTabSummoner.h"
 
 #include "Styling/AppStyle.h"
 
@@ -151,12 +152,43 @@ void FIKRetargetEditor::BindCommands()
 		Commands.ShowOverrideSets,
 		FExecuteAction::CreateLambda([this]()
 		{
-			FGlobalTabmanager::Get()->TryInvokeTab(FIKRetargetOverridesTabSummoner::TabID);
+			TSharedPtr<SDockTab> Tab = TabManager->TryInvokeTab(FIKRetargetOverridesTabSummoner::TabID);
+			if (Tab.IsValid())
+			{
+				Tab->ActivateInParent(ETabActivationCause::SetDirectly);
+				
+				if (TSharedPtr<SWindow> ParentWindow = Tab->GetParentWindow())
+				{
+					ParentWindow->BringToFront();
+				}
+			}
 		}),
 		FCanExecuteAction(),
 		FIsActionChecked::CreateLambda([this]() ->bool
 		{
-			TSharedPtr<SDockTab> ExistingTab = FGlobalTabmanager::Get()->FindExistingLiveTab(FIKRetargetOverridesTabSummoner::TabID);
+			TSharedPtr<SDockTab> ExistingTab = TabManager->FindExistingLiveTab(FIKRetargetOverridesTabSummoner::TabID);
+			return ExistingTab.IsValid() && ExistingTab->IsForeground();
+		}));
+
+	ToolkitCommands->MapAction(
+		Commands.ShowVariables,
+		FExecuteAction::CreateLambda([this]()
+		{
+			TSharedPtr<SDockTab> Tab = TabManager->TryInvokeTab(FIKRetargetVariablesTabSummoner::TabID);
+			if (Tab.IsValid())
+			{
+				Tab->ActivateInParent(ETabActivationCause::SetDirectly);
+
+				if (TSharedPtr<SWindow> ParentWindow = Tab->GetParentWindow())
+				{
+					ParentWindow->BringToFront();
+				}
+			}
+		}),
+		FCanExecuteAction(),
+		FIsActionChecked::CreateLambda([this]() ->bool
+		{
+			TSharedPtr<SDockTab> ExistingTab = TabManager->FindExistingLiveTab(FIKRetargetVariablesTabSummoner::TabID);
 			return ExistingTab.IsValid() && ExistingTab->IsForeground();
 		}));
 
@@ -304,6 +336,13 @@ void FIKRetargetEditor::FillToolbar(FToolBarBuilder& ToolbarBuilder)
 		TAttribute<FText>(),
 		TAttribute<FText>(),
 		FSlateIcon(FIKRetargetEditorStyle::Get().GetStyleSetName(),"IKRetarget.AssetSettings"));
+
+		ToolbarBuilder.AddToolBarButton(
+		FIKRetargetCommands::Get().ShowVariables,
+		NAME_None,
+		TAttribute<FText>(),
+		TAttribute<FText>(),
+		FSlateIcon(FIKRetargetEditorStyle::Get().GetStyleSetName(),"IKRetarget.AssetSettings"));
 	}
 	ToolbarBuilder.EndSection();
 }
diff --git a/Engine/Plugins/Animation/IKRig/Source/IKRigEditor/Private/RetargetEditor/IKRetargetEditorController.cpp b/Engine/Plugins/Animation/IKRig/Source/IKRigEditor/Private/RetargetEditor/IKRetargetEditorController.cpp
--- a/Engine/Plugins/Animation/IKRig/Source/IKRigEditor/Private/RetargetEditor/IKRetargetEditorController.cpp
+++ b/Engine/Plugins/Animation/IKRig/Source/IKRigEditor/Private/RetargetEditor/IKRetargetEditorController.cpp
@@ -37,6 +37,7 @@
 
 #include "Misc/ITransaction.h"
 #include "RetargetEditor/SIKRetargetOverrideManager.h"
+#include "RetargetEditor/SIKRetargetVariablesEditor.h"
 
 #define LOCTEXT_NAMESPACE "IKRetargetEditorController"
 
@@ -181,10 +182,16 @@ void FIKRetargetEditorController::Initialize(TSharedPtr<FIKRetargetEditor> InEdi
 		RefreshOpStackView();
 	});
 	// bind callback when override sets are modified
-	ProfilesModifiedDelegateHandle = AssetController->OnOverrideSetsModified().AddLambda([this]()
+	OverrideSetsModifiedDelegateHandle = AssetController->OnOverrideSetsModified().AddLambda([this]()
 	{
 		RefreshOverrideSetView();
 	});
+	// bind callback when variables are modified
+	VariablesModifiedDelegateHandle = AssetController->OnVariablesModified().AddLambda([this]()
+	{
+		RefreshVariablesView();
+	});
+	
 	// bind callback when IK Rig asset is replaced with a different asset
 	IKRigReplacedDelegateHandle = AssetController->OnIKRigReplaced().AddSP(this, &FIKRetargetEditorController::HandleIKRigReplaced);
 	// bind callback when Preview Mesh asset is replaced with a different asset
@@ -194,7 +201,8 @@ void FIKRetargetEditorController::Initialize(TSharedPtr<FIKRetargetEditor> InEdi
 void FIKRetargetEditorController::Close()
 {
 	AssetController->OnRetargeterNeedsInitialized().Remove(RetargeterReInitDelegateHandle);
-	AssetController->OnOverrideSetsModified().Remove(ProfilesModifiedDelegateHandle);
+	AssetController->OnOverrideSetsModified().Remove(OverrideSetsModifiedDelegateHandle);
+	AssetController->OnVariablesModified().Remove(VariablesModifiedDelegateHandle);
 	AssetController->OnOpStackModified().Remove(OpStackModifiedDelegateHandle);
 	AssetController->OnIKRigReplaced().Remove(IKRigReplacedDelegateHandle);
 	AssetController->OnPreviewMeshReplaced().Remove(PreviewMeshReplacedDelegateHandle);
@@ -804,6 +812,7 @@ void FIKRetargetEditorController::RefreshAllViews() const
 	RefreshAssetBrowserView();
 	RefreshHierarchyView();
 	RefreshOverrideSetView();
+	RefreshVariablesView();
 }
 
 void FIKRetargetEditorController::RefreshDetailsView() const
@@ -848,6 +857,14 @@ void FIKRetargetEditorController::RefreshOverrideSetView() const
 	}
 }
 
+void FIKRetargetEditorController::RefreshVariablesView() const
+{
+	if (VariablesView.IsValid())
+	{
+		VariablesView->RefreshList();
+	}
+}
+
 void FIKRetargetEditorController::RefreshPoseList() const
 {
 	if (HierarchyView.IsValid())
@@ -900,7 +917,7 @@ void FIKRetargetEditorController::ShowDetailsForOp(const int32 OpIndex) const
 	StructToView.Owner = AssetController->GetAsset();
 	StructToView.UniqueName = Op->GetName();
 	
-	UIKRigStructViewer* StructViewer = AssetController->GetStructViewer();
+	UIKRigStructViewer* StructViewer = AssetController->GetStructViewer(ERetargetStructViewerMode::Settings);
 	StructViewer->SetStructToView(StructToView);
 	
 	constexpr bool bForceRefresh = true; // must forcibly refresh because struct viewer is a recycled UObject
diff --git a/Engine/Plugins/Animation/IKRig/Source/IKRigEditor/Private/RetargetEditor/IKRetargetOverrideCustomization.cpp b/Engine/Plugins/Animation/IKRig/Source/IKRigEditor/Private/RetargetEditor/IKRetargetOverrideCustomization.cpp
--- a/Engine/Plugins/Animation/IKRig/Source/IKRigEditor/Private/RetargetEditor/IKRetargetOverrideCustomization.cpp
+++ b/Engine/Plugins/Animation/IKRig/Source/IKRigEditor/Private/RetargetEditor/IKRetargetOverrideCustomization.cpp
@@ -11,7 +11,6 @@
 #include "Retargeter/IKRetargeter.h"
 #include "Retargeter/IKRetargetOverrides.h"
 #include "RigEditor/IKRigStructViewer.h"
-#include "StructUtils/InstancedStruct.h"
 #include "Widgets/Input/SButton.h"
 #include "Widgets/Input/SEditableTextBox.h"
 #include "Widgets/Text/STextBlock.h"
@@ -175,19 +174,40 @@ void FRetargetOverrideCustomization::CustomizeChildren(TSharedRef<IPropertyHandl
 					continue;
 				}
 
-				// a callback to marshall edits back into the override string
-				auto UpdateValueAction = [LocalController, SetName, OpName, PropertyPath, OpSettingsScopedStruct]()
+				// inject the asset into the active transaction
+				auto PreChangeAction = [LocalController]()
 				{
 					if (LocalController.IsValid())
 					{
-						ensure(LocalController->UpdateOverrideValue(SetName, OpName, PropertyPath, *OpSettingsScopedStruct.Get()));
+						if (UIKRetargeter* Asset = LocalController->GetAsset())
+						{
+							Asset->Modify();
+						}
 					}
 				};
 
-				// bind to the main property (handles bool/float)
-				PropertyHandle->SetOnPropertyValueChanged(FSimpleDelegate::CreateLambda(UpdateValueAction));
-				// bind to the children (handles struct properties like FVector X, Y, Z changes)
-				PropertyHandle->SetOnChildPropertyValueChanged(FSimpleDelegate::CreateLambda(UpdateValueAction));
+				// marshal the data back into the asset
+				auto PostChangeAction = [LocalController, SetName, OpName, PropertyPath, OpSettingsScopedStruct]()
+				{
+					if (LocalController.IsValid())
+					{
+						constexpr bool bShouldTransact = false;
+						LocalController->UpdateOverrideValue(
+							SetName,
+							OpName,
+							PropertyPath,
+							*OpSettingsScopedStruct.Get(),
+							bShouldTransact);
+					}
+				};
+
+				// bind the PRE-change delegates
+				PropertyHandle->SetOnPropertyValuePreChange(FSimpleDelegate::CreateLambda(PreChangeAction));
+				PropertyHandle->SetOnChildPropertyValuePreChange(FSimpleDelegate::CreateLambda(PreChangeAction));
+
+				// bind the POST-change delegates
+				PropertyHandle->SetOnPropertyValueChanged(FSimpleDelegate::CreateLambda(PostChangeAction));
+				PropertyHandle->SetOnChildPropertyValueChanged(FSimpleDelegate::CreateLambda(PostChangeAction));
 
 				// add the widget for the property to the group
 				AddRowForPropertyOverride(PropertyHandle, PropertyGroup, SetName, OpName, PropertyPath);
@@ -246,34 +266,242 @@ void FRetargetOverrideCustomization::AddRowForPropertyOverride(
 	PropertyRow.CustomWidget()
 	.NameContent()
 	[
-	   DefaultNameWidget.ToSharedRef()
+		DefaultNameWidget.ToSharedRef()
 	]
 	.ValueContent()
 	[
-	   DefaultValueWidget.ToSharedRef()
-	]
-	.ExtensionContent()
-	[
-		// TODO add curve / variable binding UI here
-
-		SNew(SButton)
-		.ButtonStyle(FAppStyle::Get(), "SimpleButton")
-		.OnClicked_Lambda([LocalController, InOverrideSetName, InOpName, InPropertyPath]()
+		SNew(SBox) // wraps the default value widget so we can dynamically control its enabled state
+		.IsEnabled_Lambda([LocalController, InOverrideSetName, InOpName, InPropertyPath]()
 		{
 			if (LocalController.IsValid())
 			{
-				LocalController->RemovePropertyOverrideFromOp(InOverrideSetName, InOpName, InPropertyPath);
+				// enable the widget only if it does NOT have a binding
+				return !LocalController->GetPropertyOverrideHasBinding(InOverrideSetName, InOpName, InPropertyPath);
 			}
-			return FReply::Handled();
+			return false;
 		})
 		[
-			SNew(SImage)
-			.Image(FAppStyle::GetBrush("Icons.Delete"))
-			.ColorAndOpacity(FSlateColor::UseForeground())
+			DefaultValueWidget.ToSharedRef()
+		]
+	]
+	.ExtensionContent()
+	[
+		SNew(SHorizontalBox)
+
+		+ SHorizontalBox::Slot()
+		.Padding(0.0f, 2.0f)
+		[
+			SNew(SComboButton)
+			.ComboButtonStyle(FAppStyle::Get(), "ComboButton") 
+			.OnGetMenuContent(this, &FRetargetOverrideCustomization::GetBindingMenuContent, InPropertyHandle, InOverrideSetName, InOpName, InPropertyPath)
+			.HasDownArrow(true)
+			.ContentPadding(FMargin(6.0f, 2.0f))
+			.ButtonContent()
+			[
+				SNew(SHorizontalBox)
+				.Clipping(EWidgetClipping::ClipToBounds)
+				
+				+ SHorizontalBox::Slot()
+				.AutoWidth()
+				.VAlign(VAlign_Center)
+				.Padding(8.0f, 0.0f, 0.0f, 0.0f)
+				[
+					SNew(STextBlock)
+					.Font(IDetailLayoutBuilder::GetDetailFont())
+					.Text_Lambda([LocalController, InOverrideSetName, InOpName, InPropertyPath]()
+					{
+						if (!LocalController.IsValid())
+						{
+							return FText::GetEmpty();
+						}
+							
+						const FName BoundCurve = LocalController->GetCurveBoundToPropertyOverride(InOverrideSetName, InOpName, InPropertyPath);
+						if (!BoundCurve.IsNone())
+						{
+							return FText::FromName(BoundCurve);
+						}
+							
+						const FName BoundVariable = LocalController->GetVariableBoundToPropertyOverride(InOverrideSetName, InOpName, InPropertyPath);
+						if (!BoundVariable.IsNone())
+						{
+							return FText::FromName(BoundVariable);
+						}
+							
+						return LOCTEXT("NotBoundLabel", "Not Bound");
+					})
+					.ColorAndOpacity(FSlateColor::UseForeground())
+				]
+			]
+		]
+		
+		+ SHorizontalBox::Slot()
+		.AutoWidth()
+		.VAlign(VAlign_Center)
+		[
+			SNew(SButton)
+			.ButtonStyle(FAppStyle::Get(), "SimpleButton")
+			.OnClicked_Lambda([LocalController, InOverrideSetName, InOpName, InPropertyPath]()
+			{
+				if (LocalController.IsValid())
+				{
+					LocalController->RemovePropertyOverrideFromOp(InOverrideSetName, InOpName, InPropertyPath);
+				}
+				return FReply::Handled();
+			})
+			[
+				SNew(SImage)
+				.Image(FAppStyle::GetBrush("Icons.Delete"))
+				.ColorAndOpacity(FSlateColor::UseForeground())
+			]
 		]
 	];
 }
 
+TSharedRef<SWidget> FRetargetOverrideCustomization::GetBindingMenuContent(
+	TSharedPtr<IPropertyHandle> InPropertyHandle, 
+	FName InOverrideSetName, 
+	FName InOpName, 
+	FString InPropertyPath)
+{
+	FMenuBuilder MenuBuilder(true, nullptr);
+
+	TWeakObjectPtr<UIKRetargeterController> LocalController = AssetController;
+	if (!LocalController.IsValid())
+	{
+		return MenuBuilder.MakeWidget();
+	}
+	
+	// SECTION 1: UNBIND
+	MenuBuilder.BeginSection("UnbindSection", FText::FromString("Unbind"));
+	{
+		MenuBuilder.AddMenuEntry(
+			FText::FromString("None"),
+			FText::FromString("Clear the current variable binding"),
+			FSlateIcon(FAppStyle::GetAppStyleSetName(), "Cross"),
+			FUIAction(FExecuteAction::CreateLambda([LocalController, InOverrideSetName, InOpName, InPropertyPath]()
+			{
+				if (LocalController.IsValid())
+				{
+					LocalController->ClearPropertyOverrideBinding(InOverrideSetName, InOpName, InPropertyPath);
+				}
+			}))
+		);
+	}
+	MenuBuilder.EndSection();
+	
+	// SECTION 2: BIND TO CURVE
+	MenuBuilder.BeginSection("CurveSection", LOCTEXT("CurveBindHeaderLabel", "Bind to Curve"));
+	{
+		TSharedRef<SWidget> CurveEntryWidget = SNew(SBox)
+		.Padding(FMargin(16.0f, 4.0f, 16.0f, 4.0f))
+		.MinDesiredWidth(150.0f)
+		[
+			SNew(SEditableTextBox)
+			.HintText(LOCTEXT("EnterCurveNameLabel", "Enter Curve Name..."))
+			.IsEnabled_Lambda([LocalController, InOpName, InPropertyPath]()
+				{
+					if (!LocalController.IsValid())
+					{
+						return false;
+					}
+					return LocalController->CanPropertyOverrideBindToCurve(InOpName, InPropertyPath);
+				})
+			.Text_Lambda([LocalController, InOverrideSetName, InOpName, InPropertyPath]()
+				{
+					if (!LocalController.IsValid())
+					{
+						return FText::FromString("");
+					}
+
+					if (!LocalController->CanPropertyOverrideBindToCurve(InOpName, InPropertyPath))
+					{
+						return LOCTEXT("CannotBindToCurveLabel", "Type incompatible with curves.");
+					}
+					return FText::FromName(LocalController->GetCurveBoundToPropertyOverride(InOverrideSetName, InOpName, InPropertyPath));
+				})
+			.OnTextCommitted_Lambda([LocalController, InOverrideSetName, InOpName, InPropertyPath](const FText& NewText, ETextCommit::Type CommitInfo)
+			{
+				if (CommitInfo == ETextCommit::OnEnter && LocalController.IsValid())
+				{
+					FName CurveName = FName(*NewText.ToString());
+					LocalController->BindPropertyOverrideToCurve(CurveName, InOverrideSetName, InOpName, InPropertyPath);
+					FSlateApplication::Get().DismissAllMenus();
+				}
+			})
+		];
+
+		MenuBuilder.AddWidget(CurveEntryWidget, FText::GetEmpty(), false, false);
+	}
+	MenuBuilder.EndSection();
+	
+	// SECTION 3: BIND TO VARIABLE
+	MenuBuilder.BeginSection("VariableSection", LOCTEXT("VariableBindHeaderLabel", "Bind to Variable"));
+	{
+		const FProperty* PropertyToBind = InPropertyHandle->GetProperty();
+		const FInstancedPropertyBag& VariablePropertyBag = AssetController->GetVariables().Bag;
+		const UPropertyBag* VariablePropertyBagStruct = VariablePropertyBag.GetPropertyBagStruct();
+
+		if (VariablePropertyBagStruct)
+		{
+			int32 NumAddedProperties = 0;
+			const UEnum* PropertyTypeEnum = StaticEnum<EPropertyBagPropertyType>();
+
+			for (const FPropertyBagPropertyDesc& VariablePropDesc : VariablePropertyBagStruct->GetPropertyDescs())
+			{
+				if (!VariablePropDesc.CachedProperty)
+				{
+					continue;
+				}
+				if (!AssetController->IsPropertyCompatibleForBinding(VariablePropDesc.CachedProperty, PropertyToBind))
+				{
+					continue;
+				}
+
+				NumAddedProperties++;
+				
+				const FName VariableName = VariablePropDesc.Name;
+				FString TypeString = PropertyTypeEnum->GetNameStringByValue((int64)VariablePropDesc.ValueType).ToLower();
+				FText ItemLabel = FText::Format(
+					FText::FromString("{0} ({1})"), 
+					FText::FromName(VariableName), 
+					FText::FromString(TypeString)
+				);
+
+				MenuBuilder.AddMenuEntry(
+					ItemLabel,
+					FText::Format(LOCTEXT("BindToVariableTooltip", "Bind to variable: {0}"), FText::FromName(VariableName)),
+					FSlateIcon(),
+					FUIAction(FExecuteAction::CreateLambda([LocalController, VariableName, InOverrideSetName, InOpName, InPropertyPath]()
+					{
+						if (LocalController.IsValid())
+						{
+							LocalController->BindPropertyOverrideToVariable(VariableName, InOverrideSetName, InOpName, InPropertyPath);
+						}
+					}))
+				);
+			}
+			
+			if (NumAddedProperties == 0)
+			{
+				MenuBuilder.AddWidget(
+					SNew(SBox)
+					.Padding(FMargin(16.0f, 4.0f))
+					[
+						SNew(STextBlock)
+						.Text(LOCTEXT("NoCompatibleVariablesLabel", "No compatible variables found."))
+						.ColorAndOpacity(FSlateColor::UseSubduedForeground())
+						.Font(IDetailLayoutBuilder::GetDetailFont())
+					], 
+					FText::GetEmpty(), false, false
+				);
+			}
+		}
+	}
+	MenuBuilder.EndSection();
+
+	return MenuBuilder.MakeWidget();
+}
+
 TMap<FString, TArray<int32>> FRetargetOverrideCustomization::MapOverridesToGroups(const TSharedPtr<IPropertyHandle>& OpOverridesHandle) const
 {
     TMap<FString, TArray<int32>> GroupMap;
diff --git a/Engine/Plugins/Animation/IKRig/Source/IKRigEditor/Private/RetargetEditor/IKRetargetVariablesTabSummoner.cpp b/Engine/Plugins/Animation/IKRig/Source/IKRigEditor/Private/RetargetEditor/IKRetargetVariablesTabSummoner.cpp
new file mode 100644
--- /dev/null
+++ b/Engine/Plugins/Animation/IKRig/Source/IKRigEditor/Private/RetargetEditor/IKRetargetVariablesTabSummoner.cpp
@@ -0,0 +1,36 @@
+﻿// Copyright Epic Games, Inc. All Rights Reserved.
+#include "RetargetEditor/IKRetargetVariablesTabSummoner.h"
+
+#include "IDocumentation.h"
+#include "RetargetEditor/IKRetargetEditor.h"
+#include "RetargetEditor/SIKRetargetVariablesEditor.h"
+
+#define LOCTEXT_NAMESPACE "IKRetargetVariablesTabSummoner"
+
+const FName FIKRetargetVariablesTabSummoner::TabID(TEXT("Variables"));
+
+FIKRetargetVariablesTabSummoner::FIKRetargetVariablesTabSummoner(
+	const TSharedRef<FIKRetargetEditor>& InRetargetEditor)
+	: FWorkflowTabFactory(TabID, InRetargetEditor),
+	IKRetargetEditor(InRetargetEditor)
+{
+	bIsSingleton = true; // only allow a single instance of this tab
+	
+	TabLabel = LOCTEXT("IKRetargetVariablesTabLabel", "Variables");
+	TabIcon = FSlateIcon(FAppStyle::GetAppStyleSetName(), "Kismet.Tabs.CompilerResults");
+
+	ViewMenuDescription = LOCTEXT("IKRetargetVariables_ViewMenu_Desc", "Variables");
+	ViewMenuTooltip = LOCTEXT("IKRetargetVariables_ViewMenu_ToolTip", "Show the Retargeting Variables Tab");
+}
+
+TSharedPtr<SToolTip> FIKRetargetVariablesTabSummoner::CreateTabToolTipWidget(const FWorkflowTabSpawnInfo& Info) const
+{
+	return  IDocumentation::Get()->CreateToolTip(LOCTEXT("IKRetargetVariablesTooltip", "Create and edit variables to drive overrides."), NULL, TEXT("Shared/Editors/Persona"), TEXT("IKRetargetVariables_Window"));
+}
+
+TSharedRef<SWidget> FIKRetargetVariablesTabSummoner::CreateTabBody(const FWorkflowTabSpawnInfo& Info) const
+{
+	return SNew(SIKRetargetVariablesEditor, IKRetargetEditor.Pin()->GetController());
+}
+
+#undef LOCTEXT_NAMESPACE 
diff --git a/Engine/Plugins/Animation/IKRig/Source/IKRigEditor/Private/RetargetEditor/IKRetargeterController.cpp b/Engine/Plugins/Animation/IKRig/Source/IKRigEditor/Private/RetargetEditor/IKRetargeterController.cpp
--- a/Engine/Plugins/Animation/IKRig/Source/IKRigEditor/Private/RetargetEditor/IKRetargeterController.cpp
+++ b/Engine/Plugins/Animation/IKRig/Source/IKRigEditor/Private/RetargetEditor/IKRetargeterController.cpp
@@ -7,6 +7,7 @@
 #include "Retargeter/IKRetargeter.h"
 #include "RigEditor/IKRigController.h"
 #include "IKRigEditor.h"
+#include "PropertyBagDetails.h"
 #include "RetargetEditor/IKRetargeterPoseGenerator.h"
 #include "Retargeter/IKRetargetOps.h"
 #include "Retargeter/RetargetOps/CurveRemapOp.h"
@@ -41,8 +42,15 @@ UIKRetargeterController* UIKRetargeterController::GetController(const UIKRetarge
 
 UIKRetargeterController::UIKRetargeterController()
 {
-	StructViewer = CreateDefaultSubobject<UIKRigStructViewer>(TEXT("RetargetSettingsViewer"));
-	StructViewer->OnStructNeedsReinit().AddUObject(this, &UIKRetargeterController::OnOpPropertyChanged);
+	// initialize Op Settings viewer
+	UIKRigStructViewer* SettingsViewer = CreateDefaultSubobject<UIKRigStructViewer>(TEXT("RetargetSettingsViewer"));
+	SettingsViewer->OnNeedsReinitialized().AddUObject(this, &UIKRetargeterController::OnOpPropertyChanged);
+	StructViewers.Add(ERetargetStructViewerMode::Settings, SettingsViewer);
+
+	// initialize Override Sets viewer
+	UIKRigStructViewer* OverrideViewer = CreateDefaultSubobject<UIKRigStructViewer>(TEXT("RetargetOverrideSetsViewer"));
+	OverrideViewer->OnNeedsReinitialized().AddUObject(this, &UIKRetargeterController::OnOpPropertyChanged);
+	StructViewers.Add(ERetargetStructViewerMode::OverrideSets, OverrideViewer);
 }
 
 void UIKRetargeterController::PostInitProperties()
@@ -862,8 +870,9 @@ void UIKRetargeterController::CleanPoseList(const ERetargetSourceOrTarget Source
 	}
 }
 
-UIKRigStructViewer* UIKRetargeterController::GetStructViewer()
+UIKRigStructViewer* UIKRetargeterController::GetStructViewer(ERetargetStructViewerMode InMode)
 {
+	TObjectPtr<UIKRigStructViewer> StructViewer = StructViewers[InMode];
 	StructViewer->Reset();
 	return StructViewer;
 }
@@ -1661,6 +1670,20 @@ FRetargetOpOverrides* UIKRetargeterController::FindOpOverrides(const FName InOve
     return nullptr;
 }
 
+FRetargetOpPropertyOverride* UIKRetargeterController::FindPropertyOverride(
+	const FName InOverrideSetName,
+	const FName InOpName,
+	const FString& InPropertyPath) const
+{
+	FRetargetOpOverrides* OpOverrides = FindOpOverrides(InOverrideSetName, InOpName);
+	if (!OpOverrides)
+	{
+		return nullptr;
+	}
+
+	return OpOverrides->FindPropertyOverride(InPropertyPath);
+}
+
 FRetargetOpOverrides* UIKRetargeterController::AddOpOverrides(const FName InOverrideSetName, const FName InOpName) const
 {
     FRetargetOverrideSet* OverrideSet = FindOverrideSet(InOverrideSetName);
@@ -1720,7 +1743,8 @@ bool UIKRetargeterController::UpdateOverrideValue(
 	const FName InOverrideSetName,
 	const FName InOpName,
 	const FString& InPropertyPath,
-	const FStructOnScope& InNewValue) const
+	const FStructOnScope& InNewValue,
+	const bool bShouldTransact) const
 {
 	FRetargetOpOverrides* OpOverrides = FindOpOverrides(InOverrideSetName, InOpName);
 	if (!OpOverrides)
@@ -1733,19 +1757,135 @@ bool UIKRetargeterController::UpdateOverrideValue(
 	{
 		return false;
 	}
-	
-	constexpr bool bShouldTransact = true;
-	FScopedTransaction Transaction(LOCTEXT("UpdateOverrideValue", "Edit Override Value"), bShouldTransact);
-	Asset->Modify();
 
+	if (bShouldTransact)
+	{
+		FScopedTransaction Transaction(LOCTEXT("UpdateOverrideValue", "Edit Override Value"), bShouldTransact);
+		Asset->Modify();
+	}
+	
 	const UStruct* BaseStruct = InNewValue.GetStruct();
 	const UScriptStruct* ScriptStruct = CastChecked<UScriptStruct>(BaseStruct);
 	OpOverrides->UpdateOverrideValue(*PropertyOverride, InNewValue.GetStructMemory(), ScriptStruct);
 
+	if (bShouldTransact)
+	{
+		PostEditOverrideSets();	
+	}
+	
+	return true;
+}
+
+bool UIKRetargeterController::BindPropertyOverrideToCurve(
+	const FName& InCurveName,
+	const FName InOverrideSetName,
+	const FName InOpName,
+	const FString& InPropertyPath) const
+{
+	FRetargetOpPropertyOverride* PropertyOverride = FindPropertyOverride(InOverrideSetName, InOpName, InPropertyPath);
+	if (!PropertyOverride)
+	{
+		return false;
+	}
+	
+	if (!CanPropertyOverrideBindToCurve(InOpName, InPropertyPath))
+	{
+		return false; // override type not compatible with curves
+	}
+
+	constexpr bool bShouldTransact = true;
+	FScopedTransaction Transaction(LOCTEXT("BindCurveToProperty", "Edited Property Override Curve Binding"), bShouldTransact);
+	Asset->Modify();
+
+	PropertyOverride->SetBoundCurveName(InCurveName);
+	PropertyOverride->SetBoundVariableName(NAME_None); // mutually exclusive binding
+
 	PostEditOverrideSets();
 	return true;
 }
 
+bool UIKRetargeterController::BindPropertyOverrideToVariable(
+	const FName InVariableName,
+	const FName InOverrideSetName,
+	const FName InOpName,
+	const FString& InPropertyPath) const
+{
+	FProperty* OverrideProperty = GetPropertyForOverride(InOpName, InPropertyPath);
+	if (!OverrideProperty)
+	{
+		return false;
+	}
+
+	// user can pass in Name_None to "unbind" a property
+	// unless we are unbinding, we must validate compatibility between the overridden property type and the variable type
+	if (InVariableName != NAME_None)
+	{
+		const FProperty* VariableProperty = GetVariableProperty(InVariableName);
+		if (!VariableProperty)
+		{
+			return false; // variable not found with that name
+		}
+
+		if (!IsPropertyCompatibleForBinding(OverrideProperty, VariableProperty))
+		{
+			return false; // variable type not compatible with override property type
+		}
+	}
+
+	FRetargetOpPropertyOverride* PropertyOverride = FindPropertyOverride(InOverrideSetName, InOpName, InPropertyPath);
+	if (!PropertyOverride)
+	{
+		return false;
+	}
+	
+	constexpr bool bShouldTransact = true;
+	FScopedTransaction Transaction(LOCTEXT("BindVariableToProperty", "Edited Property Override Variable Binding"), bShouldTransact);
+	Asset->Modify();
+	
+	PropertyOverride->SetBoundVariableName(InVariableName);
+	PropertyOverride->SetBoundCurveName(NAME_None); // mutually exclusive binding
+
+	PostEditOverrideSets();
+	return true;
+}
+
+bool UIKRetargeterController::ClearPropertyOverrideBinding(const FName InOverrideSetName, const FName InOpName, const FString& InPropertyPath) const
+{
+	FRetargetOpPropertyOverride* PropertyOverride = FindPropertyOverride(InOverrideSetName, InOpName, InPropertyPath);
+	if (!PropertyOverride)
+	{
+		return false; // not found
+	}
+
+	if (PropertyOverride->GetBoundCurveName() == NAME_None && PropertyOverride->GetBoundVariableName() == NAME_None)
+	{
+		return false; // not bound
+	}
+	
+	constexpr bool bShouldTransact = true;
+	FScopedTransaction Transaction(LOCTEXT("ClearPropertyBinding", "Clear Override Binding"), bShouldTransact);
+	Asset->Modify();
+	
+	PropertyOverride->SetBoundVariableName(NAME_None);
+	PropertyOverride->SetBoundCurveName(NAME_None);
+
+	PostEditOverrideSets();
+	return true;
+}
+
+bool UIKRetargeterController::CanPropertyOverrideBindToCurve(
+	const FName InOpName,
+	const FString& InPropertyPath) const
+{
+	FProperty* OverrideProperty = GetPropertyForOverride(InOpName, InPropertyPath);
+	if (!OverrideProperty)
+	{
+		return false;
+	}
+
+	return OverrideProperty->IsA<FNumericProperty>(); // can only bind numeric properties to curves
+}
+
 bool UIKRetargeterController::AddPropertyOverrideToOp(const FName InOverrideSetName, const FName InOpName, const FString& InPropertyPath) const
 {
     FIKRetargetOpBase* Op = GetRetargetOpByName(InOpName);
@@ -1944,6 +2084,282 @@ FName UIKRetargeterController::GetOverrideSetName(const FRetargetOverrideSet* In
     return NAME_None;
 }
 
+FProperty* UIKRetargeterController::GetPropertyForOverride(const FName InOpName, const FString& InPropertyPath) const
+{
+	const FIKRetargetOpBase* Op = GetRetargetOpByName(InOpName);
+	if (!Op)
+	{
+		return nullptr; // op not found in stack
+	}
+
+	// try to parse the path into segments
+	const UScriptStruct* SettingsStruct = Op->GetSettingsType();
+	TArray<FRetargetOpPropertyOverride::FPropertySegment> Segments;
+	if (!FRetargetOpPropertyOverride::GetSegmentsFromProperyPath(InPropertyPath, SettingsStruct, Segments))
+	{
+		return nullptr; // the path is broken (ie, that property was not found on the given op settings)
+	}
+
+	// get the leaf property
+	return FRetargetOpPropertyOverride::GetLeafProperty(Segments);
+}
+
+const FRetargetVariableContainer& UIKRetargeterController::GetVariables() const
+{
+	return Asset->GetVariables();
+}
+
+const FProperty* UIKRetargeterController::GetVariableProperty(const FName InVariableName) const
+{
+	const FRetargetVariableContainer& Variables = Asset->GetVariables();
+	const FPropertyBagPropertyDesc* PropertyDesc = Variables.Bag.FindPropertyDescByName(InVariableName);
+	if (!PropertyDesc)
+	{
+		return nullptr;
+	}
+	
+	return PropertyDesc->CachedProperty;
+}
+
+bool UIKRetargeterController::IsPropertyCompatibleForBinding(const FProperty* Source, const FProperty* Target)
+{
+	if (!Source || !Target)
+	{
+		return false;
+	}
+
+	// exact match (handles structs, objects, enums)
+	if (Source->SameType(Target))
+	{
+		return true;
+	}
+
+	// numeric cross-compatibility (float <-> double)
+	if (Source->IsA<FNumericProperty>() && Target->IsA<FNumericProperty>())
+	{
+		const FNumericProperty* SourceNumeric = CastField<FNumericProperty>(Source);
+		const FNumericProperty* TargetNumeric = CastField<FNumericProperty>(Target);
+
+		// allow any float type to any float type
+		if (SourceNumeric->IsFloatingPoint() && TargetNumeric->IsFloatingPoint())
+		{
+			return true;
+		}
+
+		// allow any int type to any int type (int, int32, uint32 etc)
+		if (SourceNumeric->IsInteger() && TargetNumeric->IsInteger())
+		{
+			return true;
+		}
+	}
+
+	return false;
+};
+
+FName UIKRetargeterController::GetVariableBoundToPropertyOverride(
+	const FName InOverrideSetName,
+	const FName InOpName,
+	const FString& InPropertyPath) const
+{
+	FRetargetOpPropertyOverride* PropertyOverride = FindPropertyOverride(InOverrideSetName, InOpName, InPropertyPath);
+	if (!PropertyOverride)
+	{
+		return NAME_None; // no override found
+	}
+
+	return PropertyOverride->GetBoundVariableName();
+}
+
+FName UIKRetargeterController::GetCurveBoundToPropertyOverride(
+	const FName InOverrideSetName,
+	const FName InOpName,
+	const FString& InPropertyPath) const
+{
+	FRetargetOpPropertyOverride* PropertyOverride = FindPropertyOverride(InOverrideSetName, InOpName, InPropertyPath);
+	if (!PropertyOverride)
+	{
+		return NAME_None; // no override found
+	}
+	
+	return PropertyOverride->GetBoundCurveName();
+}
+
+bool UIKRetargeterController::GetPropertyOverrideHasBinding(
+	const FName InOverrideSetName,
+	const FName InOpName,
+	const FString& InPropertyPath) const
+{
+	FRetargetOpPropertyOverride* PropertyOverride = FindPropertyOverride(InOverrideSetName, InOpName, InPropertyPath);
+	if (!PropertyOverride)
+	{
+		return false; // no override found
+	}
+
+	return PropertyOverride->GetBoundVariableName() != NAME_None || PropertyOverride->GetBoundCurveName() != NAME_None;
+}
+
+FName UIKRetargeterController::AddNewVariable(const FName InVariableName) const
+{
+	FRetargetVariableContainer& Container = GetAsset()->GetVariables();
+	
+	auto GetSafeNewVariableName = [Container, InVariableName]() -> FName
+	{
+		// sanitize
+		FName SanitizedBase = FInstancedPropertyBag::SanitizePropertyName(InVariableName);
+		if (SanitizedBase == NAME_None)
+		{
+			SanitizedBase = TEXT("NewVar");
+		}
+
+		// make unique
+		FName NewName = SanitizedBase;
+		int32 Suffix = 0;
+		while (Container.Bag.FindPropertyDescByName(NewName) != nullptr)
+		{
+			Suffix++;
+			NewName = FName(*FString::Printf(TEXT("%s_%d"), *SanitizedBase.ToString(), Suffix));
+		}
+
+		return NewName;
+	};
+
+	FScopedTransaction Transaction(LOCTEXT("AddVariable", "Add Retarget Variable"));
+	Asset->Modify();
+
+	// get a unique and sanitized name
+	FName NewName = GetSafeNewVariableName();
+	// add as float by default
+	Container.Bag.AddProperty(NewName, EPropertyBagPropertyType::Float);
+
+	Asset->OverridesVersion++;
+	VariablesModified.Broadcast();
+	return NewName;
+}
+
+bool UIKRetargeterController::DeleteVariable(const FName InVariableName) const
+{
+	if (!GetVariableExists(InVariableName))
+	{
+		return false;
+	}
+	
+	FScopedTransaction Transaction(LOCTEXT("DeleteVar", "Delete Variable"));
+	Asset->Modify();
+
+	EPropertyBagAlterationResult Result = Asset->GetVariables().Bag.RemovePropertyByName(InVariableName);
+
+	// remove any bindings
+	IterateAllOverrides([InVariableName](FRetargetOpPropertyOverride& PropertyOverride)
+	{
+		if (PropertyOverride.GetBoundVariableName() == InVariableName)
+		{
+			PropertyOverride.SetBoundVariableName(NAME_None);
+		}
+	});
+
+	Asset->OverridesVersion++;
+	VariablesModified.Broadcast();
+	return Result == EPropertyBagAlterationResult::Success;
+}
+
+bool UIKRetargeterController::RenameVariable(const FName InVariableName, const FName InNewName) const
+{
+	if (!GetVariableExists(InVariableName))
+	{
+		return false;
+	}
+	
+	FScopedTransaction Transaction(LOCTEXT("RenameVar", "Rename Variable"));
+	Asset->Modify();
+	
+	EPropertyBagAlterationResult Result = Asset->GetVariables().Bag.RenameProperty(InVariableName, InNewName);
+
+	// update any bindings using the old name
+	IterateAllOverrides([InVariableName, InNewName](FRetargetOpPropertyOverride& PropertyOverride)
+	{
+		if (PropertyOverride.GetBoundVariableName() == InVariableName)
+		{
+			PropertyOverride.SetBoundVariableName(InNewName);
+		}
+	});
+
+	Asset->OverridesVersion++;
+	VariablesModified.Broadcast();
+	return Result == EPropertyBagAlterationResult::Success;
+}
+
+bool UIKRetargeterController::SetVariableType(const FName InVariableName, const FEdGraphPinType& InNewPinType) const
+{
+	if (!GetVariableExists(InVariableName))
+	{
+		return false;
+	}
+	
+	const FProperty* VariableProperty = GetVariableProperty(InVariableName);
+	if (!ensure(VariableProperty))
+	{
+		return false; // should not get this far with an unknown variable
+	}
+
+	FScopedTransaction Transaction(LOCTEXT("ChangeType", "Change Variable Type"));
+	Asset->Modify();
+
+	// modify the type for this variable
+	FRetargetVariableContainer& Container = Asset->GetVariables();
+	const FPropertyBagPropertyDesc* PropertyDesc = Container.Bag.FindPropertyDescByName(InVariableName);
+	FPropertyBagPropertyDesc NewDesc = *PropertyDesc;
+	UE::StructUtils::SetPropertyDescFromPin(NewDesc, InNewPinType);
+	Container.Bag.AddProperties({ NewDesc });
+	
+	// changing the variable type potentially makes it's bindings incompatible
+	const FProperty* NewVariableProperty = GetVariableProperty(InVariableName);
+	for (TTuple<FName, FRetargetOverrideSet>& OverrideSetPair : Asset->OverrideSets)
+	{
+		for (FRetargetOpOverrides& OpOverrides : OverrideSetPair.Value.OpOverrides)
+		{
+			for (FRetargetOpPropertyOverride& PropertyOverride : OpOverrides.PropertyOverrides)
+			{
+				if (PropertyOverride.GetBoundVariableName() != InVariableName)
+				{
+					continue;
+				}
+				
+				FProperty* OverrideProperty = GetPropertyForOverride(OpOverrides.OpName, PropertyOverride.GetPropertyPath());
+				if (!IsPropertyCompatibleForBinding(NewVariableProperty, OverrideProperty))
+				{
+					// variable type no longer compatible with override property type
+					PropertyOverride.SetBoundVariableName(NAME_None);
+				}
+			}
+		}
+	}
+
+	Asset->OverridesVersion++;
+	VariablesModified.Broadcast();
+	return true;
+}
+
+bool UIKRetargeterController::GetVariableExists(const FName InVariableName) const
+{
+	const FRetargetVariableContainer& Variables = Asset->GetVariables();
+	const FPropertyBagPropertyDesc* PropertyDesc = Variables.Bag.FindPropertyDescByName(InVariableName);
+	return PropertyDesc != nullptr;
+}
+
+void UIKRetargeterController::IterateAllOverrides(const TFunctionRef<void(FRetargetOpPropertyOverride&)>& InVisitor) const
+{
+	for (TTuple<FName, FRetargetOverrideSet>& OverrideSetPair : Asset->OverrideSets)
+	{
+		for (FRetargetOpOverrides& OpOverrides : OverrideSetPair.Value.OpOverrides)
+		{
+			for (FRetargetOpPropertyOverride& PropertyOverride : OpOverrides.PropertyOverrides)
+			{
+				InVisitor(PropertyOverride);
+			}
+		}
+	}
+}
+
 void UIKRetargeterController::PostEditOverrideSets() const
 {
 	// informs any attached runtime processors to rebuild their override cache
diff --git a/Engine/Plugins/Animation/IKRig/Source/IKRigEditor/Private/RetargetEditor/SIKRetargetOverrideManager.cpp b/Engine/Plugins/Animation/IKRig/Source/IKRigEditor/Private/RetargetEditor/SIKRetargetOverrideManager.cpp
--- a/Engine/Plugins/Animation/IKRig/Source/IKRigEditor/Private/RetargetEditor/SIKRetargetOverrideManager.cpp
+++ b/Engine/Plugins/Animation/IKRig/Source/IKRigEditor/Private/RetargetEditor/SIKRetargetOverrideManager.cpp
@@ -578,7 +578,7 @@ void SIKRetargetOverrideManager::OnSelectionChanged(
 	StructToView.Owner = AssetController->GetAsset();
 	StructToView.UniqueName = SetToViewName;
 	
-	UIKRigStructViewer* StructViewer = AssetController->GetStructViewer();
+	UIKRigStructViewer* StructViewer = AssetController->GetStructViewer(ERetargetStructViewerMode::OverrideSets);
 	StructViewer->SetStructToView(StructToView);
 	
 	constexpr bool bForceRefresh = true;
diff --git a/Engine/Plugins/Animation/IKRig/Source/IKRigEditor/Private/RetargetEditor/SIKRetargetVariablesEditor.cpp b/Engine/Plugins/Animation/IKRig/Source/IKRigEditor/Private/RetargetEditor/SIKRetargetVariablesEditor.cpp
new file mode 100644
--- /dev/null
+++ b/Engine/Plugins/Animation/IKRig/Source/IKRigEditor/Private/RetargetEditor/SIKRetargetVariablesEditor.cpp
@@ -0,0 +1,241 @@
+// Copyright Epic Games, Inc. All Rights Reserved.
+#include "RetargetEditor/SIKRetargetVariablesEditor.h"
+
+#include "RetargetEditor/IKRetargetEditorController.h"
+#include "Modules/ModuleManager.h"
+#include "PropertyEditorModule.h"
+#include "ISinglePropertyView.h"
+#include "PropertyBagDetails.h"
+#include "SPinTypeSelector.h"
+#include "RigEditor/IKRigStructViewer.h"
+#include "SPositiveActionButton.h"
+
+#define LOCTEXT_NAMESPACE "RetargetVariablesEditor"
+
+
+void SIKRetargetVariablesEditor::Construct(const FArguments& InArgs, TSharedRef<FIKRetargetEditorController> InEditorController)
+{
+    EditorController = InEditorController;
+	EditorController.Pin()->SetVariablesView(SharedThis(this));
+	
+    TObjectPtr<UIKRetargeterController> LocalController = EditorController.Pin()->AssetController;
+
+    ChildSlot
+    [
+        SNew(SVerticalBox)
+        + SVerticalBox::Slot()
+        .AutoHeight()
+        .Padding(5.0f)
+        [
+            SNew(SPositiveActionButton)
+            .Icon(FAppStyle::Get().GetBrush("Icons.Plus"))
+            .Text(LOCTEXT("AddVariableLabel", "Add Variable"))
+            .OnClicked_Lambda([LocalController]()
+            {
+            	if (LocalController)
+            	{
+            		LocalController->AddNewVariable(NAME_None);
+            	}
+            	
+                return FReply::Handled();
+            })
+        ]
+
+        + SVerticalBox::Slot()
+        .FillHeight(1.0f)
+        [
+            SAssignNew(VariablesListView, SListView<TSharedPtr<FVariableRow>>)
+            .ListItemsSource(&VariableRows)
+            .OnGenerateRow(this, &SIKRetargetVariablesEditor::OnGenerateVariableRow)
+        ]
+    ];
+
+    RefreshList();
+}
+
+void SIKRetargetVariablesEditor::NotifyPreChange(FProperty* PropertyAboutToChange)
+{
+	if (UIKRetargeter* Asset = GetAsset())
+	{
+		// variables are free-floating properties so we have to inject the asset into the transaction
+		Asset->Modify();
+	}
+}
+
+void SIKRetargetVariablesEditor::NotifyPostChange(const FPropertyChangedEvent& PropertyChangedEvent, FProperty* PropertyThatChanged)
+{
+	if (PropertyChangedEvent.ChangeType == EPropertyChangeType::ValueSet)
+	{
+		// only refresh when the value is finally set (not during interactive changes)
+		RefreshList();
+	}
+}
+
+TSharedRef<ITableRow> SIKRetargetVariablesEditor::OnGenerateVariableRow(
+	TSharedPtr<FVariableRow> InItem,
+	const TSharedRef<STableViewBase>& OwnerTable)
+{
+	return SNew(STableRow<TSharedPtr<FVariableRow>>, OwnerTable)
+	[
+		SNew(SHorizontalBox)
+        
+		// rename
+		+ SHorizontalBox::Slot()
+		.FillWidth(0.3f)
+		.VAlign(VAlign_Center)
+		[
+			SNew(SEditableText)
+			.Text(FText::FromName(InItem->PropertyName))
+			.OnTextCommitted_Lambda([this, InItem](const FText& NewText, ETextCommit::Type CommitType)
+			{
+				if (CommitType == ETextCommit::OnEnter)
+				{
+					if (EditorController.IsValid())
+					{
+						EditorController.Pin()->AssetController->RenameVariable(InItem->PropertyName, FName(*NewText.ToString()));
+					}
+				}
+			})
+		]
+
+		// type selector
+		+ SHorizontalBox::Slot()
+		.AutoWidth()
+		[
+			MakeTypeSelectorForVariable(InItem)
+		]
+
+		// value widget (uses single property editor for different types)
+		+ SHorizontalBox::Slot()
+		.FillWidth(0.4f)
+		[
+			MakeValueWidgetForVariable(InItem)
+		]
+
+		// delete
+		+ SHorizontalBox::Slot()
+		.AutoWidth()
+		[
+			SNew(SButton)
+			.OnClicked_Lambda([this, InItem]()
+			{
+				if (EditorController.IsValid())
+				{
+					EditorController.Pin()->AssetController->DeleteVariable(InItem->PropertyName);
+				}
+				
+				return FReply::Handled();
+			})
+			[
+				SNew(SImage).Image(FAppStyle::Get().GetBrush("Icons.Delete"))
+			]
+		]
+	];
+}
+
+UIKRetargeter* SIKRetargetVariablesEditor::GetAsset() const
+{
+	if (EditorController.IsValid())
+	{
+		return EditorController.Pin()->AssetController->GetAsset();	
+	}
+	
+	return nullptr;
+}
+
+void SIKRetargetVariablesEditor::RefreshList()
+{
+    VariableRows.Empty();
+    UIKRetargeter* Asset = GetAsset();
+    if (!Asset)
+    {
+	    return;
+    }
+
+    const FRetargetVariableContainer& Container = Asset->GetVariables();
+    if (const UPropertyBag* BagStruct = Container.Bag.GetPropertyBagStruct())
+    {
+        for (const FPropertyBagPropertyDesc& Desc : BagStruct->GetPropertyDescs())
+        {
+            TSharedPtr<FVariableRow> NewRow = MakeShared<FVariableRow>();
+            NewRow->PropertyName = Desc.Name;
+            VariableRows.Add(NewRow);
+        }
+    }
+
+    if (VariablesListView.IsValid())
+    {
+        VariablesListView->RequestListRefresh();
+    }
+}
+
+TSharedRef<SWidget> SIKRetargetVariablesEditor::MakeTypeSelectorForVariable(TSharedPtr<FVariableRow> InRow) const
+{
+	return SNew(SPinTypeSelector, FGetPinTypeTree::CreateStatic([](TArray<TSharedPtr<UEdGraphSchema_K2::FPinTypeTreeInfo>>& TypeTree, ETypeTreeFilter Filter)
+		{
+			if (const UEdGraphSchema_K2* K2Schema = GetDefault<UEdGraphSchema_K2>())
+			{
+				// this fills the dropdown with the standard Blueprint types
+				K2Schema->GetVariableTypeTree(TypeTree, Filter);
+			}
+		}))
+		.TargetPinType_Lambda([this, InRow]()
+		{
+			UIKRetargeter* Asset = GetAsset();
+			if (!Asset)
+			{
+				return FEdGraphPinType();
+			}
+			const FPropertyBagPropertyDesc* Desc = Asset->GetVariables().Bag.FindPropertyDescByName(InRow->PropertyName);
+			return Desc ? UE::StructUtils::GetPropertyDescAsPin(*Desc) : FEdGraphPinType();
+		})
+		.OnPinTypeChanged_Lambda([this, InRow](const FEdGraphPinType& NewPinType)
+		{
+			if (EditorController.IsValid())
+			{
+				EditorController.Pin()->AssetController->SetVariableType(InRow->PropertyName, NewPinType);
+			}
+		})
+		.Schema(GetDefault<UEdGraphSchema_K2>())
+		.bAllowArrays(false) // for some reason I haven't figured out, this doesn't prevent array types :(
+		.SelectorType(SPinTypeSelector::ESelectorType::Partial);
+}
+
+TSharedRef<SWidget> SIKRetargetVariablesEditor::MakeValueWidgetForVariable(TSharedPtr<FVariableRow> InRow)
+{
+	UIKRetargeter* Asset = GetAsset();
+	if (!Asset)
+	{
+		return SNullWidget::NullWidget;
+	}
+
+	FRetargetVariableContainer& VarContainer = Asset->GetVariables();
+	const UPropertyBag* BagStruct = VarContainer.Bag.GetPropertyBagStruct();
+    
+	if (!BagStruct)
+	{
+		return SNullWidget::NullWidget;
+	}
+
+	// NOTE: we're using the property editor module to generate a widget for the type stored in the property bag
+	// the property bag allows us to store generic variable types, this allows us to edit them
+
+	// the TStructureDataProvider is so the widget knows where to read/write memory
+	TSharedPtr<FStructOnScope> StructOnScope = MakeShared<FStructOnScope>(BagStruct,VarContainer.Bag.GetMutableValue().GetMemory());
+	TSharedPtr<IStructureDataProvider> StructProvider = MakeShared<FStructOnScopeStructureDataProvider>(StructOnScope);
+
+	// create the widget targeting the variable name directly in that struct
+	FSinglePropertyParams Params;
+	Params.NotifyHook = this; // routes property edit callbacks to this::NotifyPreChange() / NotifyPostChange()
+	Params.NamePlacement = EPropertyNamePlacement::Hidden;
+	FPropertyEditorModule& EditModule = FModuleManager::GetModuleChecked<FPropertyEditorModule>("PropertyEditor");
+	TSharedPtr<ISinglePropertyView> SinglePropView = EditModule.CreateSingleProperty(
+		StructProvider, 
+		InRow->PropertyName, 
+		Params
+	);
+
+	return SinglePropView.IsValid() ? SinglePropView.ToSharedRef() : SNullWidget::NullWidget;
+}
+
+#undef LOCTEXT_NAMESPACE
\ No newline at end of file
diff --git a/Engine/Plugins/Animation/IKRig/Source/IKRigEditor/Private/RigEditor/IKRigStructViewer.cpp b/Engine/Plugins/Animation/IKRig/Source/IKRigEditor/Private/RigEditor/IKRigStructViewer.cpp
--- a/Engine/Plugins/Animation/IKRig/Source/IKRigEditor/Private/RigEditor/IKRigStructViewer.cpp
+++ b/Engine/Plugins/Animation/IKRig/Source/IKRigEditor/Private/RigEditor/IKRigStructViewer.cpp
@@ -10,7 +10,7 @@
 
 #define LOCTEXT_NAMESPACE "IKRigStructViewer"
 
-void UIKRigStructViewer::TriggerReinitIfNeeded(const FPropertyChangedEvent& InEvent)
+void UIKRigStructViewer::TriggerReinitializationIfNeeded(const FPropertyChangedEvent& InEvent)
 {
 	// never reinitialize while interacting with a slider
 	if (InEvent.ChangeType == EPropertyChangeType::Interactive || !InEvent.Property)
@@ -21,14 +21,14 @@ void UIKRigStructViewer::TriggerReinitIfNeeded(const FPropertyChangedEvent& InEv
 	// always reinitialize if property is marked for it
 	if (InEvent.Property->HasMetaData(IKRigReinitOnEditMetaLabel))
 	{
-		OnStructNeedsReinit().Broadcast(StructToView.UniqueName, InEvent);
+		OnNeedsReinitialized().Broadcast(StructToView.UniqueName, InEvent);
 		return;
 	}
 
 	// if it is a name property, we want to reinitialize to validate it
 	if (CastField<FNameProperty>(InEvent.Property) != nullptr)
 	{
-		OnStructNeedsReinit().Broadcast(StructToView.UniqueName, InEvent);
+		OnNeedsReinitialized().Broadcast(StructToView.UniqueName, InEvent);
 		return;
 	}
 
@@ -37,7 +37,7 @@ void UIKRigStructViewer::TriggerReinitIfNeeded(const FPropertyChangedEvent& InEv
 	// cannot be marked with the meta-tag (IKRig plugin doesn't own it and it's only relevant here)
 	if (InEvent.Property->GetFName() == FName(TEXT("BoneName")))
 	{
-		OnStructNeedsReinit().Broadcast(StructToView.UniqueName, InEvent);
+		OnNeedsReinitialized().Broadcast(StructToView.UniqueName, InEvent);
 		return;
 	}
 }
@@ -45,33 +45,29 @@ void UIKRigStructViewer::TriggerReinitIfNeeded(const FPropertyChangedEvent& InEv
 void UIKRigStructViewer::SetupPropertyEditingCallbacks(const TSharedPtr<IPropertyHandle>& InProperty)
 {
 	UObject* ObjectToTransact = GetStructOwner();
-	
+
 	InProperty->SetOnPropertyValuePreChange(FSimpleDelegate::CreateLambda([ObjectToTransact]()
-		{
-			//UE_LOG(LogTemp, Log, TEXT("PreChange: Starting transaction for property edit"));
-			GEditor->BeginTransaction(LOCTEXT("IKRigStructViewer", "Edited property."));
-			ObjectToTransact->Modify();
-		}));
+	{
+		GEditor->BeginTransaction(LOCTEXT("IKRigStructViewer", "Edited property."));
+		ObjectToTransact->Modify();
+	}));
 		
 	InProperty->SetOnPropertyValueChangedWithData(TDelegate<void(const FPropertyChangedEvent&)>::CreateLambda([this](const FPropertyChangedEvent& InEvent)
 	{
-		//UE_LOG(LogTemp, Log, TEXT("PostChange: Ending transaction for property edit"));
 		GEditor->EndTransaction();
-		TriggerReinitIfNeeded(InEvent);
+		TriggerReinitializationIfNeeded(InEvent);
 	}));
 
 	InProperty->SetOnChildPropertyValuePreChange(FSimpleDelegate::CreateLambda([ObjectToTransact]()
 	{
-		//UE_LOG(LogTemp, Log, TEXT("ChildPreChange: Starting transaction for child property edit"));
 		GEditor->BeginTransaction(LOCTEXT("IKRigStructViewerChild", "Edited child property."));
 		ObjectToTransact->Modify();
 	}));
 
 	InProperty->SetOnChildPropertyValueChangedWithData(TDelegate<void(const FPropertyChangedEvent&)>::CreateLambda([this](const FPropertyChangedEvent& InEvent)
 	{
-		//UE_LOG(LogTemp, Log, TEXT("ChildPostChange: Ending transaction for child property edit"));
 		GEditor->EndTransaction();
-		TriggerReinitIfNeeded(InEvent);
+		TriggerReinitializationIfNeeded(InEvent);
 	}));
 }
 
@@ -207,7 +203,7 @@ void UIKRigStructWrapperBase::InitializeWithRetargeter(
 	UIKRetargeterController* InRetargeterController)
 {
 	Initialize(InStructToWrap, InWrapperPropertyName);
-	OnStructNeedsReinit().AddUObject(InRetargeterController, &UIKRetargeterController::OnOpPropertyChanged);
+	OnNeedsReinitialized().AddUObject(InRetargeterController, &UIKRetargeterController::OnOpPropertyChanged);
 }
 
 bool UIKRigStructWrapperBase::IsValid() const
diff --git a/Engine/Plugins/Animation/IKRig/Source/IKRigEditor/Public/RetargetEditor/IKRetargetCommands.h b/Engine/Plugins/Animation/IKRig/Source/IKRigEditor/Public/RetargetEditor/IKRetargetCommands.h
--- a/Engine/Plugins/Animation/IKRig/Source/IKRigEditor/Public/RetargetEditor/IKRetargetCommands.h
+++ b/Engine/Plugins/Animation/IKRig/Source/IKRigEditor/Public/RetargetEditor/IKRetargetCommands.h
@@ -29,6 +29,8 @@ public:
 	TSharedPtr< FUICommandInfo > ShowAssetSettings;
 	// open the override set manager
 	TSharedPtr< FUICommandInfo > ShowOverrideSets;
+	// open the variables editor
+	TSharedPtr< FUICommandInfo > ShowVariables;
 
 	// enable Root retarget pass
 	TSharedPtr< FUICommandInfo > EnableRoot;
diff --git a/Engine/Plugins/Animation/IKRig/Source/IKRigEditor/Public/RetargetEditor/IKRetargetEditorController.h b/Engine/Plugins/Animation/IKRig/Source/IKRigEditor/Public/RetargetEditor/IKRetargetEditorController.h
--- a/Engine/Plugins/Animation/IKRig/Source/IKRigEditor/Public/RetargetEditor/IKRetargetEditorController.h
+++ b/Engine/Plugins/Animation/IKRig/Source/IKRigEditor/Public/RetargetEditor/IKRetargetEditorController.h
@@ -14,6 +14,7 @@
 
 #include "Misc/TransactionObjectEvent.h"
 
+class SIKRetargetVariablesEditor;
 class SIKRetargetOverrideManager;
 class SRetargetOpStack;
 enum class ERetargetSourceOrTarget : uint8;
@@ -112,7 +113,8 @@ public:
 	// delegate handles to un-bind
 	FDelegateHandle RetargeterReInitDelegateHandle;
 	FDelegateHandle OpStackModifiedDelegateHandle;
-	FDelegateHandle ProfilesModifiedDelegateHandle;
+	FDelegateHandle OverrideSetsModifiedDelegateHandle;
+	FDelegateHandle VariablesModifiedDelegateHandle;
 	// callback when IK Rig asset has been swapped out
 	void HandleIKRigReplaced(ERetargetSourceOrTarget SourceOrTarget);
 	FDelegateHandle IKRigReplacedDelegateHandle;
@@ -156,6 +158,7 @@ public:
 	void SetHierarchyView(const TSharedPtr<SIKRetargetHierarchy>& InHierarchyView) { HierarchyView = InHierarchyView; };
 	void SetOpStackView(const TSharedPtr<SRetargetOpStack>& InOpStackView) { OpStackView = InOpStackView; };
 	void SetOverrideSetsView(const TSharedPtr<SIKRetargetOverrideManager>& InOverrideSetView) { OverrideSetView = InOverrideSetView; };
+	void SetVariablesView(const TSharedPtr<SIKRetargetVariablesEditor>& InVariablesView) { VariablesView = InVariablesView; };
 	bool IsObjectInDetailsView(const UObject* Object);
 	
 	// force refresh all views in the editor
@@ -165,6 +168,7 @@ public:
 	void RefreshHierarchyView() const;
 	void RefreshOpStackView() const;
 	void RefreshOverrideSetView() const;
+	void RefreshVariablesView() const;
 	void RefreshPoseList() const;
 	void SetDetailsObject(UObject* DetailsObject) const;
 	void SetDetailsObjects(const TArray<UObject*>& DetailsObjects) const;
@@ -347,6 +351,8 @@ private:
 	TSharedPtr<SRetargetOpStack> OpStackView;
 	// overrides widget
 	TSharedPtr<SIKRetargetOverrideManager> OverrideSetView;
+	// variables widget
+	TSharedPtr<SIKRetargetVariablesEditor> VariablesView;
 
 	// when prompting user to assign an IK Rig
 	TSharedPtr<SWindow> IKRigPickerWindow;
diff --git a/Engine/Plugins/Animation/IKRig/Source/IKRigEditor/Public/RetargetEditor/IKRetargetOverrideCustomization.h b/Engine/Plugins/Animation/IKRig/Source/IKRigEditor/Public/RetargetEditor/IKRetargetOverrideCustomization.h
--- a/Engine/Plugins/Animation/IKRig/Source/IKRigEditor/Public/RetargetEditor/IKRetargetOverrideCustomization.h
+++ b/Engine/Plugins/Animation/IKRig/Source/IKRigEditor/Public/RetargetEditor/IKRetargetOverrideCustomization.h
@@ -2,6 +2,7 @@
 #pragma once
 
 #include "IPropertyTypeCustomization.h"
+#include "Widgets/Input/SMenuAnchor.h"
 
 class UIKRetargeterController;
 class IDetailGroup;
@@ -47,6 +48,13 @@ private:
 		const FName& InOpName,
 		const FString& InPropertyPath);
 
+	/** Generates the binding dropdown menu content for a specific property override */
+	TSharedRef<SWidget> GetBindingMenuContent(
+		TSharedPtr<IPropertyHandle> InPropertyHandle, 
+		FName InOverrideSetName, 
+		FName InOpName, 
+		FString InPropertyPath);
+
 	/**
 	* Iterates through the overrides in an override set and groups them by their property path.
 	* @param OpOverridesHandle  Handle to the FRetargetOpOverrides struct.
diff --git a/Engine/Plugins/Animation/IKRig/Source/IKRigEditor/Public/RetargetEditor/IKRetargetVariablesTabSummoner.h b/Engine/Plugins/Animation/IKRig/Source/IKRigEditor/Public/RetargetEditor/IKRetargetVariablesTabSummoner.h
new file mode 100644
--- /dev/null
+++ b/Engine/Plugins/Animation/IKRig/Source/IKRigEditor/Public/RetargetEditor/IKRetargetVariablesTabSummoner.h
@@ -0,0 +1,23 @@
+﻿// Copyright Epic Games, Inc. All Rights Reserved.
+
+#pragma once
+
+#include "CoreMinimal.h"
+#include "WorkflowOrientedApp/WorkflowTabFactory.h"
+
+class FIKRetargetEditor;
+
+/** Summons the "Variables" tab in the IK Retarget editor. */
+struct FIKRetargetVariablesTabSummoner : public FWorkflowTabFactory
+{
+public:
+	static const FName TabID;
+	
+	FIKRetargetVariablesTabSummoner(const TSharedRef<FIKRetargetEditor>& InIKRetargetEditor);
+	virtual TSharedRef<SWidget> CreateTabBody(const FWorkflowTabSpawnInfo& Info) const override;
+	virtual TSharedPtr<SToolTip> CreateTabToolTipWidget(const FWorkflowTabSpawnInfo& Info) const override;
+
+protected:
+	TWeakPtr<FIKRetargetEditor> IKRetargetEditor;
+};
+
diff --git a/Engine/Plugins/Animation/IKRig/Source/IKRigEditor/Public/RetargetEditor/IKRetargeterController.h b/Engine/Plugins/Animation/IKRig/Source/IKRigEditor/Public/RetargetEditor/IKRetargeterController.h
--- a/Engine/Plugins/Animation/IKRig/Source/IKRigEditor/Public/RetargetEditor/IKRetargeterController.h
+++ b/Engine/Plugins/Animation/IKRig/Source/IKRigEditor/Public/RetargetEditor/IKRetargeterController.h
@@ -25,6 +25,16 @@ class UIKRetargeter;
 class USkeletalMesh;
 class UIKRigStructViewer;
 
+struct FEdGraphPinType;
+
+// used to index into the map of cached struct viewers in the editor
+UENUM()
+enum class ERetargetStructViewerMode : uint8
+{
+	Settings,
+	OverrideSets
+};
+
 // A stateless singleton (1-per-asset) class used to make modifications to a UIKRetargeter asset.
 // Use UIKRetargeter.GetController() to get the controller for the asset you want to modify.  
 UCLASS(MinimalAPI, BlueprintType, hidecategories = UObject)
@@ -502,6 +512,12 @@ public:
 	/** Get access to an override set by name. */
 	UE_API FRetargetOverrideSet* FindOverrideSet(const FName InOverrideSetName) const;
 
+	/** Helper to get a single property override */
+	FRetargetOpPropertyOverride* FindPropertyOverride(
+		const FName InOverrideSetName,
+		const FName InOpName,
+		const FString& InPropertyPath) const;
+
 	/** Get a list of all override sets that are children of the given set (recursive). */
 	UE_API const TArray<FName> GetAllChildrenOverrideSets(const FName InOverrideSetName) const;
 
@@ -518,8 +534,53 @@ public:
 	UE_API bool RemoveOpOverrides(const FName InOverrideSetName, const FName InOpName) const;
 
 	/** Update the value of a stored property override inside the given override set, for the given op and property path. */
-	UE_API bool UpdateOverrideValue(const FName InOverrideSetName, const FName InOpName, const FString& InPropertyPath, const FStructOnScope& InNewValue) const;
-
+	UE_API bool UpdateOverrideValue(
+		const FName InOverrideSetName,
+		const FName InOpName,
+		const FString& InPropertyPath,
+		const FStructOnScope& InNewValue,
+		const bool bShouldTransact=true) const;
+
+	/** Bind a property override to a curve. */
+	UE_API bool BindPropertyOverrideToCurve(
+		const FName& InCurveName,
+		const FName InOverrideSetName,
+		const FName InOpName,
+		const FString& InPropertyPath) const;
+
+	/** * Bind the given variable to the given property override.
+	* @param InVariableName The name of the variable to bind. Use NAME_None to unbind.
+	* @param InOverrideSetName The name of the override set containing the op.
+	* @param InOpName The name of the specific retarget op containing the property.
+	* @param InPropertyPath The leaf property path within the op settings.
+	* @return true if the binding was successful, false otherwise
+	*/
+	UE_API bool BindPropertyOverrideToVariable(
+		const FName InVariableName,
+		const FName InOverrideSetName,
+		const FName InOpName,
+		const FString& InPropertyPath) const;
+
+	/** * Clear all bindings on the given property override.
+	* @param InOverrideSetName The name of the override set containing the op.
+	* @param InOpName The name of the specific retarget op containing the property.
+	* @param InPropertyPath The leaf property path within the op settings.
+	* @return true if the binding was found and cleared, false otherwise
+	*/
+	UE_API bool ClearPropertyOverrideBinding(
+		const FName InOverrideSetName,
+		const FName InOpName,
+		const FString& InPropertyPath) const;
+
+	/** Check if this property override is compatible with curve binding (is a numeric type).
+	* @param InOpName The name of the specific retarget op containing the property.
+	* @param InPropertyPath The leaf property path within the op settings.
+	* @return true if the override can be bound to a curve, false otherwise
+	*/
+	UE_API bool CanPropertyOverrideBindToCurve(
+		const FName InOpName,
+		const FString& InPropertyPath) const;
+	
 	/** Add a property override for a specific op within a specific override set. */
 	UE_API bool AddPropertyOverrideToOp(const FName InOverrideSetName, const FName InOpName, const FString& InPropertyPath) const;
 
@@ -548,7 +609,101 @@ public:
 
 	/** Returns the name of the given override set. */
 	UE_API FName GetOverrideSetName(const FRetargetOverrideSet* InOverrideSet) const;
+
+	/** Helper to get an FProperty* for a given property override */
+	FProperty* GetPropertyForOverride(const FName InOpName, const FString& InPropertyPath) const;
+
+	//
+	// VARIABLES C++ ONLY API
+	//
 	
+	/** Get read-only access to the set of variables for this retarget asset. */
+	UE_API const FRetargetVariableContainer& GetVariables() const;
+
+	/** * Get the FProperty of a variable by name. 
+	 * @param InVariableName The name of the variable to find in the property bag.
+	 * @return The FProperty if found, otherwise nullptr.
+	 */
+	UE_API const FProperty* GetVariableProperty(const FName InVariableName) const;
+
+	/** * Determine if the source property (override) can be bound to the target (variable). 
+	 * This checks for type compatibility, including allowed casts like float to double.
+	 * @param Source The property override being bound. 
+	 * @param Target The property of the variable being bound to.
+	 * @return True if the properties are compatible for a data binding.
+	 */
+	UE_API static bool IsPropertyCompatibleForBinding(const FProperty* Source, const FProperty* Target);
+
+	/** * Get the name of the variable bound to the given property override. 
+	 * @param InOverrideSetName The name of the override set.
+	 * @param InOpName The name of the operation.
+	 * @param InPropertyPath The specific property path to check.
+	 * @return The name of the bound variable, or NAME_None if no binding exists.
+	 */
+	UE_API FName GetVariableBoundToPropertyOverride(
+	    const FName InOverrideSetName,
+	    const FName InOpName,
+	    const FString& InPropertyPath) const;
+
+	/** * Get the name of the curve bound to the given property override. 
+	 * @param InOverrideSetName The name of the override set.
+	 * @param InOpName The name of the operation.
+	 * @param InPropertyPath The specific property path to check.
+	 * @return The name of the bound curve, or NAME_None if no binding exists.
+	 */
+	UE_API FName GetCurveBoundToPropertyOverride(
+		const FName InOverrideSetName,
+		const FName InOpName,
+		const FString& InPropertyPath) const;
+
+	/** Get if the property override is bound to anything (curve or variable).
+	 * @param InOverrideSetName The name of the override set.
+	 * @param InOpName The name of the operation.
+	 * @param InPropertyPath The specific property path to check.
+	 * @return true if the override is bound to anything, false otherwise
+	 */
+	UE_API bool GetPropertyOverrideHasBinding(
+		const FName InOverrideSetName,
+		const FName InOpName,
+		const FString& InPropertyPath) const;
+
+	/** * Adds a new variable to the retargeter's variable container (defaults to float type)
+	 * @param InVariableName The desired name for the new variable (if Name_None, a default name is created)
+	 * @return The actual name assigned to the variable (may be suffixed if a name collision occurred).
+	 */
+	UE_API FName AddNewVariable(const FName InVariableName) const;
+
+	/** * Deletes a variable from the retargeter and clears any associated property bindings.
+	 * @param InVariableName The name of the variable to remove.
+	 * @return True if the variable was found and successfully removed.
+	 */
+	UE_API bool DeleteVariable(const FName InVariableName) const;
+
+	/** * Renames an existing variable and updates all internal property bindings to reflect the new name.
+	 * @param InVariableName The current name of the variable.
+	 * @param InNewName The new name to assign to the variable.
+	 * @return True if the rename was successful.
+	 */
+	UE_API bool RenameVariable(const FName InVariableName, const FName InNewName) const;
+
+	/** * Changes the underlying data type of a variable. This may clear existing bindings if types become incompatible.
+	 * @param InVariableName The name of the variable to modify.
+	 * @param NewPinType The new type definition (float, bool, etc.) via a Graph Pin Type.
+	 * @return True if the type was successfully changed.
+	 */
+	UE_API bool SetVariableType(const FName InVariableName, const FEdGraphPinType& NewPinType) const;
+
+	/** * Check if a variable with the given name exists in this asset.
+	 * @param InVariableName The name to check.
+	 * @return True if the variable exists in the property bag.
+	 */
+	UE_API bool GetVariableExists(const FName InVariableName) const;
+
+	/** * Helper to iterate every PropertyOverride in the asset.
+	 * @param Visitor A lambda or function to execute for every override found in the asset.
+	 */
+	void IterateAllOverrides(const TFunctionRef<void(FRetargetOpPropertyOverride&)>& Visitor) const;
+
 private:
 
 	/** Controller functions call this anytime the override sets are modified */
@@ -611,12 +766,11 @@ public:
 	// END DEPRECATED PUBLIC API
 	//
 	
+private:
+
 	//
 	// DELEGATE CALLBACKS FOR OUTSIDE SYSTEMS
 	//
-	
-private:
-	
 	DECLARE_MULTICAST_DELEGATE(FOnRetargeterNeedsInitialized);
 	FOnRetargeterNeedsInitialized RetargeterNeedsInitialized;
 
@@ -625,6 +779,9 @@ private:
 
 	DECLARE_MULTICAST_DELEGATE(FOnOverrideSetsModified);
 	FOnOverrideSetsModified OverrideSetsModified;
+
+	DECLARE_MULTICAST_DELEGATE(FOnVariablesModified);
+	FOnVariablesModified VariablesModified;
 	
 	DECLARE_MULTICAST_DELEGATE_OneParam(FOnIKRigReplaced, ERetargetSourceOrTarget);
 	FOnIKRigReplaced IKRigReplaced;
@@ -640,10 +797,12 @@ private:
 
 	// prevent reinitializing from inner operations
 	mutable int32 ReinitializeScopeCounter = 0;
-	
+
+private:
 	// a dummy UObject to wrap UStructs for the details panel
-	UPROPERTY(transient)
-	TObjectPtr<UIKRigStructViewer> StructViewer = nullptr;
+	// maps each mode to its corresponding UIKRigStructViewer wrapper
+	UPROPERTY(Transient)
+	TMap<ERetargetStructViewerMode, TObjectPtr<UIKRigStructViewer>> StructViewers;
 
 	// transient cache used for duration of editing
 	mutable TMap<const USkeletalMesh*, TSharedPtr<FIKRigSkeleton>> CachedSkeletons;
@@ -665,8 +824,12 @@ public:
 	// Attach a delegate to be notified whenever the override sets are modified.
 	FOnOverrideSetsModified& OnOverrideSetsModified(){ return OverrideSetsModified; };
 
+	// Attach a delegate to be notified whenever the variables are modified.
+	FOnVariablesModified& OnVariablesModified(){ return VariablesModified; };
+
 	// get the singleton struct viewer for editing structs belonging to the asset this controller controls
-	UE_API UIKRigStructViewer* GetStructViewer();
+	// NOTE: there are 2 details views in the retarget editor, the InMode specifies which viewer you want to reset and reuse
+	UE_API UIKRigStructViewer* GetStructViewer(ERetargetStructViewerMode InMode);
 	
 	friend class UIKRetargeter;
 	friend struct FScopedReinitializeIKRetargeter;
diff --git a/Engine/Plugins/Animation/IKRig/Source/IKRigEditor/Public/RetargetEditor/SIKRetargetVariablesEditor.h b/Engine/Plugins/Animation/IKRig/Source/IKRigEditor/Public/RetargetEditor/SIKRetargetVariablesEditor.h
new file mode 100644
--- /dev/null
+++ b/Engine/Plugins/Animation/IKRig/Source/IKRigEditor/Public/RetargetEditor/SIKRetargetVariablesEditor.h
@@ -0,0 +1,45 @@
+﻿// Copyright Epic Games, Inc. All Rights Reserved.
+#pragma once
+
+#include "CoreMinimal.h"
+#include "Retargeter/IKRetargeter.h"
+#include "Widgets/SCompoundWidget.h"
+#include "StructUtils/PropertyBag.h"
+#include "Templates/SharedPointer.h"
+#include "Misc/NotifyHook.h"
+#include "Widgets/Views/SListView.h"
+
+class FIKRetargetEditorController;
+class IStructureDetailsView;
+
+struct FVariableRow
+{
+	FName PropertyName;
+};
+
+class SIKRetargetVariablesEditor : public SCompoundWidget, public FNotifyHook
+{
+public:
+	SLATE_BEGIN_ARGS(SIKRetargetVariablesEditor) {}
+	SLATE_END_ARGS()
+	
+	void Construct(const FArguments& InArgs, TSharedRef<FIKRetargetEditorController> InEditorController);
+
+	// FNotifyHook Interface
+	virtual void NotifyPreChange(FProperty* PropertyAboutToChange) override;
+	virtual void NotifyPostChange(const FPropertyChangedEvent& PropertyChangedEvent, FProperty* PropertyThatChanged) override;
+	// END FNotifyHook Interface
+
+	void RefreshList();
+	
+private:
+	
+	UIKRetargeter* GetAsset() const;
+	TSharedRef<ITableRow> OnGenerateVariableRow(TSharedPtr<FVariableRow> InItem, const TSharedRef<STableViewBase>& OwnerTable);
+	TSharedRef<SWidget> MakeTypeSelectorForVariable(TSharedPtr<FVariableRow> InRow) const;
+	TSharedRef<SWidget> MakeValueWidgetForVariable(TSharedPtr<FVariableRow> InRow);
+	
+	TArray<TSharedPtr<FVariableRow>> VariableRows;
+	TSharedPtr<SListView<TSharedPtr<FVariableRow>>> VariablesListView;
+	TWeakPtr<FIKRetargetEditorController> EditorController;
+};
\ No newline at end of file
diff --git a/Engine/Plugins/Animation/IKRig/Source/IKRigEditor/Public/RigEditor/IKRigStructViewer.h b/Engine/Plugins/Animation/IKRig/Source/IKRigEditor/Public/RigEditor/IKRigStructViewer.h
--- a/Engine/Plugins/Animation/IKRig/Source/IKRigEditor/Public/RigEditor/IKRigStructViewer.h
+++ b/Engine/Plugins/Animation/IKRig/Source/IKRigEditor/Public/RigEditor/IKRigStructViewer.h
@@ -107,7 +107,7 @@ public:
 		return StructToView.Owner.Get();
 	};
 
-	void TriggerReinitIfNeeded(const FPropertyChangedEvent& InEvent);
+	void TriggerReinitializationIfNeeded(const FPropertyChangedEvent& InEvent);
 
 	void SetupPropertyEditingCallbacks(const TSharedPtr<IPropertyHandle>& InProperty);
 
@@ -117,8 +117,15 @@ public:
 
 protected:
 
-	DECLARE_MULTICAST_DELEGATE_TwoParams(FOnStructPropertyEdited, const FName& /*UniqueStructName*/, const FPropertyChangedEvent& /*Modified Property*/);
-	FOnStructPropertyEdited StructPropertyEditedDelegate;
+	DECLARE_MULTICAST_DELEGATE_TwoParams(FOnNeedsReinitialized, const FName& /*UniqueStructName*/, const FPropertyChangedEvent& /*Modified Property*/);
+	FOnNeedsReinitialized NeedsReinitializedDelegate;
+
+	DECLARE_MULTICAST_DELEGATE_TwoParams(FOnPreEditProperty, const FName& /*UniqueStructName*/, const FProperty* /*PropertyAboutToChange*/);
+	FOnPreEditProperty PreEditPropertyDelegate;
+
+	DECLARE_MULTICAST_DELEGATE_TwoParams(FOnPostEditProperty, const FName& /*UniqueStructName*/, const FPropertyChangedEvent& /*Modified Property*/);
+	FOnPostEditProperty PostEditPropertyDelegate;
+	
 	// a wrapper of the struct sent to the details panel
 	TSharedPtr<FStructOnScope> StructOnScope;
 	// the data needed to display and edit an instance of a struct in memory
@@ -126,8 +133,15 @@ protected:
 
 public:
 	
-	// Attach a delegate to be notified whenever a property (or child property) marked " in the currently displayed UStruct is edited
-	FOnStructPropertyEdited& OnStructNeedsReinit(){ return StructPropertyEditedDelegate; };
+	// attach a delegate to be notified whenever a property is edited that requires the system to be reinitialized
+	// NOTE: This happens in a few hard-coded special cases or for any property marked with the meta tag: ReinitializeOnEdit
+	FOnNeedsReinitialized& OnNeedsReinitialized(){ return NeedsReinitializedDelegate; };
+	
+	// attach a delegate to be notified BEFORE a property is edited
+	FOnPreEditProperty& OnPreEditProperty(){ return PreEditPropertyDelegate; };
+	
+	// attach a delegate to be notified AFTER a property is edited
+	FOnPostEditProperty& OnPostEditProperty(){ return PostEditPropertyDelegate; };	
 };
 
 class FIKRigStructViewerCustomization : public IDetailCustomization
diff --git a/Engine/Plugins/Animation/IKRig/Source/IKRigUAF/Private/RigUnit_UAFIKRetargeter.cpp b/Engine/Plugins/Animation/IKRig/Source/IKRigUAF/Private/RigUnit_UAFIKRetargeter.cpp
--- a/Engine/Plugins/Animation/IKRig/Source/IKRigUAF/Private/RigUnit_UAFIKRetargeter.cpp
+++ b/Engine/Plugins/Animation/IKRig/Source/IKRigUAF/Private/RigUnit_UAFIKRetargeter.cpp
@@ -45,7 +45,6 @@ FRigUnit_UAFIKRetargeter_Execute()
 		Params.TargetSkeletalMesh = TargetMesh;
 		Params.RetargeterAsset = IKRetargetAsset;
 		Params.CustomProfile = &WorkData.Profile;
-		Params.OverrideSetsToApply = OverrideSetNames;
 		Params.bSuppressWarnings = false;
 		Processor.Initialize(Params);
 	}
@@ -74,6 +73,8 @@ FRigUnit_UAFIKRetargeter_Execute()
 	FRetargetRunParameters Params;
 	Params.SourceGlobalPose = &WorkData.SourceMeshPoseGlobal;
 	Params.Profile = &WorkData.Profile;
+	Params.OverrideSetsToApply = OverrideSetNames;
+	Params.Variables = nullptr; // TODO, add support for dynamic variables input pins in UAF
 	Params.DeltaTime = ExecuteContext.GetDeltaTime();
 	Params.LOD = TargetRefPose.GetSourceLODLevel();
 	const TArray<FTransform>& TargetGlobalPose = Processor.RunRetargeter(Params);
