[IK Retarget] Fixes crash when opening override set.

#Synced-CL 52174688

diff --git a/Engine/Plugins/Animation/IKRig/Source/IKRig/Public/Retargeter/IKRetargetOverrides.h b/Engine/Plugins/Animation/IKRig/Source/IKRig/Public/Retargeter/IKRetargetOverrides.h
--- a/Engine/Plugins/Animation/IKRig/Source/IKRig/Public/Retargeter/IKRetargetOverrides.h
+++ b/Engine/Plugins/Animation/IKRig/Source/IKRig/Public/Retargeter/IKRetargetOverrides.h
@@ -135,13 +135,13 @@ private:
 	
 	/** property path using "->" separator like, "IKGoalSettings->BlendAlpha" and array syntax like, "Chains[0]->Weight"
 	 * NOTE: this allows generic serialization of property overrides */
-	UPROPERTY()
+	UPROPERTY(EditAnywhere, Category=Override)
 	FString PropertyPath;
 	UPROPERTY()
 	uint32 PropertyPathHash = 0;
 
 	/** Value serialized as string */
-	UPROPERTY()
+	UPROPERTY(EditAnywhere, Category=Override)
 	FString ValueAsString;
 
 	/** The name of the variable this override is bound to, if any. */
