[IK Retarget] Fixes crash deleting op with an open override set.

#Synced-CL 52174688

diff --git a/Engine/Plugins/Animation/IKRig/Source/IKRig/Private/Retargeter/IKRetargetOverrideManager.cpp b/Engine/Plugins/Animation/IKRig/Source/IKRig/Private/Retargeter/IKRetargetOverrideManager.cpp
--- a/Engine/Plugins/Animation/IKRig/Source/IKRig/Private/Retargeter/IKRetargetOverrideManager.cpp
+++ b/Engine/Plugins/Animation/IKRig/Source/IKRig/Private/Retargeter/IKRetargetOverrideManager.cpp
@@ -183,9 +183,9 @@ void FIKRetargetRuntimeOverrideManager::RebuildCache(FIKRetargetProcessor& Proce
 		for (const FRetargetOpOverrides& OpOverrides : OverrideSet->OpOverrides)
 		{
 			FIKRetargetOpBase* OpInProcessor = Processor.GetRetargetOpByName(OpOverrides.OpName);
-			if (!ensure(OpInProcessor))
+			if (!OpInProcessor)
 			{
-				continue;
+				continue; // can happen on first frame after op deleted
 			}
 			
 			FIKRetargetOpSettingsBase* OpSettings = OpInProcessor->GetSettings();
diff --git a/Engine/Plugins/Animation/IKRig/Source/IKRig/Public/Retargeter/IKRetargeter.h b/Engine/Plugins/Animation/IKRig/Source/IKRig/Public/Retargeter/IKRetargeter.h
--- a/Engine/Plugins/Animation/IKRig/Source/IKRig/Public/Retargeter/IKRetargeter.h
+++ b/Engine/Plugins/Animation/IKRig/Source/IKRig/Public/Retargeter/IKRetargeter.h
@@ -143,6 +143,8 @@ public:
 
 	// Get read-only access to all the override sets
 	const TMap<FName, FRetargetOverrideSet>& GetOverrideSets() const { return OverrideSets; };
+	// Get read-write access to all override sets (caution, all modifications should go through the controller)
+	TMap<FName, FRetargetOverrideSet>& GetOverrideSetsWriteable() { return OverrideSets; };
 	// Get all the override sets that are enabled
 	UE_API const TArray<FName>& GetOverrideSetsToApply() const;
 	// Get a set of overrides by name (may be null) 
diff --git a/Engine/Plugins/Animation/IKRig/Source/IKRigEditor/Private/RetargetEditor/IKRetargeterController.cpp b/Engine/Plugins/Animation/IKRig/Source/IKRigEditor/Private/RetargetEditor/IKRetargeterController.cpp
--- a/Engine/Plugins/Animation/IKRig/Source/IKRigEditor/Private/RetargetEditor/IKRetargeterController.cpp
+++ b/Engine/Plugins/Animation/IKRig/Source/IKRigEditor/Private/RetargetEditor/IKRetargeterController.cpp
@@ -206,11 +206,12 @@ int32 UIKRetargeterController::AddRetargetOp(const FString InIKRetargetOpType) c
 	return AddRetargetOp(OpType);
 }
 
-bool UIKRetargeterController::RemoveRetargetOp(const int32 OpIndex) const
+bool UIKRetargeterController::RemoveRetargetOp(const int32 InOpIndex) const
 {
-	if (!Asset->RetargetOps.IsValidIndex(OpIndex))
+	FIKRetargetOpBase* Op = GetRetargetOpByIndex(InOpIndex);
+	if (!Op)
 	{
-		UE_LOG(LogIKRigEditor, Warning, TEXT("Retarget Op not removed. Invalid index, %d."), OpIndex);
+		UE_LOG(LogIKRigEditor, Warning, TEXT("Retarget Op not removed. Invalid index, %d."), InOpIndex);
 		return false;
 	}
 
@@ -218,14 +219,36 @@ bool UIKRetargeterController::RemoveRetargetOp(const int32 OpIndex) const
 	FScopedReinitializeIKRetargeter Reinitialize(this, ERetargetRefreshMode::ProcessorAndOpStack);
 	Asset->Modify();
 
-	TArray<int32> IndicesToRemove = GetChildOpIndices(OpIndex);
-	IndicesToRemove.Add(OpIndex);
+	// get full list of ops to remove (including children, if any)
+	TArray<int32> IndicesToRemove = GetChildOpIndices(InOpIndex);
+	IndicesToRemove.Add(InOpIndex);
+	TSet<FName> NamesToRemove;
+	for (int32 OpIndex : IndicesToRemove)
+	{
+		const FName OpNameToRemove = GetRetargetOpByIndex(OpIndex)->GetName();
+		NamesToRemove.Add(OpNameToRemove);
+	}
+
+	// remove any overrides associated with this op
+	TMap<FName, FRetargetOverrideSet>& OverrideSets = Asset->GetOverrideSetsWriteable();
+	for (TTuple<FName, FRetargetOverrideSet>& OverrideSetPair : OverrideSets)
+	{
+		FRetargetOverrideSet& OverrideSet = OverrideSetPair.Value;
+		OverrideSet.OpOverrides.RemoveAll([&NamesToRemove](const FRetargetOpOverrides& Element)
+		{
+			return NamesToRemove.Contains(Element.OpName);
+		});
+	}
+
+	// remove the op itself
+	IndicesToRemove.Add(InOpIndex);
 	Algo::Reverse(IndicesToRemove); // high to low (children are always before parent)
 	for (const int32 OpToRemove : IndicesToRemove)
 	{
 		Asset->RetargetOps.RemoveAt(OpToRemove);
 	}
-	
+
+	OverrideSetsModified.Broadcast();
 	return true;
 }
 
