[IK Retarget] Added ability to enable/disable ops from profiles

#Synced-CL 52174688

diff --git a/Engine/Plugins/Animation/IKRig/Source/IKRig/Private/Retargeter/IKRetargetOps.cpp b/Engine/Plugins/Animation/IKRig/Source/IKRig/Private/Retargeter/IKRetargetOps.cpp
--- a/Engine/Plugins/Animation/IKRig/Source/IKRig/Private/Retargeter/IKRetargetOps.cpp
+++ b/Engine/Plugins/Animation/IKRig/Source/IKRig/Private/Retargeter/IKRetargetOps.cpp
@@ -76,7 +76,7 @@ FText FIKRetargetOpBase::GetWarningMessage() const
 {
 	if (bIsInitialized)
 	{
-		return bIsEnabled ? LOCTEXT("OpReadyAndOn", "Running.") : LOCTEXT("OpReadyAndOff", "Ready, but disabled.");
+		return IsEnabled() ? LOCTEXT("OpReadyAndOn", "Running.") : LOCTEXT("OpReadyAndOff", "Ready, but disabled.");
 	}
 	
 	return LOCTEXT("OpNotReady", "Not initialized. See output log.");
diff --git a/Engine/Plugins/Animation/IKRig/Source/IKRig/Private/Retargeter/IKRetargeter.cpp b/Engine/Plugins/Animation/IKRig/Source/IKRig/Private/Retargeter/IKRetargeter.cpp
--- a/Engine/Plugins/Animation/IKRig/Source/IKRig/Private/Retargeter/IKRetargeter.cpp
+++ b/Engine/Plugins/Animation/IKRig/Source/IKRig/Private/Retargeter/IKRetargeter.cpp
@@ -1050,6 +1050,12 @@ void UIKRetargeter::PostLoadOpStack()
 	{
 		FIKRetargetOpBase& Op = OpStruct.GetMutable<FIKRetargetOpBase>();
 		Op.PostLoad(CustomVersion);
+
+		// in 5.8 we moved the enabled state from the op itself into the op settings
+		if (CustomVersion < FIKRigObjectVersion::EnabledStateMovedFromOpToOpSettings)
+		{
+			Op.GetSettings()->bEnabled = Op.GetDeprecatedEnabledState();
+		}
 	}
 };
 
diff --git a/Engine/Plugins/Animation/IKRig/Source/IKRig/Public/IKRigObjectVersion.h b/Engine/Plugins/Animation/IKRig/Source/IKRig/Public/IKRigObjectVersion.h
--- a/Engine/Plugins/Animation/IKRig/Source/IKRig/Public/IKRigObjectVersion.h
+++ b/Engine/Plugins/Animation/IKRig/Source/IKRig/Public/IKRigObjectVersion.h
@@ -56,6 +56,9 @@ struct FIKRigObjectVersion
 		// IK Rig retarget definition now stores separate root bone and pelvis bone
 		AddedRootBoneToRetargetDefinition,
 
+		// the op settings now contain bEnabled instead of the ops themselves
+		EnabledStateMovedFromOpToOpSettings,
+
 		// -----<new versions can be added above this line>-------------------------------------------------
 		VersionPlusOne,
 		LatestVersion = VersionPlusOne - 1
diff --git a/Engine/Plugins/Animation/IKRig/Source/IKRig/Public/Retargeter/IKRetargetOps.h b/Engine/Plugins/Animation/IKRig/Source/IKRig/Public/Retargeter/IKRetargetOps.h
--- a/Engine/Plugins/Animation/IKRig/Source/IKRig/Public/Retargeter/IKRetargetOps.h
+++ b/Engine/Plugins/Animation/IKRig/Source/IKRig/Public/Retargeter/IKRetargetOps.h
@@ -73,6 +73,10 @@ struct FIKRetargetOpSettingsBase : public FIKRetargetOpSkeletonProvider
 	
 	virtual ~FIKRetargetOpSettingsBase() = default;
 
+	/** Enable/disable the entire operation. Disabled ops are skipped entirely and do not run at all.*/
+	UPROPERTY()
+	bool bEnabled = true;
+
 	/** The maximum LOD that this Op is allowed to run at.
 	 * For example if you have LODThreshold of 2, the Op will run until LOD 2 (based on 0 index). When the component LOD becomes 3, it will stop running.
 	 * A value of -1 forces the Op to execute at all LOD levels. Default is -1. */
@@ -228,10 +232,10 @@ public:
 	virtual void OnAssignIKRig(const ERetargetSourceOrTarget SourceOrTarget, const UIKRigDefinition* InIKRig, const FIKRetargetOpBase* InParentOp) {};
 	
 	/** get if this op is enabled */
-	bool IsEnabled() const { return bIsEnabled; };
+	bool IsEnabled() const { return GetSettingsConst()->bEnabled; };
 
 	/** turn this operation on/off (will be skipped during execution if disabled) */
-	void SetEnabled(const bool bEnabled){ bIsEnabled = bEnabled; };
+	void SetEnabled(const bool bEnabled){ GetSettings()->bEnabled = bEnabled; };
 
 	/** return true if this op is initialized and ready to run */
 	bool IsInitialized() const { return bIsInitialized; };
@@ -425,12 +429,10 @@ public:
 		TArray<FLinearColor> Colors;
 	};
 
+	bool GetDeprecatedEnabledState() const {return bIsEnabled_DEPRECATED; }
+
 private:
 	
-	/** when false, execution of this op is skipped */
-	UPROPERTY()
-	bool bIsEnabled = true;
-
 	/** the text label given to the op, used to refer to it from script */
 	UPROPERTY()
 	FName Name;
@@ -439,6 +441,10 @@ private:
 	UPROPERTY()
 	FName ParentOpName;
 
+	/** DEPRECATED enabled flag, moved to base op settings */
+	UPROPERTY(meta = (DeprecatedProperty, DeprecationMessage = "bIsEnabled is deprecated. Please use FIKRetargetOpSettingsBase::bEnabled instead."))
+	bool bIsEnabled_DEPRECATED = true;
+
 protected:
 	
 	/** set to true when the op is ready to run */
