[IK Retarget] Cleaning up and unifying op settings.

#Synced-CL 52174688

diff --git a/Engine/Plugins/Animation/IKRig/Source/IKRig/Private/Retargeter/RetargetOps/FilterBoneOp.cpp b/Engine/Plugins/Animation/IKRig/Source/IKRig/Private/Retargeter/RetargetOps/FilterBoneOp.cpp
--- a/Engine/Plugins/Animation/IKRig/Source/IKRig/Private/Retargeter/RetargetOps/FilterBoneOp.cpp
+++ b/Engine/Plugins/Animation/IKRig/Source/IKRig/Private/Retargeter/RetargetOps/FilterBoneOp.cpp
@@ -74,8 +74,9 @@ void FIKRetargetFilterBoneOp::Run(
 		}
 
 		FTransform TargetGlobalTransform = OutTargetGlobalPose[FilterBoneData.TargetBone.BoneIndex];
-		FQuat FilteredRotation = FilterBoneData.Filter.Update(TargetGlobalTransform.GetRotation(), InDeltaTime, Settings.FilterSettings);
-		TargetGlobalTransform.SetRotation(FilteredRotation);
+		const FQuat FilteredRotation = FilterBoneData.Filter.Update(TargetGlobalTransform.GetRotation(), InDeltaTime, Settings.FilterSettings);
+		const FQuat FinalRotation = FQuat::FastLerp(TargetGlobalTransform.GetRotation(), FilteredRotation, Settings.Alpha).GetNormalized();
+		TargetGlobalTransform.SetRotation(FinalRotation);
 
 		// assign result and update children
 		TargetSkeleton.SetGlobalTransformAndUpdateChildren(FilterBoneData.TargetBone.BoneIndex, TargetGlobalTransform,OutTargetGlobalPose);
diff --git a/Engine/Plugins/Animation/IKRig/Source/IKRig/Private/Retargeter/RetargetOps/FloorConstraintOp.cpp b/Engine/Plugins/Animation/IKRig/Source/IKRig/Private/Retargeter/RetargetOps/FloorConstraintOp.cpp
--- a/Engine/Plugins/Animation/IKRig/Source/IKRig/Private/Retargeter/RetargetOps/FloorConstraintOp.cpp
+++ b/Engine/Plugins/Animation/IKRig/Source/IKRig/Private/Retargeter/RetargetOps/FloorConstraintOp.cpp
@@ -226,6 +226,7 @@ void FToeConstraint::RunAfterParent(
 	const FTargetSkeleton& InTargetSkeleton,
 	TArray<FTransform>& OutGlobalPose,
 	const FFloorConstraintChainSettings* InChainSettings,
+	FIKRetargetFloorConstraintOpSettings& InOpSettings,
 	const double InDeltaTime)
 {
 	if (!bToeReady)
@@ -248,7 +249,8 @@ void FToeConstraint::RunAfterParent(
 	ThetaRadians = AngleFilter.Update(ThetaRadians, InDeltaTime, OneEuroFilterSettings);
 
 	// alpha the whole effect
-	ThetaRadians = FMath::Lerp(0, ThetaRadians, AllToeSettings->Alpha * InChainSettings->Alpha);
+	const float Alpha = AllToeSettings->Alpha * InChainSettings->Alpha * InOpSettings.Alpha;
+	ThetaRadians = FMath::Lerp(0, ThetaRadians, Alpha);
 
 	// convert to a quaternion that rotates around the pitch axis by theta
 	const FQuat DeltaRotation = FQuat(PitchAxis, ThetaRadians);
@@ -424,7 +426,8 @@ double FToeConstraint::CalcToeCorrectionAngle()
 void FFootConstraint::Run(
 	FIKRigGoal& InOutGoal,
 	const TArray<FTransform>& OutTargetGlobalPose,
-	const FFloorConstraintChainSettings& InSettings)
+	const FFloorConstraintChainSettings& InChainSettings,
+	const FIKRetargetFloorConstraintOpSettings& InOpSettings)
 {
 	if (!bReadyToRun)
 	{
@@ -435,10 +438,11 @@ void FFootConstraint::Run(
 	UpdateBasePoints(OutTargetGlobalPose, &InOutGoal);
 	const int32 LowestBasePointIndex = GetIndexOfLowestPoint();
 	const double BaseHeight = Points[LowestBasePointIndex].Z;
-	const double BaseDelta = (BaseHeight < 0.0f ? FMath::Abs(BaseHeight) : 0.0) * InSettings.Alpha;
+	const double BaseDelta = BaseHeight < 0.0f ? FMath::Abs(BaseHeight) : 0.0;
 	
 	// move the goal up
-	InOutGoal.Position.Z += BaseDelta;
+	const double Alpha = InChainSettings.Alpha * InOpSettings.Alpha;
+	InOutGoal.Position.Z += BaseDelta * Alpha;
 }
 
 void FFloorConstraint::Initialize(
@@ -498,7 +502,7 @@ void FFloorConstraint::Run(
 		// calculate the target height of the goal by blending towards the source height when close to floor
 		const double FalloffRange = FMath::Max(1.0,FalloffEndHeight - FalloffStartHeight);
 		const double Falloff = SourceHeight < FalloffStartHeight ? 1.0f : 1.0f - (SourceHeight - FalloffStartHeight) / FalloffRange;
-		const double Weight = Falloff * Settings->Alpha;
+		const double Weight = Falloff * Settings->Alpha * InOpSettings.Alpha;
 		const double HeightOffset = FMath::Lerp(0.0f, HeightOffsetInRefPose, Settings->MaintainHeightOffset);
 		const double TargetHeight = SourceHeight + HeightOffset;
 		
@@ -507,7 +511,7 @@ void FFloorConstraint::Run(
 	}
 
 	// run the foot/floor constraint based on new goal location
-	FootConstraint.Run(Goal, OutTargetGlobalPose, *Settings);
+	FootConstraint.Run(Goal, OutTargetGlobalPose, *Settings, InOpSettings);
 }
 
 bool FIKRetargetFloorConstraintOp::Initialize(
@@ -628,7 +632,7 @@ void FIKRetargetFloorConstraintOp::RunAfterParent(
 		
 		for (FToeConstraint& Toe : FloorConstraint.ToeConstraints)
 		{
-			Toe.RunAfterParent(TargetSkeleton, OutTargetGlobalPose, FloorConstraint.Settings, InDeltaTime);
+			Toe.RunAfterParent(TargetSkeleton, OutTargetGlobalPose, FloorConstraint.Settings, Settings, InDeltaTime);
 		}
 	}
 }
diff --git a/Engine/Plugins/Animation/IKRig/Source/IKRig/Private/Retargeter/RetargetOps/SpeedPlantingOp.cpp b/Engine/Plugins/Animation/IKRig/Source/IKRig/Private/Retargeter/RetargetOps/SpeedPlantingOp.cpp
--- a/Engine/Plugins/Animation/IKRig/Source/IKRig/Private/Retargeter/RetargetOps/SpeedPlantingOp.cpp
+++ b/Engine/Plugins/Animation/IKRig/Source/IKRig/Private/Retargeter/RetargetOps/SpeedPlantingOp.cpp
@@ -18,6 +18,11 @@
 
 #define LOCTEXT_NAMESPACE "SpeedPlantingOp"
 
+bool FRetargetSpeedPlantingSettings::operator==(const FRetargetSpeedPlantingSettings& Other) const
+{
+	return SpeedCurveName == Other.SpeedCurveName && FMath::IsNearlyEqualByULP(Alpha,Other.Alpha);
+}
+
 const UClass* FIKRetargetSpeedPlantingOpSettings::GetControllerType() const
 {
 	return UIKRetargetSpeedPlantingController::StaticClass();
@@ -150,7 +155,8 @@ void FSpeedPlantingGoalState::Update(
 		}
 	}
 
-	IKRigGoal->Position = CurrentPosition;
+	const float Alpha = Settings->Alpha * InSettings.Alpha;
+	IKRigGoal->Position = FMath::Lerp(IKRigGoal->Position, CurrentPosition, Alpha);
 }
 
 void FSpeedPlantingGoalState::Reset()
diff --git a/Engine/Plugins/Animation/IKRig/Source/IKRig/Private/Retargeter/RetargetOps/StrideWarpingOp.cpp b/Engine/Plugins/Animation/IKRig/Source/IKRig/Private/Retargeter/RetargetOps/StrideWarpingOp.cpp
--- a/Engine/Plugins/Animation/IKRig/Source/IKRig/Private/Retargeter/RetargetOps/StrideWarpingOp.cpp
+++ b/Engine/Plugins/Animation/IKRig/Source/IKRig/Private/Retargeter/RetargetOps/StrideWarpingOp.cpp
@@ -168,7 +168,7 @@ bool FIKRetargetStrideWarpingOp::Initialize(
 		}
 		
 		// store goal to warp
-		GoalsToWarp.Emplace(Goal->Name, TargetGlobalRetargetPose[GoalBoneIndex]);
+		GoalsToWarp.Emplace(Goal->Name, TargetGlobalRetargetPose[GoalBoneIndex], ChainSettings);
 	}
 
 	bIsInitialized = !GoalsToWarp.IsEmpty();
@@ -334,24 +334,30 @@ void FIKRetargetStrideWarpingOp::Run(
 
 		// get initial goal position
 		const FVector& InitialPosition = GoalToWarp.GlobalRefPoseOfGoalBone.GetLocation();
+		// accumulate a warped position to apply to the goal
+		FVector WarpedPosition = Goal->Position;
 		
 		// forward warping
 		const FVector InitialGoalInOrigSpace = InitialBodyTransform.InverseTransformPosition(InitialPosition);
 		const FVector IntialGoalInCurrentSpace = CurrentBodyTransform.TransformPosition(InitialGoalInOrigSpace);
 		const FPlane FwdPlane(IntialGoalInCurrentSpace, Fwd);
 		const FVector GoalProjOnFwdPlane = FwdPlane.PointPlaneProject(Goal->Position, FwdPlane);
-		Goal->Position = GoalProjOnFwdPlane + ((Goal->Position - GoalProjOnFwdPlane) * Settings.WarpForwards);
+		WarpedPosition = GoalProjOnFwdPlane + ((WarpedPosition - GoalProjOnFwdPlane) * Settings.WarpForwards);
 
 		// sideways offset
 		// first determine which side the goal is on originally
 		float GoalSideMultiplier = static_cast<float>(FVector::DotProduct(InitialPosition.GetSafeNormal(), SideOrig));
 		// push goal by offset in the newly calculated sideways direction
-		Goal->Position += Side * Settings.SidewaysOffset * GoalSideMultiplier;
+		WarpedPosition += Side * Settings.SidewaysOffset * GoalSideMultiplier;
 
 		// splay warping
 		FVector SplayOrigin = CurrentBodyPosition;
-		SplayOrigin.Z = Goal->Position.Z;
-		Goal->Position = SplayOrigin + (Goal->Position - SplayOrigin) * Settings.WarpSplay;
+		SplayOrigin.Z = WarpedPosition.Z;
+		WarpedPosition = SplayOrigin + (WarpedPosition - SplayOrigin) * Settings.WarpSplay;
+
+		// alpha blend the goal location towards the warped position
+		const float Alpha = GoalToWarp.Settings->Alpha * Settings.Alpha;
+		Goal->Position = FMath::Lerp(Goal->Position, WarpedPosition, Alpha);
 
 		// goals are additive by default, this one is now in component space
 		Goal->PositionSpace = EIKRigGoalSpace::Component;
diff --git a/Engine/Plugins/Animation/IKRig/Source/IKRig/Private/Retargeter/RetargetOps/WeaponGoalsOp.cpp b/Engine/Plugins/Animation/IKRig/Source/IKRig/Private/Retargeter/RetargetOps/WeaponGoalsOp.cpp
--- a/Engine/Plugins/Animation/IKRig/Source/IKRig/Private/Retargeter/RetargetOps/WeaponGoalsOp.cpp
+++ b/Engine/Plugins/Animation/IKRig/Source/IKRig/Private/Retargeter/RetargetOps/WeaponGoalsOp.cpp
@@ -140,6 +140,7 @@ void FIKRetargetWeaponGoalOp::Run(
 			(int32)ERetargetWeaponBone::LeftHandAttach,
 			(int32)ERetargetWeaponBone::RightHandAttach,
 			(int32)ERetargetWeaponBone::HandGun };
+		const FVector FinalScale = FMath::Lerp(FVector::OneVector, Scale, Settings.Alpha);
 		for (int32 Index : BonesToScale)
 		{
 			const int32 BoneIndex = BoneIndices[Index];
@@ -147,7 +148,7 @@ void FIKRetargetWeaponGoalOp::Run(
 			{
 				// scale the bone and propagate to children
 				FTransform BoneTransform = OutTargetGlobalPose[BoneIndex];
-				BoneTransform.SetScale3D(Scale);
+				BoneTransform.SetScale3D(FinalScale);
 				TargetSkeleton.SetGlobalTransformAndUpdateChildren(BoneIndex, BoneTransform, OutTargetGlobalPose);
 			}
 		}
@@ -164,7 +165,8 @@ void FIKRetargetWeaponGoalOp::Run(
 		{
 			const FTransform& LeftHand = OutTargetGlobalPose[LeftHandBoneIndex];
 			const FTransform& RightHand = OutTargetGlobalPose[RightHandBoneIndex];
-			const FVector BlendedPosition = FMath::Lerp(LeftHand.GetLocation(), RightHand.GetLocation(), Settings.LeftRightHandIKAlpha);
+			const float Alpha = Settings.LeftRightHandIKAlpha * Settings.Alpha;
+			const FVector BlendedPosition = FMath::Lerp(LeftHand.GetLocation(), RightHand.GetLocation(), Alpha);
 			FTransform HandGunTransform = OutTargetGlobalPose[HandGunBoneIndex];
 			HandGunTransform.SetTranslation(BlendedPosition);
 			TargetSkeleton.SetGlobalTransformAndUpdateChildren(HandGunBoneIndex, HandGunTransform, OutTargetGlobalPose);
@@ -182,11 +184,13 @@ void FIKRetargetWeaponGoalOp::Run(
 		};
 	if (FIKRigGoal* LeftHandGoal = GoalContainer.FindGoalByName(LeftHandGoalName))
 	{
-		BlendGoalToIKHandBone(LeftHandGoal, BoneIndices[(int32)ERetargetWeaponBone::IKHandLeft], Settings.LeftHandIKAlpha);
+		const float Alpha = Settings.LeftHandIKAlpha * Settings.Alpha;
+		BlendGoalToIKHandBone(LeftHandGoal, BoneIndices[(int32)ERetargetWeaponBone::IKHandLeft], Alpha);
 	}
 	if (FIKRigGoal* RightHandGoal = GoalContainer.FindGoalByName(RightHandGoalName))
 	{
-		BlendGoalToIKHandBone(RightHandGoal, BoneIndices[(int32)ERetargetWeaponBone::IKHandRight], Settings.RightHandIKAlpha);
+		const float Alpha = Settings.RightHandIKAlpha * Settings.Alpha;
+		BlendGoalToIKHandBone(RightHandGoal, BoneIndices[(int32)ERetargetWeaponBone::IKHandRight], Alpha);
 	}
 }
 
diff --git a/Engine/Plugins/Animation/IKRig/Source/IKRig/Public/Retargeter/RetargetOps/AlignPoleVectorOp.h b/Engine/Plugins/Animation/IKRig/Source/IKRig/Public/Retargeter/RetargetOps/AlignPoleVectorOp.h
--- a/Engine/Plugins/Animation/IKRig/Source/IKRig/Public/Retargeter/RetargetOps/AlignPoleVectorOp.h
+++ b/Engine/Plugins/Animation/IKRig/Source/IKRig/Public/Retargeter/RetargetOps/AlignPoleVectorOp.h
@@ -103,6 +103,10 @@ struct FIKRetargetAlignPoleVectorOpSettings : public FIKRetargetOpSettingsBase
 	UPROPERTY(EditAnywhere, BlueprintReadWrite, Category = "Chains")
 	TArray<FRetargetPoleVectorSettings> ChainsToAlign;
 
+	/** Range 0 to 1. Default 1. Blends the effect of the entire op. */
+	UPROPERTY(EditAnywhere, BlueprintReadWrite, Category = "Op Settings", meta = (ClampMin = "0.0", ClampMax = "1.0", UIMin = "0.0", UIMax = "1.0"))
+	double Alpha = 1.0f;
+
 	UE_API void MergePoleVectorSettings(const FRetargetPoleVectorSettings& InSettingsToMerge);
 
 	UE_API virtual const UClass* GetControllerType() const override;
diff --git a/Engine/Plugins/Animation/IKRig/Source/IKRig/Public/Retargeter/RetargetOps/FilterBoneOp.h b/Engine/Plugins/Animation/IKRig/Source/IKRig/Public/Retargeter/RetargetOps/FilterBoneOp.h
--- a/Engine/Plugins/Animation/IKRig/Source/IKRig/Public/Retargeter/RetargetOps/FilterBoneOp.h
+++ b/Engine/Plugins/Animation/IKRig/Source/IKRig/Public/Retargeter/RetargetOps/FilterBoneOp.h
@@ -37,12 +37,16 @@ struct FIKRetargetFilterBoneOpSettings : public FIKRetargetOpSettingsBase
 	UPROPERTY(EditAnywhere, BlueprintReadWrite, Category=Setup, meta=(ReinitializeOnEdit))
 	TArray<FFilterBoneData> BonesToFilter;
 
+	/** Range 0 to 1. Default 1. Blends the effect of the entire op. */
+	UPROPERTY(EditAnywhere, BlueprintReadWrite, Category = "Op Settings", meta = (ClampMin = "0.0", ClampMax = "1.0", UIMin = "0.0", UIMax = "1.0"))
+	double Alpha = 1.0f;
+
 	/** The settings to control the One-Euro filter behavior */
-	UPROPERTY(EditAnywhere, BlueprintReadWrite, Category="One Euro Filter")
+	UPROPERTY(EditAnywhere, BlueprintReadWrite, Category="Op Settings")
 	FOneEuroFilterSettings FilterSettings;
 
 	/* If true, filter is reset when playback loops. */
-	UPROPERTY(EditAnywhere, BlueprintReadWrite, Category="One Euro Filter")
+	UPROPERTY(EditAnywhere, BlueprintReadWrite, Category="Op Settings")
 	bool bResetPlayback = true;
 
 	UE_API virtual const UClass* GetControllerType() const override;
diff --git a/Engine/Plugins/Animation/IKRig/Source/IKRig/Public/Retargeter/RetargetOps/FloorConstraintOp.h b/Engine/Plugins/Animation/IKRig/Source/IKRig/Public/Retargeter/RetargetOps/FloorConstraintOp.h
--- a/Engine/Plugins/Animation/IKRig/Source/IKRig/Public/Retargeter/RetargetOps/FloorConstraintOp.h
+++ b/Engine/Plugins/Animation/IKRig/Source/IKRig/Public/Retargeter/RetargetOps/FloorConstraintOp.h
@@ -156,15 +156,19 @@ struct FIKRetargetFloorConstraintOpSettings : public FIKRetargetOpSettingsBase
 	UPROPERTY(EditAnywhere, BlueprintReadWrite, Category = Chains, meta=(ReinitializeOnEdit))
 	TArray<FFloorConstraintChainSettings> ChainsToAffect;
 
+	/** Range 0 to 1. Default is 1. Blend the effect of the entire op on/off. */
+	UPROPERTY(EditAnywhere, BlueprintReadWrite, Category = "Op Settings", meta = (ClampMin = "0.0", ClampMax = "1.0", UIMin = "0.0", UIMax = "1.0"))
+	double Alpha = 1.0f;
+
 	/** Range 0 to inf. Default is 8. The height in cm from the floor below which the goal is snapped directly to the source bone height.
      * NOTE: if the source bone height is greater than this value, but lower than FloorHeightFalloffEnd, then the height will smoothly blend from the source
      * bone height, to the height of the goal in its normal retargeted position. */
-    UPROPERTY(EditAnywhere, BlueprintReadWrite, Category = "Floor Constraint Op Settings", meta = (ClampMin = "0.1", UIMin = "0.1", UIMax = "100.0"))
+    UPROPERTY(EditAnywhere, BlueprintReadWrite, Category = "Op Settings", meta = (ClampMin = "0.1", UIMin = "0.1", UIMax = "100.0"))
     double HeightFalloffOffset = 8.0f;
     
     /** Range 0 to inf. Default is 20. The height in cm from the floor below which the goal is gradually blended towards the source bone height.
      * NOTE: if the source bone is higher than this value, the height of the goal is left at its normal retargeted height. */
-    UPROPERTY(EditAnywhere, BlueprintReadWrite, Category = "Floor Constraint Op Settings", meta = (ClampMin = "0.1", UIMin = "0.1", UIMax = "100.0"))
+    UPROPERTY(EditAnywhere, BlueprintReadWrite, Category = "Op Settings", meta = (ClampMin = "0.1", UIMin = "0.1", UIMax = "100.0"))
     double HeightFalloffDistance = 20.0f;
 
 	UE_API virtual const UClass* GetControllerType() const override;
@@ -206,6 +210,7 @@ struct FToeConstraint
 		const FTargetSkeleton& InTargetSkeleton,
 		TArray<FTransform>& OutGlobalPose,
 		const FFloorConstraintChainSettings* InChainSettings,
+		FIKRetargetFloorConstraintOpSettings& InOpSettings,
 		const double InDeltaTime);
 
 	void UpdateFramePoints(
@@ -243,7 +248,8 @@ struct FFootConstraint
 	void Run(
 		FIKRigGoal& InOutGoal,
 		const TArray<FTransform>& OutTargetGlobalPose,
-		const FFloorConstraintChainSettings& InSettings);
+		const FFloorConstraintChainSettings& InChainSettings,
+		const FIKRetargetFloorConstraintOpSettings& InOpSettings);
 	
 	void UpdateBasePoints(const TArray<FTransform>& InGlobalPose, const FIKRigGoal* InGoal);
 	
diff --git a/Engine/Plugins/Animation/IKRig/Source/IKRig/Public/Retargeter/RetargetOps/OffsetGoalsOp.h b/Engine/Plugins/Animation/IKRig/Source/IKRig/Public/Retargeter/RetargetOps/OffsetGoalsOp.h
--- a/Engine/Plugins/Animation/IKRig/Source/IKRig/Public/Retargeter/RetargetOps/OffsetGoalsOp.h
+++ b/Engine/Plugins/Animation/IKRig/Source/IKRig/Public/Retargeter/RetargetOps/OffsetGoalsOp.h
@@ -29,7 +29,7 @@ struct FIKRetargetOffsetGoalsChainSettings
 
 	/** Range 0 to 1. Default 0. Blend the offsets on this chain on/off. */
 	UPROPERTY(EditAnywhere, BlueprintReadWrite, Category = "Op Settings", meta = (ClampMin = "0.0", ClampMax = "1.0", UIMin = "0.0", UIMax = "1.0"))
-	double Alpha = 0.0f;
+	double Alpha = 1.0f;
 	
 	/** Default 0,0,0. Apply a global-space offset to the IK goal position. */
 	UPROPERTY(EditAnywhere, BlueprintReadWrite, Category="Chain Settings")
diff --git a/Engine/Plugins/Animation/IKRig/Source/IKRig/Public/Retargeter/RetargetOps/PelvisMotionOp.h b/Engine/Plugins/Animation/IKRig/Source/IKRig/Public/Retargeter/RetargetOps/PelvisMotionOp.h
--- a/Engine/Plugins/Animation/IKRig/Source/IKRig/Public/Retargeter/RetargetOps/PelvisMotionOp.h
+++ b/Engine/Plugins/Animation/IKRig/Source/IKRig/Public/Retargeter/RetargetOps/PelvisMotionOp.h
@@ -34,28 +34,28 @@ struct FIKRetargetPelvisMotionOpSettings : public FIKRetargetOpSettingsBase
 	 * When ON, the floor constraint will adjust the vertical Pelvis motion separately according to the following rules:
 	 * 1. When the source pelvis is LOWER than the ref pose, the target pelvis will be lowered proportional to their relative crotch heights.
 	 * 2. When the source pelvis is HIGHER than the ref pose, the target pelvis will be raised an identical amount while maintaining any vertical offset in the ref pose.*/
-	UPROPERTY(EditAnywhere, BlueprintReadWrite, Category=FloorConstraint, meta = (ClampMin = "0.0", ClampMax = "1.0", UIMin = "0.0", UIMax = "1.0"))
+	UPROPERTY(EditAnywhere, BlueprintReadWrite, Category="Floor Constraint", meta = (ClampMin = "0.0", ClampMax = "1.0", UIMin = "0.0", UIMax = "1.0"))
 	double FloorConstraintWeight = 0.0f;
 	
 	/** A negative vertical offset in cm relative to the Pelvis bone of the SOURCE.
 	 * NOTE: Adjust this until the green dot is located roughly at the crotch of the character. */
-	UPROPERTY(EditAnywhere, BlueprintReadWrite, Category = FloorConstraint)
+	UPROPERTY(EditAnywhere, BlueprintReadWrite, Category = "Floor Constraint")
 	double SourceCrotchOffset = 0.0f;
 
 	/** A negative vertical offset in cm relative to the Pelvis bone of the TARGET.
 	 * NOTE: Adjust this until the green dot is located roughly at the crotch of the character. */
-	UPROPERTY(EditAnywhere, BlueprintReadWrite, Category = FloorConstraint)
+	UPROPERTY(EditAnywhere, BlueprintReadWrite, Category = "Floor Constraint")
 	double TargetCrotchOffset = 0.0f;
 
 	/** When true, the effect of the floor constraint decays as the pelvis gets further from the floor.
 	 * NOTE: This can help avoid over dampening vertical motion on the pelvis when retargeting to characters that are shorter than the source. */
-	UPROPERTY(EditAnywhere, BlueprintReadWrite, Category = FloorConstraint)
+	UPROPERTY(EditAnywhere, BlueprintReadWrite, Category = "Floor Constraint")
 	bool bUseGroundFalloff = false;
 	
 	/** Controls the height at which the floor constraint will begin blending in as a percentage of the pelvis height in the retarget pose.
 	 * NOTE: Use this to prevent the floor constraint from dampening vertical pelvis motion at normal standing heights.
 	 * At 1, the constraint is active at normal standing height. At 0.1, the constraint will only activate when within 10% of the height off the floor. */
-	UPROPERTY(EditAnywhere, BlueprintReadWrite, Category = FloorConstraint, meta = (ClampMin = "0.1", ClampMax = "1.0", UIMin = "0.1", UIMax = "1.0"))
+	UPROPERTY(EditAnywhere, BlueprintReadWrite, Category = "Floor Constraint", meta = (ClampMin = "0.1", ClampMax = "1.0", UIMin = "0.1", UIMax = "1.0"))
 	double GroundFalloffHeightPercent = 0.9f;
 	
 	/** Range 0 to 1. Default 1. Blends the amount of retargeted pelvis rotation to apply.
@@ -117,11 +117,11 @@ struct FIKRetargetPelvisMotionOpSettings : public FIKRetargetOpSettingsBase
 	double AffectIKVertical = 0.0f;
 
 	/** Adjust size of the debug drawing */
-	UPROPERTY(EditAnywhere, BlueprintReadWrite, Category=Debug)
+	UPROPERTY(EditAnywhere, BlueprintReadWrite, Category="Debug")
 	double DebugDrawSize = 5.0f;
 	
 	/** Adjust thickness of the debug drawing */
-	UPROPERTY(EditAnywhere, BlueprintReadWrite, Category=Debug)
+	UPROPERTY(EditAnywhere, BlueprintReadWrite, Category="Debug")
 	double DebugDrawThickness = 5.0f;
 
 	IKRIG_API virtual const UClass* GetControllerType() const override;
diff --git a/Engine/Plugins/Animation/IKRig/Source/IKRig/Public/Retargeter/RetargetOps/ScaleGoalsOp.h b/Engine/Plugins/Animation/IKRig/Source/IKRig/Public/Retargeter/RetargetOps/ScaleGoalsOp.h
--- a/Engine/Plugins/Animation/IKRig/Source/IKRig/Public/Retargeter/RetargetOps/ScaleGoalsOp.h
+++ b/Engine/Plugins/Animation/IKRig/Source/IKRig/Public/Retargeter/RetargetOps/ScaleGoalsOp.h
@@ -52,7 +52,7 @@ struct FIKRetargetScaleGoalsOpSettings : public FIKRetargetOpSettingsBase
 	TArray<FIKRetargetScaleGoalsChainSettings> Chains;
 
 	/** Range 0 to 1. Default 1. Blend the effect on all chains on/off. */
-	UPROPERTY(EditAnywhere, BlueprintReadWrite, Category = "Alpha", meta = (ClampMin = "0.0", ClampMax = "1.0", UIMin = "0.0", UIMax = "1.0"))
+	UPROPERTY(EditAnywhere, BlueprintReadWrite, Category = "Op Settings", meta = (ClampMin = "0.0", ClampMax = "1.0", UIMin = "0.0", UIMax = "1.0"))
 	double Alpha = 1.0f;
 
 	virtual const UClass* GetControllerType() const override;
diff --git a/Engine/Plugins/Animation/IKRig/Source/IKRig/Public/Retargeter/RetargetOps/ScaleSourceOp.h b/Engine/Plugins/Animation/IKRig/Source/IKRig/Public/Retargeter/RetargetOps/ScaleSourceOp.h
--- a/Engine/Plugins/Animation/IKRig/Source/IKRig/Public/Retargeter/RetargetOps/ScaleSourceOp.h
+++ b/Engine/Plugins/Animation/IKRig/Source/IKRig/Public/Retargeter/RetargetOps/ScaleSourceOp.h
@@ -23,16 +23,16 @@ struct FIKRetargetScaleSourceOpSettings : public FIKRetargetOpSettingsBase
 	GENERATED_BODY()
 	
 	/** Range 0.01 to +inf. Default 1. Scales the incoming source pose. Affects entire skeleton and all IK goals.*/
-	UPROPERTY(EditAnywhere, BlueprintReadWrite, Category = Scale, meta = (ReinitializeOnEdit, ClampMin = "0.01", UIMin = "0.01", UIMax = "10.0"))
+	UPROPERTY(EditAnywhere, BlueprintReadWrite, Category = "Op Settings", meta = (ReinitializeOnEdit, ClampMin = "0.01", UIMin = "0.01", UIMax = "10.0"))
 	double SourceScaleFactor = 1.0f;
 
-	UPROPERTY(EditAnywhere, BlueprintReadWrite, Category = ScalePivot)
+	UPROPERTY(EditAnywhere, BlueprintReadWrite, Category = "Op Settings")
 	EScaleSourcePivot ScalePivot = EScaleSourcePivot::ComponentOrigin;
 
-	UPROPERTY(EditAnywhere, Category = ScalePivot, meta=(NotOverrideable))
+	UPROPERTY(EditAnywhere, Category = "Op Settings", meta=(NotOverrideable))
 	FBoneReference ScalePivotBone;
 
-	UPROPERTY(EditAnywhere, BlueprintReadWrite, Category = ScalePivot)
+	UPROPERTY(EditAnywhere, BlueprintReadWrite, Category = "Op Settings")
 	bool bProjectScalePivotToFloor = false;
 
 	UE_API virtual const UClass* GetControllerType() const override;
diff --git a/Engine/Plugins/Animation/IKRig/Source/IKRig/Public/Retargeter/RetargetOps/SpeedPlantingOp.h b/Engine/Plugins/Animation/IKRig/Source/IKRig/Public/Retargeter/RetargetOps/SpeedPlantingOp.h
--- a/Engine/Plugins/Animation/IKRig/Source/IKRig/Public/Retargeter/RetargetOps/SpeedPlantingOp.h
+++ b/Engine/Plugins/Animation/IKRig/Source/IKRig/Public/Retargeter/RetargetOps/SpeedPlantingOp.h
@@ -25,11 +25,15 @@ struct FRetargetSpeedPlantingSettings
 	UPROPERTY(EditAnywhere, BlueprintReadWrite, Category = "Speed Planting", meta=(ReinitializeOnEdit, NotOverrideable))
 	FName TargetChainName;
 
+	/** Range 0 to 1. Default 1. Blends the effect of the speed planting on this chain only. At 0, there is no effect. */
+	UPROPERTY(EditAnywhere, BlueprintReadWrite, Category = "Chain Settings", meta = (UIMin = "0.0", UIMax = "1.0", ClampMin = "0.0", ClampMax = "1.0"))
+	float Alpha = 1.0;
+
 	/** The name of the curve on the source animation that contains the speed of the end effector bone.*/
-	UPROPERTY(EditAnywhere, BlueprintReadWrite, BlueprintReadWrite, Category = "Plant IK by Speed", meta = (ClampMin = "0.0", ClampMax = "100.0", UIMin = "0.0", UIMax = "100.0"))
+	UPROPERTY(EditAnywhere, BlueprintReadWrite, BlueprintReadWrite, Category = "Chain Settings", meta = (ClampMin = "0.0", ClampMax = "100.0", UIMin = "0.0", UIMax = "100.0"))
 	FName SpeedCurveName;
 
-	bool operator==(const FRetargetSpeedPlantingSettings& Other) const { return SpeedCurveName == Other.SpeedCurveName; };
+	bool operator==(const FRetargetSpeedPlantingSettings& Other) const;
 	FName GetName() const { return TargetChainName; };
 	void SetName(const FName InName) { TargetChainName = InName; };
 };
@@ -49,25 +53,29 @@ struct FIKRetargetSpeedPlantingOpSettings : public FIKRetargetOpSettingsBase
 	UPROPERTY(EditAnywhere, BlueprintReadWrite, Category = "Setup", meta=(ReinitializeOnEdit))
 	TArray<FRetargetSpeedPlantingSettings> ChainsToSpeedPlant;
 
+	/** Range 0 to 1. Default 1. Blends the effect of the speed planting. At 0, there is no effect. */
+	UPROPERTY(EditAnywhere, BlueprintReadWrite, Category = "Op Settings", meta = (UIMin = "0.0", UIMax = "1.0", ClampMin = "0.0", ClampMax = "1.0"))
+	float Alpha = 1.0;
+	
 	/** Range 0 to 100. Default 15. The maximum speed a source bone can be moving while being considered 'planted'.
 	*  The target IK goal will not be allowed to move whenever the source bone speed drops below this threshold speed. */
-	UPROPERTY(EditAnywhere, BlueprintReadWrite, Category = "Speed Planting", meta = (ClampMin = "0.0", UIMin = "0.0", UIMax = "100.0"))
+	UPROPERTY(EditAnywhere, BlueprintReadWrite, Category = "Op Settings", meta = (ClampMin = "0.0", UIMin = "0.0", UIMax = "100.0"))
 	double SpeedThreshold = 15.0f;
 
 	/** The method to use when blending the IK goal in/out of a planted position.
 	 * The Spring method uses a critically damped spring equation to blend into and out of the planted states.
 	 * The Move Towards method can more accurately track fast motion without looking mushy. It only blends out of planted states. */
-	UPROPERTY(EditAnywhere, BlueprintReadWrite, Category = "Speed Planting")
+	UPROPERTY(EditAnywhere, BlueprintReadWrite, Category = "Op Settings")
 	ESpeedPlantBlendMethod BlendMethod = ESpeedPlantBlendMethod::Spring;
 
 	/** The maximum speed in cm/s that the goal will move to catch up to the animation after unplanting. This blended towards the animated velocity
 	 * as the goal catches up to animated location after unplanting. */
-	UPROPERTY(EditAnywhere, BlueprintReadWrite, Category = "Speed Planting", meta = (ClampMin = "0.0", UIMin = "0.0", UIMax = "2000.0", 
+	UPROPERTY(EditAnywhere, BlueprintReadWrite, Category = "Op Settings", meta = (ClampMin = "0.0", UIMin = "0.0", UIMax = "2000.0", 
 		EditCondition = "BlendMethod == ESpeedPlantBlendMethod::MoveTowards", EditConditionHides))
 	double MaxBlendSpeed = 900.0;
 
 	/** Seconds to reach MaxBlendSpeed. Determines how rapidly the goal accelerates after unplanting. */
-	UPROPERTY(EditAnywhere, BlueprintReadWrite, Category = "Speed Planting", meta = (ClampMin = "0.001", UIMin = "0.01", UIMax = "1.0", 
+	UPROPERTY(EditAnywhere, BlueprintReadWrite, Category = "Op Settings", meta = (ClampMin = "0.001", UIMin = "0.01", UIMax = "1.0", 
 		EditCondition = "BlendMethod == ESpeedPlantBlendMethod::MoveTowards", EditConditionHides))
 	double TimeToMaxSpeed = 0.2;
 
@@ -75,14 +83,14 @@ struct FIKRetargetSpeedPlantingOpSettings : public FIKRetargetOpSettingsBase
 	 * This allows the goal speed to gradually blend towards the target speed as the goal location approaches the animated location.
 	 * Reduce this to remove velocity discontinuities when catching up to the animation.
 	 * 1.0 is very soft, 20.0 is very sharp. Default is 10.0. */
-	UPROPERTY(EditAnywhere, BlueprintReadWrite, Category = "Speed Planting", meta = (ClampMin = "0.001", UIMin = "0.01", UIMax = "20.0", 
+	UPROPERTY(EditAnywhere, BlueprintReadWrite, Category = "Op Settings", meta = (ClampMin = "0.001", UIMin = "0.01", UIMax = "20.0", 
 		EditCondition = "BlendMethod == ESpeedPlantBlendMethod::MoveTowards", EditConditionHides))
 	double ArrivalSpeedGain = 10.0f;
 	
 	/**
 	* The spring strength determines how hard it will pull towards the target when planting and unplanting.
 	* The value is the frequency (hz) at which it will oscillate when there is no damping. */
-	UPROPERTY(EditAnywhere, BlueprintReadWrite, Category = "Speed Planting", meta=(Input, ClampMin = "0.0", UIMin = "0.0", UIMax = "30.0",
+	UPROPERTY(EditAnywhere, BlueprintReadWrite, Category = "Op Settings", meta=(Input, ClampMin = "0.0", UIMin = "0.0", UIMax = "30.0",
 		EditCondition = "BlendMethod == ESpeedPlantBlendMethod::Spring", EditConditionHides))
 	double SpringStrength = 3.0f;
 	// Deprecated in 5.8, replaced with higher-level "Strength" param
diff --git a/Engine/Plugins/Animation/IKRig/Source/IKRig/Public/Retargeter/RetargetOps/StretchChainOp.h b/Engine/Plugins/Animation/IKRig/Source/IKRig/Public/Retargeter/RetargetOps/StretchChainOp.h
--- a/Engine/Plugins/Animation/IKRig/Source/IKRig/Public/Retargeter/RetargetOps/StretchChainOp.h
+++ b/Engine/Plugins/Animation/IKRig/Source/IKRig/Public/Retargeter/RetargetOps/StretchChainOp.h
@@ -49,17 +49,17 @@ struct FRetargetStretchChainSettings
 	FName TargetChainName;
 
 	/** Enable/disable stretching. */
-	UPROPERTY(EditAnywhere, BlueprintReadWrite, Category = "Stretch Chain")
+	UPROPERTY(EditAnywhere, BlueprintReadWrite, Category = "Chain Settings")
 	bool bEnabled = false;
 
 	/** Range 0 to 1. Default 0. Matches the total length of this chain with the mapped source chain.
 	*  At 0, the chain's length will be left alone
 	*  At 1, the chain will be stretched to match the length of the source.*/
-	UPROPERTY(EditAnywhere, BlueprintReadWrite, Category = "Stretch Chain", meta = (EditCondition = "bEnabled", UIMin = "0.0", UIMax = "1.0", ClampMin = "0.0", ClampMax = "1.0"))
+	UPROPERTY(EditAnywhere, BlueprintReadWrite, Category = "Chain Settings", meta = (EditCondition = "bEnabled", UIMin = "0.0", UIMax = "1.0", ClampMin = "0.0", ClampMax = "1.0"))
 	double MatchSourceLength = 0.0f;
 
 	/** Range 0 to inf. Default 1. A multiplier on the length of the chain, applied after matching the source length. */
-	UPROPERTY(EditAnywhere, BlueprintReadWrite, Category = "Stretch Chain", meta = (EditCondition = "bEnabled", UIMin = "0.0", UIMax = "5.0", ClampMin = "0.0"))
+	UPROPERTY(EditAnywhere, BlueprintReadWrite, Category = "Chain Settings", meta = (EditCondition = "bEnabled", UIMin = "0.0", UIMax = "5.0", ClampMin = "0.0"))
 	double ScaleChainLength = 1.0f;
 	
 	bool operator==(const FRetargetStretchChainSettings& Other) const;
diff --git a/Engine/Plugins/Animation/IKRig/Source/IKRig/Public/Retargeter/RetargetOps/StrideWarpingOp.h b/Engine/Plugins/Animation/IKRig/Source/IKRig/Public/Retargeter/RetargetOps/StrideWarpingOp.h
--- a/Engine/Plugins/Animation/IKRig/Source/IKRig/Public/Retargeter/RetargetOps/StrideWarpingOp.h
+++ b/Engine/Plugins/Animation/IKRig/Source/IKRig/Public/Retargeter/RetargetOps/StrideWarpingOp.h
@@ -32,9 +32,13 @@ struct FRetargetStrideWarpChainSettings
 	FName TargetChainName;
 	
 	/** Whether to warp the location of the IK goal on this chain. Default is true. */
-	UPROPERTY(EditAnywhere, BlueprintReadWrite, Category = "Setup", meta=(ReinitializeOnEdit))
+	UPROPERTY(EditAnywhere, BlueprintReadWrite, Category = "Chain Settings", meta=(ReinitializeOnEdit))
 	bool EnableStrideWarping = true;
 
+	/** Range 0 to 1. Default 1. Blends the effect of the stride warping on this chain only. At 0, there is no effect. */
+	UPROPERTY(EditAnywhere, BlueprintReadWrite, Category = "Chain Settings", meta = (EditCondition = "EnableStrideWarping", UIMin = "0.0", UIMax = "1.0", ClampMin = "0.0", ClampMax = "1.0"))
+	float Alpha = 1.0;
+
 	UE_API bool operator==(const FRetargetStrideWarpChainSettings& Other) const;
 	FName GetName() const { return TargetChainName; };
 	void SetName(const FName InName) { TargetChainName = InName; };
@@ -49,6 +53,10 @@ struct FIKRetargetStrideWarpingOpSettings : public FIKRetargetOpSettingsBase
 	UPROPERTY(EditAnywhere, BlueprintReadWrite, Category = Setup, meta=(ReinitializeOnEdit))
 	TArray<FRetargetStrideWarpChainSettings> ChainSettings;
 
+	/** Range 0 to 1. Default 1. Blends the effect of all stride warping on/off. At 0, there is no effect. */
+	UPROPERTY(EditAnywhere, BlueprintReadWrite, Category = "Op Settings", meta = (UIMin = "0.0", UIMax = "1.0", ClampMin = "0.0", ClampMax = "1.0"))
+	float Alpha = 1.0;
+
 	/** Defines the source used to determine the forward direction as the character animates in world space. Default is "Goals".
 	 * This method uses various positions on the character to create a "best fit" global rotation that approximates the facing direction of the character over time.
 	 * This global rotation is used to define the forward and sideways directions used when warping goals along those axes.
@@ -57,40 +65,40 @@ struct FIKRetargetStrideWarpingOpSettings : public FIKRetargetOpSettingsBase
 	 * Chain: uses the positions of the bones in a retarget chain to approximate the facing direction. This is best when used with the spine chain for characters with a horizontal spine, like quadrupeds.
 	 * Root Bone: uses the rotation of the root bone of the skeleton. This is most robust, but character must have correct root motion with yaw rotation in movement direction.
 	 */
-	UPROPERTY(EditAnywhere, BlueprintReadWrite, Category = Warping)
+	UPROPERTY(EditAnywhere, BlueprintReadWrite, Category = "Op Settings")
 	EWarpingDirectionSource DirectionSource = EWarpingDirectionSource::Goals;
 	
 	/** The world space axis that represents the forward facing direction for your character. By default, in Unreal, this is +Y.  */
-	UPROPERTY(EditAnywhere, BlueprintReadWrite, Category = Warping)
+	UPROPERTY(EditAnywhere, BlueprintReadWrite, Category = "Op Settings")
 	EBasicAxis ForwardDirection = EBasicAxis::Y;
 
 	/** When using the "Chain" option as a Direction Source, this defines the chain to use to determine the facing direction of the character.
 	 * Typically this would be the chain that contains the Spine bones. */
-	UPROPERTY(EditAnywhere, BlueprintReadWrite, Category = Warping, meta=(EditCondition="DirectionSource == EWarpingDirectionSource::Chain"))
+	UPROPERTY(EditAnywhere, BlueprintReadWrite, Category = "Op Settings", meta=(EditCondition="DirectionSource == EWarpingDirectionSource::Chain"))
 	FName DirectionChain;
 
 	/** Range 0 to Inf. Default 1. Warps IK goal positions in the forward direction. Useful for stride warping.
 	 * Values below 1 will create smaller, squashed strides. Values greater than 1 will create stretched, longer strides.*/
-	UPROPERTY(EditAnywhere, BlueprintReadWrite, Category = Warping, meta = (UIMin = "0.0", UIMax = "5.0", ClampMin = "0.0"))
+	UPROPERTY(EditAnywhere, BlueprintReadWrite, Category = "Op Settings", meta = (UIMin = "0.0", UIMax = "5.0", ClampMin = "0.0"))
 	double WarpForwards = 1.0f;
 
 	/** Range -+Inf. Default is 0. A static offset in world units to move the IK goals perpendicular to the forward direction.
 	 * Values less than zero will move the goals towards the center-line of the character. Values greater than zero push the goals outwards.*/
-	UPROPERTY(EditAnywhere, BlueprintReadWrite, Category = Warping, meta = (UIMin = "-10.0", UIMax = "100.0"))
+	UPROPERTY(EditAnywhere, BlueprintReadWrite, Category = "Op Settings", meta = (UIMin = "-10.0", UIMax = "100.0"))
 	double SidewaysOffset = 0.0f;
 
 	/** Range 0 to +Inf. Default is 1.0f.
 	 * Values below 1 pull all the goals towards the average of all the goals (towards each other).
 	 * Values greater than 1 push the goals apart.*/
-	UPROPERTY(EditAnywhere, BlueprintReadWrite, Category = Warping, meta = (UIMin = "0.0", UIMax = "2.0", ClampMin = "0.0"))
+	UPROPERTY(EditAnywhere, BlueprintReadWrite, Category = "Op Settings", meta = (UIMin = "0.0", UIMax = "2.0", ClampMin = "0.0"))
 	double WarpSplay = 1.0f;
 
 	/** Adjust the size of the debug drawing */
-	UPROPERTY(EditAnywhere, BlueprintReadWrite, Category = Debug)
+	UPROPERTY(EditAnywhere, BlueprintReadWrite, Category = "Debug")
 	double DebugDrawSize = 5.0;
 	
 	/** Adjust the thickness of the debug drawing */
-	UPROPERTY(EditAnywhere, BlueprintReadWrite, Category = Debug)
+	UPROPERTY(EditAnywhere, BlueprintReadWrite, Category = "Debug")
 	double DebugDrawThickness = 5.0;
 
 	UE_API virtual const UClass* GetControllerType() const override;
@@ -116,12 +124,17 @@ struct TStructOpsTypeTraits<FIKRetargetStrideWarpingOpSettings> : public TStruct
 
 struct FStrideWarpGoalData
 {
-	FStrideWarpGoalData(const FName InGoalName, const FTransform& InGlobalRefPoseOfGoalBone) :
+	FStrideWarpGoalData(
+		const FName InGoalName,
+		const FTransform& InGlobalRefPoseOfGoalBone,
+		const FRetargetStrideWarpChainSettings& InChainSettings) :
 		IKRigGoalName(InGoalName),
-		GlobalRefPoseOfGoalBone(InGlobalRefPoseOfGoalBone) {};
+		GlobalRefPoseOfGoalBone(InGlobalRefPoseOfGoalBone),
+		Settings(&InChainSettings) {};
 	
 	FName IKRigGoalName;
 	FTransform GlobalRefPoseOfGoalBone;
+	const FRetargetStrideWarpChainSettings* Settings;
 };
 
 USTRUCT(BlueprintType, meta = (DisplayName = "Stride Warp IK Goals"))
diff --git a/Engine/Plugins/Animation/IKRig/Source/IKRig/Public/Retargeter/RetargetOps/WeaponGoalsOp.h b/Engine/Plugins/Animation/IKRig/Source/IKRig/Public/Retargeter/RetargetOps/WeaponGoalsOp.h
--- a/Engine/Plugins/Animation/IKRig/Source/IKRig/Public/Retargeter/RetargetOps/WeaponGoalsOp.h
+++ b/Engine/Plugins/Animation/IKRig/Source/IKRig/Public/Retargeter/RetargetOps/WeaponGoalsOp.h
@@ -62,6 +62,10 @@ struct FIKRetargetWeaponGoalsOpSettings : public FIKRetargetOpSettingsBase
 	UPROPERTY(EditAnywhere, Category = "Bones", meta=(NotOverrideable))
 	FBoneReference IKHandGun;
 
+	/** Range 0 to 1. Default 1. Blends the effect of the op off. At 0, there is no effect. */
+	UPROPERTY(EditAnywhere, BlueprintReadWrite, Category = "Op Settings", meta = (UIMin = "0.0", UIMax = "1.0", ClampMin = "0.0", ClampMax = "1.0"))
+	float Alpha = 1.0;
+	
 	/* When false, all weapon scaling is skipped and WeaponScale has no effect. */
 	UPROPERTY(EditAnywhere, BlueprintReadWrite, Category = "Op Settings")
 	bool bEnableWeaponScale = true;
diff --git a/Engine/Plugins/Animation/IKRig/Source/IKRigEditor/Private/RetargetEditor/IKRetargetOpDetails.cpp b/Engine/Plugins/Animation/IKRig/Source/IKRigEditor/Private/RetargetEditor/IKRetargetOpDetails.cpp
--- a/Engine/Plugins/Animation/IKRig/Source/IKRigEditor/Private/RetargetEditor/IKRetargetOpDetails.cpp
+++ b/Engine/Plugins/Animation/IKRig/Source/IKRigEditor/Private/RetargetEditor/IKRetargetOpDetails.cpp
@@ -198,6 +198,13 @@ void FIKRetargetOpBaseSettingsCustomization::AddBaseOpProperties(
 	for (TFieldIterator<FProperty> It(OpBaseStructType); It; ++It)
 	{
 		FProperty* BaseProp = *It;
+
+		// the enabled property is displayed as a checkbox on the op itself
+		if (BaseProp->GetFName() == GET_MEMBER_NAME_CHECKED(FIKRetargetOpSettingsBase, bEnabled))
+		{
+			continue;
+		}
+		
 		TSharedPtr<IPropertyHandle> BasePropHandle = StructPropertyHandle->GetChildHandle(BaseProp->GetFName());
 		if (BasePropHandle.IsValid())
 		{
@@ -370,11 +377,8 @@ void FStrideWarpOpCustomization::CustomizeChildren(
 		SNew(SIKRetargetChainMapList).InChainMapListConfig(ChainListConfig)
 	];
 
-	// add the warping properties
-	AddChildPropertiesInCategory(StructPropertyHandle, ChildBuilder, FName("Warping"));
-	// add the debug properties
+	AddChildPropertiesInCategory(StructPropertyHandle, ChildBuilder, FName("Op Settings"));
 	AddChildPropertiesInCategory(StructPropertyHandle, ChildBuilder, FName("Debug"));
-	// add the properties from the base op
 	AddBaseOpProperties(StructPropertyHandle, ChildBuilder);
 }
 
@@ -412,10 +416,8 @@ void FSpeedPlantOpCustomization::CustomizeChildren(
 	[
 		SNew(SIKRetargetChainMapList).InChainMapListConfig(ChainListConfig)
 	];
-
-	// add speed planting properties
-	AddChildPropertiesInCategory(StructPropertyHandle, ChildBuilder, FName("Speed Planting"));
-	// add the properties from the base op
+	
+	AddChildPropertiesInCategory(StructPropertyHandle, ChildBuilder, FName("Op Settings"));
 	AddBaseOpProperties(StructPropertyHandle, ChildBuilder);
 }
 
@@ -439,7 +441,7 @@ void FPoleVectorOpCustomization::CustomizeChildren(
 	ChainListConfig.OpWithChainMapping = OpName;
 	ChainListConfig.Controller = AssetController;
 	ChainListConfig.bEnableGoalColumn = true;
-	ChainListConfig.bEnableChainMapping = true;
+	ChainListConfig.bEnableChainMapping = false;
 	const FName SettingsPropertyName = GET_MEMBER_NAME_CHECKED(UPoleVectorSettingsWrapper, Settings);
 	const TArray<FName> HiddenSettings = {GET_MEMBER_NAME_CHECKED(FRetargetPoleVectorSettings, TargetChainName)};
 	ChainListConfig.ChainSettingsGetterFunc =
@@ -453,7 +455,7 @@ void FPoleVectorOpCustomization::CustomizeChildren(
 		SNew(SIKRetargetChainMapList).InChainMapListConfig(ChainListConfig)
 	];
 
-	// add the properties from the base op
+	AddChildPropertiesInCategory(StructPropertyHandle, ChildBuilder, FName("Op Settings"));
 	AddBaseOpProperties(StructPropertyHandle, ChildBuilder);
 }
 
@@ -579,7 +581,7 @@ void FStretchChainOpCustomization::CustomizeChildren(
 		SNew(SIKRetargetChainMapList).InChainMapListConfig(ChainListConfig)
 	];
 
-	// add the properties from the base op
+	AddChildPropertiesInCategory(StructPropertyHandle, ChildBuilder, FName("Op Settings"));
 	AddBaseOpProperties(StructPropertyHandle, ChildBuilder);
 }
 
@@ -618,10 +620,7 @@ void FFloorConstraintOpCustomization::CustomizeChildren(
 		SNew(SIKRetargetChainMapList).InChainMapListConfig(ChainListConfig)
 	];
 
-	// add op properties
-	AddChildPropertiesInCategory(StructPropertyHandle, ChildBuilder, FName("Floor Constraint Op Settings"));
-	
-	// add the properties from the base op
+	AddChildPropertiesInCategory(StructPropertyHandle, ChildBuilder, FName("Op Settings"));
 	AddBaseOpProperties(StructPropertyHandle, ChildBuilder);
 }
 
@@ -703,10 +702,7 @@ void FOffsetGoalsOpCustomization::CustomizeChildren(
 		SNew(SIKRetargetChainMapList).InChainMapListConfig(ChainListConfig)
 	];
 
-	// add op properties
 	AddChildPropertiesInCategory(StructPropertyHandle, ChildBuilder, FName("Op Settings"));
-	
-	// add the properties from the base op
 	AddBaseOpProperties(StructPropertyHandle, ChildBuilder);
 }
 
@@ -745,10 +741,7 @@ void FScaleGoalsOpCustomization::CustomizeChildren(
 		SNew(SIKRetargetChainMapList).InChainMapListConfig(ChainListConfig)
 	];
 
-	// add op properties
-	AddChildPropertiesInCategory(StructPropertyHandle, ChildBuilder, FName("Alpha"));
-	
-	// add the properties from the base op
+	AddChildPropertiesInCategory(StructPropertyHandle, ChildBuilder, FName("Op Settings"));
 	AddBaseOpProperties(StructPropertyHandle, ChildBuilder);
 }
 
