[IK Retarget] Fixes regression deleting ops.

#Synced-CL 52174688

diff --git a/Engine/Plugins/Animation/IKRig/Source/IKRigEditor/Private/RetargetEditor/IKRetargeterController.cpp b/Engine/Plugins/Animation/IKRig/Source/IKRigEditor/Private/RetargetEditor/IKRetargeterController.cpp
--- a/Engine/Plugins/Animation/IKRig/Source/IKRigEditor/Private/RetargetEditor/IKRetargeterController.cpp
+++ b/Engine/Plugins/Animation/IKRig/Source/IKRigEditor/Private/RetargetEditor/IKRetargeterController.cpp
@@ -220,8 +220,9 @@ bool UIKRetargeterController::RemoveRetargetOp(const int32 InOpIndex) const
 	Asset->Modify();
 
 	// get full list of ops to remove (including children, if any)
-	TArray<int32> IndicesToRemove = GetChildOpIndices(InOpIndex);
-	IndicesToRemove.Add(InOpIndex);
+	const TArray<int32> ChildrenOpIndices = GetChildOpIndices(InOpIndex);
+	TArray<int32> IndicesToRemove = {InOpIndex};
+	IndicesToRemove.Append(ChildrenOpIndices);
 	TSet<FName> NamesToRemove;
 	for (int32 OpIndex : IndicesToRemove)
 	{
@@ -241,7 +242,7 @@ bool UIKRetargeterController::RemoveRetargetOp(const int32 InOpIndex) const
 	}
 
 	// remove the op itself
-	IndicesToRemove.Add(InOpIndex);
+	Algo::Sort(IndicesToRemove); 
 	Algo::Reverse(IndicesToRemove); // high to low (children are always before parent)
 	for (const int32 OpToRemove : IndicesToRemove)
 	{
