[IKRetarget] Fixed enable setting not showing in override menu.

#Synced-CL 52174688

diff --git a/Engine/Plugins/Animation/IKRig/Source/IKRig/Public/Retargeter/IKRetargetOps.h b/Engine/Plugins/Animation/IKRig/Source/IKRig/Public/Retargeter/IKRetargetOps.h
--- a/Engine/Plugins/Animation/IKRig/Source/IKRig/Public/Retargeter/IKRetargetOps.h
+++ b/Engine/Plugins/Animation/IKRig/Source/IKRig/Public/Retargeter/IKRetargetOps.h
@@ -74,7 +74,7 @@ struct FIKRetargetOpSettingsBase : public FIKRetargetOpSkeletonProvider
 	virtual ~FIKRetargetOpSettingsBase() = default;
 
 	/** Enable/disable the entire operation. Disabled ops are skipped entirely and do not run at all.*/
-	UPROPERTY()
+	UPROPERTY(EditAnywhere, BlueprintReadWrite, Category = "Enable Op")
 	bool bEnabled = true;
 
 	/** The maximum LOD that this Op is allowed to run at.
