[IKRetarget] Fixes crash when naming an Op None.

#Synced-CL 52174688

diff --git a/Engine/Plugins/Animation/IKRig/Source/IKRigEditor/Private/RetargetEditor/IKRetargeterController.cpp b/Engine/Plugins/Animation/IKRig/Source/IKRigEditor/Private/RetargetEditor/IKRetargeterController.cpp
--- a/Engine/Plugins/Animation/IKRig/Source/IKRigEditor/Private/RetargetEditor/IKRetargeterController.cpp
+++ b/Engine/Plugins/Animation/IKRig/Source/IKRigEditor/Private/RetargetEditor/IKRetargeterController.cpp
@@ -486,6 +486,9 @@ int32 UIKRetargeterController::GetParentOpIndex(const int32 InOpIndex) const
 
 FName UIKRetargeterController::GetUniqueOpName(const FName InNameToMakeUnique, int32 InOpIndexToIgnore) const
 {
+	// handle the "None" case by providing a default base name
+	const FName BaseName = InNameToMakeUnique.IsNone() ? FName(TEXT("NewOp")) : InNameToMakeUnique;
+
 	// check if any other op using name
 	auto OpNameInUse = [this, InOpIndexToIgnore](const FName InOpNameToCheck)
 	{
@@ -495,9 +498,9 @@ FName UIKRetargeterController::GetUniqueOpName(const FName InNameToMakeUnique, i
 			{
 				continue;
 			}
-		
+   
 			const FIKRetargetOpBase* Op = Asset->RetargetOps[OpIndex].GetPtr<FIKRetargetOpBase>();
-			if (Op->GetName() == InOpNameToCheck)
+			if (Op && Op->GetName() == InOpNameToCheck)
 			{
 				return true;
 			}
@@ -505,19 +508,20 @@ FName UIKRetargeterController::GetUniqueOpName(const FName InNameToMakeUnique, i
 		return false;
 	};
 
-	if (!OpNameInUse(InNameToMakeUnique))
+	// check if the (potentially corrected) base name is available
+	if (!OpNameInUse(BaseName))
 	{
-		return InNameToMakeUnique;
+		return BaseName;
 	}
-	
-	// keep concatenating an incremented integer suffix until name is unique
-	int32 Number = InNameToMakeUnique.GetNumber() + 1;
-	while(OpNameInUse(FName(InNameToMakeUnique, Number)))
+    
+	// increment the internal FName number suffix until we hit a unique one
+	int32 Number = BaseName.GetNumber() + 1;
+	while(OpNameInUse(FName(BaseName, Number)))
 	{
 		Number++;
 	}
 
-	return FName(InNameToMakeUnique, Number);
+	return FName(BaseName, Number);
 }
 
 int32 UIKRetargeterController::GetIndexOfRetargetOp(FIKRetargetOpBase* RetargetOp) const
