[IKRetarget] Fixed crash when trying to run uninitialized IK Rig.

#Synced-CL 52174688

diff --git a/Engine/Plugins/Animation/IKRig/Source/IKRig/Private/Retargeter/RetargetOps/RunIKRigOp.cpp b/Engine/Plugins/Animation/IKRig/Source/IKRig/Private/Retargeter/RetargetOps/RunIKRigOp.cpp
--- a/Engine/Plugins/Animation/IKRig/Source/IKRig/Private/Retargeter/RetargetOps/RunIKRigOp.cpp
+++ b/Engine/Plugins/Animation/IKRig/Source/IKRig/Private/Retargeter/RetargetOps/RunIKRigOp.cpp
@@ -281,7 +281,7 @@ void FIKRetargetRunIKRigOp::Run(
 	const TArray<FTransform>& InSourceGlobalPose,
 	TArray<FTransform>& OutTargetGlobalPose)
 {
-	if (InProcessor.IsIKForcedOff())
+	if (InProcessor.IsIKForcedOff() || !InProcessor.IsInitialized())
 	{
 		return; // skip this op when IK is off
 	}
diff --git a/Engine/Plugins/Animation/IKRig/Source/IKRig/Private/Rig/IKRigProcessor.cpp b/Engine/Plugins/Animation/IKRig/Source/IKRig/Private/Rig/IKRigProcessor.cpp
--- a/Engine/Plugins/Animation/IKRig/Source/IKRig/Private/Rig/IKRigProcessor.cpp
+++ b/Engine/Plugins/Animation/IKRig/Source/IKRig/Private/Rig/IKRigProcessor.cpp
@@ -378,7 +378,10 @@ void FIKRigProcessor::Solve(const FTransform& ComponentToWorld)
 		// reinitialize
 		Initialize(AssetInitializedWith, Skeleton.SkeletalMesh, CopiedGoalContainer);
 		// restore input pose after reinit
-		SetInputPoseGlobal(InputPoseGlobal);
+		if (IsInitialized())
+		{
+			SetInputPoseGlobal(InputPoseGlobal);
+		}
 	}
 
 	if (!bInitialized)
