[IKRetarget] New Retarget Pose Op

#Synced-CL 52284185

diff --git a/Engine/Plugins/Animation/IKRig/Source/IKRig/Private/Retargeter/IKRetargetProcessor.cpp b/Engine/Plugins/Animation/IKRig/Source/IKRig/Private/Retargeter/IKRetargetProcessor.cpp
--- a/Engine/Plugins/Animation/IKRig/Source/IKRig/Private/Retargeter/IKRetargetProcessor.cpp
+++ b/Engine/Plugins/Animation/IKRig/Source/IKRig/Private/Retargeter/IKRetargetProcessor.cpp
@@ -14,6 +14,7 @@
 #include "Retargeter/RetargetOps/CopyBasePoseOp.h"
 #include "Retargeter/RetargetOps/ScaleSourceOp.h"
 #include "Retargeter/RetargetOps/RetargetCurvesOp.h"
+#include "Retargeter/RetargetOps/RetargetPoseOp.h"
 
 #include UE_INLINE_GENERATED_CPP_BY_NAME(IKRetargetProcessor)
 
@@ -2097,17 +2098,44 @@ void FIKRetargetProcessor::UpdateSettingsAtRuntime(
 
 	// apply stored property overrides
 	OverrideManager.ApplyPropertyOverridesToOps(*this, InOverrideSetsToApply, InVariables, InBoundCurveValues);
-	
+
+	// apply flag to force IK off
 	if (InProfile)
 	{
-		// apply flag to force IK off
 		bIKForcedOff = InProfile->bForceAllIKOff;
+	}
+
+	// update retarget poses
+	{
+		FName SourcePoseToApply = NAME_None;
+		FName TargetPoseToApply = NAME_None;
+
+		// allow profile to override retarget poses
+		if (InProfile)
+		{
+			SourcePoseToApply = InProfile->SourceRetargetPoseName;
+			TargetPoseToApply = InProfile->TargetRetargetPoseName;
+		}
+
+		// allow retarget pose op to override
+		if (const FIKRetargetPoseOp* RetargetPoseOp = GetFirstRetargetOpOfType<FIKRetargetPoseOp>())
+		{
+			const FIKRetargetPoseOpSettings* Settings = reinterpret_cast<const FIKRetargetPoseOpSettings*>(RetargetPoseOp->GetSettingsConst());
+			if (Settings->bOverrideSourcePose && !Settings->SourcePoseToUse.IsNone())
+			{
+				SourcePoseToApply = Settings->SourcePoseToUse;
+			}
+			if (Settings->bOverrideTargetPose && !Settings->TargetPoseToUse.IsNone())
+			{
+				TargetPoseToApply = Settings->TargetPoseToUse;
+			}
+		}
 	
 		// apply retarget poses specified in the profile
 		// NOTE: must be done AFTER applying op settings because it uses the SourceScaleFactor in the ScaleSourceOp settings
-		UpdateRetargetPoseAtRuntime(InProfile->SourceRetargetPoseName, ERetargetSourceOrTarget::Source);
-		UpdateRetargetPoseAtRuntime(InProfile->TargetRetargetPoseName, ERetargetSourceOrTarget::Target);
-	}
+		UpdateRetargetPoseAtRuntime(SourcePoseToApply, ERetargetSourceOrTarget::Source);
+		UpdateRetargetPoseAtRuntime(TargetPoseToApply, ERetargetSourceOrTarget::Target);
+	}	
 }
 
 void FIKRetargetProcessor::UpdateRetargetPoseAtRuntime(
diff --git a/Engine/Plugins/Animation/IKRig/Source/IKRig/Private/Retargeter/RetargetOps/RetargetPoseOp.cpp b/Engine/Plugins/Animation/IKRig/Source/IKRig/Private/Retargeter/RetargetOps/RetargetPoseOp.cpp
new file mode 100644
--- /dev/null
+++ b/Engine/Plugins/Animation/IKRig/Source/IKRig/Private/Retargeter/RetargetOps/RetargetPoseOp.cpp
@@ -0,0 +1,69 @@
+﻿// Copyright Epic Games, Inc. All Rights Reserved.
+
+#include "Retargeter/RetargetOps/RetargetPoseOp.h"
+
+#include "Retargeter/IKRetargetProcessor.h"
+
+#include UE_INLINE_GENERATED_CPP_BY_NAME(RetargetPoseOp)
+
+const UClass* FIKRetargetPoseOpSettings::GetControllerType() const
+{
+	return UIKRetargetPoseController::StaticClass();
+}
+
+void FIKRetargetPoseOpSettings::CopySettingsAtRuntime(const FIKRetargetOpSettingsBase* InSettingsToCopyFrom)
+{
+	// copy all properties
+	static TArray<FName> PropertiesToIgnore = {};
+	FIKRetargetOpBase::CopyStructProperties(
+		FIKRetargetPoseOpSettings::StaticStruct(),
+		InSettingsToCopyFrom,
+		this,
+		PropertiesToIgnore);
+}
+
+bool FIKRetargetPoseOp::Initialize(
+	const FIKRetargetProcessor& InProcessor,
+	const FRetargetSkeleton& InSourceSkeleton,
+	const FTargetSkeleton& InTargetSkeleton,
+	const FIKRetargetOpBase* InParentOp,
+	FIKRigLogger& Log)
+{
+	bIsInitialized = true;
+	return true;
+}
+
+void FIKRetargetPoseOp::Run(
+	FIKRetargetProcessor& InProcessor,
+	const double InDeltaTime,
+	const TArray<FTransform>& InSourceGlobalPose,
+	TArray<FTransform>& OutTargetGlobalPose)
+{
+	// NOTE: this op is read-from by the processor and does no work of its own
+}
+
+FIKRetargetOpSettingsBase* FIKRetargetPoseOp::GetSettings()
+{
+	return &Settings;
+}
+
+const UScriptStruct* FIKRetargetPoseOp::GetSettingsType() const
+{
+	return FIKRetargetPoseOpSettings::StaticStruct();
+}
+
+const UScriptStruct* FIKRetargetPoseOp::GetType() const
+{
+	return FIKRetargetPoseOp::StaticStruct();
+}
+
+FIKRetargetPoseOpSettings UIKRetargetPoseController::GetSettings()
+{
+	return *reinterpret_cast<FIKRetargetPoseOpSettings*>(OpSettingsToControl);
+}
+
+void UIKRetargetPoseController::SetSettings(FIKRetargetPoseOpSettings InSettings)
+{
+	reinterpret_cast<FIKRetargetPoseOpSettings*>(OpSettingsToControl)->CopySettingsAtRuntime(&InSettings);
+}
+
diff --git a/Engine/Plugins/Animation/IKRig/Source/IKRig/Public/Retargeter/RetargetOps/RetargetPoseOp.h b/Engine/Plugins/Animation/IKRig/Source/IKRig/Public/Retargeter/RetargetOps/RetargetPoseOp.h
new file mode 100644
--- /dev/null
+++ b/Engine/Plugins/Animation/IKRig/Source/IKRig/Public/Retargeter/RetargetOps/RetargetPoseOp.h
@@ -0,0 +1,90 @@
+﻿// Copyright Epic Games, Inc. All Rights Reserved.
+#pragma once
+
+#include "Retargeter/IKRetargetOps.h"
+
+#include "RetargetPoseOp.generated.h"
+
+#define UE_API IKRIG_API
+
+#define LOCTEXT_NAMESPACE "RetargetPoseOp"
+
+USTRUCT(BlueprintType, meta = (DisplayName = "Retarget Pose Op Settings"))
+struct FIKRetargetPoseOpSettings : public FIKRetargetOpSettingsBase
+{
+	GENERATED_BODY()
+	
+	// When true, overrides the default retarget pose for the source skeleton.
+	UPROPERTY(EditAnywhere, BlueprintReadWrite, Category="Source Pose")
+	bool bOverrideSourcePose = false;
+
+	// A retarget pose that is applied to the source skeleton. Must match the name of a retarget pose stored in this asset.
+	UPROPERTY(EditAnywhere, BlueprintReadWrite, Category="Source Pose")
+	FName SourcePoseToUse;
+
+	// When true, overrides the default retarget pose for the target skeleton.
+	UPROPERTY(EditAnywhere, BlueprintReadWrite, Category="Target Pose")
+	bool bOverrideTargetPose = false;
+
+	// A retarget pose that is applied to the target skeleton. Must match the name of a retarget pose stored in this asset.
+	UPROPERTY(EditAnywhere, BlueprintReadWrite, Category="Target Pose")
+	FName TargetPoseToUse;
+	
+	UE_API virtual const UClass* GetControllerType() const override;
+	
+	UE_API virtual void CopySettingsAtRuntime(const FIKRetargetOpSettingsBase* InSettingsToCopyFrom) override;
+};
+
+USTRUCT(BlueprintType, meta = (DisplayName = "Retarget Pose"))
+struct FIKRetargetPoseOp : public FIKRetargetOpBase
+{
+	GENERATED_BODY()
+	
+	UE_API virtual bool Initialize(
+		const FIKRetargetProcessor& InProcessor,
+		const FRetargetSkeleton& InSourceSkeleton,
+		const FTargetSkeleton& InTargetSkeleton,
+		const FIKRetargetOpBase* InParentOp,
+		FIKRigLogger& Log) override;
+	
+	UE_API virtual void Run(
+		FIKRetargetProcessor& InProcessor,
+		const double InDeltaTime,
+		const TArray<FTransform>& InSourceGlobalPose,
+		TArray<FTransform>& OutTargetGlobalPose) override;
+
+	UE_API virtual FIKRetargetOpSettingsBase* GetSettings() override;
+	
+	UE_API virtual const UScriptStruct* GetSettingsType() const override;
+	
+	UE_API virtual const UScriptStruct* GetType() const override;
+
+	virtual bool IsSingleton() const override { return true; }
+
+private:
+	
+	UPROPERTY()
+	FIKRetargetPoseOpSettings Settings;
+};
+
+/* The blueprint/python API for editing a Retarget Pose Op */
+UCLASS(BlueprintType)
+class UIKRetargetPoseController : public UIKRetargetOpControllerBase
+{
+	GENERATED_BODY()
+	
+public:
+	/* Get the current op settings as a struct.
+	 * @return FIKRetargetPoseOpSettings struct with the current settings used by the op. */
+	UFUNCTION(BlueprintCallable, Category = Settings)
+	UE_API FIKRetargetPoseOpSettings GetSettings();
+
+	/* Set the op settings. Input is a custom struct type for this op.
+	 * @param InSettings a FIKRetargetPoseOpSettings struct containing all the settings to apply to this op */
+	UFUNCTION(BlueprintCallable, Category = Settings)
+	UE_API void SetSettings(FIKRetargetPoseOpSettings InSettings);
+};
+
+#undef LOCTEXT_NAMESPACE
+
+#undef UE_API
