[IK Retargeter] Fix crash in GetTransformAtChainParam with degenerate chain parameterization

Fixed checkNoEntry() crash in FResolvedBoneChain::GetTransformAtChainParam when the interpolation loop fails to find a matching segment. This can happen with degenerate chains containing zero-length bones or duplicate parameterization values. Replaced with a safe fallback returning the last transform in the chain.

#Synced-CL 52515588

diff --git a/Engine/Plugins/Animation/IKRig/Source/IKRig/Private/Retargeter/IKRetargetProcessor.cpp b/Engine/Plugins/Animation/IKRig/Source/IKRig/Private/Retargeter/IKRetargetProcessor.cpp
--- a/Engine/Plugins/Animation/IKRig/Source/IKRig/Private/Retargeter/IKRetargetProcessor.cpp
+++ b/Engine/Plugins/Animation/IKRig/Source/IKRig/Private/Retargeter/IKRetargetProcessor.cpp
@@ -710,8 +710,8 @@ FTransform FResolvedBoneChain::GetTransformAtChainParam(const TArray<FTransform>
 		return FTransform(Rotation,Position, Scale);
 	}
 
-	checkNoEntry();
-	return FTransform::Identity;
+	// fallback — can be reached with degenerate chains (zero-length bones, duplicate params)
+	return Transforms.Last();
 }
 
 double FResolvedBoneChain::GetStretchAtParam(
