[IK Retargeter] Fix crash in batch retarget NotifyUserOfResults with null asset

Fixed access violation crash when constructing FAssetData from a null UObject pointer in NotifyUserOfResults. GetNewAssets can return null entries if an asset failed to create during the batch retarget. Added null check before constructing FAssetData.

#Synced-CL 52515588

diff --git a/Engine/Plugins/Animation/IKRig/Source/IKRigEditor/Private/RetargetEditor/IKRetargetBatchOperation.cpp b/Engine/Plugins/Animation/IKRig/Source/IKRigEditor/Private/RetargetEditor/IKRetargetBatchOperation.cpp
--- a/Engine/Plugins/Animation/IKRig/Source/IKRigEditor/Private/RetargetEditor/IKRetargetBatchOperation.cpp
+++ b/Engine/Plugins/Animation/IKRig/Source/IKRigEditor/Private/RetargetEditor/IKRetargetBatchOperation.cpp
@@ -738,7 +738,10 @@ void UIKRetargetBatchOperation::NotifyUserOfResults(
 	TArray<FAssetData> CurrentSelection;
 	for(UObject* NewObject : NewAssets)
 	{
-		CurrentSelection.Add(FAssetData(NewObject));
+		if (NewObject)
+		{
+			CurrentSelection.Add(FAssetData(NewObject));
+		}
 	}
 	const FContentBrowserModule& ContentBrowserModule = FModuleManager::Get().LoadModuleChecked<FContentBrowserModule>("ContentBrowser");
 	ContentBrowserModule.Get().SyncBrowserToAssets(CurrentSelection);
