[IK Retargeter] Add override set dropdown to Batch Retarget window

When using a manually supplied IK Retargeter in the Batch/Auto Retarget window, there was no way to select override sets. Added a dropdown (FName with GetOptions) that lists available override sets from the assigned retargeter asset.

Selected override set is applied to both the live preview viewport and animation export. Dropdown is hidden when using auto-generated retargeter. Settings reset to defaults when reopening the dialog.

#Synced-CL 52549651

diff --git a/Engine/Plugins/Animation/IKRig/Source/IKRigEditor/Private/RetargetEditor/IKRetargetAnimInstance.cpp b/Engine/Plugins/Animation/IKRig/Source/IKRigEditor/Private/RetargetEditor/IKRetargetAnimInstance.cpp
--- a/Engine/Plugins/Animation/IKRig/Source/IKRigEditor/Private/RetargetEditor/IKRetargetAnimInstance.cpp
+++ b/Engine/Plugins/Animation/IKRig/Source/IKRigEditor/Private/RetargetEditor/IKRetargetAnimInstance.cpp
@@ -152,6 +152,12 @@ void UIKRetargetAnimInstance::ConfigureAnimInstance(
 	Proxy.ConfigureAnimInstance(InSourceOrTarget, InAsset, InSourceMeshComponent);
 }
 
+void UIKRetargetAnimInstance::SetOverrideSets(const TArray<FName>& InOverrideSets)
+{
+	RetargetNode.bUseCustomOverrideSets = !InOverrideSets.IsEmpty();
+	RetargetNode.OverrideSetsToApply = InOverrideSets;
+}
+
 FIKRetargetProcessor* UIKRetargetAnimInstance::GetRetargetProcessor()
 {
 	return RetargetNode.GetRetargetProcessor();
diff --git a/Engine/Plugins/Animation/IKRig/Source/IKRigEditor/Private/RetargetEditor/SRetargetAnimAssetsWindow.cpp b/Engine/Plugins/Animation/IKRig/Source/IKRigEditor/Private/RetargetEditor/SRetargetAnimAssetsWindow.cpp
--- a/Engine/Plugins/Animation/IKRig/Source/IKRigEditor/Private/RetargetEditor/SRetargetAnimAssetsWindow.cpp
+++ b/Engine/Plugins/Animation/IKRig/Source/IKRigEditor/Private/RetargetEditor/SRetargetAnimAssetsWindow.cpp
@@ -550,6 +550,14 @@ void SRetargetPoseViewport::SetRetargetAsset(UIKRetargeter* RetargetAsset)
 	TargetAnimInstance->ConfigureAnimInstance(ERetargetSourceOrTarget::Target, RetargetAsset, SourceComponent);
 }
 
+void SRetargetPoseViewport::SetOverrideSets(const TArray<FName>& InOverrideSets)
+{
+	if (TargetAnimInstance)
+	{
+		TargetAnimInstance->SetOverrideSets(InOverrideSets);
+	}
+}
+
 void SRetargetPoseViewport::PlayAnimation(UAnimationAsset* AnimationAsset)
 {
 	SourceComponent->EnablePreview(true, AnimationAsset);
@@ -903,6 +911,11 @@ SRetargetAnimAssetsWindow::SRetargetAnimAssetsWindow()
 	// get the global settings uobject
 	Settings = UBatchRetargetSettings::GetInstance();
 
+	// reset settings to defaults when opening a new dialog
+	Settings->bAutoGenerateRetargeter = true;
+	Settings->RetargetAsset = nullptr;
+	Settings->OverrideSetToApply = NAME_None;
+
 	// assign default retargeter
 	BatchContext.IKRetargetAsset = ProceduralAssets.Retargeter;
 
@@ -970,7 +983,7 @@ void SRetargetAnimAssetsWindow::Construct(const FArguments& InArgs)
 				.Value(0.6f)
 				[
 					SNew(SVerticalBox)
-					
+
 					+SVerticalBox::Slot()
 					[
 						SAssignNew(AssetBrowser, SRetargetExporterAssetBrowser, SharedThis(this))
@@ -1047,9 +1060,29 @@ void SRetargetAnimAssetsWindow::OnFinishedChangingSelectionProperties(const FPro
 		{
 			TargetMesh = Settings->RetargetAsset->GetPreviewMesh(ERetargetSourceOrTarget::Target);
 		}
-		
+
+		// clear override set if it no longer exists on the new retargeter
+		if (Settings->RetargetAsset)
+		{
+			const TMap<FName, FRetargetOverrideSet>& ValidSets = Settings->RetargetAsset->GetOverrideSets();
+			if (!Settings->OverrideSetToApply.IsNone() && !ValidSets.Contains(Settings->OverrideSetToApply))
+			{
+				Settings->OverrideSetToApply = NAME_None;
+			}
+		}
+		else
+		{
+			Settings->OverrideSetToApply = NAME_None;
+		}
+
 		SetAssets(Settings->SourceSkeletalMesh, TargetMesh, Settings->RetargetAsset);
 	}
+
+	const bool ChangedOverrideSet = PropertyChangedEvent.Property->GetName() == "OverrideSetToApply";
+	if (ChangedOverrideSet)
+	{
+		ApplyOverrideSetToPreview();
+	}
 }
 
 bool SRetargetAnimAssetsWindow::CanExportAnimations() const
@@ -1085,6 +1118,13 @@ FReply SRetargetAnimAssetsWindow::OnExportAnimations()
 		}
 	}
 
+	// apply override set to batch context
+	BatchContext.OverrideSetNames.Reset();
+	if (!Settings->OverrideSetToApply.IsNone())
+	{
+		BatchContext.OverrideSetNames.Add(Settings->OverrideSetToApply);
+	}
+
 	// run the batch retarget
 	const TStrongObjectPtr<UIKRetargetBatchOperation> BatchOperation(NewObject<UIKRetargetBatchOperation>());
 	BatchOperation->RunRetarget(BatchContext);
@@ -1211,6 +1251,37 @@ FText SRetargetAnimAssetsWindow::GetWarningText() const
 	return FText::GetEmpty();
 }
 
+TArray<FString> UBatchRetargetSettings::GetOverrideSetOptions() const
+{
+	TArray<FString> Options;
+	Options.Add(TEXT("None"));
+
+	if (RetargetAsset)
+	{
+		const TMap<FName, FRetargetOverrideSet>& AllSets = RetargetAsset->GetOverrideSets();
+		for (const TTuple<FName, FRetargetOverrideSet>& Entry : AllSets)
+		{
+			Options.Add(Entry.Key.ToString());
+		}
+	}
+
+	return Options;
+}
+
+void SRetargetAnimAssetsWindow::ApplyOverrideSetToPreview()
+{
+	BatchContext.OverrideSetNames.Reset();
+	if (!Settings->OverrideSetToApply.IsNone())
+	{
+		BatchContext.OverrideSetNames.Add(Settings->OverrideSetToApply);
+	}
+
+	if (Viewport.IsValid())
+	{
+		Viewport->SetOverrideSets(BatchContext.OverrideSetNames);
+	}
+}
+
 void SRetargetAnimAssetsWindow::ShowWindow(TArray<UObject*> InSelectedAssets)
 {	
 	if(Window.IsValid())
@@ -1314,6 +1385,9 @@ void SRetargetAnimAssetsWindow::SetAssets(
 	Viewport->SetSkeletalMesh(TargetMesh, ERetargetSourceOrTarget::Target);
 	Viewport->SetRetargetAsset(BatchContext.IKRetargetAsset);
 
+	// apply override set to preview
+	ApplyOverrideSetToPreview();
+
 	ShowAssetWarnings();
 	
 	if (ReplacedSourceMesh)
diff --git a/Engine/Plugins/Animation/IKRig/Source/IKRigEditor/Public/RetargetEditor/IKRetargetAnimInstance.h b/Engine/Plugins/Animation/IKRig/Source/IKRigEditor/Public/RetargetEditor/IKRetargetAnimInstance.h
--- a/Engine/Plugins/Animation/IKRig/Source/IKRigEditor/Public/RetargetEditor/IKRetargetAnimInstance.h
+++ b/Engine/Plugins/Animation/IKRig/Source/IKRigEditor/Public/RetargetEditor/IKRetargetAnimInstance.h
@@ -60,6 +60,8 @@ public:
 		UIKRetargeter* InIKRetargetAsset,
 		TWeakObjectPtr<USkeletalMeshComponent> InSourceMeshComponent);
 
+	void SetOverrideSets(const TArray<FName>& InOverrideSets);
+
 	FIKRetargetProcessor* GetRetargetProcessor();
 	
 	const FAnimNode_RetargetPoseFromMesh& GetRetargetAnimNode();
diff --git a/Engine/Plugins/Animation/IKRig/Source/IKRigEditor/Public/RetargetEditor/SRetargetAnimAssetsWindow.h b/Engine/Plugins/Animation/IKRig/Source/IKRigEditor/Public/RetargetEditor/SRetargetAnimAssetsWindow.h
--- a/Engine/Plugins/Animation/IKRig/Source/IKRigEditor/Public/RetargetEditor/SRetargetAnimAssetsWindow.h
+++ b/Engine/Plugins/Animation/IKRig/Source/IKRigEditor/Public/RetargetEditor/SRetargetAnimAssetsWindow.h
@@ -132,6 +132,7 @@ public:
 	
 	void SetSkeletalMesh(USkeletalMesh* InSkeletalMesh, ERetargetSourceOrTarget SourceOrTarget);
 	void SetRetargetAsset(UIKRetargeter* RetargetAsset);
+	void SetOverrideSets(const TArray<FName>& InOverrideSets);
 	void PlayAnimation(UAnimationAsset* AnimationAsset);
 	bool IsRetargeterValid();
 
@@ -232,6 +233,14 @@ public:
 	UPROPERTY(EditAnywhere, BlueprintReadWrite, Category = "Retargeter", meta = (EditCondition = "!bAutoGenerateRetargeter"))
 	TObjectPtr<UIKRetargeter> RetargetAsset;
 
+	// Override set from the retargeter to apply during preview and export.
+	// Only available when using a manually supplied retargeter asset.
+	UPROPERTY(EditAnywhere, BlueprintReadWrite, Category = "Retargeter", meta = (EditCondition = "!bAutoGenerateRetargeter", GetOptions = "GetOverrideSetOptions"))
+	FName OverrideSetToApply;
+
+	UFUNCTION()
+	TArray<FString> GetOverrideSetOptions() const;
+
 private:
 	static UBatchRetargetSettings* SingletonInstance;
 };
@@ -317,6 +326,8 @@ private:
 	// necessary data collected from UI to run retarget.
 	FIKRetargetBatchOperationContext BatchContext;
 
+	void ApplyOverrideSetToPreview();
+
 	// ui elements
 	TSharedPtr<SRetargetPoseViewport> Viewport;
 	TSharedPtr<SRetargetExporterAssetBrowser> AssetBrowser;
