IK Rig Template Updates: Split LeftFoot/RightFoot chains, FBIK defaults

Updated all 17 biped templates to split the existing LeftLeg/RightLeg chains:

LeftLeg/RightLeg chains now end at the ankle/foot bone (previously ended at ball/toe)

New LeftFoot/RightFoot chains cover the ball/toe bone

IK goals remain on the leg chains, now created at the ankle/foot bone

Updated FBIK auto-setup defaults:

SubIterations: 0 -> 10

Goal ChainDepth: 0 -> 2

Updated Biped Templates, Added new Templates for Rigify and Meshcapade

#Synced-CL 52631650

#Synced-CL 52631650

diff --git a/Engine/Plugins/Animation/IKRig/Source/IKRigEditor/Private/RigEditor/IKRigAutoCharacterizer.cpp b/Engine/Plugins/Animation/IKRig/Source/IKRigEditor/Private/RigEditor/IKRigAutoCharacterizer.cpp
--- a/Engine/Plugins/Animation/IKRig/Source/IKRigEditor/Private/RigEditor/IKRigAutoCharacterizer.cpp
+++ b/Engine/Plugins/Animation/IKRig/Source/IKRigEditor/Private/RigEditor/IKRigAutoCharacterizer.cpp
@@ -16,6 +16,9 @@ const FName FCharacterizationStandard::Tail("Tail");
 // legs
 const FName FCharacterizationStandard::LeftLeg("LeftLeg");
 const FName FCharacterizationStandard::RightLeg("RightLeg");
+// feet
+const FName FCharacterizationStandard::LeftFoot("LeftFoot");
+const FName FCharacterizationStandard::RightFoot("RightFoot");
 // arms
 const FName FCharacterizationStandard::LeftClavicle("LeftClavicle");
 const FName FCharacterizationStandard::RightClavicle("RightClavicle");
@@ -436,6 +439,12 @@ void FAllBoneSettingsForIK::SetRotationStiffness(const FName BoneName, const flo
 	Settings.RotationStiffness = RotationStiffness;
 }
 
+void FAllBoneSettingsForIK::SetPositionStiffness(const FName BoneName, const float PositionStiffness)
+{
+	FBoneSettingsForIK& Settings = GetOrAddBoneSettings(BoneName);
+	Settings.PositionStiffness = PositionStiffness;
+}
+
 void FAllBoneSettingsForIK::SetExcluded(const FName BoneName, const bool bExclude)
 {
 	FBoneSettingsForIK& Settings = GetOrAddBoneSettings(BoneName);
@@ -471,11 +480,12 @@ FKnownTemplateHierarchies::FKnownTemplateHierarchies()
 		// core
 		UE4Retarget.PelvisBone = FName("pelvis");
 		UE4Retarget.RootBone = FName("root");
-		UE4Retarget.AddBoneChain(FCharacterizationStandard::Spine, FName("spine_01"), FName("spine_02"));
+		UE4Retarget.AddBoneChain(FCharacterizationStandard::Spine, FName("spine_01"), FName("spine_03"));
 		UE4Retarget.AddBoneChain(FCharacterizationStandard::Neck, FName("neck_01"), FName("neck_01"));
 		UE4Retarget.AddBoneChain(FCharacterizationStandard::Head, FName("head"), FName("head"));
 		// left
-		UE4Retarget.AddBoneChain(FCharacterizationStandard::LeftLeg, FName("thigh_l"), FName("ball_l"),FCharacterizationStandard::LeftFootIKGoal);
+		UE4Retarget.AddBoneChain(FCharacterizationStandard::LeftLeg, FName("thigh_l"), FName("foot_l"),FCharacterizationStandard::LeftFootIKGoal);
+		UE4Retarget.AddBoneChain(FCharacterizationStandard::LeftFoot, FName("ball_l"), FName("ball_l"));
 		UE4Retarget.AddBoneChain(FCharacterizationStandard::LeftClavicle, FName("clavicle_l"), FName("clavicle_l"));
 		UE4Retarget.AddBoneChain(FCharacterizationStandard::LeftArm, FName("upperarm_l"), FName("hand_l"),FCharacterizationStandard::LeftHandIKGoal);
 		UE4Retarget.AddBoneChain(FCharacterizationStandard::LeftThumb, FName("thumb_01_l"), FName("thumb_03_l"));
@@ -484,7 +494,8 @@ FKnownTemplateHierarchies::FKnownTemplateHierarchies()
 		UE4Retarget.AddBoneChain(FCharacterizationStandard::LeftRing, FName("ring_01_l"), FName("ring_03_l"));
 		UE4Retarget.AddBoneChain(FCharacterizationStandard::LeftPinky, FName("pinky_01_l"), FName("pinky_03_l"));
 		// right
-		UE4Retarget.AddBoneChain(FCharacterizationStandard::RightLeg, FName("thigh_r"), FName("ball_r"), FCharacterizationStandard::RightFootIKGoal);
+		UE4Retarget.AddBoneChain(FCharacterizationStandard::RightLeg, FName("thigh_r"), FName("foot_r"), FCharacterizationStandard::RightFootIKGoal);
+		UE4Retarget.AddBoneChain(FCharacterizationStandard::RightFoot, FName("ball_r"), FName("ball_r"));
 		UE4Retarget.AddBoneChain(FCharacterizationStandard::RightClavicle, FName("clavicle_r"), FName("clavicle_r"));
 		UE4Retarget.AddBoneChain(FCharacterizationStandard::RightArm, FName("upperarm_r"), FName("hand_r"),FCharacterizationStandard::RightHandIKGoal);
 		UE4Retarget.AddBoneChain(FCharacterizationStandard::RightThumb, FName("thumb_01_r"), FName("thumb_03_r"));
@@ -531,7 +542,8 @@ FKnownTemplateHierarchies::FKnownTemplateHierarchies()
 		UE5Retarget.AddBoneChain(FCharacterizationStandard::Neck, FName("neck_01"), FName("neck_02"));
 		UE5Retarget.AddBoneChain(FCharacterizationStandard::Head, FName("head"), FName("head"));
 		// left
-		UE5Retarget.AddBoneChain(FCharacterizationStandard::LeftLeg, FName("thigh_l"), FName("ball_l"),FCharacterizationStandard::LeftFootIKGoal);
+		UE5Retarget.AddBoneChain(FCharacterizationStandard::LeftLeg, FName("thigh_l"), FName("foot_l"),FCharacterizationStandard::LeftFootIKGoal);
+		UE5Retarget.AddBoneChain(FCharacterizationStandard::LeftFoot, FName("ball_l"), FName("ball_l"));
 		UE5Retarget.AddBoneChain(FCharacterizationStandard::LeftClavicle, FName("clavicle_l"), FName("clavicle_l"));
 		UE5Retarget.AddBoneChain(FCharacterizationStandard::LeftArm, FName("upperarm_l"), FName("hand_l"),FCharacterizationStandard::LeftHandIKGoal);
 		UE5Retarget.AddBoneChain(FCharacterizationStandard::LeftIndexMetacarpal, FName("index_metacarpal_l"), FName("index_metacarpal_l"));
@@ -544,7 +556,8 @@ FKnownTemplateHierarchies::FKnownTemplateHierarchies()
 		UE5Retarget.AddBoneChain(FCharacterizationStandard::LeftRing, FName("ring_01_l"), FName("ring_03_l"));
 		UE5Retarget.AddBoneChain(FCharacterizationStandard::LeftPinky, FName("pinky_01_l"), FName("pinky_03_l"));
 		// right
-		UE5Retarget.AddBoneChain(FCharacterizationStandard::RightLeg, FName("thigh_r"), FName("ball_r"),FCharacterizationStandard::RightFootIKGoal);
+		UE5Retarget.AddBoneChain(FCharacterizationStandard::RightLeg, FName("thigh_r"), FName("foot_r"),FCharacterizationStandard::RightFootIKGoal);
+		UE5Retarget.AddBoneChain(FCharacterizationStandard::RightFoot, FName("ball_r"), FName("ball_r"));
 		UE5Retarget.AddBoneChain(FCharacterizationStandard::RightClavicle, FName("clavicle_r"), FName("clavicle_r"));
 		UE5Retarget.AddBoneChain(FCharacterizationStandard::RightArm, FName("upperarm_r"), FName("hand_r"),FCharacterizationStandard::RightHandIKGoal);
 		UE5Retarget.AddBoneChain(FCharacterizationStandard::RightIndexMetacarpal, FName("index_metacarpal_r"), FName("index_metacarpal_r"));
@@ -594,7 +607,8 @@ FKnownTemplateHierarchies::FKnownTemplateHierarchies()
 		DazRetarget.AddBoneChain(FCharacterizationStandard::Neck, FName("neckLower"), FName("neckUpper"));
 		DazRetarget.AddBoneChain(FCharacterizationStandard::Head, FName("head"), FName("head"));
 		// left
-		DazRetarget.AddBoneChain(FCharacterizationStandard::LeftLeg, FName("lThighBend"), FName("lToe"),FCharacterizationStandard::LeftFootIKGoal);
+		DazRetarget.AddBoneChain(FCharacterizationStandard::LeftLeg, FName("lThighBend"), FName("lFoot"),FCharacterizationStandard::LeftFootIKGoal);
+		DazRetarget.AddBoneChain(FCharacterizationStandard::LeftFoot, FName("lToe"), FName("lToe"));
 		DazRetarget.AddBoneChain(FCharacterizationStandard::LeftClavicle, FName("lCollar"), FName("lCollar"));
 		DazRetarget.AddBoneChain(FCharacterizationStandard::LeftArm, FName("lShldrBend"), FName("lHand"),FCharacterizationStandard::LeftHandIKGoal);
 		DazRetarget.AddBoneChain(FCharacterizationStandard::LeftIndexMetacarpal, FName("lCarpal1"), FName("lCarpal1"));
@@ -612,7 +626,8 @@ FKnownTemplateHierarchies::FKnownTemplateHierarchies()
 		DazRetarget.AddBoneChain(FCharacterizationStandard::LeftRingToe, FName("lSmallToe3"), FName("lSmallToe3_2"));
 		DazRetarget.AddBoneChain(FCharacterizationStandard::LeftPinkyToe, FName("lSmallToe4"), FName("lSmallToe4_2"));
 		// right
-		DazRetarget.AddBoneChain(FCharacterizationStandard::RightLeg, FName("rThighBend"), FName("rToe"),FCharacterizationStandard::RightFootIKGoal);
+		DazRetarget.AddBoneChain(FCharacterizationStandard::RightLeg, FName("rThighBend"), FName("rFoot"),FCharacterizationStandard::RightFootIKGoal);
+		DazRetarget.AddBoneChain(FCharacterizationStandard::RightFoot, FName("rToe"), FName("rToe"));
 		DazRetarget.AddBoneChain(FCharacterizationStandard::RightClavicle, FName("rCollar"), FName("rCollar"));
 		DazRetarget.AddBoneChain(FCharacterizationStandard::RightArm, FName("rShldrBend"), FName("rHand"),FCharacterizationStandard::RightHandIKGoal);
 		DazRetarget.AddBoneChain(FCharacterizationStandard::RightIndexMetacarpal, FName("rCarpal1"), FName("rCarpal1"));
@@ -667,7 +682,8 @@ FKnownTemplateHierarchies::FKnownTemplateHierarchies()
 		MixamoRetarget.AddBoneChain(FCharacterizationStandard::Neck, FName("Neck"), FName("Neck"));
 		MixamoRetarget.AddBoneChain(FCharacterizationStandard::Head, FName("Head"), FName("Head"));
 		// left
-		MixamoRetarget.AddBoneChain(FCharacterizationStandard::LeftLeg, FName("LeftUpLeg"), FName("LeftToeBase"),FCharacterizationStandard::LeftFootIKGoal);
+		MixamoRetarget.AddBoneChain(FCharacterizationStandard::LeftLeg, FName("LeftUpLeg"), FName("LeftFoot"),FCharacterizationStandard::LeftFootIKGoal);
+		MixamoRetarget.AddBoneChain(FCharacterizationStandard::LeftFoot, FName("LeftToeBase"), FName("LeftToeBase"));
 		MixamoRetarget.AddBoneChain(FCharacterizationStandard::LeftClavicle, FName("LeftShoulder"), FName("LeftShoulder"));
 		MixamoRetarget.AddBoneChain(FCharacterizationStandard::LeftArm, FName("LeftArm"), FName("LeftHand"),FCharacterizationStandard::LeftHandIKGoal);
 		MixamoRetarget.AddBoneChain(FCharacterizationStandard::LeftThumb, FName("LeftHandThumb1"), FName("LeftHandThumb3"));
@@ -676,7 +692,8 @@ FKnownTemplateHierarchies::FKnownTemplateHierarchies()
 		MixamoRetarget.AddBoneChain(FCharacterizationStandard::LeftRing, FName("LeftHandRing1"), FName("LeftHandRing3"));
 		MixamoRetarget.AddBoneChain(FCharacterizationStandard::LeftPinky, FName("LeftHandPinky1"), FName("LeftHandPinky3"));
 		// right
-		MixamoRetarget.AddBoneChain(FCharacterizationStandard::RightLeg, FName("RightUpLeg"), FName("RightToeBase"),FCharacterizationStandard::RightFootIKGoal);
+		MixamoRetarget.AddBoneChain(FCharacterizationStandard::RightLeg, FName("RightUpLeg"), FName("RightFoot"),FCharacterizationStandard::RightFootIKGoal);
+		MixamoRetarget.AddBoneChain(FCharacterizationStandard::RightFoot, FName("RightToeBase"), FName("RightToeBase"));
 		MixamoRetarget.AddBoneChain(FCharacterizationStandard::RightClavicle, FName("RightShoulder"), FName("RightShoulder"));
 		MixamoRetarget.AddBoneChain(FCharacterizationStandard::RightArm, FName("RightArm"), FName("RightHand"),FCharacterizationStandard::RightHandIKGoal);
 		MixamoRetarget.AddBoneChain(FCharacterizationStandard::RightThumb, FName("RightHandThumb1"), FName("RightHandThumb3"));
@@ -713,7 +730,8 @@ FKnownTemplateHierarchies::FKnownTemplateHierarchies()
 		CC4Retarget.AddBoneChain(FCharacterizationStandard::Neck, FName("NeckTwist01"), FName("NeckTwist02"));
 		CC4Retarget.AddBoneChain(FCharacterizationStandard::Head, FName("Head"), FName("Head"));
 		// left
-		CC4Retarget.AddBoneChain(FCharacterizationStandard::LeftLeg, FName("L_Thigh"), FName("L_ToeBase"),FCharacterizationStandard::LeftFootIKGoal);
+		CC4Retarget.AddBoneChain(FCharacterizationStandard::LeftLeg, FName("L_Thigh"), FName("L_Foot"),FCharacterizationStandard::LeftFootIKGoal);
+		CC4Retarget.AddBoneChain(FCharacterizationStandard::LeftFoot, FName("L_ToeBase"), FName("L_ToeBase"));
 		CC4Retarget.AddBoneChain(FCharacterizationStandard::LeftClavicle, FName("L_Clavicle"), FName("L_Clavicle"));
 		CC4Retarget.AddBoneChain(FCharacterizationStandard::LeftArm, FName("L_Upperarm"), FName("L_Hand"),FCharacterizationStandard::LeftHandIKGoal);
 		CC4Retarget.AddBoneChain(FCharacterizationStandard::LeftThumb, FName("L_Thumb1"), FName("L_Thumb3"));
@@ -727,7 +745,8 @@ FKnownTemplateHierarchies::FKnownTemplateHierarchies()
 		CC4Retarget.AddBoneChain(FCharacterizationStandard::LeftRingToe, FName("L_RingToe1"), FName("L_RingToe1"));
 		CC4Retarget.AddBoneChain(FCharacterizationStandard::LeftPinkyToe, FName("L_PinkyToe1"), FName("L_PinkyToe1"));
 		// right
-		CC4Retarget.AddBoneChain(FCharacterizationStandard::RightLeg, FName("R_Thigh"), FName("R_ToeBase"),FCharacterizationStandard::RightFootIKGoal);
+		CC4Retarget.AddBoneChain(FCharacterizationStandard::RightLeg, FName("R_Thigh"), FName("R_Foot"),FCharacterizationStandard::RightFootIKGoal);
+		CC4Retarget.AddBoneChain(FCharacterizationStandard::RightFoot, FName("R_ToeBase"), FName("R_ToeBase"));
 		CC4Retarget.AddBoneChain(FCharacterizationStandard::RightClavicle, FName("R_Clavicle"), FName("R_Clavicle"));
 		CC4Retarget.AddBoneChain(FCharacterizationStandard::RightArm, FName("R_Upperarm"), FName("R_Hand"),FCharacterizationStandard::RightHandIKGoal);
 		CC4Retarget.AddBoneChain(FCharacterizationStandard::RightThumb, FName("R_Thumb1"), FName("R_Thumb3"));
@@ -770,7 +789,8 @@ FKnownTemplateHierarchies::FKnownTemplateHierarchies()
 		XsensRetarget.AddBoneChain(FCharacterizationStandard::Neck, FName("Neck"), FName("Neck"));
 		XsensRetarget.AddBoneChain(FCharacterizationStandard::Head, FName("Head"), FName("Head"));
 		// left
-		XsensRetarget.AddBoneChain(FCharacterizationStandard::LeftLeg, FName("LeftUpLeg"), FName("LeftToeBase"),FCharacterizationStandard::LeftFootIKGoal);
+		XsensRetarget.AddBoneChain(FCharacterizationStandard::LeftLeg, FName("LeftUpLeg"), FName("LeftFoot"),FCharacterizationStandard::LeftFootIKGoal);
+		XsensRetarget.AddBoneChain(FCharacterizationStandard::LeftFoot, FName("LeftToeBase"), FName("LeftToeBase"));
 		XsensRetarget.AddBoneChain(FCharacterizationStandard::LeftClavicle, FName("LeftShoulder"), FName("LeftShoulder"));
 		XsensRetarget.AddBoneChain(FCharacterizationStandard::LeftArm, FName("LeftArm"), FName("LeftHand"),FCharacterizationStandard::LeftHandIKGoal);
 		XsensRetarget.AddBoneChain(FCharacterizationStandard::LeftThumb, FName("LeftHandThumb1"), FName("LeftHandThumb3"));
@@ -779,7 +799,8 @@ FKnownTemplateHierarchies::FKnownTemplateHierarchies()
 		XsensRetarget.AddBoneChain(FCharacterizationStandard::LeftRing, FName("LeftHandRing1"), FName("LeftHandRing3"));
 		XsensRetarget.AddBoneChain(FCharacterizationStandard::LeftPinky, FName("LeftHandPinky1"), FName("LeftHandPinky3"));
 		// right
-		XsensRetarget.AddBoneChain(FCharacterizationStandard::RightLeg, FName("RightUpLeg"), FName("RightToeBase"),FCharacterizationStandard::RightFootIKGoal);
+		XsensRetarget.AddBoneChain(FCharacterizationStandard::RightLeg, FName("RightUpLeg"), FName("RightFoot"),FCharacterizationStandard::RightFootIKGoal);
+		XsensRetarget.AddBoneChain(FCharacterizationStandard::RightFoot, FName("RightToeBase"), FName("RightToeBase"));
 		XsensRetarget.AddBoneChain(FCharacterizationStandard::RightClavicle, FName("RightShoulder"), FName("RightShoulder"));
 		XsensRetarget.AddBoneChain(FCharacterizationStandard::RightArm, FName("RightArm"), FName("RightHand"),FCharacterizationStandard::RightHandIKGoal);
 		XsensRetarget.AddBoneChain(FCharacterizationStandard::RightThumb, FName("RightHandThumb1"), FName("RightHandThumb3"));
@@ -816,7 +837,8 @@ FKnownTemplateHierarchies::FKnownTemplateHierarchies()
 		mGearRetarget.AddBoneChain(FCharacterizationStandard::Neck, FName("neck_C0_01_jnt"), FName("neck_C0_02_jnt"));
 		mGearRetarget.AddBoneChain(FCharacterizationStandard::Head, FName("neck_C0_head_jnt"), FName("neck_C0_head_jnt"));
 		// left
-		mGearRetarget.AddBoneChain(FCharacterizationStandard::LeftLeg, FName("leg_L0_thigh_jnt"), FName("foot_L0_ball_jnt"),FCharacterizationStandard::LeftFootIKGoal);
+		mGearRetarget.AddBoneChain(FCharacterizationStandard::LeftLeg, FName("leg_L0_thigh_jnt"), FName("leg_L0_foot_jnt"),FCharacterizationStandard::LeftFootIKGoal);
+		mGearRetarget.AddBoneChain(FCharacterizationStandard::LeftFoot, FName("foot_L0_ball_jnt"), FName("foot_L0_ball_jnt"));
 		mGearRetarget.AddBoneChain(FCharacterizationStandard::LeftClavicle, FName("shoulder_L0_0_jnt"), FName("shoulder_L0_0_jnt"));
 		mGearRetarget.AddBoneChain(FCharacterizationStandard::LeftArm, FName("arm_L0_upperarm_jnt"), FName("arm_L0_hand_jnt"),FCharacterizationStandard::LeftHandIKGoal);
 		mGearRetarget.AddBoneChain(FCharacterizationStandard::LeftThumb, FName("thumb_L0_0_jnt"), FName("thumb_L0_2_jnt"));
@@ -829,7 +851,8 @@ FKnownTemplateHierarchies::FKnownTemplateHierarchies()
 		mGearRetarget.AddBoneChain(FCharacterizationStandard::LeftRingMetacarpal, FName("meta_L0_2_jnt"), FName("meta_L0_2_jnt"));
 		mGearRetarget.AddBoneChain(FCharacterizationStandard::LeftPinkyMetacarpal, FName("meta_L0_3_jnt"), FName("meta_L0_3_jnt"));
 		// right
-		mGearRetarget.AddBoneChain(FCharacterizationStandard::RightLeg, FName("leg_R0_thigh_jnt"), FName("foot_R0_ball_jnt"),FCharacterizationStandard::RightFootIKGoal);
+		mGearRetarget.AddBoneChain(FCharacterizationStandard::RightLeg, FName("leg_R0_thigh_jnt"), FName("leg_R0_foot_jnt"),FCharacterizationStandard::RightFootIKGoal);
+		mGearRetarget.AddBoneChain(FCharacterizationStandard::RightFoot, FName("foot_R0_ball_jnt"), FName("foot_R0_ball_jnt"));
 		mGearRetarget.AddBoneChain(FCharacterizationStandard::RightClavicle, FName("shoulder_R0_0_jnt"), FName("shoulder_R0_0_jnt"));
 		mGearRetarget.AddBoneChain(FCharacterizationStandard::RightArm, FName("arm_R0_upperarm_jnt"), FName("arm_R0_hand_jnt"),FCharacterizationStandard::RightHandIKGoal);
 		mGearRetarget.AddBoneChain(FCharacterizationStandard::RightThumb, FName("thumb_R0_0_jnt"), FName("thumb_R0_2_jnt"));
@@ -878,7 +901,8 @@ FKnownTemplateHierarchies::FKnownTemplateHierarchies()
 		HumanIKRetarget.AddBoneChain(FCharacterizationStandard::Neck, FName("Neck"), FName("Neck"));
 		HumanIKRetarget.AddBoneChain(FCharacterizationStandard::Head, FName("Head"), FName("Head"));
 		// left
-		HumanIKRetarget.AddBoneChain(FCharacterizationStandard::LeftLeg, FName("LeftUpLeg"), FName("LeftToeBase"),FCharacterizationStandard::LeftFootIKGoal);
+		HumanIKRetarget.AddBoneChain(FCharacterizationStandard::LeftLeg, FName("LeftUpLeg"), FName("LeftFoot"),FCharacterizationStandard::LeftFootIKGoal);
+		HumanIKRetarget.AddBoneChain(FCharacterizationStandard::LeftFoot, FName("LeftToeBase"), FName("LeftToeBase"));
 		HumanIKRetarget.AddBoneChain(FCharacterizationStandard::LeftClavicle, FName("LeftShoulder"), FName("LeftShoulder"));
 		HumanIKRetarget.AddBoneChain(FCharacterizationStandard::LeftArm, FName("LeftArm"), FName("LeftHand"),FCharacterizationStandard::LeftHandIKGoal);
 		HumanIKRetarget.AddBoneChain(FCharacterizationStandard::LeftThumb, FName("LeftHandThumb1"), FName("LeftHandThumb3"));
@@ -887,7 +911,8 @@ FKnownTemplateHierarchies::FKnownTemplateHierarchies()
 		HumanIKRetarget.AddBoneChain(FCharacterizationStandard::LeftRing, FName("LeftHandRing1"), FName("LeftHandRing3"));
 		HumanIKRetarget.AddBoneChain(FCharacterizationStandard::LeftPinky, FName("LeftHandPinky1"), FName("LeftHandPinky3"));
 		// right
-		HumanIKRetarget.AddBoneChain(FCharacterizationStandard::RightLeg, FName("RightUpLeg"), FName("RightToeBase"),FCharacterizationStandard::RightFootIKGoal);
+		HumanIKRetarget.AddBoneChain(FCharacterizationStandard::RightLeg, FName("RightUpLeg"), FName("RightFoot"),FCharacterizationStandard::RightFootIKGoal);
+		HumanIKRetarget.AddBoneChain(FCharacterizationStandard::RightFoot, FName("RightToeBase"), FName("RightToeBase"));
 		HumanIKRetarget.AddBoneChain(FCharacterizationStandard::RightClavicle, FName("RightShoulder"), FName("RightShoulder"));
 		HumanIKRetarget.AddBoneChain(FCharacterizationStandard::RightArm, FName("RightArm"), FName("RightHand"),FCharacterizationStandard::RightHandIKGoal);
 		HumanIKRetarget.AddBoneChain(FCharacterizationStandard::RightThumb, FName("RightHandThumb1"), FName("RightHandThumb3"));
@@ -924,7 +949,8 @@ FKnownTemplateHierarchies::FKnownTemplateHierarchies()
 		ViconRetarget.AddBoneChain(FCharacterizationStandard::Neck, FName("Neck"), FName("Neck1"));
 		ViconRetarget.AddBoneChain(FCharacterizationStandard::Head, FName("Head"), FName("Head"));
 		// left
-		ViconRetarget.AddBoneChain(FCharacterizationStandard::LeftLeg, FName("LeftUpLeg"), FName("LeftToeBase"),FCharacterizationStandard::LeftFootIKGoal);
+		ViconRetarget.AddBoneChain(FCharacterizationStandard::LeftLeg, FName("LeftUpLeg"), FName("LeftFoot"),FCharacterizationStandard::LeftFootIKGoal);
+		ViconRetarget.AddBoneChain(FCharacterizationStandard::LeftFoot, FName("LeftToeBase"), FName("LeftToeBase"));
 		ViconRetarget.AddBoneChain(FCharacterizationStandard::LeftClavicle, FName("LeftShoulder"), FName("LeftShoulder"));
 		ViconRetarget.AddBoneChain(FCharacterizationStandard::LeftArm, FName("LeftArm"), FName("LeftHand"),FCharacterizationStandard::LeftHandIKGoal);
 		ViconRetarget.AddBoneChain(FCharacterizationStandard::LeftThumb, FName("LeftHandThumb1"), FName("LeftHandThumb3"));
@@ -935,7 +961,8 @@ FKnownTemplateHierarchies::FKnownTemplateHierarchies()
 		ViconRetarget.AddBoneChain(FCharacterizationStandard::LeftIndexMetacarpal, FName("LeftHandIndex"), FName("LeftHandIndex"));
 		ViconRetarget.AddBoneChain(FCharacterizationStandard::LeftPinkyMetacarpal, FName("LeftHandPinky"), FName("LeftHandPinky"));
 		// right
-		ViconRetarget.AddBoneChain(FCharacterizationStandard::RightLeg, FName("RightUpLeg"), FName("RightToeBase"),FCharacterizationStandard::RightFootIKGoal);
+		ViconRetarget.AddBoneChain(FCharacterizationStandard::RightLeg, FName("RightUpLeg"), FName("RightFoot"),FCharacterizationStandard::RightFootIKGoal);
+		ViconRetarget.AddBoneChain(FCharacterizationStandard::RightFoot, FName("RightToeBase"), FName("RightToeBase"));
 		ViconRetarget.AddBoneChain(FCharacterizationStandard::RightClavicle, FName("RightShoulder"), FName("RightShoulder"));
 		ViconRetarget.AddBoneChain(FCharacterizationStandard::RightArm, FName("RightArm"), FName("RightHand"),FCharacterizationStandard::RightHandIKGoal);
 		ViconRetarget.AddBoneChain(FCharacterizationStandard::RightThumb, FName("RightHandThumb1"), FName("RightHandThumb3"));
@@ -975,11 +1002,13 @@ FKnownTemplateHierarchies::FKnownTemplateHierarchies()
 		OptitrakRetarget.AddBoneChain(FCharacterizationStandard::Neck, FName("neck_01"), FName("neck_01"));
 		OptitrakRetarget.AddBoneChain(FCharacterizationStandard::Head, FName("head"), FName("head"));
 		// left
-		OptitrakRetarget.AddBoneChain(FCharacterizationStandard::LeftLeg, FName("thigh_l"), FName("ball_l"),FCharacterizationStandard::LeftFootIKGoal);
+		OptitrakRetarget.AddBoneChain(FCharacterizationStandard::LeftLeg, FName("thigh_l"), FName("foot_l"),FCharacterizationStandard::LeftFootIKGoal);
+		OptitrakRetarget.AddBoneChain(FCharacterizationStandard::LeftFoot, FName("ball_l"), FName("ball_l"));
 		OptitrakRetarget.AddBoneChain(FCharacterizationStandard::LeftClavicle, FName("clavicle_l"), FName("clavicle_l"));
 		OptitrakRetarget.AddBoneChain(FCharacterizationStandard::LeftArm, FName("upperarm_l"), FName("hand_l"),FCharacterizationStandard::LeftHandIKGoal);
 		// right
-		OptitrakRetarget.AddBoneChain(FCharacterizationStandard::RightLeg, FName("thigh_r"), FName("ball_r"),FCharacterizationStandard::RightFootIKGoal);
+		OptitrakRetarget.AddBoneChain(FCharacterizationStandard::RightLeg, FName("thigh_r"), FName("foot_r"),FCharacterizationStandard::RightFootIKGoal);
+		OptitrakRetarget.AddBoneChain(FCharacterizationStandard::RightFoot, FName("ball_r"), FName("ball_r"));
 		OptitrakRetarget.AddBoneChain(FCharacterizationStandard::RightClavicle, FName("clavicle_r"), FName("clavicle_r"));
 		OptitrakRetarget.AddBoneChain(FCharacterizationStandard::RightArm, FName("upperarm_r"), FName("hand_r"),FCharacterizationStandard::RightHandIKGoal);
 		// bone settings for IK
@@ -1009,11 +1038,13 @@ FKnownTemplateHierarchies::FKnownTemplateHierarchies()
 		MocopiRetarget.AddBoneChain(FCharacterizationStandard::Neck, FName("neck_1"), FName("neck_2"));
 		MocopiRetarget.AddBoneChain(FCharacterizationStandard::Head, FName("head"), FName("head"));
 		// left
-		MocopiRetarget.AddBoneChain(FCharacterizationStandard::LeftLeg, FName("l_up_leg"), FName("l_toes"),FCharacterizationStandard::LeftFootIKGoal);
+		MocopiRetarget.AddBoneChain(FCharacterizationStandard::LeftLeg, FName("l_up_leg"), FName("l_foot"),FCharacterizationStandard::LeftFootIKGoal);
+		MocopiRetarget.AddBoneChain(FCharacterizationStandard::LeftFoot, FName("l_toes"), FName("l_toes"));
 		MocopiRetarget.AddBoneChain(FCharacterizationStandard::LeftClavicle, FName("l_shoulder"), FName("l_shoulder"));
 		MocopiRetarget.AddBoneChain(FCharacterizationStandard::LeftArm, FName("l_up_arm"), FName("l_hand"),FCharacterizationStandard::LeftHandIKGoal);
 		// right
-		MocopiRetarget.AddBoneChain(FCharacterizationStandard::RightLeg, FName("r_up_leg"), FName("r_toes"),FCharacterizationStandard::RightFootIKGoal);
+		MocopiRetarget.AddBoneChain(FCharacterizationStandard::RightLeg, FName("r_up_leg"), FName("r_foot"),FCharacterizationStandard::RightFootIKGoal);
+		MocopiRetarget.AddBoneChain(FCharacterizationStandard::RightFoot, FName("r_toes"), FName("r_toes"));
 		MocopiRetarget.AddBoneChain(FCharacterizationStandard::RightClavicle, FName("r_shoulder"), FName("r_shoulder"));
 		MocopiRetarget.AddBoneChain(FCharacterizationStandard::RightArm, FName("r_up_arm"), FName("r_hand"),FCharacterizationStandard::RightHandIKGoal);
 		// bone settings for IK
@@ -1046,7 +1077,8 @@ FKnownTemplateHierarchies::FKnownTemplateHierarchies()
 		AdvSkelRetarget.AddBoneChain(FCharacterizationStandard::Neck, FName("Neck_M"), FName("NeckPart2_M"));
 		AdvSkelRetarget.AddBoneChain(FCharacterizationStandard::Head, FName("Head_M"), FName("Head_M"));
 		// left
-		AdvSkelRetarget.AddBoneChain(FCharacterizationStandard::LeftLeg, FName("Hip_L"), FName("Toes_L"),FCharacterizationStandard::LeftFootIKGoal);
+		AdvSkelRetarget.AddBoneChain(FCharacterizationStandard::LeftLeg, FName("Hip_L"), FName("Ankle_L"),FCharacterizationStandard::LeftFootIKGoal);
+		AdvSkelRetarget.AddBoneChain(FCharacterizationStandard::LeftFoot, FName("Toes_L"), FName("Toes_L"));
 		AdvSkelRetarget.AddBoneChain(FCharacterizationStandard::LeftClavicle, FName("Scapula_L"), FName("Scapula_L"));
 		AdvSkelRetarget.AddBoneChain(FCharacterizationStandard::LeftArm, FName("Shoulder_L"), FName("Wrist_L"),FCharacterizationStandard::LeftHandIKGoal);
 		AdvSkelRetarget.AddBoneChain(FCharacterizationStandard::LeftThumb, FName("ThumbFinger2_L"), FName("ThumbFinger4_L"));
@@ -1056,7 +1088,8 @@ FKnownTemplateHierarchies::FKnownTemplateHierarchies()
 		AdvSkelRetarget.AddBoneChain(FCharacterizationStandard::LeftPinky, FName("PinkyFinger1_L"), FName("PinkyFinger3_L"));
 		AdvSkelRetarget.AddBoneChain(FCharacterizationStandard::LeftPinkyMetacarpal, FName("Cup_L"), FName("Cup_L"));
 		// right
-		AdvSkelRetarget.AddBoneChain(FCharacterizationStandard::RightLeg, FName("Hip_R"), FName("Toes_R"),FCharacterizationStandard::RightFootIKGoal);
+		AdvSkelRetarget.AddBoneChain(FCharacterizationStandard::RightLeg, FName("Hip_R"), FName("Ankle_R"),FCharacterizationStandard::RightFootIKGoal);
+		AdvSkelRetarget.AddBoneChain(FCharacterizationStandard::RightFoot, FName("Toes_R"), FName("Toes_R"));
 		AdvSkelRetarget.AddBoneChain(FCharacterizationStandard::RightClavicle, FName("Scapula_R"), FName("Scapula_R"));
 		AdvSkelRetarget.AddBoneChain(FCharacterizationStandard::RightArm, FName("Shoulder_R"), FName("Wrist_R"),FCharacterizationStandard::RightHandIKGoal);
 		AdvSkelRetarget.AddBoneChain(FCharacterizationStandard::RightThumb, FName("ThumbFinger2_R"), FName("ThumbFinger4_R"));
@@ -1105,7 +1138,8 @@ FKnownTemplateHierarchies::FKnownTemplateHierarchies()
 		MoveOneRetarget.AddBoneChain(FCharacterizationStandard::Neck, FName("Neck"), FName("Neck"));
 		MoveOneRetarget.AddBoneChain(FCharacterizationStandard::Head, FName("Head"), FName("Head"));
 		// left
-		MoveOneRetarget.AddBoneChain(FCharacterizationStandard::LeftLeg, FName("LeftUpLeg"), FName("LeftToeBase"),FCharacterizationStandard::LeftFootIKGoal);
+		MoveOneRetarget.AddBoneChain(FCharacterizationStandard::LeftLeg, FName("LeftUpLeg"), FName("LeftFoot"),FCharacterizationStandard::LeftFootIKGoal);
+		MoveOneRetarget.AddBoneChain(FCharacterizationStandard::LeftFoot, FName("LeftToeBase"), FName("LeftToeBase"));
 		MoveOneRetarget.AddBoneChain(FCharacterizationStandard::LeftClavicle, FName("LeftShoulder"), FName("LeftShoulder"));
 		MoveOneRetarget.AddBoneChain(FCharacterizationStandard::LeftArm, FName("LeftArm"), FName("LeftHand"),FCharacterizationStandard::LeftHandIKGoal);
 		MoveOneRetarget.AddBoneChain(FCharacterizationStandard::LeftThumb, FName("LeftHandThumb1"), FName("LeftHandThumb3"));
@@ -1118,7 +1152,8 @@ FKnownTemplateHierarchies::FKnownTemplateHierarchies()
 		MoveOneRetarget.AddBoneChain(FCharacterizationStandard::LeftRingMetacarpal, FName("palm_03_L"), FName("palm_03_L"));
 		MoveOneRetarget.AddBoneChain(FCharacterizationStandard::LeftPinkyMetacarpal, FName("palm_04_L"), FName("palm_04_L"));
 		// right
-		MoveOneRetarget.AddBoneChain(FCharacterizationStandard::RightLeg, FName("RightUpLeg"), FName("RightToeBase"),FCharacterizationStandard::RightFootIKGoal);
+		MoveOneRetarget.AddBoneChain(FCharacterizationStandard::RightLeg, FName("RightUpLeg"), FName("RightFoot"),FCharacterizationStandard::RightFootIKGoal);
+		MoveOneRetarget.AddBoneChain(FCharacterizationStandard::RightFoot, FName("RightToeBase"), FName("RightToeBase"));
 		MoveOneRetarget.AddBoneChain(FCharacterizationStandard::RightClavicle, FName("RightShoulder"), FName("RightShoulder"));
 		MoveOneRetarget.AddBoneChain(FCharacterizationStandard::RightArm, FName("RightArm"), FName("RightHand"),FCharacterizationStandard::RightHandIKGoal);
 		MoveOneRetarget.AddBoneChain(FCharacterizationStandard::RightThumb, FName("RightHandThumb1"), FName("RightHandThumb3"));
@@ -1159,7 +1194,8 @@ FKnownTemplateHierarchies::FKnownTemplateHierarchies()
 		QualisysRetarget.AddBoneChain(FCharacterizationStandard::Neck, FName("Neck"), FName("Neck"));
 		QualisysRetarget.AddBoneChain(FCharacterizationStandard::Head, FName("Head"), FName("Head"));
 		// left
-		QualisysRetarget.AddBoneChain(FCharacterizationStandard::LeftLeg, FName("LeftUpLeg"), FName("LeftToeBase"),FCharacterizationStandard::LeftFootIKGoal);
+		QualisysRetarget.AddBoneChain(FCharacterizationStandard::LeftLeg, FName("LeftUpLeg"), FName("LeftFoot"),FCharacterizationStandard::LeftFootIKGoal);
+		QualisysRetarget.AddBoneChain(FCharacterizationStandard::LeftFoot, FName("LeftToeBase"), FName("LeftToeBase"));
 		QualisysRetarget.AddBoneChain(FCharacterizationStandard::LeftClavicle, FName("LeftShoulder"), FName("LeftShoulder"));
 		QualisysRetarget.AddBoneChain(FCharacterizationStandard::LeftArm, FName("LeftArm"), FName("LeftHand"),FCharacterizationStandard::LeftHandIKGoal);
 		QualisysRetarget.AddBoneChain(FCharacterizationStandard::LeftThumb, FName("LeftHandThumb1"), FName("LeftHandThumb3"));
@@ -1173,7 +1209,8 @@ FKnownTemplateHierarchies::FKnownTemplateHierarchies()
 		QualisysRetarget.AddBoneChain(FCharacterizationStandard::LeftRingMetacarpal, FName("LeftInHandRing"), FName("LeftInHandRing"));
 		QualisysRetarget.AddBoneChain(FCharacterizationStandard::LeftPinkyMetacarpal, FName("LeftInHandPinky"), FName("LeftInHandPinky"));
 		// right
-		QualisysRetarget.AddBoneChain(FCharacterizationStandard::RightLeg, FName("RightUpLeg"), FName("RightToeBase"),FCharacterizationStandard::RightFootIKGoal);
+		QualisysRetarget.AddBoneChain(FCharacterizationStandard::RightLeg, FName("RightUpLeg"), FName("RightFoot"),FCharacterizationStandard::RightFootIKGoal);
+		QualisysRetarget.AddBoneChain(FCharacterizationStandard::RightFoot, FName("RightToeBase"), FName("RightToeBase"));
 		QualisysRetarget.AddBoneChain(FCharacterizationStandard::RightClavicle, FName("RightShoulder"), FName("RightShoulder"));
 		QualisysRetarget.AddBoneChain(FCharacterizationStandard::RightArm, FName("RightArm"), FName("RightHand"),FCharacterizationStandard::RightHandIKGoal);
 		QualisysRetarget.AddBoneChain(FCharacterizationStandard::RightThumb, FName("RightHandThumb1"), FName("RightHandThumb3"));
@@ -1216,7 +1253,8 @@ FKnownTemplateHierarchies::FKnownTemplateHierarchies()
 		FNHumanRetarget.AddBoneChain(FCharacterizationStandard::Neck, FName("neck_01"), FName("neck_02"));
 		FNHumanRetarget.AddBoneChain(FCharacterizationStandard::Head, FName("head"), FName("head"));
 		// left
-		FNHumanRetarget.AddBoneChain(FCharacterizationStandard::LeftLeg, FName("thigh_l"), FName("ball_l"),FCharacterizationStandard::LeftFootIKGoal);
+		FNHumanRetarget.AddBoneChain(FCharacterizationStandard::LeftLeg, FName("thigh_l"), FName("foot_l"),FCharacterizationStandard::LeftFootIKGoal);
+		FNHumanRetarget.AddBoneChain(FCharacterizationStandard::LeftFoot, FName("ball_l"), FName("ball_l"));
 		FNHumanRetarget.AddBoneChain(FCharacterizationStandard::LeftClavicle, FName("clavicle_l"), FName("clavicle_l"));
 		FNHumanRetarget.AddBoneChain(FCharacterizationStandard::LeftArm, FName("upperarm_l"), FName("hand_l"),FCharacterizationStandard::LeftHandIKGoal);
 		FNHumanRetarget.AddBoneChain(FCharacterizationStandard::LeftThumb, FName("thumb_01_l"), FName("thumb_03_l"));
@@ -1229,7 +1267,8 @@ FKnownTemplateHierarchies::FKnownTemplateHierarchies()
 		FNHumanRetarget.AddBoneChain(FCharacterizationStandard::LeftRingMetacarpal, FName("ring_metacarpal_l"), FName("ring_metacarpal_l"));
 		FNHumanRetarget.AddBoneChain(FCharacterizationStandard::LeftPinkyMetacarpal, FName("pinky_metacarpal_l"), FName("pinky_metacarpal_l"));
 		// right
-		FNHumanRetarget.AddBoneChain(FCharacterizationStandard::RightLeg, FName("thigh_r"), FName("ball_r"), FCharacterizationStandard::RightFootIKGoal);
+		FNHumanRetarget.AddBoneChain(FCharacterizationStandard::RightLeg, FName("thigh_r"), FName("foot_r"), FCharacterizationStandard::RightFootIKGoal);
+		FNHumanRetarget.AddBoneChain(FCharacterizationStandard::RightFoot, FName("ball_r"), FName("ball_r"));
 		FNHumanRetarget.AddBoneChain(FCharacterizationStandard::RightClavicle, FName("clavicle_r"), FName("clavicle_r"));
 		FNHumanRetarget.AddBoneChain(FCharacterizationStandard::RightArm, FName("upperarm_r"), FName("hand_r"),FCharacterizationStandard::RightHandIKGoal);
 		FNHumanRetarget.AddBoneChain(FCharacterizationStandard::RightThumb, FName("thumb_01_r"), FName("thumb_03_r"));
@@ -1416,7 +1455,8 @@ FKnownTemplateHierarchies::FKnownTemplateHierarchies()
 		FNBipedRetarget.AddBoneChain(FCharacterizationStandard::Neck, FName("C_BipedHeadNeck_A_NeckBase_Jnt"), FName("C_BipedHeadNeck_A_NeckMid_Jnt"));
 		FNBipedRetarget.AddBoneChain(FCharacterizationStandard::Head, FName("C_BipedHeadNeck_A_Head_Jnt"), FName("C_BipedHeadNeck_A_Head_Jnt"));
 		// left
-		FNBipedRetarget.AddBoneChain(FCharacterizationStandard::LeftLeg, FName("L_BipedLeg_A_Thigh_Jnt"), FName("L_BipedLeg_A_Ball_Jnt"),FCharacterizationStandard::LeftFootIKGoal);
+		FNBipedRetarget.AddBoneChain(FCharacterizationStandard::LeftLeg, FName("L_BipedLeg_A_Thigh_Jnt"), FName("L_BipedLeg_A_Ankle_Jnt"),FCharacterizationStandard::LeftFootIKGoal);
+		FNBipedRetarget.AddBoneChain(FCharacterizationStandard::LeftFoot, FName("L_BipedLeg_A_Ball_Jnt"), FName("L_BipedLeg_A_Ball_Jnt"));
 		FNBipedRetarget.AddBoneChain(FCharacterizationStandard::LeftClavicle, FName("L_BipedArm_A_Clavicle_Jnt"), FName("L_BipedArm_A_Clavicle_Jnt"));
 		FNBipedRetarget.AddBoneChain(FCharacterizationStandard::LeftArm, FName("L_BipedArm_A_Shoulder_Jnt"), FName("L_BipedArm_A_Wrist_Jnt"),FCharacterizationStandard::LeftHandIKGoal);
 		FNBipedRetarget.AddBoneChain(FCharacterizationStandard::LeftThumb, FName("L_BipedHand_A_ThumbBase_Jnt"), FName("L_BipedHand_A_ThumbTip_Jnt"));
@@ -1429,7 +1469,8 @@ FKnownTemplateHierarchies::FKnownTemplateHierarchies()
 		FNBipedRetarget.AddBoneChain(FCharacterizationStandard::LeftRingMetacarpal, FName("L_BipedHand_A_RingCarpal_Jnt"), FName("L_BipedHand_A_RingCarpal_Jnt"));
 		FNBipedRetarget.AddBoneChain(FCharacterizationStandard::LeftPinkyMetacarpal, FName("L_BipedHand_A_PinkyCarpal_Jnt"), FName("L_BipedHand_A_PinkyCarpal_Jnt"));
 		// right
-		FNBipedRetarget.AddBoneChain(FCharacterizationStandard::RightLeg, FName("R_BipedLeg_A_Thigh_Jnt"), FName("R_BipedLeg_A_Ball_Jnt"),FCharacterizationStandard::RightFootIKGoal);
+		FNBipedRetarget.AddBoneChain(FCharacterizationStandard::RightLeg, FName("R_BipedLeg_A_Thigh_Jnt"), FName("R_BipedLeg_A_Ankle_Jnt"),FCharacterizationStandard::RightFootIKGoal);
+		FNBipedRetarget.AddBoneChain(FCharacterizationStandard::RightFoot, FName("R_BipedLeg_A_Ball_Jnt"), FName("R_BipedLeg_A_Ball_Jnt"));
 		FNBipedRetarget.AddBoneChain(FCharacterizationStandard::RightClavicle, FName("R_BipedArm_A_Clavicle_Jnt"), FName("R_BipedArm_A_Clavicle_Jnt"));
 		FNBipedRetarget.AddBoneChain(FCharacterizationStandard::RightArm, FName("R_BipedArm_A_Shoulder_Jnt"), FName("R_BipedArm_A_Wrist_Jnt"),FCharacterizationStandard::RightHandIKGoal);
 		FNBipedRetarget.AddBoneChain(FCharacterizationStandard::RightThumb, FName("R_BipedHand_A_ThumbBase_Jnt"), FName("R_BipedHand_A_ThumbTip_Jnt"));
@@ -1482,7 +1523,8 @@ FKnownTemplateHierarchies::FKnownTemplateHierarchies()
 		RokokoRetarget.AddBoneChain(FCharacterizationStandard::Neck, FName("Neck"), FName("Neck"));
 		RokokoRetarget.AddBoneChain(FCharacterizationStandard::Head, FName("Head"), FName("Head"));
 		// left
-		RokokoRetarget.AddBoneChain(FCharacterizationStandard::LeftLeg, FName("LeftThigh"), FName("LeftToe"),FCharacterizationStandard::LeftFootIKGoal);
+		RokokoRetarget.AddBoneChain(FCharacterizationStandard::LeftLeg, FName("LeftThigh"), FName("LeftFoot"),FCharacterizationStandard::LeftFootIKGoal);
+	RokokoRetarget.AddBoneChain(FCharacterizationStandard::LeftFoot, FName("LeftToe"), FName("LeftToe"));
 		RokokoRetarget.AddBoneChain(FCharacterizationStandard::LeftClavicle, FName("LeftShoulder"), FName("LeftShoulder"));
 		RokokoRetarget.AddBoneChain(FCharacterizationStandard::LeftArm, FName("LeftArm"), FName("LeftHand"),FCharacterizationStandard::LeftHandIKGoal);
 		RokokoRetarget.AddBoneChain(FCharacterizationStandard::LeftThumb, FName("LeftFinger1Metacarpal"), FName("LeftFinger1Distal"));
@@ -1495,7 +1537,8 @@ FKnownTemplateHierarchies::FKnownTemplateHierarchies()
 		RokokoRetarget.AddBoneChain(FCharacterizationStandard::LeftRingMetacarpal, FName("LeftFinger4Metacarpal"), FName("LeftFinger4Metacarpal"));
 		RokokoRetarget.AddBoneChain(FCharacterizationStandard::LeftPinkyMetacarpal, FName("LeftFinger5Metacarpal"), FName("LeftFinger5Metacarpal"));
 		// right
-		RokokoRetarget.AddBoneChain(FCharacterizationStandard::RightLeg, FName("RightThigh"), FName("RightToe"),FCharacterizationStandard::RightFootIKGoal);
+		RokokoRetarget.AddBoneChain(FCharacterizationStandard::RightLeg, FName("RightThigh"), FName("RightFoot"),FCharacterizationStandard::RightFootIKGoal);
+	RokokoRetarget.AddBoneChain(FCharacterizationStandard::RightFoot, FName("RightToe"), FName("RightToe"));
 		RokokoRetarget.AddBoneChain(FCharacterizationStandard::RightClavicle, FName("RightShoulder"), FName("RightShoulder"));
 		RokokoRetarget.AddBoneChain(FCharacterizationStandard::RightArm, FName("RightArm"), FName("RightHand"),FCharacterizationStandard::RightHandIKGoal);
 		RokokoRetarget.AddBoneChain(FCharacterizationStandard::RightThumb, FName("RightFinger1Metacarpal"), FName("RightFinger1Distal"));
@@ -1521,6 +1564,117 @@ FKnownTemplateHierarchies::FKnownTemplateHierarchies()
 		Rokoko.AutoRetargetDefinition.BonesToExcludeFromAutoPose.Add("LeftFoot");
 		Rokoko.AutoRetargetDefinition.BonesToExcludeFromAutoPose.Add("RightFoot");
 	}();
+
+	// Meshcapade
+	[this]()
+	{
+		static FName MeshcapadeName = "Meshcapade";
+		static TArray<FName> MeshcapadeBones = {"root", "pelvis", "spine_01", "spine_02", "spine_03", "clavicle_l", "upperarm_l", "lowerarm_l", "hand_l", "index_01_l", "index_02_l", "index_03_l", "middle_01_l", "middle_02_l", "middle_03_l", "pinky_01_l", "pinky_02_l", "pinky_03_l", "ring_01_l", "ring_02_l", "ring_03_l", "thumb_01_l", "thumb_02_l", "thumb_03_l", "clavicle_r", "upperarm_r", "lowerarm_r", "hand_r", "index_01_r", "index_02_r", "index_03_r", "middle_01_r", "middle_02_r", "middle_03_r", "pinky_01_r", "pinky_02_r", "pinky_03_r", "ring_01_r", "ring_02_r", "ring_03_r", "thumb_01_r", "thumb_02_r", "thumb_03_r", "neck_01", "head", "thigh_l", "calf_l", "foot_l", "ball_l", "thigh_r", "calf_r", "foot_r", "ball_r"};
+		static TArray<int32> MeshcapadeParentIndices = {-1, 0, 1, 2, 3, 4, 5, 6, 7, 8, 9, 10, 8, 12, 13, 8, 15, 16, 8, 18, 19, 8, 21, 22, 4, 24, 25, 26, 27, 28, 29, 27, 31, 32, 27, 34, 35, 27, 37, 38, 27, 40, 41, 4, 43, 1, 45, 46, 47, 1, 49, 50, 51};
+		FTemplateHierarchy& Meshcapade = AddTemplateHierarchy(MeshcapadeName, MeshcapadeBones, MeshcapadeParentIndices);
+		FRetargetDefinition& MeshcapadeRetarget = Meshcapade.AutoRetargetDefinition.RetargetDefinition;
+		// core
+		MeshcapadeRetarget.PelvisBone = FName("pelvis");
+		MeshcapadeRetarget.RootBone = FName("root");
+		MeshcapadeRetarget.AddBoneChain(FCharacterizationStandard::Spine, FName("spine_01"), FName("spine_02"));
+		MeshcapadeRetarget.AddBoneChain(FCharacterizationStandard::Neck, FName("neck_01"), FName("neck_01"));
+		MeshcapadeRetarget.AddBoneChain(FCharacterizationStandard::Head, FName("head"), FName("head"));
+		// left
+		MeshcapadeRetarget.AddBoneChain(FCharacterizationStandard::LeftLeg, FName("thigh_l"), FName("foot_l"),FCharacterizationStandard::LeftFootIKGoal);
+		MeshcapadeRetarget.AddBoneChain(FCharacterizationStandard::LeftFoot, FName("ball_l"), FName("ball_l"));
+		MeshcapadeRetarget.AddBoneChain(FCharacterizationStandard::LeftClavicle, FName("clavicle_l"), FName("clavicle_l"));
+		MeshcapadeRetarget.AddBoneChain(FCharacterizationStandard::LeftArm, FName("upperarm_l"), FName("hand_l"),FCharacterizationStandard::LeftHandIKGoal);
+		MeshcapadeRetarget.AddBoneChain(FCharacterizationStandard::LeftThumb, FName("thumb_01_l"), FName("thumb_03_l"));
+		MeshcapadeRetarget.AddBoneChain(FCharacterizationStandard::LeftIndex, FName("index_01_l"), FName("index_03_l"));
+		MeshcapadeRetarget.AddBoneChain(FCharacterizationStandard::LeftMiddle, FName("middle_01_l"), FName("middle_03_l"));
+		MeshcapadeRetarget.AddBoneChain(FCharacterizationStandard::LeftRing, FName("ring_01_l"), FName("ring_03_l"));
+		MeshcapadeRetarget.AddBoneChain(FCharacterizationStandard::LeftPinky, FName("pinky_01_l"), FName("pinky_03_l"));
+		// right
+		MeshcapadeRetarget.AddBoneChain(FCharacterizationStandard::RightLeg, FName("thigh_r"), FName("foot_r"), FCharacterizationStandard::RightFootIKGoal);
+		MeshcapadeRetarget.AddBoneChain(FCharacterizationStandard::RightFoot, FName("ball_r"), FName("ball_r"));
+		MeshcapadeRetarget.AddBoneChain(FCharacterizationStandard::RightClavicle, FName("clavicle_r"), FName("clavicle_r"));
+		MeshcapadeRetarget.AddBoneChain(FCharacterizationStandard::RightArm, FName("upperarm_r"), FName("hand_r"),FCharacterizationStandard::RightHandIKGoal);
+		MeshcapadeRetarget.AddBoneChain(FCharacterizationStandard::RightThumb, FName("thumb_01_r"), FName("thumb_03_r"));
+		MeshcapadeRetarget.AddBoneChain(FCharacterizationStandard::RightIndex, FName("index_01_r"), FName("index_03_r"));
+		MeshcapadeRetarget.AddBoneChain(FCharacterizationStandard::RightMiddle, FName("middle_01_r"), FName("middle_03_r"));
+		MeshcapadeRetarget.AddBoneChain(FCharacterizationStandard::RightRing, FName("ring_01_r"), FName("ring_03_r"));
+		MeshcapadeRetarget.AddBoneChain(FCharacterizationStandard::RightPinky, FName("pinky_01_r"), FName("pinky_03_r"));
+		// bone settings for IK
+		Meshcapade.AutoRetargetDefinition.BoneSettingsForIK.SetRotationStiffness(FName("pelvis"), 0.7f);
+		Meshcapade.AutoRetargetDefinition.BoneSettingsForIK.SetPositionStiffness(FName("pelvis"), 0.7f);
+		Meshcapade.AutoRetargetDefinition.BoneSettingsForIK.SetRotationStiffness(FName("spine_02"), 0.95f);
+		Meshcapade.AutoRetargetDefinition.BoneSettingsForIK.SetRotationStiffness(FName("spine_03"), 0.9f);
+		Meshcapade.AutoRetargetDefinition.BoneSettingsForIK.SetRotationStiffness(FName("clavicle_l"), 0.95f);
+		Meshcapade.AutoRetargetDefinition.BoneSettingsForIK.SetRotationStiffness(FName("clavicle_r"), 0.95f);
+		Meshcapade.AutoRetargetDefinition.BoneSettingsForIK.SetRotationStiffness(FName("foot_l"), 0.85f);
+		Meshcapade.AutoRetargetDefinition.BoneSettingsForIK.SetRotationStiffness(FName("foot_r"), 0.85f);
+		Meshcapade.AutoRetargetDefinition.BoneSettingsForIK.SetPreferredAxis(FName("calf_l"), EPreferredAxis::PositiveX);
+		Meshcapade.AutoRetargetDefinition.BoneSettingsForIK.SetPreferredAxis(FName("calf_r"), EPreferredAxis::PositiveX);
+		Meshcapade.AutoRetargetDefinition.BoneSettingsForIK.SetPreferredAxis(FName("lowerarm_l"), EPreferredAxis::PositiveY);
+		Meshcapade.AutoRetargetDefinition.BoneSettingsForIK.SetPreferredAxis(FName("lowerarm_r"), EPreferredAxis::PositiveY);
+		Meshcapade.AutoRetargetDefinition.BoneSettingsForIK.SetExcluded(FName("spine_03"), true);
+		// exclude feet from auto-pose
+		Meshcapade.AutoRetargetDefinition.BonesToExcludeFromAutoPose.Add("foot_l");
+		Meshcapade.AutoRetargetDefinition.BonesToExcludeFromAutoPose.Add("foot_r");
+	}();
+
+	// Blender Rigify
+	[this]()
+	{
+		static FName RigifyName = "Blender Rigify";
+		static TArray<FName> RigifyBones = {"metarig", "spine", "pelvis_L", "pelvis_R", "spine_001", "spine_002", "spine_003", "breast_L", "breast_R", "shoulder_L", "upper_arm_L", "forearm_L", "hand_L", "palm_01_L", "f_index_01_L", "f_index_02_L", "f_index_03_L", "thumb_01_L", "thumb_02_L", "thumb_03_L", "palm_02_L", "f_middle_01_L", "f_middle_02_L", "f_middle_03_L", "palm_03_L", "f_ring_01_L", "f_ring_02_L", "f_ring_03_L", "palm_04_L", "f_pinky_01_L", "f_pinky_02_L", "f_pinky_03_L", "shoulder_R", "upper_arm_R", "forearm_R", "hand_R", "palm_01_R", "f_index_01_R", "f_index_02_R", "f_index_03_R", "thumb_01_R", "thumb_02_R", "thumb_03_R", "palm_02_R", "f_middle_01_R", "f_middle_02_R", "f_middle_03_R", "palm_03_R", "f_ring_01_R", "f_ring_02_R", "f_ring_03_R", "palm_04_R", "f_pinky_01_R", "f_pinky_02_R", "f_pinky_03_R", "spine_004", "spine_005", "spine_006", "face", "thigh_L", "shin_L", "foot_L", "heel_02_L", "toe_L", "thigh_R", "shin_R", "foot_R", "heel_02_R", "toe_R"};
+		static TArray<int32> RigifyParentIndices = {-1, 0, 1, 1, 1, 4, 5, 6, 6, 6, 9, 10, 11, 12, 13, 14, 15, 13, 17, 18, 12, 20, 21, 22, 12, 24, 25, 26, 12, 28, 29, 30, 6, 32, 33, 34, 35, 36, 37, 38, 36, 40, 41, 35, 43, 44, 45, 35, 47, 48, 49, 35, 51, 52, 53, 6, 55, 56, 57, 1, 59, 60, 61, 61, 1, 64, 65, 66, 66};
+		FTemplateHierarchy& Rigify = AddTemplateHierarchy(RigifyName, RigifyBones, RigifyParentIndices);
+		FRetargetDefinition& RigifyRetarget = Rigify.AutoRetargetDefinition.RetargetDefinition;
+		// core
+		RigifyRetarget.PelvisBone = FName("spine");
+		RigifyRetarget.RootBone = FName("metarig");
+		RigifyRetarget.AddBoneChain(FCharacterizationStandard::Spine, FName("spine_001"), FName("spine_003"));
+		RigifyRetarget.AddBoneChain(FCharacterizationStandard::Neck, FName("spine_004"), FName("spine_005"));
+		RigifyRetarget.AddBoneChain(FCharacterizationStandard::Head, FName("spine_006"), FName("spine_006"));
+		// left
+		RigifyRetarget.AddBoneChain(FCharacterizationStandard::LeftLeg, FName("thigh_L"), FName("foot_L"),FCharacterizationStandard::LeftFootIKGoal);
+		RigifyRetarget.AddBoneChain(FCharacterizationStandard::LeftFoot, FName("toe_L"), FName("toe_L"));
+		RigifyRetarget.AddBoneChain(FCharacterizationStandard::LeftClavicle, FName("shoulder_L"), FName("shoulder_L"));
+		RigifyRetarget.AddBoneChain(FCharacterizationStandard::LeftArm, FName("upper_arm_L"), FName("hand_L"),FCharacterizationStandard::LeftHandIKGoal);
+		RigifyRetarget.AddBoneChain(FCharacterizationStandard::LeftThumb, FName("thumb_01_L"), FName("thumb_03_L"));
+		RigifyRetarget.AddBoneChain(FCharacterizationStandard::LeftIndex, FName("f_index_01_L"), FName("f_index_03_L"));
+		RigifyRetarget.AddBoneChain(FCharacterizationStandard::LeftMiddle, FName("f_middle_01_L"), FName("f_middle_03_L"));
+		RigifyRetarget.AddBoneChain(FCharacterizationStandard::LeftRing, FName("f_ring_01_L"), FName("f_ring_03_L"));
+		RigifyRetarget.AddBoneChain(FCharacterizationStandard::LeftPinky, FName("f_pinky_01_L"), FName("f_pinky_03_L"));
+		RigifyRetarget.AddBoneChain(FCharacterizationStandard::LeftIndexMetacarpal, FName("palm_01_L"), FName("palm_01_L"));
+		RigifyRetarget.AddBoneChain(FCharacterizationStandard::LeftMiddleMetacarpal, FName("palm_02_L"), FName("palm_02_L"));
+		RigifyRetarget.AddBoneChain(FCharacterizationStandard::LeftRingMetacarpal, FName("palm_03_L"), FName("palm_03_L"));
+		RigifyRetarget.AddBoneChain(FCharacterizationStandard::LeftPinkyMetacarpal, FName("palm_04_L"), FName("palm_04_L"));
+		// right
+		RigifyRetarget.AddBoneChain(FCharacterizationStandard::RightLeg, FName("thigh_R"), FName("foot_R"),FCharacterizationStandard::RightFootIKGoal);
+		RigifyRetarget.AddBoneChain(FCharacterizationStandard::RightFoot, FName("toe_R"), FName("toe_R"));
+		RigifyRetarget.AddBoneChain(FCharacterizationStandard::RightClavicle, FName("shoulder_R"), FName("shoulder_R"));
+		RigifyRetarget.AddBoneChain(FCharacterizationStandard::RightArm, FName("upper_arm_R"), FName("hand_R"),FCharacterizationStandard::RightHandIKGoal);
+		RigifyRetarget.AddBoneChain(FCharacterizationStandard::RightThumb, FName("thumb_01_R"), FName("thumb_03_R"));
+		RigifyRetarget.AddBoneChain(FCharacterizationStandard::RightIndex, FName("f_index_01_R"), FName("f_index_03_R"));
+		RigifyRetarget.AddBoneChain(FCharacterizationStandard::RightMiddle, FName("f_middle_01_R"), FName("f_middle_03_R"));
+		RigifyRetarget.AddBoneChain(FCharacterizationStandard::RightRing, FName("f_ring_01_R"), FName("f_ring_03_R"));
+		RigifyRetarget.AddBoneChain(FCharacterizationStandard::RightPinky, FName("f_pinky_01_R"), FName("f_pinky_03_R"));
+		RigifyRetarget.AddBoneChain(FCharacterizationStandard::RightIndexMetacarpal, FName("palm_01_R"), FName("palm_01_R"));
+		RigifyRetarget.AddBoneChain(FCharacterizationStandard::RightMiddleMetacarpal, FName("palm_02_R"), FName("palm_02_R"));
+		RigifyRetarget.AddBoneChain(FCharacterizationStandard::RightRingMetacarpal, FName("palm_03_R"), FName("palm_03_R"));
+		RigifyRetarget.AddBoneChain(FCharacterizationStandard::RightPinkyMetacarpal, FName("palm_04_R"), FName("palm_04_R"));
+		// bone settings for IK
+		Rigify.AutoRetargetDefinition.BoneSettingsForIK.SetRotationStiffness(FName("spine"), FCharacterizationStandard::PelvisRotationStiffness);
+		Rigify.AutoRetargetDefinition.BoneSettingsForIK.SetRotationStiffness(FName("shoulder_L"), FCharacterizationStandard::ClavicleRotationStiffness);
+		Rigify.AutoRetargetDefinition.BoneSettingsForIK.SetRotationStiffness(FName("shoulder_R"), FCharacterizationStandard::ClavicleRotationStiffness);
+		Rigify.AutoRetargetDefinition.BoneSettingsForIK.SetRotationStiffness(FName("foot_L"), FCharacterizationStandard::FootRotationStiffness);
+		Rigify.AutoRetargetDefinition.BoneSettingsForIK.SetRotationStiffness(FName("foot_R"), FCharacterizationStandard::FootRotationStiffness);
+		Rigify.AutoRetargetDefinition.BoneSettingsForIK.SetPreferredAxis(FName("shin_L"), EPreferredAxis::PositiveZ);
+		Rigify.AutoRetargetDefinition.BoneSettingsForIK.SetPreferredAxis(FName("shin_R"), EPreferredAxis::PositiveZ);
+		Rigify.AutoRetargetDefinition.BoneSettingsForIK.SetPreferredAxis(FName("forearm_L"), EPreferredAxis::PositiveZ);
+		Rigify.AutoRetargetDefinition.BoneSettingsForIK.SetPreferredAxis(FName("forearm_R"), EPreferredAxis::PositiveZ);
+		Rigify.AutoRetargetDefinition.BoneSettingsForIK.SetExcluded(FName("spine_003"), true);
+		// exclude feet from auto-pose
+		Rigify.AutoRetargetDefinition.BonesToExcludeFromAutoPose.Add("foot_L");
+		Rigify.AutoRetargetDefinition.BonesToExcludeFromAutoPose.Add("foot_R");
+	}();
 }
 
 void FKnownTemplateHierarchies::GetClosestMatchingKnownHierarchy(
diff --git a/Engine/Plugins/Animation/IKRig/Source/IKRigEditor/Private/RigEditor/IKRigAutoFBIK.cpp b/Engine/Plugins/Animation/IKRig/Source/IKRigEditor/Private/RigEditor/IKRigAutoFBIK.cpp
--- a/Engine/Plugins/Animation/IKRig/Source/IKRigEditor/Private/RigEditor/IKRigAutoFBIK.cpp
+++ b/Engine/Plugins/Animation/IKRig/Source/IKRigEditor/Private/RigEditor/IKRigAutoFBIK.cpp
@@ -80,6 +80,14 @@ void FAutoFBIKCreator::CreateFBIKSetup(const UIKRigController& IKRigController,
 	Solver->Settings.RootBehavior = EPBIKRootBehavior::Free;
 	// removing pull chain alpha on all goals "calms" the motion down, especially when retargeting arms
 	Solver->Settings.GlobalPullChainAlpha = 0.0f;
+	// sub-iterations for chain-level sub-solves before each main iteration
+	Solver->Settings.SubIterations = 10;
+
+	// set chain depth on all goals
+	for (FIKRigFBIKGoalSettings& GoalSettings : Solver->AllGoalSettings)
+	{
+		GoalSettings.ChainDepth = 2;
+	}
 
 	// assign bone settings from template
 	const FAutoCharacterizer& Characterizer = IKRigController.GetAutoCharacterizer();
@@ -124,7 +132,14 @@ void FAutoFBIKCreator::CreateFBIKSetup(const UIKRigController& IKRigController,
 			FIKRigFBIKBoneSettings* Settings = GetOrAddBoneSettings();
 			Settings->RotationStiffness = FMath::Clamp(BoneSettings.RotationStiffness, 0.f, 1.0f);
 		}
-		
+
+		// apply position stiffness
+		if (BoneSettings.PositionStiffness > 0.f)
+		{
+			FIKRigFBIKBoneSettings* Settings = GetOrAddBoneSettings();
+			Settings->PositionStiffness = FMath::Clamp(BoneSettings.PositionStiffness, 0.f, 1.0f);
+		}
+
 		// apply preferred angles
 		if (BoneSettings.PreferredAxis != EPreferredAxis::None)
 		{
diff --git a/Engine/Plugins/Animation/IKRig/Source/IKRigEditor/Public/RigEditor/IKRigAutoCharacterizer.h b/Engine/Plugins/Animation/IKRig/Source/IKRigEditor/Public/RigEditor/IKRigAutoCharacterizer.h
--- a/Engine/Plugins/Animation/IKRig/Source/IKRigEditor/Public/RigEditor/IKRigAutoCharacterizer.h
+++ b/Engine/Plugins/Animation/IKRig/Source/IKRigEditor/Public/RigEditor/IKRigAutoCharacterizer.h
@@ -36,6 +36,9 @@ struct FCharacterizationStandard
 	// legs
 	static const FName LeftLeg;
 	static const FName RightLeg;
+	// feet
+	static const FName LeftFoot;
+	static const FName RightFoot;
 	// arms
 	static const FName LeftClavicle;
 	static const FName RightClavicle;
@@ -244,6 +247,7 @@ struct FBoneSettingsForIK
 	
 	FName BoneToApplyTo;
 	float RotationStiffness = 0.f;
+	float PositionStiffness = 0.f;
 	EPreferredAxis PreferredAxis = EPreferredAxis::None;
 	bool bIsHinge = false;
 	bool bExcluded = false;
@@ -259,6 +263,7 @@ struct FAllBoneSettingsForIK
 {
 	void SetPreferredAxis(const FName BoneName, const EPreferredAxis PreferredAxis, const bool bTreatAsHinge=true);
 	void SetRotationStiffness(const FName BoneName, const float RotationStiffness);
+	void SetPositionStiffness(const FName BoneName, const float PositionStiffness);
 	void SetExcluded(const FName BoneName, const bool bExclude);
 	const TArray<FBoneSettingsForIK>& GetBoneSettings() const { return AllBoneSettings; };
 
