[IK Retargeter] Expose bRetainAdditiveFlags to RunBatchRetarget Blueprint/Python API

Added bRetainAdditiveFlags to FIKRetargetBatchOperationInputs so that
Blueprint and Python scripts can control whether additive animation
attributes are preserved during batch retarget.

#Synced-CL 52631650

diff --git a/Engine/Plugins/Animation/IKRig/Source/IKRigEditor/Private/RetargetEditor/IKRetargetBatchOperation.cpp b/Engine/Plugins/Animation/IKRig/Source/IKRigEditor/Private/RetargetEditor/IKRetargetBatchOperation.cpp
--- a/Engine/Plugins/Animation/IKRig/Source/IKRigEditor/Private/RetargetEditor/IKRetargetBatchOperation.cpp
+++ b/Engine/Plugins/Animation/IKRig/Source/IKRigEditor/Private/RetargetEditor/IKRetargetBatchOperation.cpp
@@ -877,6 +877,7 @@ TArray<FAssetData> UIKRetargetBatchOperation::RunBatchRetarget(const FIKRetarget
 	Context.bUseSourcePath = Inputs.bUseSourcePath;
 	Context.bIncludeReferencedAssets = Inputs.bIncludeReferencedAssets;
 	Context.bOverwriteExistingFiles = Inputs.bOverwriteExistingFiles;
+	Context.bRetainAdditiveFlags = Inputs.bRetainAdditiveFlags;
 	Context.OverrideSetNames = Inputs.InOverrideSetNames;
 
 	// actually run the batch operation
diff --git a/Engine/Plugins/Animation/IKRig/Source/IKRigEditor/Public/RetargetEditor/IKRetargetBatchOperation.h b/Engine/Plugins/Animation/IKRig/Source/IKRigEditor/Public/RetargetEditor/IKRetargetBatchOperation.h
--- a/Engine/Plugins/Animation/IKRig/Source/IKRigEditor/Public/RetargetEditor/IKRetargetBatchOperation.h
+++ b/Engine/Plugins/Animation/IKRig/Source/IKRigEditor/Public/RetargetEditor/IKRetargetBatchOperation.h
@@ -169,6 +169,11 @@ struct FIKRetargetBatchOperationInputs
 	/** Whether to overwrite any existing output files with the same name, or create new files with a unique name */
 	UPROPERTY(BlueprintReadWrite, Category="IK Batch Retarget")
 	bool bOverwriteExistingFiles = false;
+
+	/** Keep the additive animation sequence attributes on the retargeted results.
+	 * NOTE: results may not be WYSIWYG with the editor, but the behavior should remain intact. */
+	UPROPERTY(BlueprintReadWrite, Category="IK Batch Retarget")
+	bool bRetainAdditiveFlags = true;
 };
 
 // Encapsulate ability to batch duplicate and retarget a set of animation assets
