[IK Retargeter] Override Set editor: Add Child, Search Bar, Keyboard Shortcuts, Duplicate, Reordering

Added DuplicateRetargetOverrideSet() to the controller API (Blueprint/Python callable). Creates a new override set with a unique name, copying all property overrides, parent assignment, and active-by-default state from the source set.

Added "Duplicate" entry to the Override Set tree view context menu between Rename and Unparent.

Add Child: Context menu action to create a new override set parented under the selected set.

Search Bar: Enabled the built-in details panel search bar for the Override Set editor.

Reordering: Added ability to reorder override sets in override set tree

Keyboard Shortcuts: Delete/Backspace deletes selected override set, F2 renames it.

Override Set Reordering: Fixed drag-drop to distinguish between reordering (AboveItem/BelowItem) and parenting (OntoItem). Added DisplayOrder field for persistent sort order. Reordering makes the dragged set a sibling of the target at the correct position.

#Synced-CL 52631650

diff --git a/Engine/Plugins/Animation/IKRig/Source/IKRig/Public/Retargeter/IKRetargetOverrides.h b/Engine/Plugins/Animation/IKRig/Source/IKRig/Public/Retargeter/IKRetargetOverrides.h
--- a/Engine/Plugins/Animation/IKRig/Source/IKRig/Public/Retargeter/IKRetargetOverrides.h
+++ b/Engine/Plugins/Animation/IKRig/Source/IKRig/Public/Retargeter/IKRetargetOverrides.h
@@ -237,6 +237,10 @@ struct FRetargetOverrideSet
 	UPROPERTY(EditAnywhere, Category=Override)
 	bool bActiveByDefault = false;
 
+	/** Display order for sorting among sibling override sets in the editor tree view */
+	UPROPERTY()
+	int32 DisplayOrder = 0;
+
 	/** get an op override by the op name */
 	UE_API FRetargetOpOverrides* FindOpOverrideByOpName(const FName InOpName);
 };
diff --git a/Engine/Plugins/Animation/IKRig/Source/IKRigEditor/Private/RetargetEditor/IKRetargeterController.cpp b/Engine/Plugins/Animation/IKRig/Source/IKRigEditor/Private/RetargetEditor/IKRetargeterController.cpp
--- a/Engine/Plugins/Animation/IKRig/Source/IKRigEditor/Private/RetargetEditor/IKRetargeterController.cpp
+++ b/Engine/Plugins/Animation/IKRig/Source/IKRigEditor/Private/RetargetEditor/IKRetargeterController.cpp
@@ -1512,6 +1512,90 @@ FName UIKRetargeterController::AddNewRetargetOverrideSet(FName InOverrideSetName
     return UniqueName;
 }
 
+FName UIKRetargeterController::DuplicateRetargetOverrideSet(FName InOverrideSetName)
+{
+	const FRetargetOverrideSet* SourceSet = FindOverrideSet(InOverrideSetName);
+	if (!SourceSet)
+	{
+		return NAME_None;
+	}
+
+	// copy source data before Add() which may rehash the TMap and invalidate SourceSet
+	const FName SourceParentName = SourceSet->ParentName;
+	const bool bSourceActiveByDefault = SourceSet->bActiveByDefault;
+	const TArray<FRetargetOpOverrides> SourceOpOverrides = SourceSet->OpOverrides;
+
+	constexpr bool bShouldTransact = true;
+	FScopedTransaction Transaction(LOCTEXT("DuplicateRetargetOverrideSet", "Duplicate Retarget Override Set"), bShouldTransact);
+	Asset->Modify();
+
+	const FName UniqueName = GetUniqueOverrideSetName(InOverrideSetName);
+	FRetargetOverrideSet& NewSet = Asset->OverrideSets.Add(UniqueName);
+	NewSet.Name = UniqueName;
+	NewSet.ParentName = SourceParentName;
+	NewSet.bActiveByDefault = bSourceActiveByDefault;
+	NewSet.OpOverrides = SourceOpOverrides;
+
+	PostEditOverrideSets();
+
+	return UniqueName;
+}
+
+void UIKRetargeterController::ReorderOverrideSet(FName InSetToMove, FName InTargetSet, bool bInsertAfter)
+{
+	if (InSetToMove == InTargetSet)
+	{
+		return;
+	}
+
+	FRetargetOverrideSet* SetToMove = FindOverrideSet(InSetToMove);
+	FRetargetOverrideSet* TargetSet = FindOverrideSet(InTargetSet);
+	if (!SetToMove || !TargetSet)
+	{
+		return;
+	}
+
+	constexpr bool bShouldTransact = true;
+	FScopedTransaction Transaction(LOCTEXT("ReorderOverrideSet", "Reorder Override Set"), bShouldTransact);
+	Asset->Modify();
+
+	// make the moved set a sibling of the target (same parent)
+	const FName TargetParent = TargetSet->ParentName;
+	SetToMove->ParentName = TargetParent;
+
+	// collect all siblings (same parent) sorted by current DisplayOrder, excluding the moved set
+	TArray<FName> SiblingOrder;
+	for (const TTuple<FName, FRetargetOverrideSet>& Entry : Asset->GetOverrideSets())
+	{
+		if (Entry.Value.ParentName == TargetParent && Entry.Key != InSetToMove)
+		{
+			SiblingOrder.Add(Entry.Key);
+		}
+	}
+	SiblingOrder.Sort([this](const FName& A, const FName& B)
+	{
+		const FRetargetOverrideSet* SetA = FindOverrideSet(A);
+		const FRetargetOverrideSet* SetB = FindOverrideSet(B);
+		return SetA && SetB && SetA->DisplayOrder < SetB->DisplayOrder;
+	});
+
+	// find the target's position in the sorted list and insert the moved set before or after
+	const int32 TargetIdx = SiblingOrder.IndexOfByKey(InTargetSet);
+	const int32 InsertIdx = bInsertAfter ? TargetIdx + 1 : TargetIdx;
+	SiblingOrder.Insert(InSetToMove, FMath::Clamp(InsertIdx, 0, SiblingOrder.Num()));
+
+	// reassign sequential display orders
+	for (int32 Index = 0; Index < SiblingOrder.Num(); ++Index)
+	{
+		if (FRetargetOverrideSet* Set = FindOverrideSet(SiblingOrder[Index]))
+		{
+			Set->DisplayOrder = Index;
+		}
+	}
+
+	PostEditOverrideSets();
+}
+
 bool UIKRetargeterController::RemoveRetargetOverrideSet(FName InOverrideSetName)
 {
     if (!Asset->OverrideSets.Contains(InOverrideSetName))
diff --git a/Engine/Plugins/Animation/IKRig/Source/IKRigEditor/Private/RetargetEditor/SIKRetargetOverrideManager.cpp b/Engine/Plugins/Animation/IKRig/Source/IKRigEditor/Private/RetargetEditor/SIKRetargetOverrideManager.cpp
--- a/Engine/Plugins/Animation/IKRig/Source/IKRigEditor/Private/RetargetEditor/SIKRetargetOverrideManager.cpp
+++ b/Engine/Plugins/Animation/IKRig/Source/IKRigEditor/Private/RetargetEditor/SIKRetargetOverrideManager.cpp
@@ -13,6 +13,7 @@
 #include "Net/ReplayPlaylistTracker.h"
 #include "RetargetEditor/IKRetargetEditorStyle.h"
 #include "RigEditor/IKRigStructViewer.h"
+#include "Framework/Commands/GenericCommands.h"
 #include "Widgets/Input/SSearchBox.h"
 #include "Widgets/Text/SInlineEditableTextBlock.h"
 
@@ -283,8 +284,10 @@ void SIKRetargetOverrideManager::Construct(
 	
 	TextFilter = MakeShareable(new FTextFilterExpressionEvaluator(ETextFilterExpressionEvaluatorMode::BasicString));
 
+	BindCommands();
+
 	FDetailsViewArgs DetailsViewArgs;
-	DetailsViewArgs.bAllowSearch = false;
+	DetailsViewArgs.bAllowSearch = true;
 	DetailsViewArgs.bLockable = false;
 	DetailsViewArgs.bUpdatesFromSelection = false;
 	DetailsViewArgs.bAllowMultipleTopLevelObjects = false;
@@ -505,6 +508,30 @@ void SIKRetargetOverrideManager::RefreshTreeView(bool IsInitialSetup)
 		}
 	}
 
+	// sort root elements and children by DisplayOrder
+	auto SortByDisplayOrder = [&AllOverrideSets](const TSharedPtr<FIKRetargetOverrideSetElement>& A, const TSharedPtr<FIKRetargetOverrideSetElement>& B)
+	{
+		if (!A || !B)
+		{
+			return false;
+		}
+		const FRetargetOverrideSet* SetA = AllOverrideSets.Find(A->Name);
+		const FRetargetOverrideSet* SetB = AllOverrideSets.Find(B->Name);
+		if (SetA && SetB)
+		{
+			return SetA->DisplayOrder < SetB->DisplayOrder;
+		}
+		return false;
+	};
+	RootElements.Sort(SortByDisplayOrder);
+	for (const TSharedPtr<FIKRetargetOverrideSetElement>& Element : AllElements)
+	{
+		Element->Children.Sort([&SortByDisplayOrder](const TWeakPtr<FIKRetargetOverrideSetElement>& A, const TWeakPtr<FIKRetargetOverrideSetElement>& B)
+		{
+			return SortByDisplayOrder(A.Pin(), B.Pin());
+		});
+	}
+
 	if (IsInitialSetup)
 	{
 		for (const TSharedPtr<FIKRetargetOverrideSetElement>& Root : RootElements)
@@ -648,8 +675,9 @@ TSharedPtr<SWidget> SIKRetargetOverrideManager::OnGenerateContextMenu()
 			TSharedPtr<ITableRow> RowWidget = TreeView->WidgetFromItem(Selection[0]);
 			if (RowWidget.IsValid())
 			{
-				TSharedPtr<SIKRetargetOverrideSetRow> SetRow = StaticCastSharedRef<SIKRetargetOverrideSetRow>(RowWidget->AsWidget());
-				if (SetRow.IsValid())
+				TSharedRef<SWidget> Widget = RowWidget->AsWidget();
+				SIKRetargetOverrideSetRow* SetRow = static_cast<SIKRetargetOverrideSetRow*>(&Widget.Get());
+				if (SetRow)
 				{
 					SetRow->OnRequestRename();
 				}
@@ -660,6 +688,30 @@ TSharedPtr<SWidget> SIKRetargetOverrideManager::OnGenerateContextMenu()
 			return TreeView.IsValid() && TreeView->GetNumItemsSelected() == 1;
 		})));
 
+		MenuBuilder.AddMenuEntry(
+		LOCTEXT("AddChildItem", "Add Child"),
+		LOCTEXT("AddChildLabel", "Create a new override set parented under the selected set."),
+		FSlateIcon(FAppStyle::GetAppStyleSetName(), "Icons.Plus"),
+		FUIAction(FExecuteAction::CreateLambda([this]()
+		{
+			TArray<TSharedPtr<FIKRetargetOverrideSetElement>> Selection = TreeView->GetSelectedItems();
+			if (Selection.IsEmpty())
+			{
+				return;
+			}
+
+			const FName ParentName = Selection[0]->Name;
+			const FName ChildName = EditorController.Pin()->AssetController->AddNewRetargetOverrideSet(FName("NewChildOverrideSet"));
+			if (ChildName != NAME_None)
+			{
+				EditorController.Pin()->AssetController->SetParentOverrideSet(ChildName, ParentName);
+			}
+		}),
+		FCanExecuteAction::CreateLambda([this]()
+		{
+			return TreeView.IsValid() && TreeView->GetNumItemsSelected() == 1;
+		})));
+
 		MenuBuilder.AddMenuEntry(
 		LOCTEXT("UnparentItem", "Unparent"),
 		LOCTEXT("UnparentLabel", "Unparent the selected override set."),
@@ -737,6 +789,67 @@ TSharedPtr<SWidget> SIKRetargetOverrideManager::OnGenerateContextMenu()
 	return MenuBuilder.MakeWidget();
 }
 
+void SIKRetargetOverrideManager::BindCommands()
+{
+	CommandList = MakeShared<FUICommandList>();
+
+	CommandList->MapAction(
+		FGenericCommands::Get().Delete,
+		FExecuteAction::CreateSP(this, &SIKRetargetOverrideManager::DeleteSelectedOverrideSet),
+		FCanExecuteAction::CreateSP(this, &SIKRetargetOverrideManager::HasSelectedOverrideSet));
+
+	CommandList->MapAction(
+		FGenericCommands::Get().Rename,
+		FExecuteAction::CreateSP(this, &SIKRetargetOverrideManager::RenameSelectedOverrideSet),
+		FCanExecuteAction::CreateSP(this, &SIKRetargetOverrideManager::HasSelectedOverrideSet));
+}
+
+FReply SIKRetargetOverrideManager::OnKeyDown(const FGeometry& InGeometry, const FKeyEvent& InKeyEvent)
+{
+	if (CommandList.IsValid() && CommandList->ProcessCommandBindings(InKeyEvent))
+	{
+		return FReply::Handled();
+	}
+	return SCompoundWidget::OnKeyDown(InGeometry, InKeyEvent);
+}
+
+void SIKRetargetOverrideManager::DeleteSelectedOverrideSet()
+{
+	TArray<TSharedPtr<FIKRetargetOverrideSetElement>> Selection = TreeView->GetSelectedItems();
+	if (Selection.IsEmpty())
+	{
+		return;
+	}
+
+	const FName SetToDelete = Selection[0]->Name;
+	EditorController.Pin()->AssetController->RemoveRetargetOverrideSet(SetToDelete);
+}
+
+void SIKRetargetOverrideManager::RenameSelectedOverrideSet()
+{
+	TArray<TSharedPtr<FIKRetargetOverrideSetElement>> Selection = TreeView->GetSelectedItems();
+	if (Selection.IsEmpty())
+	{
+		return;
+	}
+
+	TSharedPtr<ITableRow> RowWidget = TreeView->WidgetFromItem(Selection[0]);
+	if (RowWidget.IsValid())
+	{
+		TSharedRef<SWidget> Widget = RowWidget->AsWidget();
+		SIKRetargetOverrideSetRow* SetRow = static_cast<SIKRetargetOverrideSetRow*>(&Widget.Get());
+		if (SetRow)
+		{
+			SetRow->OnRequestRename();
+		}
+	}
+}
+
+bool SIKRetargetOverrideManager::HasSelectedOverrideSet() const
+{
+	return TreeView.IsValid() && TreeView->GetNumItemsSelected() > 0;
+}
+
 FReply SIKRetargetOverrideManager::OnDragDetected(
 	const FGeometry& MyGeometry,
 	const FPointerEvent& MouseEvent)
@@ -801,7 +914,17 @@ FReply SIKRetargetOverrideManager::OnAcceptDrop(
 		return FReply::Unhandled();
 	}
 
-	EditorController.Pin()->AssetController->SetParentOverrideSet(DraggedElement.Get()->Name, TargetElement.Get()->Name);
+	if (DropZone == EItemDropZone::OntoItem)
+	{
+		// drop onto item = parent the dragged set to the target
+		EditorController.Pin()->AssetController->SetParentOverrideSet(DraggedElement.Get()->Name, TargetElement.Get()->Name);
+	}
+	else
+	{
+		// drop above/below item = reorder as sibling (same parent as target)
+		const bool bInsertAfter = (DropZone == EItemDropZone::BelowItem);
+		EditorController.Pin()->AssetController->ReorderOverrideSet(DraggedElement.Get()->Name, TargetElement.Get()->Name, bInsertAfter);
+	}
 
 	return FReply::Handled();
 }
diff --git a/Engine/Plugins/Animation/IKRig/Source/IKRigEditor/Public/RetargetEditor/IKRetargeterController.h b/Engine/Plugins/Animation/IKRig/Source/IKRigEditor/Public/RetargetEditor/IKRetargeterController.h
--- a/Engine/Plugins/Animation/IKRig/Source/IKRigEditor/Public/RetargetEditor/IKRetargeterController.h
+++ b/Engine/Plugins/Animation/IKRig/Source/IKRigEditor/Public/RetargetEditor/IKRetargeterController.h
@@ -485,6 +485,19 @@ public:
 	UFUNCTION(BlueprintCallable, Category="IK Retarget")
 	UE_API bool RenameRetargetOverrideSet(FName InOverrideSetName, FName InNewName);
 
+	/** Duplicate an existing override set, creating a new set with all the same property overrides.
+	* @param InOverrideSetName    The name of the override set to duplicate.
+	* @return The name of the new duplicated override set (guaranteed unique). */
+	UFUNCTION(BlueprintCallable, Category="IK Retarget")
+	UE_API FName DuplicateRetargetOverrideSet(FName InOverrideSetName);
+
+	/** Reorder an override set relative to another set.
+	* The dragged set becomes a sibling of the target (same parent) and is placed before or after it.
+	* @param InSetToMove      The override set to move.
+	* @param InTargetSet      The override set to position relative to.
+	* @param bInsertAfter     If true, insert after target. If false, insert before. */
+	UE_API void ReorderOverrideSet(FName InSetToMove, FName InTargetSet, bool bInsertAfter);
+
 	/** Assign a parent to the given override set.
 	* NOTE: Override sets can be applied hierarchically to make for easy composition of property overrides.
 	* @param InOverrideSetName    The name of the override set to re-parent.
diff --git a/Engine/Plugins/Animation/IKRig/Source/IKRigEditor/Public/RetargetEditor/SIKRetargetOverrideManager.h b/Engine/Plugins/Animation/IKRig/Source/IKRigEditor/Public/RetargetEditor/SIKRetargetOverrideManager.h
--- a/Engine/Plugins/Animation/IKRig/Source/IKRigEditor/Public/RetargetEditor/SIKRetargetOverrideManager.h
+++ b/Engine/Plugins/Animation/IKRig/Source/IKRigEditor/Public/RetargetEditor/SIKRetargetOverrideManager.h
@@ -7,6 +7,7 @@
 #include "SBaseHierarchyTreeView.h"
 #include "SIKRetargetPoseEditor.h"
 #include "DragAndDrop/DecoratedDragDropOp.h"
+#include "Framework/Commands/UICommandList.h"
 #include "Widgets/SCompoundWidget.h"
 
 class FTextFilterExpressionEvaluator;
@@ -149,6 +150,17 @@ private:
        TSharedPtr<FIKRetargetOverrideSetElement> TargetElement);
     /** END drag and drop operations */
 
+    /** keyboard shortcuts */
+    void BindCommands();
+    void DeleteSelectedOverrideSet();
+    void RenameSelectedOverrideSet();
+    bool HasSelectedOverrideSet() const;
+    TSharedPtr<FUICommandList> CommandList;
+
+    //~ Begin SWidget
+    virtual FReply OnKeyDown(const FGeometry& InGeometry, const FKeyEvent& InKeyEvent) override;
+    //~ End SWidget
+
     friend SIKRetargetOverrideSetRow;
     friend FIKRetargetEditorController;
 };
\ No newline at end of file
